#UE4 Fix it so Linux builds only link in SteamController if they're building with Engine, as SteamController depends on engine. Fixes linux crash report client

diff --git a/Engine/Source/Runtime/Core/Core.Build.cs b/Engine/Source/Runtime/Core/Core.Build.cs
--- a/Engine/Source/Runtime/Core/Core.Build.cs
+++ b/Engine/Source/Runtime/Core/Core.Build.cs
@@ -111,9 +111,13 @@ public class Core : ModuleRules
 				"zlib",
 				"jemalloc",
 				"elftoolchain",
-				"SDL2",
-                		"SteamController"
+				"SDL2"
                 );
+
+            if (UEBuildConfiguration.bCompileAgainstEngine == true)
+            {
+                AddThirdPartyPrivateStaticDependencies(Target, "SteamController");
+            }
         }
 		else if (Target.Platform == UnrealTargetPlatform.HTML5 && Target.Architecture == "-win32")
 		{
diff --git a/Engine/Source/Runtime/Core/Private/Linux/LinuxApplication.cpp b/Engine/Source/Runtime/Core/Private/Linux/LinuxApplication.cpp
--- a/Engine/Source/Runtime/Core/Private/Linux/LinuxApplication.cpp
+++ b/Engine/Source/Runtime/Core/Private/Linux/LinuxApplication.cpp
@@ -4,7 +4,9 @@
 #include "LinuxWindow.h"
 #include "LinuxCursor.h"
 #include "GenericApplicationMessageHandler.h"
+#if WITH_ENGINE
 #include "SteamControllerInterface.h"
+#endif
 
 //	todo:
 //	need to change to a linux one
@@ -74,7 +76,10 @@ FLinuxApplication* FLinuxApplication::CreateLinuxApplication()
 }
 
 
-FLinuxApplication::FLinuxApplication() : GenericApplication( MakeShareable( new FLinuxCursor() ) ), SteamInput( SteamControllerInterface::Create(MessageHandler) )
+FLinuxApplication::FLinuxApplication() : GenericApplication( MakeShareable( new FLinuxCursor() ) )
+#if WITH_ENGINE
+	, SteamInput( SteamControllerInterface::Create(MessageHandler) )
+#endif
 {
 	bUsingHighPrecisionMouseInput = false;
 	bAllowedToDeferMessageProcessing = false;
@@ -119,7 +124,9 @@ void FLinuxApplication::InitializeWindow(	const TSharedRef< FGenericWindow >& In
 void FLinuxApplication::SetMessageHandler( const TSharedRef< FGenericApplicationMessageHandler >& InMessageHandler )
 {
 	GenericApplication::SetMessageHandler(InMessageHandler);
+#if WITH_ENGINE
 	SteamInput->SetMessageHandler(InMessageHandler);
+#endif
 }
 
 static TSharedPtr< FLinuxWindow > FindWindowBySDLWindow( const TArray< TSharedRef< FLinuxWindow > >& WindowsToSearch, SDL_HWindow const WindowHandle )
@@ -679,8 +686,10 @@ void FLinuxApplication::ProcessDeferredEvents( const float TimeDelta )
 
 void FLinuxApplication::PollGameDeviceState( const float TimeDelta )
 {
+#if WITH_ENGINE
 	// Poll game device states and send new events
 	SteamInput->SendControllerEvents();
+#endif
 }
 
 TCHAR FLinuxApplication::ConvertChar( SDL_Keysym Keysym )
diff --git a/Engine/Source/Runtime/Core/Public/Linux/LinuxApplication.h b/Engine/Source/Runtime/Core/Public/Linux/LinuxApplication.h
--- a/Engine/Source/Runtime/Core/Public/Linux/LinuxApplication.h
+++ b/Engine/Source/Runtime/Core/Public/Linux/LinuxApplication.h
@@ -95,7 +95,9 @@ private:
 
 	SDLControllerState *ControllerStates;
 
+#if WITH_ENGINE
 	TSharedPtr< class SteamControllerInterface > SteamInput;
+#endif
 };
 
 
