#ue4 Adding also SetVector4 to FConfigCacheIni class.

diff --git a/Engine/Source/Runtime/Core/Private/Misc/ConfigCacheIni.cpp b/Engine/Source/Runtime/Core/Private/Misc/ConfigCacheIni.cpp
--- a/Engine/Source/Runtime/Core/Private/Misc/ConfigCacheIni.cpp
+++ b/Engine/Source/Runtime/Core/Private/Misc/ConfigCacheIni.cpp
@@ -2066,6 +2066,17 @@ void FConfigCacheIni::SetVector
 	SetString( Section, Key, *Value.ToString(), Filename );
 }
 
+void FConfigCacheIni::SetVector4
+(
+ const TCHAR*		Section,
+ const TCHAR*		Key,
+ const FVector4		 Value,
+ const FString&	Filename
+)
+{
+	SetString(Section, Key, *Value.ToString(), Filename);
+}
+
 void FConfigCacheIni::SetRotator
 (
  const TCHAR*		Section,
diff --git a/Engine/Source/Runtime/Core/Public/Misc/ConfigCacheIni.h b/Engine/Source/Runtime/Core/Public/Misc/ConfigCacheIni.h
--- a/Engine/Source/Runtime/Core/Public/Misc/ConfigCacheIni.h
+++ b/Engine/Source/Runtime/Core/Public/Misc/ConfigCacheIni.h
@@ -435,6 +435,13 @@ public:
 		FVector				Value,
 		const FString&	Filename
 	);
+	void SetVector4
+	(
+		const TCHAR*		Section,
+		const TCHAR*		Key,
+		FVector4			Value,
+		const FString&	Filename
+	);
 	void SetRotator
 	(
 		const TCHAR*		Section,
