#UE4
- Platform memory
  - Added LogMemory to log all memory related things
  - Separated PlatformMemory stats from Allocator stats, added MemoryAllocator group for allocator stats, added MemoryPlatform for platform memory specific stats, stats for these groups are updated once per frame in FEngineLoop::Tick, stats are duplicated, so they can be accessed if STATS is not enabled
 - Removed outdated code, references to unsupported platforms like PS3 or Xbox360
 - Updated MallocProfiler to support these changed, increased malloc profiler file version to 4
 - From now FGenericPlatformMemory.GetStats contains FPlatformMemoryConstants so there is no need to call the second method
 - Improved memory reporting for Windows platform, other platforms need to be updated separately
 - Misc tweak and fixes

diff --git a/Engine/Source/Editor/MainFrame/Private/MainFrameModule.cpp b/Engine/Source/Editor/MainFrame/Private/MainFrameModule.cpp
--- a/Engine/Source/Editor/MainFrame/Private/MainFrameModule.cpp
+++ b/Engine/Source/Editor/MainFrame/Private/MainFrameModule.cpp
@@ -225,7 +225,7 @@ TSharedRef<SWidget> FMainFrameModule::MakeDeveloperTools() const
 			{
 				// Query OS for process memory used
 				FPlatformMemoryStats MemoryStats = FPlatformMemory::GetStats();
-				StaticLastTotalAllocated = MemoryStats.WorkingSetSize;
+				StaticLastTotalAllocated = MemoryStats.UsedPhysical;
 
 				// Wait 60 queries until we refresh memory again
 				QueriesUntilUpdate = 60;
diff --git a/Engine/Source/Runtime/Core/Private/GenericPlatform/GenericPlatformMemory.cpp b/Engine/Source/Runtime/Core/Private/GenericPlatform/GenericPlatformMemory.cpp
--- a/Engine/Source/Runtime/Core/Private/GenericPlatform/GenericPlatformMemory.cpp
+++ b/Engine/Source/Runtime/Core/Private/GenericPlatform/GenericPlatformMemory.cpp
@@ -7,12 +7,35 @@
 #include "CorePrivate.h"
 #include "MallocAnsi.h"
 #include "GenericPlatformMemoryPoolStats.h"
+#include "MemoryMisc.h"
 
 
 DEFINE_STAT(MCR_Physical);
 DEFINE_STAT(MCR_GPU);
 DEFINE_STAT(MCR_TexturePool);
 
+DECLARE_MEMORY_STAT(TEXT("Total Physical"),		STAT_TotalPhysical,STATGROUP_MemoryPlatform);
+DECLARE_MEMORY_STAT(TEXT("Total Virtual"),		STAT_TotalVirtual,STATGROUP_MemoryPlatform);
+DECLARE_MEMORY_STAT(TEXT("PageSize"),			STAT_PageSize,STATGROUP_MemoryPlatform);
+DECLARE_MEMORY_STAT(TEXT("Total Physical GB"),	STAT_TotalPhysicalGB,STATGROUP_MemoryPlatform);
+
+DECLARE_MEMORY_STAT(TEXT("AvailablePhysical"),	STAT_AvailablePhysical,STATGROUP_MemoryPlatform);
+DECLARE_MEMORY_STAT(TEXT("AvailableVirtual"),	STAT_AvailableVirtual,STATGROUP_MemoryPlatform);
+DECLARE_MEMORY_STAT(TEXT("UsedPhysical"),		STAT_UsedPhysical,STATGROUP_MemoryPlatform);
+DECLARE_MEMORY_STAT(TEXT("PeakUsedPhysical"),	STAT_PeakUsedPhysical,STATGROUP_MemoryPlatform);
+DECLARE_MEMORY_STAT(TEXT("UsedVirtual"),		STAT_UsedVirtual,STATGROUP_MemoryPlatform);
+DECLARE_MEMORY_STAT(TEXT("PeakUsedVirtual"),	STAT_PeakUsedVirtual,STATGROUP_MemoryPlatform);
+
+FGenericPlatformMemoryStats::FGenericPlatformMemoryStats()
+	: FGenericPlatformMemoryConstants( FPlatformMemory::GetConstants() )
+	, AvailablePhysical( 0 )
+	, AvailableVirtual( 0 )
+	, UsedPhysical( 0 )
+	, PeakUsedPhysical( 0 )
+	, UsedVirtual( 0 )
+	, PeakUsedVirtual( 0 )
+{}
+
 void FGenericPlatformMemory::SetupMemoryPools()
 {
 	SET_MEMORY_STAT(MCR_Physical, 0); // "unlimited" physical memory, we still need to make this call to set the short name, etc
@@ -23,12 +46,12 @@ void FGenericPlatformMemory::SetupMemoryPools()
 void FGenericPlatformMemory::Init()
 {
 	SetupMemoryPools();
-	UE_LOG(LogHAL, Warning, TEXT("FGenericPlatformMemory::Init not implemented on this platform"));
+	UE_LOG(LogMemory, Warning, TEXT("FGenericPlatformMemory::Init not implemented on this platform"));
 }
 
 void FGenericPlatformMemory::OnOutOfMemory(uint64 Size, uint32 Alignment)
 {
-	UE_LOG(LogHAL, Fatal, TEXT("Ran out of memory allocating %llu bytes with alignment %u"), Size, Alignment);
+	UE_LOG(LogMemory, Fatal, TEXT("Ran out of memory allocating %llu bytes with alignment %u"), Size, Alignment);
 }
 
 FMalloc* FGenericPlatformMemory::BaseAllocator()
@@ -38,14 +61,32 @@ FMalloc* FGenericPlatformMemory::BaseAllocator()
 
 FPlatformMemoryStats FGenericPlatformMemory::GetStats()
 {
-	UE_LOG(LogHAL, Warning, TEXT("FGenericPlatformMemory::GetStats not implemented on this platform"));
+	UE_LOG(LogMemory, Warning, TEXT("FGenericPlatformMemory::GetStats not implemented on this platform"));
 	return FPlatformMemoryStats();
 }
 
+void FGenericPlatformMemory::GetStatsForMallocProfiler( FGenericMemoryStats& out_Stats )
+{
+#if	STATS
+	FPlatformMemoryStats Stats = FPlatformMemory::GetStats();
+
+	// Base common stats for all platforms.
+	out_Stats.Add( GET_STATFNAME( STAT_TotalPhysical ), Stats.TotalPhysical );
+	out_Stats.Add( GET_STATFNAME( STAT_TotalVirtual ), Stats.TotalVirtual );
+	out_Stats.Add( GET_STATFNAME( STAT_PageSize ), Stats.PageSize );
+	out_Stats.Add( GET_STATFNAME( STAT_TotalPhysicalGB ), (SIZE_T)Stats.TotalPhysicalGB );
+	out_Stats.Add( GET_STATFNAME( STAT_AvailablePhysical ), Stats.AvailablePhysical );
+	out_Stats.Add( GET_STATFNAME( STAT_AvailableVirtual ), Stats.AvailableVirtual );
+	out_Stats.Add( GET_STATFNAME( STAT_UsedPhysical ), Stats.UsedPhysical );
+	out_Stats.Add( GET_STATFNAME( STAT_PeakUsedPhysical ), Stats.PeakUsedPhysical );
+	out_Stats.Add( GET_STATFNAME( STAT_UsedVirtual ), Stats.UsedVirtual );
+	out_Stats.Add( GET_STATFNAME( STAT_PeakUsedVirtual ), Stats.PeakUsedVirtual );
+#endif // STATS
+}
 
 const FPlatformMemoryConstants& FGenericPlatformMemory::GetConstants()
 {
-	UE_LOG(LogHAL, Warning, TEXT("FGenericPlatformMemory::GetConstants not implemented on this platform"));
+	UE_LOG(LogMemory, Warning, TEXT("FGenericPlatformMemory::GetConstants not implemented on this platform"));
 	static FPlatformMemoryConstants MemoryConstants;
 	return MemoryConstants;
 }
@@ -55,34 +96,51 @@ uint32 FGenericPlatformMemory::GetPhysicalGBRam()
 	return FPlatformMemory::GetConstants().TotalPhysicalGB;
 }
 
+void FGenericPlatformMemory::UpdateStats()
+{
+	FPlatformMemoryStats MemoryStats = FPlatformMemory::GetStats();
+
+	SET_MEMORY_STAT(STAT_TotalPhysical,MemoryStats.TotalPhysical);
+	SET_MEMORY_STAT(STAT_TotalVirtual,MemoryStats.TotalVirtual);
+	SET_MEMORY_STAT(STAT_PageSize,MemoryStats.PageSize);
+	SET_MEMORY_STAT(STAT_TotalPhysicalGB,MemoryStats.TotalPhysicalGB);
+
+	SET_MEMORY_STAT(STAT_AvailablePhysical,MemoryStats.AvailablePhysical);
+	SET_MEMORY_STAT(STAT_AvailableVirtual,MemoryStats.AvailableVirtual);
+	SET_MEMORY_STAT(STAT_UsedPhysical,MemoryStats.UsedPhysical);
+	SET_MEMORY_STAT(STAT_PeakUsedPhysical,MemoryStats.PeakUsedPhysical);
+	SET_MEMORY_STAT(STAT_UsedVirtual,MemoryStats.UsedVirtual);
+	SET_MEMORY_STAT(STAT_PeakUsedVirtual,MemoryStats.PeakUsedVirtual);
+}
 
 void* FGenericPlatformMemory::BinnedAllocFromOS( SIZE_T Size )
 {
-	UE_LOG(LogHAL, Error, TEXT("FGenericPlatformMemory::BinnedAllocFromOS not implemented on this platform"));
+	UE_LOG(LogMemory, Error, TEXT("FGenericPlatformMemory::BinnedAllocFromOS not implemented on this platform"));
 	return nullptr;
 }
 
 void FGenericPlatformMemory::BinnedFreeToOS( void* Ptr )
 {
-	UE_LOG(LogHAL, Error, TEXT("FGenericPlatformMemory::BinnedFreeToOS not implemented on this platform"));
+	UE_LOG(LogMemory, Error, TEXT("FGenericPlatformMemory::BinnedFreeToOS not implemented on this platform"));
 }
 
 void FGenericPlatformMemory::DumpStats( class FOutputDevice& Ar )
 {
 	const float InvMB = 1.0f / 1024.0f / 1024.0f;
-	const FPlatformMemoryConstants& MemoryConstants = FPlatformMemory::GetConstants();
 	FPlatformMemoryStats MemoryStats = FPlatformMemory::GetStats();
-	Ar.Logf(TEXT("Platform Memory Status"));
-	Ar.Logf(TEXT("Working Set: %.2f MB used, %.2f MB peak"), MemoryStats.WorkingSetSize*InvMB, MemoryStats.PeakWorkingSetSize*InvMB );
-	Ar.Logf(TEXT("Physical Memory: %.2f MB allocated, %.2f MB total"), (MemoryConstants.TotalPhysical - MemoryStats.AvailablePhysical)*InvMB, MemoryConstants.TotalPhysical*InvMB );
-	Ar.Logf(TEXT("Virtual Memory: %.2f MB allocated, %.2f MB total"), (MemoryConstants.TotalVirtual - MemoryStats.AvailableVirtual)*InvMB, MemoryConstants.TotalVirtual*InvMB );
+	Ar.CategorizedLogf( LogMemory.GetCategoryName(), ELogVerbosity::Log, TEXT("Platform Memory Stats for %s"), ANSI_TO_TCHAR(FPlatformProperties::PlatformName()) );
+	Ar.CategorizedLogf( LogMemory.GetCategoryName(), ELogVerbosity::Log, TEXT("Process Physical Memory: %.2f MB used, %.2f MB peak"), MemoryStats.UsedPhysical*InvMB, MemoryStats.PeakUsedPhysical*InvMB );
+	Ar.CategorizedLogf( LogMemory.GetCategoryName(), ELogVerbosity::Log, TEXT("Process Virtual Memory: %.2f MB used, %.2f MB peak"), MemoryStats.UsedVirtual*InvMB, MemoryStats.PeakUsedVirtual*InvMB );
+
+	Ar.CategorizedLogf( LogMemory.GetCategoryName(), ELogVerbosity::Log, TEXT("Physical Memory: %.2f MB used, %.2f MB total"), (MemoryStats.TotalPhysical - MemoryStats.AvailablePhysical)*InvMB, MemoryStats.TotalPhysical*InvMB );
+	Ar.CategorizedLogf( LogMemory.GetCategoryName(), ELogVerbosity::Log, TEXT("Virtual Memory: %.2f MB used, %.2f MB total"), (MemoryStats.TotalVirtual - MemoryStats.AvailableVirtual)*InvMB, MemoryStats.TotalVirtual*InvMB );
 
 }
 
 void FGenericPlatformMemory::DumpPlatformAndAllocatorStats( class FOutputDevice& Ar )
 {
 	FPlatformMemory::DumpStats( Ar );
-	GMalloc->DumpAllocations( Ar );
+	GMalloc->DumpAllocatorStats( Ar );
 }
 
 void FGenericPlatformMemory::Memswap( void* Ptr1, void* Ptr2, SIZE_T Size )
diff --git a/Engine/Source/Runtime/Core/Private/HAL/MallocBinned.cpp b/Engine/Source/Runtime/Core/Private/HAL/MallocBinned.cpp
new file mode 100644
--- /dev/null
+++ b/Engine/Source/Runtime/Core/Private/HAL/MallocBinned.cpp
@@ -0,0 +1,47 @@
+// Copyright 1998-2014 Epic Games, Inc. All Rights Reserved.
+
+/*=============================================================================
+	MallocBinned.cpp: Binned memory allocator
+=============================================================================*/
+
+#include "CorePrivate.h"
+
+#if	!FORCE_ANSI_ALLOCATOR
+
+#include "MallocBinned.h"
+#include "MemoryMisc.h"
+
+/** Malloc binned allocator specific stats. */
+DEFINE_STAT(STAT_Binned_OsCurrent);
+DEFINE_STAT(STAT_Binned_OsPeak);
+DEFINE_STAT(STAT_Binned_WasteCurrent);
+DEFINE_STAT(STAT_Binned_WastePeak);
+DEFINE_STAT(STAT_Binned_UsedCurrent);
+DEFINE_STAT(STAT_Binned_UsedPeak);
+DEFINE_STAT(STAT_Binned_CurrentAllocs);
+DEFINE_STAT(STAT_Binned_TotalAllocs);
+DEFINE_STAT(STAT_Binned_SlackCurrent);
+
+void FMallocBinned::GetAllocatorStats( FGenericMemoryStats& out_Stats )
+{
+#ifdef USE_INTERNAL_LOCKS
+	FScopeLock ScopedLock(&AccessGuard);
+#endif
+
+	FMalloc::GetAllocatorStats( out_Stats );
+	UpdateSlackStat();
+
+#if	STATS
+	// Malloc binned stats.
+	out_Stats.Add( GET_STATFNAME( STAT_Binned_OsCurrent ), OsCurrent );
+	out_Stats.Add( GET_STATFNAME( STAT_Binned_OsPeak ), OsPeak );
+	out_Stats.Add( GET_STATFNAME( STAT_Binned_WasteCurrent ), WasteCurrent );
+	out_Stats.Add( GET_STATFNAME( STAT_Binned_WastePeak ), WastePeak );
+	out_Stats.Add( GET_STATFNAME( STAT_Binned_UsedCurrent ), UsedCurrent );
+	out_Stats.Add( GET_STATFNAME( STAT_Binned_UsedPeak ), UsedPeak );
+	out_Stats.Add( GET_STATFNAME( STAT_Binned_CurrentAllocs ), CurrentAllocs );
+	out_Stats.Add( GET_STATFNAME( STAT_Binned_SlackCurrent ), SlackCurrent );
+#endif // STATS
+}
+
+#endif // !FORCE_ANSI_ALLOCATOR
\ No newline at end of file
diff --git a/Engine/Source/Runtime/Core/Private/HAL/MallocJemalloc.cpp b/Engine/Source/Runtime/Core/Private/HAL/MallocJemalloc.cpp
--- a/Engine/Source/Runtime/Core/Private/HAL/MallocJemalloc.cpp
+++ b/Engine/Source/Runtime/Core/Private/HAL/MallocJemalloc.cpp
@@ -133,26 +133,12 @@ namespace
 	}
 }
 
-void FMallocJemalloc::DumpAllocations( FOutputDevice& Ar ) 
+void FMallocJemalloc::DumpAllocatorStats( FOutputDevice& Ar ) 
 {
-	STAT(Ar.Logf( TEXT("Memory Allocation Status") ));
 	MEM_TIME(Ar.Logf( TEXT("Seconds     % 5.3f"), MemTime ));
-
 	je_malloc_stats_print(JemallocStatsPrintCallback, &Ar, NULL);
 }
 
-void FMallocJemalloc::GetAllocationInfo( FMemoryAllocationStats_DEPRECATED& MemStats )
-{
-	size_t TotalAllocated, TotalAllocatedLen = sizeof(TotalAllocated), 
-		TotalMapped, TotalMappedLen = sizeof(TotalMapped);
-	je_mallctl("stats.allocated", &TotalAllocated, &TotalAllocatedLen, NULL, 0);
-	je_mallctl("stats.mapped", &TotalMapped, &TotalMappedLen, NULL, 0);
-
-	MemStats.TotalUsed = TotalMapped;
-	MemStats.TotalAllocated = TotalAllocated;
-	MemStats.CPUUsed = TotalMapped;
-}
-
 bool FMallocJemalloc::GetAllocationSize(void *Original, SIZE_T &SizeOut)
 {
 	SizeOut = je_malloc_usable_size(Original);
diff --git a/Engine/Source/Runtime/Core/Private/HAL/MallocTBB.cpp b/Engine/Source/Runtime/Core/Private/HAL/MallocTBB.cpp
--- a/Engine/Source/Runtime/Core/Private/HAL/MallocTBB.cpp
+++ b/Engine/Source/Runtime/Core/Private/HAL/MallocTBB.cpp
@@ -26,23 +26,11 @@
 
 #define MEM_TIME(st)
 
-void FMallocTBB::GetAllocationInfo( FMemoryAllocationStats_DEPRECATED& MemStats )
-{
-	FPlatformMemoryStats MemoryStats = FPlatformMemory::GetStats();
-
-	MemStats.TotalUsed = MemoryStats.WorkingSetSize;
-	MemStats.TotalAllocated = MemoryStats.WorkingSetSize;
-	MemStats.CPUUsed = MemoryStats.PagefileUsage;
-}
-
-void FMallocTBB::DumpAllocations( FOutputDevice& Ar )
-{
-	STAT(Ar.Logf( TEXT("Memory Allocation Status") ));
-	MEM_TIME(Ar.Logf( TEXT("Seconds     % 5.3f"), MemTime ));
-}
 
 void* FMallocTBB::Malloc( SIZE_T Size, uint32 Alignment )
 {
+	IncrementTotalMallocCalls();
+
 	MEM_TIME(MemTime -= FPlatformTime::Seconds());
 
 	void* Free = NULL;
@@ -72,6 +60,8 @@ void* FMallocTBB::Malloc( SIZE_T Size, uint32 Alignment )
 
 void* FMallocTBB::Realloc( void* Ptr, SIZE_T NewSize, uint32 Alignment )
 {
+	IncrementTotalReallocCalls();
+
 	MEM_TIME(MemTime -= FPlatformTime::Seconds())
 #if UE_BUILD_DEBUG || UE_BUILD_DEVELOPMENT
 	SIZE_T OldSize = 0;
@@ -118,6 +108,7 @@ void FMallocTBB::Free( void* Ptr )
 #if UE_BUILD_DEBUG || UE_BUILD_DEVELOPMENT
 	FMemory::Memset(Ptr, DEBUG_FILL_FREED, scalable_msize(Ptr)); 
 #endif
+	IncrementTotalFreeCalls();
 	scalable_free(Ptr);
 
 	MEM_TIME(MemTime += FPlatformTime::Seconds())
diff --git a/Engine/Source/Runtime/Core/Private/HAL/MemoryBase.cpp b/Engine/Source/Runtime/Core/Private/HAL/MemoryBase.cpp
new file mode 100644
--- /dev/null
+++ b/Engine/Source/Runtime/Core/Private/HAL/MemoryBase.cpp
@@ -0,0 +1,75 @@
+// Copyright 1998-2014 Epic Games, Inc. All Rights Reserved.
+
+#include "CorePrivate.h"
+#include "MemoryMisc.h"
+
+/** Memory allocator base stats. */
+DECLARE_DWORD_COUNTER_STAT(TEXT("Malloc calls"),			STAT_MallocCalls,STATGROUP_MemoryAllocator);
+DECLARE_DWORD_COUNTER_STAT(TEXT("Free calls"),				STAT_FreeCalls,STATGROUP_MemoryAllocator);
+DECLARE_DWORD_COUNTER_STAT(TEXT("Realloc calls"),			STAT_ReallocCalls,STATGROUP_MemoryAllocator);
+DECLARE_DWORD_COUNTER_STAT(TEXT("Total Allocator calls"),	STAT_TotalAllocatorCalls,STATGROUP_MemoryAllocator);
+
+uint64 FMalloc::TotalMallocCalls = 0;
+uint64 FMalloc::TotalFreeCalls = 0;
+uint64 FMalloc::TotalReallocCalls = 0;
+
+struct FCurrentFrameCalls
+{
+	uint64 LastMallocCalls;
+	uint64 LastReallocCalls;
+	uint64 LastFreeCalls;
+	
+	uint32 MallocCalls;
+	uint32 ReallocCalls;
+	uint32 FreeCalls;
+	uint32 AllocatorCalls;
+
+	FCurrentFrameCalls()
+		: LastMallocCalls(0)
+		, LastReallocCalls(0)
+		, LastFreeCalls(0)
+		, MallocCalls(0)
+		, ReallocCalls(0)
+		, FreeCalls(0)
+		, AllocatorCalls(0)
+	{}
+
+	void Update()
+	{
+		MallocCalls      = FMalloc::TotalMallocCalls - LastMallocCalls;
+		ReallocCalls     = FMalloc::TotalReallocCalls - LastReallocCalls;
+		FreeCalls        = FMalloc::TotalFreeCalls - LastFreeCalls;
+		AllocatorCalls   = MallocCalls + ReallocCalls + FreeCalls;
+
+		LastMallocCalls  = FMalloc::TotalMallocCalls;
+		LastReallocCalls = FMalloc::TotalReallocCalls;
+		LastFreeCalls    = FMalloc::TotalFreeCalls;
+	}
+};
+
+static FCurrentFrameCalls& GetCurrentFrameCalls()
+{
+	static FCurrentFrameCalls CurrentFrameCalls;
+	return CurrentFrameCalls;
+}
+
+void FMalloc::UpdateStats()
+{
+	GetCurrentFrameCalls().Update();
+
+	SET_DWORD_STAT( STAT_MallocCalls, GetCurrentFrameCalls().MallocCalls );
+	SET_DWORD_STAT( STAT_ReallocCalls, GetCurrentFrameCalls().ReallocCalls );
+	SET_DWORD_STAT( STAT_FreeCalls, GetCurrentFrameCalls().FreeCalls );
+	SET_DWORD_STAT( STAT_TotalAllocatorCalls, GetCurrentFrameCalls().AllocatorCalls );
+}
+
+
+void FMalloc::GetAllocatorStats( FGenericMemoryStats& out_Stats )
+{
+#if	STATS
+	out_Stats.Add( GET_STATFNAME( STAT_MallocCalls ), GetCurrentFrameCalls().MallocCalls );
+	out_Stats.Add( GET_STATFNAME( STAT_ReallocCalls ), GetCurrentFrameCalls().ReallocCalls );
+	out_Stats.Add( GET_STATFNAME( STAT_FreeCalls ), GetCurrentFrameCalls().FreeCalls );
+	out_Stats.Add( GET_STATFNAME( STAT_TotalAllocatorCalls ), GetCurrentFrameCalls().AllocatorCalls );
+#endif // STATS
+}
\ No newline at end of file
diff --git a/Engine/Source/Runtime/Core/Private/IOS/IOSPlatformMemory.cpp b/Engine/Source/Runtime/Core/Private/IOS/IOSPlatformMemory.cpp
--- a/Engine/Source/Runtime/Core/Private/IOS/IOSPlatformMemory.cpp
+++ b/Engine/Source/Runtime/Core/Private/IOS/IOSPlatformMemory.cpp
@@ -39,8 +39,8 @@ FPlatformMemoryStats FIOSPlatformMemory::GetStats()
 
 	MemoryStats.AvailablePhysical = FreeMem;
 	MemoryStats.AvailableVirtual = 0;
-	MemoryStats.WorkingSetSize = Info.resident_size;
-	MemoryStats.PagefileUsage = Info.virtual_size;
+	MemoryStats.UsedPhysical = Info.resident_size;
+	MemoryStats.UsedVirtual = Info.virtual_size;
 
 	return MemoryStats;
 }
diff --git a/Engine/Source/Runtime/Core/Private/Mac/MacPlatformMemory.cpp b/Engine/Source/Runtime/Core/Private/Mac/MacPlatformMemory.cpp
--- a/Engine/Source/Runtime/Core/Private/Mac/MacPlatformMemory.cpp
+++ b/Engine/Source/Runtime/Core/Private/Mac/MacPlatformMemory.cpp
@@ -46,14 +46,13 @@ FPlatformMemoryStats FMacPlatformMemory::GetStats()
 	MemoryStats.AvailablePhysical = FreeMem;
 	MemoryStats.AvailableVirtual = 0;
 
-	//MemoryStats.WorkingSetSize = UsedMem;??
-	MemoryStats.PagefileUsage = 0;
+	MemoryStats.UsedVirtual = 0;
 
 	// Just get memory information for the process and report the working set instead
 	task_basic_info_64_data_t TaskInfo;
 	mach_msg_type_number_t TaskInfoCount = TASK_BASIC_INFO_COUNT;
 	task_info( mach_task_self(), TASK_BASIC_INFO, (task_info_t)&TaskInfo, &TaskInfoCount );
-	MemoryStats.WorkingSetSize = TaskInfo.resident_size;
+	MemoryStats.UsedPhysical = TaskInfo.resident_size;
 
 	return MemoryStats;
 }
diff --git a/Engine/Source/Runtime/Core/Private/Misc/Core.cpp b/Engine/Source/Runtime/Core/Private/Misc/Core.cpp
--- a/Engine/Source/Runtime/Core/Private/Misc/Core.cpp
+++ b/Engine/Source/Runtime/Core/Private/Misc/Core.cpp
@@ -193,8 +193,6 @@ bool					GIsFirstInstance				= true;
 #endif
 /** Threshold for a frame to be considered a hitch (in seconds. */
 float GHitchThreshold = 0.075f;
-/** Whether to forcefully enable capturing of stats due to a profiler attached								*/
-bool					GProfilerAttached = false;
 /** Size to break up data into when saving compressed data													*/
 int32					GSavingCompressionChunkSize		= SAVING_COMPRESSION_CHUNK_SIZE;
 /** Whether we are using the seekfree/ cooked loading codepath.												*/
@@ -242,25 +240,9 @@ bool					GIsAutomationTesting					= false;
 /** Whether or not messages are being pumped outside of the main loop										*/
 bool					GPumpingMessagesOutsideOfMainLoop = false;
 
-/** Total number of calls to Malloc, if implemented by derived class.										*/
-uint64 FMalloc::TotalMallocCalls;
-/** Total number of calls to Free, if implemented by derived class.											*/
-uint64 FMalloc::TotalFreeCalls;
-/** Total number of calls to Realloc, if implemented by derived class.										*/
-uint64 FMalloc::TotalReallocCalls;
-
 /** Total blueprint compile time.																			*/
 double GBlueprintCompileTime = 0.0;
 
-
-/** Memory stats objects 
- *
- * If you add new Stat Memory stats please update:  FMemoryChartEntry so the automated memory chart has the info
- */
-
-
-DEFINE_STAT(STAT_PhysicalAllocSize);
-DEFINE_STAT(STAT_VirtualAllocSize);
 DEFINE_STAT(STAT_AudioMemory);
 DEFINE_STAT(STAT_TextureMemory);
 DEFINE_STAT(STAT_MemoryPhysXTotalAllocationSize);
@@ -285,10 +267,7 @@ DEFINE_STAT(STAT_ResourceVertexColorMemory);
 DEFINE_STAT(STAT_InstVertexColorMemory);
 DEFINE_STAT(STAT_StaticMeshIndexMemory);
 
-DEFINE_STAT(STAT_MallocCalls);
-DEFINE_STAT(STAT_ReallocCalls);
-DEFINE_STAT(STAT_FreeCalls);
-DEFINE_STAT(STAT_TotalAllocatorCalls);
+
 
 /** Threading stats objects */
 
@@ -335,6 +314,7 @@ DEFINE_LOG_CATEGORY(LogSerialization);
 DEFINE_LOG_CATEGORY(LogContentComparisonCommandlet);
 DEFINE_LOG_CATEGORY(LogNetPackageMap);
 DEFINE_LOG_CATEGORY(LogNetSerialization);
+DEFINE_LOG_CATEGORY(LogMemory);
 
 DEFINE_LOG_CATEGORY(LogTemp);
 
diff --git a/Engine/Source/Runtime/Core/Private/Misc/OutputDevice.cpp b/Engine/Source/Runtime/Core/Private/Misc/OutputDevice.cpp
--- a/Engine/Source/Runtime/Core/Private/Misc/OutputDevice.cpp
+++ b/Engine/Source/Runtime/Core/Private/Misc/OutputDevice.cpp
@@ -117,7 +117,7 @@ void FOutputDevice::Log( const FText& T )
 /** Number of top function calls to hide when dumping the callstack as text. */
 #define CALLSTACK_IGNOREDEPTH 2
 
-VARARG_BODY( void, FOutputDevice::CategorizedLogf, const TCHAR*, VARARG_EXTRA(ELogVerbosity::Type Verbosity) VARARG_EXTRA(const class FName& Category) )
+VARARG_BODY( void, FOutputDevice::CategorizedLogf, const TCHAR*, VARARG_EXTRA(const class FName& Category) VARARG_EXTRA(ELogVerbosity::Type Verbosity)  )
 {
 	GROWABLE_LOGF(Serialize(Buffer, Verbosity, Category))
 }
diff --git a/Engine/Source/Runtime/Core/Private/ProfilingDebugging/MallocProfiler.cpp b/Engine/Source/Runtime/Core/Private/ProfilingDebugging/MallocProfiler.cpp
--- a/Engine/Source/Runtime/Core/Private/ProfilingDebugging/MallocProfiler.cpp
+++ b/Engine/Source/Runtime/Core/Private/ProfilingDebugging/MallocProfiler.cpp
@@ -10,10 +10,11 @@
 
 #include "MallocProfiler.h"
 #include "ModuleManager.h"
+#include "MemoryMisc.h"
 
 DEFINE_LOG_CATEGORY(LogProfilingDebugging);
 
-FMallocProfiler* GMallocProfiler;
+CORE_API FMallocProfiler* GMallocProfiler;
 
 /**
  * Maximum depth of stack backtrace.
@@ -29,7 +30,7 @@ FMallocProfiler* GMallocProfiler;
 /** Magic value, determining that file is a memory profiler file.				*/
 #define MEMORY_PROFILER_MAGIC						0xDA15F7D8
 /** Version of memory profiler.													*/
-#define MEMORY_PROFILER_VERSION						3
+#define MEMORY_PROFILER_VERSION						4
 
 /*=============================================================================
 	Profiler header.
@@ -41,8 +42,8 @@ struct FProfilerHeader
 	uint32	Magic;
 	/** Version number to detect version mismatches.	*/
 	uint32	Version;
-	/** Platform that was captured.						*/
-	uint32	Platform;
+	/** Platform that this file was captured on.		*/
+	FString	PlatformName;
 	/** Whether symbol information is being serialized. */
 	uint32	bShouldSerializeSymbolInfo;
 	/** Name of executable, used for finding symbols.	*/
@@ -79,9 +80,11 @@ struct FProfilerHeader
 	friend FArchive& operator << ( FArchive& Ar, FProfilerHeader Header )
 	{
 		Ar	<< Header.Magic
-			<< Header.Version
-			<< Header.Platform
-			<< Header.bShouldSerializeSymbolInfo
+			<< Header.Version;
+
+		Header.PlatformName.SerializeAsANSICharArray(Ar,255);
+
+		Ar	<< Header.bShouldSerializeSymbolInfo
 			<< Header.NameTableOffset
 			<< Header.NameTableEntries
 			<< Header.CallStackAddressTableOffset
@@ -92,14 +95,8 @@ struct FProfilerHeader
 			<< Header.ModuleEntries
 			<< Header.NumDataFiles;
 
-		// @todo: Backwards compatibility with MemoryProfiler2 header reader
-		uint32 Dummy = -1;
-		Ar
-			<< Dummy	// ScriptCallstackTableOffset (-1 means, "no script data")
-			<< Dummy;	// ScriptNameTableOffset (-1 means, "no script data")
-
 		check( Ar.IsSaving() );
-		Header.ExecutableName.SerializeAsANSICharArray(Ar, 255);
+		Header.ExecutableName.SerializeAsANSICharArray(Ar,255);
 		return Ar;
 	}
 };
@@ -257,8 +254,8 @@ struct FProfilerOtherInfo
 	 *
 	 * @param	Ar			Archive to serialize to
 	 * @param	OtherInfo	Info to serialize
- 	 * @return	Passed in archive
- 	 */
+	 * @return	Passed in archive
+	 */
 	friend FArchive& operator << ( FArchive& Ar, FProfilerOtherInfo OtherInfo )
 	{
 		Ar	<< OtherInfo.DummyPointer
@@ -374,6 +371,8 @@ void FMallocProfiler::TrackSpecialMemory()
 		// Avoid tracking operations caused by tracking!
 		if (SyncObjectLockCount == 0)
 		{
+			FScopedMallocProfilerLock MallocProfilerLock;
+
 			// Write marker snapshot to stream.
 			FProfilerOtherInfo SnapshotMarker;
 			SnapshotMarker.DummyPointer	= TYPE_Other;
@@ -508,7 +507,7 @@ void FMallocProfiler::EndProfiling()
 		FProfilerHeader Header;
 		Header.Magic				= MEMORY_PROFILER_MAGIC;
 		Header.Version				= MEMORY_PROFILER_VERSION;
-		Header.Platform				= 0; //@todo this is now a string FPlatformProperties::PlatformName()
+		Header.PlatformName			= FPlatformProperties::PlatformName();
 		Header.bShouldSerializeSymbolInfo = SERIALIZE_SYMBOL_INFO ? 1 : 0;
 #if PLATFORM_WINDOWS
 		Header.ExecutableName		= FPlatformProcess::ExecutableName();
@@ -889,19 +888,24 @@ void FMallocProfiler::EmbedDwordMarker(EProfilingPayloadSubType SubType, uint32
 /** Writes memory allocations stats. */
 void FMallocProfiler::WriteMemoryAllocationStats()
 {
-	FMemoryAllocationStats_DEPRECATED MemStats;
-	UsedMalloc->GetAllocationInfo( MemStats );
-	GetTexturePoolSize( MemStats );
+	FGenericMemoryStats Stats;
+	FPlatformMemory::GetStatsForMallocProfiler( Stats ); 
+	UsedMalloc->GetAllocatorStats( Stats );
+
+	static const FName NAME_MemoryProfilingOverhead = TEXT("Memory Profiling Overhead");
+	int64 MemoryProfilingOverhead = CalculateMemoryProfilingOverhead();
+	Stats.Add( NAME_MemoryProfilingOverhead, MemoryProfilingOverhead );
 
-	uint8 StatsCount = FMemoryAllocationStats_DEPRECATED::GetStatsNum();
+	check( Stats.Data.Num() < MAX_uint8 );
+	uint8 StatsCount = Stats.Data.Num();
 	BufferedFileWriter << StatsCount;
 
-	// Convert memory allocations stats into more reliable format. Size of SIZE_T is platform-dependent.
-	for( int32 StatIndex = 0; StatIndex < StatsCount; StatIndex++ )
+	for( const auto& MapIt : Stats.Data )
 	{
-		// Serialize as SQWORDs.
-		int64 Value = (int64)(*( (SIZE_T*)&MemStats + StatIndex ));
-		BufferedFileWriter << Value;
+		FString StatName = MapIt.Key.ToString();
+		int64 StatValue = (int64)MapIt.Value;
+		int32 StatNameIndex = GetNameTableIndex(StatName);
+		BufferedFileWriter << StatNameIndex << StatValue;
 	}
 }
 
@@ -914,34 +918,6 @@ void FMallocProfiler::WriteAdditionalSnapshotMemoryStats()
 
 	// Write "stat levels" information.
 	WriteLoadedLevels( NULL );
-
-	// Write memory metrics, need to be redone. It could be used to store platform dependent memory metrics.
-	FMemoryAllocationStats_DEPRECATED MemStats;
-	UsedMalloc->GetAllocationInfo( MemStats );
-
-	int64 MemoryProfilingOverhead = CalculateMemoryProfilingOverhead();
-
-	int64 AllocatedFromOS = 0;
-	int64 MaxAllocatedFromOS = 0;
-	int64 TotalUsedChunks = 0;
-	int64 AllocateFromGame = 0;
-	int64 TextureLightmapMemory = 0;
-	int64 TextureShadowmapMemory = 0;
-#if STATS
-//	TextureLightmapMemory = GStatManager.GetStatValueDWORD(STAT_TextureLightmapMemory);
-//	TextureShadowmapMemory = GStatManager.GetStatValueDWORD(STAT_TextureShadowmapMemory);
-#endif
-
-	// IMPORTANT: All metrics MUST be serialized as SQWORDs!
-	uint8 NumMetrics = 7;
-	BufferedFileWriter << NumMetrics
-						<< MemoryProfilingOverhead
-						<< AllocatedFromOS
-						<< MaxAllocatedFromOS
-						<< AllocateFromGame
-						<< TotalUsedChunks
-						<< TextureLightmapMemory
-						<< TextureShadowmapMemory;
 }
 
 /** Snapshot taken when engine has started the cleaning process before loading a new level. */
@@ -1020,10 +996,8 @@ void FMallocProfiler::WriteLoadedLevels( UWorld* InWorld )
 /** 
  * Gather texture memory stats. 
  */
-void FMallocProfiler::GetTexturePoolSize( FMemoryAllocationStats_DEPRECATED& MemoryStats )
+void FMallocProfiler::GetTexturePoolSize( FGenericMemoryStats& out_Stats )
 {
-	MemoryStats.AllocatedTextureMemorySize = 0;
-	MemoryStats.AvailableTextureMemorySize = 0;
 }
 
 /*=============================================================================
diff --git a/Engine/Source/Runtime/Core/Private/Windows/WindowsPlatformMemory.cpp b/Engine/Source/Runtime/Core/Private/Windows/WindowsPlatformMemory.cpp
--- a/Engine/Source/Runtime/Core/Private/Windows/WindowsPlatformMemory.cpp
+++ b/Engine/Source/Runtime/Core/Private/Windows/WindowsPlatformMemory.cpp
@@ -9,7 +9,7 @@
 #include "MallocTBB.h"
 #include "MallocAnsi.h"
 #include "GenericPlatformMemoryPoolStats.h"
-
+#include "MemoryMisc.h"
 
 #if !FORCE_ANSI_ALLOCATOR
 #include "MallocBinned.h"
@@ -19,6 +19,8 @@
 #include <Psapi.h>
 #pragma comment(lib, "psapi.lib")
 
+DECLARE_MEMORY_STAT(TEXT("Windows Specific Memory Stat"),	STAT_WindowsSpecificMemoryStat, STATGROUP_MemoryPlatform);
+
 /** Enable this to track down windows allocations not wrapped by our wrappers
 int WindowsAllocHook(int nAllocType, void *pvData,
 				  size_t nSize, int nBlockUse, long lRequest,
@@ -47,11 +49,13 @@ void FWindowsPlatformMemory::Init()
 
 
 	const FPlatformMemoryConstants& MemoryConstants = FPlatformMemory::GetConstants();
-	UE_LOG(LogInit, Log, TEXT("Memory total: Physical=%.1fGB (%dGB approx) Pagefile=%.1fGB Virtual=%.1fGB"), 
+	UE_LOG(LogMemory, Log, TEXT("Memory total: Physical=%.1fGB (%dGB approx) Virtual=%.1fGB"), 
 		float(MemoryConstants.TotalPhysical/1024.0/1024.0/1024.0),
 		MemoryConstants.TotalPhysicalGB, 
-		float((MemoryConstants.TotalVirtual-MemoryConstants.TotalPhysical)/1024.0/1024.0/1024.0), 
 		float(MemoryConstants.TotalVirtual/1024.0/1024.0/1024.0) );
+
+	UpdateStats();
+	DumpStats( *GLog );
 }
 
 FMalloc* FWindowsPlatformMemory::BaseAllocator()
@@ -80,8 +84,9 @@ FPlatformMemoryStats FWindowsPlatformMemory::GetStats()
 	 *	GetProcessMemoryInfo
 	 *	PROCESS_MEMORY_COUNTERS
 	 *		WorkingSetSize
-	 *		PagefileUsage
-	 *	
+	 *		UsedVirtual
+	 *		PeakUsedVirtual
+	 *		
 	 *	GetPerformanceInfo
 	 *		PPERFORMANCE_INFORMATION 
 	 *		PageSize
@@ -100,13 +105,27 @@ FPlatformMemoryStats FWindowsPlatformMemory::GetStats()
 	MemoryStats.AvailablePhysical = MemoryStatusEx.ullAvailPhys;
 	MemoryStats.AvailableVirtual = MemoryStatusEx.ullAvailVirtual;
 	
-	MemoryStats.WorkingSetSize = ProcessMemoryCounters.WorkingSetSize;
-	MemoryStats.PeakWorkingSetSize = ProcessMemoryCounters.PeakWorkingSetSize;
-	MemoryStats.PagefileUsage = ProcessMemoryCounters.PagefileUsage;
+	MemoryStats.UsedPhysical = ProcessMemoryCounters.WorkingSetSize;
+	MemoryStats.PeakUsedPhysical = ProcessMemoryCounters.PeakWorkingSetSize;
+	MemoryStats.UsedVirtual = ProcessMemoryCounters.PagefileUsage;
+	MemoryStats.PeakUsedVirtual = ProcessMemoryCounters.PeakPagefileUsage;
 
 	return MemoryStats;
 }
 
+void FWindowsPlatformMemory::GetStatsForMallocProfiler( FGenericMemoryStats& out_Stats )
+{
+#if	STATS
+	FGenericPlatformMemory::GetStatsForMallocProfiler( out_Stats );
+
+	FPlatformMemoryStats Stats = GetStats();
+
+	// Windows specific stats.
+	static const FName NAME_WindowsSpecificMemoryStat   = TEXT("Windows Specific Memory Stat");
+	out_Stats.Add( NAME_WindowsSpecificMemoryStat, Stats.WindowsSpecificMemoryStat );
+#endif // STATS
+}
+
 const FPlatformMemoryConstants& FWindowsPlatformMemory::GetConstants()
 {
 	static FPlatformMemoryConstants MemoryConstants;
@@ -131,6 +150,15 @@ const FPlatformMemoryConstants& FWindowsPlatformMemory::GetConstants()
 	return MemoryConstants;	
 }
 
+void FWindowsPlatformMemory::UpdateStats()
+{
+	FGenericPlatformMemory::UpdateStats();
+
+	// Windows specific stats.
+	FPlatformMemoryStats MemoryStats = FPlatformMemory::GetStats();
+	SET_MEMORY_STAT(STAT_WindowsSpecificMemoryStat,MemoryStats.WindowsSpecificMemoryStat);
+}
+
 void* FWindowsPlatformMemory::BinnedAllocFromOS( SIZE_T Size )
 {
 	return VirtualAlloc( NULL, Size, MEM_COMMIT, PAGE_READWRITE );
diff --git a/Engine/Source/Runtime/Core/Public/Core.h b/Engine/Source/Runtime/Core/Public/Core.h
--- a/Engine/Source/Runtime/Core/Public/Core.h
+++ b/Engine/Source/Runtime/Core/Public/Core.h
@@ -35,7 +35,6 @@ template<typename T,typename Allocator = FDefaultAllocator> class TArray; // @to
 // Container forward declarations
 #include "Containers/ContainerAllocationPolicies.h"
 
-template<typename ElementType,typename Allocator = FDefaultAllocator> class TBaseArray;
 template<typename T> class TTransArray;
 template<typename KeyType,typename ValueType,bool bInAllowDuplicateKeys> struct TDefaultMapKeyFuncs;
 template<typename KeyType,typename ValueType,typename SetAllocator = FDefaultSetAllocator, typename KeyFuncs = TDefaultMapKeyFuncs<KeyType,ValueType,false> > class TMap;
@@ -319,10 +318,6 @@ extern CORE_API uint32 GFrameNumberRenderThread;
 /** Threshold for a frame to be considered a hitch (in seconds. */
 extern CORE_API float GHitchThreshold;
 
-
-/** Whether to forcefully enable capturing of stats due to profiler attached */
-extern CORE_API bool GProfilerAttached;
-
 /** Size to break up data into when saving compressed data */
 extern CORE_API int32 GSavingCompressionChunkSize;
 
diff --git a/Engine/Source/Runtime/Core/Public/GenericPlatform/GenericPlatformMemory.h b/Engine/Source/Runtime/Core/Public/GenericPlatform/GenericPlatformMemory.h
--- a/Engine/Source/Runtime/Core/Public/GenericPlatform/GenericPlatformMemory.h
+++ b/Engine/Source/Runtime/Core/Public/GenericPlatform/GenericPlatformMemory.h
@@ -7,7 +7,8 @@
 
 #pragma once
 
-// @todo: Add stats
+/** Holds generic memory stats, internally implemented as a map. */
+struct FGenericMemoryStats;
 
 /** 
  * Struct used to hold common memory constants for all platforms.
@@ -34,13 +35,23 @@ struct FGenericPlatformMemoryConstants
 		, PageSize( 0 )
 		, TotalPhysicalGB( 1 )
 	{}
+
+	/** Copy constructor, used by the generic platform memory stats. */
+	FGenericPlatformMemoryConstants( const FGenericPlatformMemoryConstants& Other )
+		: TotalPhysical( Other.TotalPhysical )
+		, TotalVirtual( Other.TotalVirtual )
+		, PageSize( Other.PageSize )
+		, TotalPhysicalGB( Other.TotalPhysicalGB )
+	{}
 };
 
+typedef FGenericPlatformMemoryConstants FPlatformMemoryConstants;
+
 /** 
  * Struct used to hold common memory stats for all platforms.
  * These values may change over the entire life of the executable.
  */
-struct FGenericPlatformMemoryStats
+struct FGenericPlatformMemoryStats : public FPlatformMemoryConstants
 {
 	/** The amount of physical memory currently available, in bytes. */
 	SIZE_T AvailablePhysical;
@@ -48,29 +59,24 @@ struct FGenericPlatformMemoryStats
 	/** The amount of virtual memory currently available, in bytes. */
 	SIZE_T AvailableVirtual;
 
-	/** The current working set size, in bytes. This is the amount of physical memory used by the process. */
-	SIZE_T WorkingSetSize;
+	/** The amount of physical memory used by the process, in bytes. */
+	SIZE_T UsedPhysical;
 
-	/** The peak working set size, in bytes. */
-	SIZE_T PeakWorkingSetSize;
+	/** The peak amount of physical memory used by the process, in bytes. */
+	SIZE_T PeakUsedPhysical;
 
 	/** Total amount of virtual memory used by the process. */
-	SIZE_T PagefileUsage;
+	SIZE_T UsedVirtual;
+
+	/** The peak amount of virtual memory used by the process. */
+	SIZE_T PeakUsedVirtual;
 	
 	/** Default constructor, clears all variables. */
-	FGenericPlatformMemoryStats()
-		: AvailablePhysical( 0 )
-		, AvailableVirtual( 0 )
-		, WorkingSetSize( 0 )
-		, PeakWorkingSetSize( 0 )
-		, PagefileUsage( 0 )		
-	{}
+	FGenericPlatformMemoryStats();
 };
 
 struct FPlatformMemoryStats;
 
-typedef FGenericPlatformMemoryConstants FPlatformMemoryConstants;
-
 /**
  * FMemory_Alloca/alloca implementation. This can't be a function, even FORCEINLINE'd because there's no guarantee that 
  * the memory returned in a function will stick around for the caller to use.
@@ -90,7 +96,7 @@ struct CORE_API FGenericPlatformMemory
 	 * Various memory regions that can be used with memory stats. The exact meaning of
 	 * the enums are relatively platform-dependent, although the general ones (Physical, GPU)
 	 * are straightforward. A platform can add more of these, and it won't affect other 
-	 * platforms, other than a miniscule amount of memory for the StatManager to track the
+	 * platforms, other than a minuscule amount of memory for the StatManager to track the
 	 * max available memory for each region (uses an array FPlatformMemory::MCR_MAX big)
 	 */
 	enum EMemoryCounterRegion
@@ -170,6 +176,11 @@ struct CORE_API FGenericPlatformMemory
 	 */
 	static FPlatformMemoryStats GetStats();
 
+	/**
+	 * Writes all platform specific current memory statistics in the format usable by the malloc profiler.
+	 */
+	static void GetStatsForMallocProfiler( FGenericMemoryStats& out_Stats );
+
 	/**
 	 * @return platform specific memory constants.
 	 */
@@ -180,6 +191,9 @@ struct CORE_API FGenericPlatformMemory
 	 */
 	static uint32 GetPhysicalGBRam();
 
+	/** Called once per frame, gathers and sets all platform memory statistics into the corresponding stats. */
+	static void UpdateStats();
+
 	/**
 	 * Allocates pages from the OS.
 	 *
diff --git a/Engine/Source/Runtime/Core/Public/HAL/MallocAnsi.h b/Engine/Source/Runtime/Core/Public/HAL/MallocAnsi.h
--- a/Engine/Source/Runtime/Core/Public/HAL/MallocAnsi.h
+++ b/Engine/Source/Runtime/Core/Public/HAL/MallocAnsi.h
@@ -121,35 +121,35 @@ public:
 		return false;
 	}
 
-#if PLATFORM_MAC
-	virtual void GetAllocationInfo( FMemoryAllocationStats_DEPRECATED& MemStats ) OVERRIDE
-	{
-		// @todo implement in platform memory
-		// Just get memory information for the process and report the working set instead
-		task_basic_info_64_data_t TaskInfo;
-		mach_msg_type_number_t TaskInfoCount = TASK_BASIC_INFO_COUNT;
-		task_info( mach_task_self(), TASK_BASIC_INFO, (task_info_t)&TaskInfo, &TaskInfoCount );
-		MemStats.TotalUsed = MemStats.TotalAllocated = TaskInfo.resident_size;
-	}
-#elif PLATFORM_IOS
-	virtual void GetAllocationInfo( FMemoryAllocationStats_DEPRECATED& MemStats ) OVERRIDE
-	{
-		// @todo implement in platform memory
-		task_basic_info info;
-		mach_msg_type_number_t size = sizeof(info);
-		kern_return_t kerr = task_info(mach_task_self(), TASK_BASIC_INFO, (task_info_t)&info, &size);
-		MemStats.TotalUsed = MemStats.TotalAllocated = info.resident_size;
-
-		mach_port_t host_port = mach_host_self();
-		mach_msg_type_number_t host_size = sizeof(vm_statistics_data_t) / sizeof(integer_t);
-		vm_size_t pagesize;
-		vm_statistics_data_t vm_stat;
-
-		host_page_size(host_port, &pagesize);
-		(void) host_statistics(host_port, HOST_VM_INFO, (host_info_t)&vm_stat, &host_size);
-		MemStats.OSReportedFree = vm_stat.free_count * pagesize;
-	}
-#endif
+// #if PLATFORM_MAC
+// 	virtual void GetAllocationInfo( FMemoryAllocationStats_DEPRECATED& MemStats ) OVERRIDE
+// 	{
+// 		// @todo implement in platform memory
+// 		// Just get memory information for the process and report the working set instead
+// 		task_basic_info_64_data_t TaskInfo;
+// 		mach_msg_type_number_t TaskInfoCount = TASK_BASIC_INFO_COUNT;
+// 		task_info( mach_task_self(), TASK_BASIC_INFO, (task_info_t)&TaskInfo, &TaskInfoCount );
+// 		MemStats.TotalUsed___UsedPhysical = MemStats.TotalAllocated___UsedVirtual = TaskInfo.resident_size;
+// 	}
+// #elif PLATFORM_IOS
+// 	virtual void GetAllocationInfo( FMemoryAllocationStats_DEPRECATED& MemStats ) OVERRIDE
+// 	{
+// 		// @todo implement in platform memory
+// 		task_basic_info info;
+// 		mach_msg_type_number_t size = sizeof(info);
+// 		kern_return_t kerr = task_info(mach_task_self(), TASK_BASIC_INFO, (task_info_t)&info, &size);
+// 		MemStats.TotalUsed___UsedPhysical = MemStats.TotalAllocated___UsedVirtual = info.resident_size;
+// 
+// 		mach_port_t host_port = mach_host_self();
+// 		mach_msg_type_number_t host_size = sizeof(vm_statistics_data_t) / sizeof(integer_t);
+// 		vm_size_t pagesize;
+// 		vm_statistics_data_t vm_stat;
+// 
+// 		host_page_size(host_port, &pagesize);
+// 		(void) host_statistics(host_port, HOST_VM_INFO, (host_info_t)&vm_stat, &host_size);
+// 		//MemStats.OSReportedFree = vm_stat.free_count * pagesize;
+// 	}
+// #endif
 
 	/**
 	 * Returns if the allocator is guaranteed to be thread-safe and therefore
diff --git a/Engine/Source/Runtime/Core/Public/HAL/MallocBinned.h b/Engine/Source/Runtime/Core/Public/HAL/MallocBinned.h
--- a/Engine/Source/Runtime/Core/Public/HAL/MallocBinned.h
+++ b/Engine/Source/Runtime/Core/Public/HAL/MallocBinned.h
@@ -33,6 +33,17 @@
 #include "LockFreeList.h"
 #include "Array.h"
 
+/** Malloc binned allocator specific stats. */
+DECLARE_MEMORY_STAT_EXTERN(TEXT("Binned OsCurrent"),	STAT_Binned_OsCurrent,STATGROUP_MemoryAllocator, CORE_API);
+DECLARE_MEMORY_STAT_EXTERN(TEXT("Binned OsPeak"),		STAT_Binned_OsPeak,STATGROUP_MemoryAllocator, CORE_API);
+DECLARE_MEMORY_STAT_EXTERN(TEXT("Binned WasteCurrent"),	STAT_Binned_WasteCurrent,STATGROUP_MemoryAllocator, CORE_API);
+DECLARE_MEMORY_STAT_EXTERN(TEXT("Binned WastePeak"),	STAT_Binned_WastePeak,STATGROUP_MemoryAllocator, CORE_API);
+DECLARE_MEMORY_STAT_EXTERN(TEXT("Binned UsedCurrent"),	STAT_Binned_UsedCurrent,STATGROUP_MemoryAllocator, CORE_API);
+DECLARE_MEMORY_STAT_EXTERN(TEXT("Binned UsedPeak"),		STAT_Binned_UsedPeak,STATGROUP_MemoryAllocator, CORE_API);
+DECLARE_MEMORY_STAT_EXTERN(TEXT("Binned CurrentAllocs"),STAT_Binned_CurrentAllocs,STATGROUP_MemoryAllocator, CORE_API);
+DECLARE_MEMORY_STAT_EXTERN(TEXT("Binned TotalAllocs"),	STAT_Binned_TotalAllocs,STATGROUP_MemoryAllocator, CORE_API);
+DECLARE_MEMORY_STAT_EXTERN(TEXT("Binned SlackCurrent"),	STAT_Binned_SlackCurrent,STATGROUP_MemoryAllocator, CORE_API);
+
 //
 // Optimized virtual memory allocator.
 //
@@ -292,14 +303,16 @@ private:
 #endif
 
 #if STATS
-	uint32		OsCurrent;
-	uint32		OsPeak;
-	uint32		WasteCurrent;
-	uint32		WastePeak;
-	uint32		UsedCurrent;
-	uint32		UsedPeak;
-	uint32		CurrentAllocs;
-	uint32		TotalAllocs;
+	SIZE_T		OsCurrent;
+	SIZE_T		OsPeak;
+	SIZE_T		WasteCurrent;
+	SIZE_T		WastePeak;
+	SIZE_T		UsedCurrent;
+	SIZE_T		UsedPeak;
+	SIZE_T		CurrentAllocs;
+	SIZE_T		TotalAllocs;
+	/** OsCurrent - WasteCurrent - UsedCurrent. */
+	SIZE_T		SlackCurrent;
 	double		MemTime;
 #endif
 
@@ -334,7 +347,7 @@ private:
 		{
 			OutOfMemory(IndirectPoolBlockSize * sizeof(FPoolInfo));
 		}
-        FMemory::Memset(Indirect, 0, IndirectPoolBlockSize*sizeof(FPoolInfo));
+		FMemory::Memset(Indirect, 0, IndirectPoolBlockSize*sizeof(FPoolInfo));
 		STAT(OsPeak = FMath::Max(OsPeak, OsCurrent += Align(IndirectPoolBlockSize * sizeof(FPoolInfo), PageSize)));
 		STAT(WastePeak = FMath::Max(WastePeak, WasteCurrent += Align(IndirectPoolBlockSize * sizeof(FPoolInfo), PageSize)));
 		return Indirect;
@@ -369,7 +382,7 @@ private:
 			{
 				if (!collision->FirstPool)
 				{
-                    collision->Key=Key;
+					collision->Key=Key;
 					InitializeHashBucket(collision);
 				}
 				return &collision->FirstPool[PoolIndex];
@@ -775,6 +788,7 @@ public:
 		,	UsedPeak		( 0 )
 		,	CurrentAllocs	( 0 )
 		,	TotalAllocs		( 0 )
+		,	SlackCurrent	( 0 )
 		,	MemTime			( 0.0 )
 #endif
 	{
@@ -1067,36 +1081,6 @@ public:
 		return true;
 	}
 
-	/**
-	 * Gathers memory allocations for both virtual and physical allocations. Subclasses should override to add additional info
-	 *
-	 * @param FMemoryAllocationStats_DEPRECATED	[out] structure containing information about the size of allocations
-	 */
-	virtual void GetAllocationInfo( FMemoryAllocationStats_DEPRECATED& MemStats ) OVERRIDE
-	{
-		const FPlatformMemoryConstants& MemoryConstants = FPlatformMemory::GetConstants();
-		FPlatformMemoryStats MemoryStats = FPlatformMemory::GetStats();
-
-		// determine how much memory has been allocated from the OS
-		MemStats.TotalUsed = MemoryStats.WorkingSetSize;
-		MemStats.TotalAllocated = MemoryConstants.TotalVirtual - MemoryStats.AvailableVirtual;
-		MemStats.CPUAvailable = MemoryStats.AvailableVirtual;
-
-#if STATS
-		MemStats.CPUUsed = UsedCurrent;
-		MemStats.CPUWaste = WasteCurrent;
-		
-		double Waste = 0.0;
-		for( int32 PoolIndex = 0; PoolIndex < POOL_COUNT; PoolIndex++ )
-		{
-			Waste += ( ( double )PoolTable[PoolIndex].TotalWaste / ( double )PoolTable[PoolIndex].TotalRequests ) * ( double )PoolTable[PoolIndex].ActiveRequests;
-			Waste += PoolTable[PoolIndex].NumActivePools * ( BINNED_ALLOC_POOL_SIZE - ( ( BINNED_ALLOC_POOL_SIZE / PoolTable[PoolIndex].BlockSize ) * PoolTable[PoolIndex].BlockSize ) );
-		}
-		MemStats.CPUWaste += ( uint32 )Waste;
-		MemStats.CPUSlack = OsCurrent - WasteCurrent - UsedCurrent;
-#endif
-	}
-
 	/**
 	 * Validates the allocator's heap
 	 */
@@ -1129,32 +1113,65 @@ public:
 		return( true );
 	}
 
+
+	/** Called once per frame, gathers and sets all memory allocator statistics into the corresponding stats. */
+	virtual void UpdateStats() OVERRIDE
+	{
+#ifdef USE_INTERNAL_LOCKS
+		FScopeLock ScopedLock(&AccessGuard);
+#endif
+
+		FMalloc::UpdateStats();
+
+#if STATS
+		SET_MEMORY_STAT(STAT_Binned_OsCurrent,OsCurrent);
+		SET_MEMORY_STAT(STAT_Binned_OsPeak,OsPeak);
+		SET_MEMORY_STAT(STAT_Binned_WasteCurrent,WasteCurrent);
+		SET_MEMORY_STAT(STAT_Binned_WastePeak,WastePeak);
+		SET_MEMORY_STAT(STAT_Binned_UsedCurrent,UsedCurrent);
+		SET_MEMORY_STAT(STAT_Binned_UsedPeak,UsedPeak);
+		SET_MEMORY_STAT(STAT_Binned_CurrentAllocs,CurrentAllocs);
+		SET_MEMORY_STAT(STAT_Binned_TotalAllocs,TotalAllocs);
+		SET_MEMORY_STAT(STAT_Binned_SlackCurrent,SlackCurrent);
+#endif
+	}
+
+	/** Writes allocator stats from the last update into the specified destination. */
+	virtual void GetAllocatorStats( FGenericMemoryStats& out_Stats ) OVERRIDE;
+
 	/**
-	 * Dumps details about all allocations to an output device. Subclasses should override to add additional info
+	 * Dumps allocator stats to an output device. Subclasses should override to add additional info
 	 *
 	 * @param Ar	[in] Output device
 	 */
-	virtual void DumpAllocations( class FOutputDevice& Ar ) OVERRIDE
+	virtual void DumpAllocatorStats( class FOutputDevice& Ar ) OVERRIDE
 	{
 #ifdef USE_INTERNAL_LOCKS
 		FScopeLock ScopedLock(&AccessGuard);
 #endif
 		ValidateHeap();
-
 #if STATS
+		UpdateSlackStat();
+
 		// This is all of the memory including stuff too big for the pools
-		Ar.Logf( TEXT("Tracked Allocation Status") );
+		Ar.CategorizedLogf( LogMemory.GetCategoryName(), ELogVerbosity::Log, TEXT("Allocator Stats for %s:" ), GetDescriptiveName() );
 		// Waste is the total overhead of the memory system
-		Ar.Logf( TEXT("Current Memory %.2f MB used, plus %.2f MB waste"), UsedCurrent / ( 1024.0f * 1024.0f ), ( OsCurrent - UsedCurrent ) / ( 1024.0f * 1024.0f ) );
-		Ar.Logf( TEXT("Peak Memory %.2f MB used, plus %.2f MB waste"), UsedPeak / ( 1024.0f * 1024.0f ), ( OsPeak - UsedPeak ) / ( 1024.0f * 1024.0f ) );
-		Ar.Logf( TEXT("Allocs      % 6i Current / % 6i Total"), CurrentAllocs, TotalAllocs );
-		MEM_TIME(Ar.Logf( TEXT( "Seconds     % 5.3f" ), MemTime ));
-		MEM_TIME(Ar.Logf( TEXT( "MSec/Allc   % 5.5f" ), 1000.0 * MemTime / MemAllocs ));
+		Ar.CategorizedLogf( LogMemory.GetCategoryName(), ELogVerbosity::Log, TEXT("Current Memory %.2f MB used, plus %.2f MB waste"), UsedCurrent / ( 1024.0f * 1024.0f ), ( OsCurrent - UsedCurrent ) / ( 1024.0f * 1024.0f ) );
+		Ar.CategorizedLogf( LogMemory.GetCategoryName(), ELogVerbosity::Log, TEXT("Peak Memory %.2f MB used, plus %.2f MB waste"), UsedPeak / ( 1024.0f * 1024.0f ), ( OsPeak - UsedPeak ) / ( 1024.0f * 1024.0f ) );
+
+		Ar.CategorizedLogf( LogMemory.GetCategoryName(), ELogVerbosity::Log, TEXT("Current OS Memory %.2f MB, peak %.2f MB"), OsCurrent / ( 1024.0f * 1024.0f ), OsPeak / ( 1024.0f * 1024.0f ) );
+		Ar.CategorizedLogf( LogMemory.GetCategoryName(), ELogVerbosity::Log, TEXT("Current Waste %.2f MB, peak %.2f MB"), WasteCurrent / ( 1024.0f * 1024.0f ), WastePeak / ( 1024.0f * 1024.0f ) );
+		Ar.CategorizedLogf( LogMemory.GetCategoryName(), ELogVerbosity::Log, TEXT("Current Used %.2f MB, peak %.2f MB"), UsedCurrent / ( 1024.0f * 1024.0f ), UsedPeak / ( 1024.0f * 1024.0f ) );
+		Ar.CategorizedLogf( LogMemory.GetCategoryName(), ELogVerbosity::Log, TEXT("Current Slack %.2f MB"), SlackCurrent / ( 1024.0f * 1024.0f ) );
+
+		Ar.CategorizedLogf( LogMemory.GetCategoryName(), ELogVerbosity::Log, TEXT("Allocs      % 6i Current / % 6i Total"), CurrentAllocs, TotalAllocs );
+		MEM_TIME(Ar.CategorizedLogf( LogMemory.GetCategoryName(), ELogVerbosity::Log, TEXT( "Seconds     % 5.3f" ), MemTime ));
+		MEM_TIME(Ar.CategorizedLogf( LogMemory.GetCategoryName(), ELogVerbosity::Log, TEXT( "MSec/Allc   % 5.5f" ), 1000.0 * MemTime / MemAllocs ));
 
 		// This is the memory tracked inside individual allocation pools
-		Ar.Logf( TEXT("") );
-		Ar.Logf( TEXT("Block Size Num Pools Max Pools Cur Allocs Total Allocs Min Req Max Req Mem Used Mem Slack Mem Waste Efficiency") );
-		Ar.Logf( TEXT("---------- --------- --------- ---------- ------------ ------- ------- -------- --------- --------- ----------") );
+		Ar.CategorizedLogf( LogMemory.GetCategoryName(), ELogVerbosity::Log, TEXT("") );
+		Ar.CategorizedLogf( LogMemory.GetCategoryName(), ELogVerbosity::Log, TEXT("Block Size Num Pools Max Pools Cur Allocs Total Allocs Min Req Max Req Mem Used Mem Slack Mem Waste Efficiency") );
+		Ar.CategorizedLogf( LogMemory.GetCategoryName(), ELogVerbosity::Log, TEXT("---------- --------- --------- ---------- ------------ ------- ------- -------- --------- --------- ----------") );
 
 		uint32 TotalMemory = 0;
 		uint32 TotalWaste = 0;
@@ -1183,7 +1200,7 @@ public:
 			// Memory that is reserved in active pools and ready for future use
 			uint32 MemSlack = MemAllocated - MemUsed - PoolMemWaste;
 
-			Ar.Logf( TEXT("% 10i % 9i % 9i % 10i % 12i % 7i % 7i % 7iK % 8iK % 8iK % 9.2f%%"),
+			Ar.CategorizedLogf( LogMemory.GetCategoryName(), ELogVerbosity::Log, TEXT("% 10i % 9i % 9i % 10i % 12i % 7i % 7i % 7iK % 8iK % 8iK % 9.2f%%"),
 				Table->BlockSize,
 				Table->NumActivePools,
 				Table->MaxActivePools,
@@ -1204,10 +1221,10 @@ public:
 			TotalPools += Table->NumActivePools;
 		}
 
-		Ar.Logf( TEXT( "" ) );
-		Ar.Logf( TEXT( "%iK allocated in pools (with %iK slack and %iK waste). Efficiency %.2f%%" ), TotalMemory, TotalSlack, TotalWaste, TotalMemory ? 100.0f * (TotalMemory - TotalWaste) / TotalMemory : 100.0f );
-		Ar.Logf( TEXT( "Allocations %i Current / %i Total (in %i pools)"), TotalActiveRequests, TotalTotalRequests, TotalPools );
-		Ar.Logf( TEXT("") );
+		Ar.CategorizedLogf( LogMemory.GetCategoryName(), ELogVerbosity::Log, TEXT("") );
+		Ar.CategorizedLogf( LogMemory.GetCategoryName(), ELogVerbosity::Log, TEXT("%iK allocated in pools (with %iK slack and %iK waste). Efficiency %.2f%%"), TotalMemory, TotalSlack, TotalWaste, TotalMemory ? 100.0f * (TotalMemory - TotalWaste) / TotalMemory : 100.0f );
+		Ar.CategorizedLogf( LogMemory.GetCategoryName(), ELogVerbosity::Log, TEXT("Allocations %i Current / %i Total (in %i pools)"), TotalActiveRequests, TotalTotalRequests, TotalPools );
+		Ar.CategorizedLogf( LogMemory.GetCategoryName(), ELogVerbosity::Log, TEXT("") );
 #endif
 	}
 
@@ -1215,6 +1232,21 @@ public:
 
 protected:
 
+	void UpdateSlackStat()
+	{
+#if	STATS
+		SIZE_T LocalWaste = WasteCurrent;
+		double Waste = 0.0;
+		for( int32 PoolIndex = 0; PoolIndex < POOL_COUNT; PoolIndex++ )
+		{
+			Waste += ( ( double )PoolTable[PoolIndex].TotalWaste / ( double )PoolTable[PoolIndex].TotalRequests ) * ( double )PoolTable[PoolIndex].ActiveRequests;
+			Waste += PoolTable[PoolIndex].NumActivePools * ( BINNED_ALLOC_POOL_SIZE - ( ( BINNED_ALLOC_POOL_SIZE / PoolTable[PoolIndex].BlockSize ) * PoolTable[PoolIndex].BlockSize ) );
+		}
+		LocalWaste += ( uint32 )Waste;
+		SlackCurrent = OsCurrent - LocalWaste - UsedCurrent;
+#endif // STATS
+	}
+
 	///////////////////////////////////
 	//// API platforms must implement
 	///////////////////////////////////
diff --git a/Engine/Source/Runtime/Core/Public/HAL/MallocDebug.h b/Engine/Source/Runtime/Core/Public/HAL/MallocDebug.h
--- a/Engine/Source/Runtime/Core/Public/HAL/MallocDebug.h
+++ b/Engine/Source/Runtime/Core/Public/HAL/MallocDebug.h
@@ -6,8 +6,6 @@
 
 #pragma once
 
-extern FRunnableThread* GRenderingThread;
-
 // Debug memory allocator.
 class FMallocDebug : public FMalloc
 {
@@ -31,7 +29,7 @@ private:
 	struct FMemDebug
 	{
 		SIZE_T		Size;
-		int32			RefCount;
+		int32		RefCount;
 		int32*		PreTag;
 		FMemDebug*	Next;
 		FMemDebug**	PrevLink;
@@ -42,7 +40,6 @@ private:
 	/** Total size of allocations */
 	SIZE_T		TotalAllocationSize;
 	/** Total size of allocation overhead */
-
 	SIZE_T		TotalWasteSize;
 
 	static const uint32 AllocatorOverhead = sizeof(FMemDebug) + sizeof(FMemDebug*) + sizeof(int32) + ALLOCATION_ALIGNMENT + sizeof(int32);
@@ -168,26 +165,16 @@ public:
 		return true;
 	}
 
-	/**
-	 * Gathers memory allocations for both virtual and physical allocations.
-	 *
-	 * @param FMemoryAllocationStats_DEPRECATED	[out] structure containing information about the size of allocations
-	 */
-	virtual void GetAllocationInfo( FMemoryAllocationStats_DEPRECATED& MemStats ) OVERRIDE
-	{
-		MemStats.TotalUsed = TotalAllocationSize;
-		MemStats.TotalAllocated = TotalAllocationSize + TotalWasteSize;
-		MemStats.CPUUsed = TotalAllocationSize;
-		MemStats.CPUWaste = TotalWasteSize;
-	}
-
 	/**
 	 * Dumps details about all allocations to an output device
 	 *
 	 * @param Ar	[in] Output device
 	 */
-	virtual void DumpAllocations( FOutputDevice& Ar ) OVERRIDE
+	virtual void DumpAllocatorStats( FOutputDevice& Ar ) OVERRIDE
 	{
+		Ar.Logf( TEXT( "Total Allocation Size: %u" ), TotalAllocationSize );
+		Ar.Logf( TEXT( "Total Waste Size: %u" ), TotalWasteSize );
+
 		int32 Count = 0;
 		int32 Chunks = 0;
 
diff --git a/Engine/Source/Runtime/Core/Public/HAL/MallocJemalloc.h b/Engine/Source/Runtime/Core/Public/HAL/MallocJemalloc.h
--- a/Engine/Source/Runtime/Core/Public/HAL/MallocJemalloc.h
+++ b/Engine/Source/Runtime/Core/Public/HAL/MallocJemalloc.h
@@ -37,8 +37,7 @@ public:
 	virtual void* Malloc( SIZE_T Size, uint32 Alignment ) OVERRIDE;
 	virtual void* Realloc( void* Ptr, SIZE_T NewSize, uint32 Alignment ) OVERRIDE;
 	virtual void Free( void* Ptr ) OVERRIDE;
-	virtual void DumpAllocations( FOutputDevice& Ar ) OVERRIDE;
-	CORE_API virtual void GetAllocationInfo( FMemoryAllocationStats_DEPRECATED& MemStats ) OVERRIDE;
+	virtual void DumpAllocatorStats( FOutputDevice& Ar ) OVERRIDE;
 	virtual bool GetAllocationSize(void *Original, SIZE_T &SizeOut) OVERRIDE;
 	virtual bool IsInternallyThreadSafe() const OVERRIDE {  return true; }
 	virtual const TCHAR * GetDescriptiveName() OVERRIDE { return TEXT("jemalloc"); }
diff --git a/Engine/Source/Runtime/Core/Public/HAL/MallocTBB.h b/Engine/Source/Runtime/Core/Public/HAL/MallocTBB.h
--- a/Engine/Source/Runtime/Core/Public/HAL/MallocTBB.h
+++ b/Engine/Source/Runtime/Core/Public/HAL/MallocTBB.h
@@ -42,20 +42,6 @@ public:
 	 * Free
 	 */
 	virtual void Free( void* Ptr ) OVERRIDE;
-	
-	/**
-	 * Dumps details about all allocations to an output device
-	 *
-	 * @param Ar	[in] Output device
-	 */
-	virtual void DumpAllocations( FOutputDevice& Ar ) OVERRIDE;
-
-	/**
-	 * Returns stats about current memory usage.
-	 *
-	 * @param FMemoryAllocationStats_DEPRECATED	[out] structure containing information about the size of allocations
-	 */
-	CORE_API virtual void GetAllocationInfo( FMemoryAllocationStats_DEPRECATED& MemStats ) OVERRIDE;
 
 	/**
 	 * If possible determine the size of the memory allocated at the given address
diff --git a/Engine/Source/Runtime/Core/Public/HAL/MallocThreadSafeProxy.h b/Engine/Source/Runtime/Core/Public/HAL/MallocThreadSafeProxy.h
--- a/Engine/Source/Runtime/Core/Public/HAL/MallocThreadSafeProxy.h
+++ b/Engine/Source/Runtime/Core/Public/HAL/MallocThreadSafeProxy.h
@@ -77,27 +77,25 @@ public:
 		}
 	}
 
-	/**
-	 * Passes request for gathering memory allocations for both virtual and physical allocations
-	 * on to used memory manager.
-	 *
-	 * @param FMemoryAllocationStats_DEPRECATED	[out] structure containing information about the size of allocations
-	 */
-	void GetAllocationInfo( FMemoryAllocationStats_DEPRECATED& MemStats ) OVERRIDE
+	/** Called once per frame, gathers and sets all memory allocator statistics into the corresponding stats. */
+	virtual void UpdateStats()
 	{
 		FScopeLock ScopeLock( &SynchronizationObject );
-		UsedMalloc->GetAllocationInfo( MemStats );
+		UsedMalloc->UpdateStats();
 	}
 
-	/**
-	 * Dumps details about all allocations to an output device
-	 *
-	 * @param Ar	[in] Output device
-	 */
-	virtual void DumpAllocations( class FOutputDevice& Ar ) OVERRIDE
+	/** Writes allocator stats from the last update into the specified destination. */
+	virtual void GetAllocatorStats( FGenericMemoryStats& out_Stats )
 	{
-		FScopeLock Lock( &SynchronizationObject );
-		UsedMalloc->DumpAllocations( Ar );
+		FScopeLock ScopeLock( &SynchronizationObject );
+		UsedMalloc->GetAllocatorStats( out_Stats );
+	}
+
+	/** Dumps allocator stats to an output device. */
+	virtual void DumpAllocatorStats( class FOutputDevice& Ar )
+	{
+		FScopeLock ScopeLock( &SynchronizationObject );
+		UsedMalloc->DumpAllocatorStats( Ar );
 	}
 
 	/**
@@ -115,14 +113,6 @@ public:
 		return UsedMalloc->Exec( InWorld, Cmd, Ar);
 	}
 
-
-	/** Called every game thread tick */
-	void Tick( float DeltaTime ) OVERRIDE
-	{
-		FScopeLock ScopeLock( &SynchronizationObject );
-		return UsedMalloc->Tick( DeltaTime );
-	}
-
 	/**
 	* If possible determine the size of the memory allocated at the given address
 	*
diff --git a/Engine/Source/Runtime/Core/Public/HAL/MemoryBase.h b/Engine/Source/Runtime/Core/Public/HAL/MemoryBase.h
--- a/Engine/Source/Runtime/Core/Public/HAL/MemoryBase.h
+++ b/Engine/Source/Runtime/Core/Public/HAL/MemoryBase.h
@@ -15,85 +15,10 @@
 CORE_API extern class FMalloc* GMalloc;
 
 /** Global FMallocProfiler variable to allow multiple malloc profilers to communicate. */
-MALLOC_PROFILER( extern class FMallocProfiler* GMallocProfiler; )
+MALLOC_PROFILER( CORE_API extern class FMallocProfiler* GMallocProfiler; )
 
-// Memory allocator.
-
-/** 
- * Struct used to hold memory allocation statistics. 
- * NOTE: Be sure to check hardcoded/mirrored usage found in MemoryProfiler2 StreamToken.cs
- * Only for backward compatibility with MemoryProfiler2.
- * @todo remove it later.
- * Contains only allocator specific stats.
- */
-struct FMemoryAllocationStats_DEPRECATED
-{
-	SIZE_T	TotalUsed;			/** The total amount of memory used by the game. */	
-	SIZE_T	TotalAllocated;		/** The total amount of memory allocated from the OS. */
-
-	// Virtual memory for Xbox and PC / Main memory for PS3 (tracked in the allocators)
-	SIZE_T	CPUUsed;			/** The allocated in use by the application virtual memory. */
-	SIZE_T	CPUSlack;			/** The allocated from the OS/allocator, but not used by the application. */
-	SIZE_T	CPUWaste;			/** Alignment waste from a pooled allocator plus book keeping overhead. */
-	SIZE_T	CPUAvailable;		/** MemoryStats.AvailableVirtual */
-
-	SIZE_T	NotUsed0;			
-	SIZE_T	NotUsed1;
-	SIZE_T	NotUsed2;
-	SIZE_T	NotUsed3;
-
-	SIZE_T	NotUsed4;
-	SIZE_T	OSReportedFree;		/** Free memory as reported by the operating system. (Xbox360, PS3) */
-	SIZE_T	NotUsed5;
-	SIZE_T	NotUsed6;
-
-	SIZE_T	NotUsed7;
-	SIZE_T	NotUsed8;
-	SIZE_T	NotUsed9;
-	SIZE_T	NotUsedA;
-
-	SIZE_T	AllocatedTextureMemorySize; /** Size of allocated memory in the texture pool. */
-	SIZE_T	AvailableTextureMemorySize; /** Size of available memory in the texture pool. */
-
-	/** Returns statistics count. */
-	static uint8 GetStatsNum()
-	{
-		const uint8 ItemsCount = (uint8)(sizeof(FMemoryAllocationStats_DEPRECATED) / sizeof(SIZE_T));
-		return ItemsCount;
-	}
-
-	/** 
-	 * Constructor. 
-	 */
-	FMemoryAllocationStats_DEPRECATED()
-	{
-		TotalUsed = 0;
-		TotalAllocated = 0;
-
-		CPUUsed = 0;
-		CPUSlack = 0;
-		CPUWaste = 0;
-		CPUAvailable = 0;
-
-		NotUsed0 = 0;
-		NotUsed1 = 0;	
-		NotUsed2 = 0;
-		NotUsed3 = 0;
-
-		NotUsed4 = 0;
-		OSReportedFree = 0;
-		NotUsed5 = 0;
-		NotUsed6 = 0;
-
-		NotUsed6 = 0;	
-		NotUsed8 = 0;	
-		NotUsed9 = 0;
-		NotUsedA = 0;
-
-		AllocatedTextureMemorySize = 0;
-		AvailableTextureMemorySize = 0;
-	}
-};
+/** Holds generic memory stats, internally implemented as a map. */
+struct FGenericMemoryStats;
 
 /**
  * Inherit from FUseSystemMallocForNew if you want your objects to be placed in memory
@@ -190,12 +115,17 @@ public:
 	{ 
 		return false; 
 	}
-	
-	/** 
-	 * Called every game thread tick 
-	 */
-	virtual void Tick( float DeltaTime ) 
-	{ 
+
+	/** Called once per frame, gathers and sets all memory allocator statistics into the corresponding stats. */
+	virtual void UpdateStats();
+
+	/** Writes allocator stats from the last update into the specified destination. */
+	virtual void GetAllocatorStats( FGenericMemoryStats& out_Stats );
+
+	/** Dumps current allocator stats to the log. */
+	virtual void DumpAllocatorStats( class FOutputDevice& Ar )
+	{
+		Ar.Logf( TEXT("Allocator Stats for %s: (not implemented)" ), GetDescriptiveName() );
 	}
 
 	/**
@@ -207,25 +137,6 @@ public:
 		return false; 
 	}
 
-	/**
-	 * Gathers all current memory stats
-	 *
-	 * @param FMemoryAllocationStats_DEPRECATED	[out] structure containing information about the size of allocations
-	 */
-	virtual void GetAllocationInfo( FMemoryAllocationStats_DEPRECATED& MemStats ) 
-	{
-	}
-
-	/**
-	 * Dumps details about all allocations to an output device
-	 *
-	 * @param Ar	[in] Output device
-	 */
-	virtual void DumpAllocations( class FOutputDevice& Ar ) 
-	{
-		Ar.Logf( TEXT( "DumpAllocations not implemented" ) );
-	}
-
 	/**
 	 * Validates the allocator's heap
 	 */
@@ -256,6 +167,27 @@ public:
 		return TEXT("Unspecified allocator");
 	}
 
+protected:
+	friend struct FCurrentFrameCalls;
+
+	/** Atomically increment total malloc calls. */
+	FORCEINLINE void IncrementTotalMallocCalls()
+	{
+		FPlatformAtomics::InterlockedIncrement( (volatile int64*)&FMalloc::TotalMallocCalls );
+	}
+
+	/** Atomically increment total free calls. */
+	FORCEINLINE void IncrementTotalFreeCalls()
+	{
+		FPlatformAtomics::InterlockedIncrement( (volatile int64*)&FMalloc::TotalFreeCalls );
+	}
+
+	/** Atomically increment total realloc calls. */
+	FORCEINLINE void IncrementTotalReallocCalls()
+	{
+		FPlatformAtomics::InterlockedIncrement( (volatile int64*)&FMalloc::TotalReallocCalls );
+	}
+
 	/** Total number of calls Malloc, if implemented by derived class. */
 	static uint64 TotalMallocCalls;
 	/** Total number of calls Malloc, if implemented by derived class. */
diff --git a/Engine/Source/Runtime/Core/Public/HAL/MemoryMisc.h b/Engine/Source/Runtime/Core/Public/HAL/MemoryMisc.h
new file mode 100644
--- /dev/null
+++ b/Engine/Source/Runtime/Core/Public/HAL/MemoryMisc.h
@@ -0,0 +1,14 @@
+// Copyright 1998-2014 Epic Games, Inc. All Rights Reserved.
+
+#pragma once
+
+/** Holds generic memory stats, internally implemented as a map. */
+struct FGenericMemoryStats
+{
+	void Add( const FName StatName, const SIZE_T StatValue )
+	{
+		Data.Add( StatName, StatValue );
+	}
+
+	TMap<FName,SIZE_T> Data;
+};
\ No newline at end of file
diff --git a/Engine/Source/Runtime/Core/Public/HAL/OutputDevices.h b/Engine/Source/Runtime/Core/Public/HAL/OutputDevices.h
--- a/Engine/Source/Runtime/Core/Public/HAL/OutputDevices.h
+++ b/Engine/Source/Runtime/Core/Public/HAL/OutputDevices.h
@@ -501,6 +501,7 @@ CORE_API DECLARE_LOG_CATEGORY_EXTERN(LogUnrealMatrix, Log, All);
 CORE_API DECLARE_LOG_CATEGORY_EXTERN(LogContentComparisonCommandlet, Log, All);
 CORE_API DECLARE_LOG_CATEGORY_EXTERN(LogNetPackageMap, Warning, All);
 CORE_API DECLARE_LOG_CATEGORY_EXTERN(LogNetSerialization, Warning, All);
+CORE_API DECLARE_LOG_CATEGORY_EXTERN(LogMemory, Log, All);
 
 // Temporary log category, generally you should not check things in that use this
 CORE_API DECLARE_LOG_CATEGORY_EXTERN(LogTemp, Log, All);
diff --git a/Engine/Source/Runtime/Core/Public/Misc/CoreStats.h b/Engine/Source/Runtime/Core/Public/Misc/CoreStats.h
--- a/Engine/Source/Runtime/Core/Public/Misc/CoreStats.h
+++ b/Engine/Source/Runtime/Core/Public/Misc/CoreStats.h
@@ -7,7 +7,6 @@
 #pragma once
 
 #include "Stats.h"
-#include "Stats2.h"
 
 CORE_API DECLARE_LOG_CATEGORY_EXTERN(LogSHA, Warning, All);
 CORE_API DECLARE_LOG_CATEGORY_EXTERN(LogStats, Log, All);
@@ -24,8 +23,6 @@ CORE_API DECLARE_LOG_CATEGORY_EXTERN(LogLoad, Log, All);
 
 /** Memory stats */
 DECLARE_MEMORY_STAT_EXTERN(TEXT("Audio Memory Used"),STAT_AudioMemory,STATGROUP_Memory, CORE_API);
-DECLARE_MEMORY_STAT_EXTERN(TEXT("Physical Memory Used"),STAT_PhysicalAllocSize,STATGROUP_Memory, CORE_API);
-DECLARE_MEMORY_STAT_EXTERN(TEXT("Virtual Memory Used"),STAT_VirtualAllocSize,STATGROUP_Memory, CORE_API);
 DECLARE_MEMORY_STAT_EXTERN(TEXT("Animation Memory"),STAT_AnimationMemory,STATGROUP_Memory, CORE_API);
 DECLARE_MEMORY_STAT_EXTERN(TEXT("Precomputed Visibility Memory"),STAT_PrecomputedVisibilityMemory,STATGROUP_Memory, CORE_API);
 DECLARE_MEMORY_STAT_EXTERN(TEXT("Precomputed Light Volume Memory"),STAT_PrecomputedLightVolumeMemory,STATGROUP_Memory, CORE_API);
@@ -50,12 +47,6 @@ DECLARE_MEMORY_STAT_EXTERN(TEXT("Texture Memory Used"),STAT_TextureMemory,STATGR
 DECLARE_MEMORY_STAT_EXTERN(TEXT("Reflection Capture Texture Memory"),STAT_ReflectionCaptureTextureMemory,STATGROUP_Memory, CORE_API);
 DECLARE_MEMORY_STAT_EXTERN(TEXT("Reflection Capture Memory"),STAT_ReflectionCaptureMemory,STATGROUP_Memory, CORE_API);
 
-/** Memory churn stats. */
-DECLARE_DWORD_COUNTER_STAT_EXTERN(TEXT("Malloc calls"),STAT_MallocCalls,STATGROUP_MemoryChurn, CORE_API);
-DECLARE_DWORD_COUNTER_STAT_EXTERN(TEXT("Free calls"),STAT_FreeCalls,STATGROUP_MemoryChurn, CORE_API);
-DECLARE_DWORD_COUNTER_STAT_EXTERN(TEXT("Realloc calls"),STAT_ReallocCalls,STATGROUP_MemoryChurn, CORE_API);
-DECLARE_DWORD_COUNTER_STAT_EXTERN(TEXT("Total Allocator calls"),STAT_TotalAllocatorCalls,STATGROUP_MemoryChurn, CORE_API);
-
 DECLARE_CYCLE_STAT_EXTERN(TEXT("Total Render thread idle time"),STAT_RenderingIdleTime,STATGROUP_Threading, CORE_API);
 DECLARE_CYCLE_STAT_EXTERN(TEXT("Wait for GPU Query"),STAT_RenderingIdleTime_WaitingForGPUQuery,STATGROUP_Threading, CORE_API);
 DECLARE_CYCLE_STAT_EXTERN(TEXT("Wait for GPU Present"),STAT_RenderingIdleTime_WaitingForGPUPresent,STATGROUP_Threading, CORE_API);
diff --git a/Engine/Source/Runtime/Core/Public/Misc/OutputDevice.h b/Engine/Source/Runtime/Core/Public/Misc/OutputDevice.h
--- a/Engine/Source/Runtime/Core/Public/Misc/OutputDevice.h
+++ b/Engine/Source/Runtime/Core/Public/Misc/OutputDevice.h
@@ -317,7 +317,7 @@ public:
 
 	VARARG_DECL( void, void, {}, Logf, VARARG_NONE, const TCHAR*, VARARG_NONE, VARARG_NONE );
 	VARARG_DECL( void, void, {}, Logf, VARARG_NONE, const TCHAR*, VARARG_EXTRA(ELogVerbosity::Type Verbosity), VARARG_EXTRA(Verbosity) );
-	VARARG_DECL( void, void, {}, CategorizedLogf, VARARG_NONE, const TCHAR*, VARARG_EXTRA(ELogVerbosity::Type Verbosity) VARARG_EXTRA(const class FName& Category), VARARG_EXTRA(Verbosity) VARARG_EXTRA(Category) );
+	VARARG_DECL( void, void, {}, CategorizedLogf, VARARG_NONE, const TCHAR*, VARARG_EXTRA(const class FName& Category) VARARG_EXTRA(ELogVerbosity::Type Verbosity), VARARG_EXTRA(Category) VARARG_EXTRA(Verbosity) );
 protected:
 	/** Whether to output the 'Log: ' type front... */
 	bool bSuppressEventTag;
diff --git a/Engine/Source/Runtime/Core/Public/ProfilingDebugging/MallocProfiler.h b/Engine/Source/Runtime/Core/Public/ProfilingDebugging/MallocProfiler.h
--- a/Engine/Source/Runtime/Core/Public/ProfilingDebugging/MallocProfiler.h
+++ b/Engine/Source/Runtime/Core/Public/ProfilingDebugging/MallocProfiler.h
@@ -49,40 +49,7 @@ enum EProfilingPayloadSubType
 	/** Marker used to determine when a new frame has started. */
 	SUBTYPE_FrameTimeMarker						= 3,
 
-	/** Not used. Only for backward compatibility. Use a new snapshot marker instead. */
-	SUBTYPE_TextMarker							= 4,
-
-	// Marker types for periodic non-GMalloc memory status updates. Only for backward compatibility, replaced with SUBTYPE_MemoryAllocationStats
-
-	/** Marker used to store the total amount of memory used by the game. */
-	SUBTYPE_TotalUsed							= 5,
-
-	/** Marker used to store the total amount of memory allocated from the OS. */
-	SUBTYPE_TotalAllocated						= 6,
-
-	/** Marker used to store the allocated in use by the application virtual memory. */
-	SUBTYPE_CPUUsed								= 7,
-
-	/** Marker used to store the allocated from the OS/allocator, but not used by the application. */
-	SUBTYPE_CPUSlack							= 8,
-
-	/** Marker used to store the alignment waste from a pooled allocator plus book keeping overhead. */
-	SUBTYPE_CPUWaste							= 9,
-
-	/** Marker used to store the allocated in use by the application physical memory. */
-	SUBTYPE_GPUUsed								= 10,
-
-	/** Marker used to store the allocated from the OS, but not used by the application. */
-	SUBTYPE_GPUSlack							= 11,
-
-	/** Marker used to store the alignment waste from a pooled allocator plus book keeping overhead. */
-	SUBTYPE_GPUWaste							= 12,
-
-	/** Marker used to store the overhead of the operating system. */
-	SUBTYPE_OSOverhead							= 13,
-
-	/** Marker used to store the size of loaded executable, stack, static, and global object size. */
-	SUBTYPE_ImageSize							= 14,
+	// Markers from 4 to 20 are deprecated and have been removed.
 
 	/// Version 3
 	// Marker types for automatic snapshots.
@@ -108,7 +75,7 @@ enum EProfilingPayloadSubType
 	/** Marker used to determine when a previously streamed level has been made visible. */
 	SUBTYPE_SnapshotMarker_LevelStream_End		= 27,
 
-	/** Marker used to store a generic malloc statistics. @see FMemoryAllocationStats_DEPRECATED. */
+	/** Marker used to store a generic malloc statistics. */
 	SUBTYPE_MemoryAllocationStats				= 31,
 
 	/** Start licensee-specific subtypes from here. */
@@ -245,7 +212,7 @@ protected:
 	FCriticalSection						CriticalSection;
 
 	/** Mapping from program counter to index in program counter array.								*/
-	TMap<uint64,int32>							ProgramCounterToIndexMap;
+	TMap<uint64,int32>						ProgramCounterToIndexMap;
 	/** Array of unique call stack address infos.													*/
 	TArray<struct FCallStackAddressInfo>	CallStackAddressInfoArray;
 
@@ -266,7 +233,7 @@ protected:
 	FCriticalSection						SyncObject;
 
 	/** Whether operations should be tracked. false e.g. in tracking internal functions.			*/
-	int32										SyncObjectLockCount;
+	int32									SyncObjectLockCount;
 
 	/** The currently executing thread's id. */
 	uint32									ThreadId;
@@ -386,7 +353,7 @@ protected:
 	/** 
 	 * Gather texture memory stats. 
 	 */
-	virtual void GetTexturePoolSize( FMemoryAllocationStats_DEPRECATED& MemoryStats );
+	virtual void GetTexturePoolSize( FGenericMemoryStats& out_Stats );
 
 	/** 
 		Added for untracked memory calculation
@@ -502,27 +469,25 @@ public:
 		return true; 
 	}
 
-	/**
-	 * Passes request for gathering memory allocations for both virtual and physical allocations
-	 * on to used memory manager.
-	 *
-	 * @param FMemoryAllocationStats_DEPRECATED	[out] structure containing information about the size of allocations
-	 */
-	virtual void GetAllocationInfo( FMemoryAllocationStats_DEPRECATED& MemStats ) OVERRIDE
+	/** Called once per frame, gathers and sets all memory allocator statistics into the corresponding stats. */
+	virtual void UpdateStats() OVERRIDE
 	{
 		FScopeLock Lock( &CriticalSection );
-		UsedMalloc->GetAllocationInfo( MemStats );
+		UsedMalloc->UpdateStats();
 	}
 
-	/**
-	 * Dumps details about all allocations to an output device
-	 *
-	 * @param Ar	[in] Output device
-	 */
-	virtual void DumpAllocations( class FOutputDevice& Ar ) OVERRIDE
+	/** Writes allocator stats from the last update into the specified destination. */
+	virtual void GetAllocatorStats( FGenericMemoryStats& out_Stats ) OVERRIDE
+	{
+		FScopeLock Lock( &CriticalSection );
+		UsedMalloc->GetAllocatorStats( out_Stats );
+	}
+
+	/** Dumps allocator stats to an output device. */
+	virtual void DumpAllocatorStats( class FOutputDevice& Ar ) OVERRIDE
 	{
 		FScopeLock Lock( &CriticalSection );
-		UsedMalloc->DumpAllocations( Ar );
+		UsedMalloc->DumpAllocatorStats( Ar );
 	}
 
 	/**
@@ -560,13 +525,6 @@ public:
 	bool HandleSnapshotMemoryCommand( const TCHAR* Cmd, FOutputDevice& Ar );
 	bool HandleSnapshotMemoryFrameCommand( const TCHAR* Cmd, FOutputDevice& Ar );
 
-	/** Called every game thread tick */
-	virtual void Tick( float DeltaTime ) OVERRIDE
-	{ 
-		FScopeLock Lock( &CriticalSection );
-		UsedMalloc->Tick(DeltaTime);
-	}
-
 	virtual const TCHAR * GetDescriptiveName() OVERRIDE
 	{ 
 		FScopeLock ScopeLock( &CriticalSection );
diff --git a/Engine/Source/Runtime/Core/Public/Stats/Stats2.h b/Engine/Source/Runtime/Core/Public/Stats/Stats2.h
--- a/Engine/Source/Runtime/Core/Public/Stats/Stats2.h
+++ b/Engine/Source/Runtime/Core/Public/Stats/Stats2.h
@@ -1542,7 +1542,6 @@ DECLARE_STATS_GROUP(TEXT("Landscape"),STATGROUP_Landscape);
 DECLARE_STATS_GROUP(TEXT("LightRendering"),STATGROUP_LightRendering);
 DECLARE_STATS_GROUP(TEXT("Load Time"), STATGROUP_LoadTime);
 DECLARE_STATS_GROUP(TEXT("Memory"),STATGROUP_Memory);
-DECLARE_STATS_GROUP(TEXT("MemoryChurn"),STATGROUP_MemoryChurn);
 DECLARE_STATS_GROUP(TEXT("MemoryStaticMesh"),STATGROUP_MemoryStaticMesh);
 DECLARE_STATS_GROUP(TEXT("MeshParticles"),STATGROUP_MeshParticles);
 DECLARE_STATS_GROUP(TEXT("Morph"),STATGROUP_MorphTarget);
@@ -1580,6 +1579,9 @@ DECLARE_STATS_GROUP(TEXT("UI"),STATGROUP_UI);
 DECLARE_STATS_GROUP(TEXT("UObjects"),STATGROUP_UObjects);
 DECLARE_STATS_GROUP(TEXT("User"),STATGROUP_User);
 
+DECLARE_STATS_GROUP(TEXT("MemoryAllocator"),STATGROUP_MemoryAllocator);
+DECLARE_STATS_GROUP(TEXT("MemoryPlatform"),STATGROUP_MemoryPlatform);
+
 DECLARE_CYCLE_STAT_EXTERN(TEXT("FrameTime"),STAT_FrameTime,STATGROUP_Engine, CORE_API);
 DECLARE_FLOAT_COUNTER_STAT_EXTERN(TEXT("StatUnit FPS"), STAT_FPS, STATGROUP_Engine, CORE_API);
 
diff --git a/Engine/Source/Runtime/Core/Public/Windows/WindowsPlatformMemory.h b/Engine/Source/Runtime/Core/Public/Windows/WindowsPlatformMemory.h
--- a/Engine/Source/Runtime/Core/Public/Windows/WindowsPlatformMemory.h
+++ b/Engine/Source/Runtime/Core/Public/Windows/WindowsPlatformMemory.h
@@ -7,11 +7,10 @@
 
 #pragma once
 
-// TODO: @JarekS 2013-12-10, 20:36 As stats?
 /**
  *	Windows implementation of the FGenericPlatformMemoryStats.
  *	At this moment it's just the same as the FGenericPlatformMemoryStats.
- *	Can be extended as showed in the following example.
+ *	Can be extended as shown in the following example.
  */
 struct FPlatformMemoryStats : public FGenericPlatformMemoryStats
 {
@@ -64,7 +63,9 @@ struct CORE_API FWindowsPlatformMemory : public FGenericPlatformMemory
 	static void Init();
 	static class FMalloc* BaseAllocator();
 	static FPlatformMemoryStats GetStats();
+	static void GetStatsForMallocProfiler( FGenericMemoryStats& out_Stats );
 	static const FPlatformMemoryConstants& GetConstants();
+	static void UpdateStats();
 	static void* BinnedAllocFromOS( SIZE_T Size );
 	static void BinnedFreeToOS( void* Ptr );
 	static FSharedMemoryRegion * MapNamedSharedMemoryRegion(const FString & InName, bool bCreate, uint32 AccessMode, SIZE_T Size);
diff --git a/Engine/Source/Runtime/Engine/Classes/Engine/Engine.h b/Engine/Source/Runtime/Engine/Classes/Engine/Engine.h
--- a/Engine/Source/Runtime/Engine/Classes/Engine/Engine.h
+++ b/Engine/Source/Runtime/Engine/Classes/Engine/Engine.h
@@ -1587,7 +1587,7 @@ public:
 	bool HandleTrackParticleRenderingStatsCommand( const TCHAR* Cmd, FOutputDevice& Ar );
 	bool HandleDumpParticleRenderingStatsCommand( const TCHAR* Cmd, FOutputDevice& Ar );
 	bool HandleDumpParticleFrameRenderingStatsCommand( const TCHAR* Cmd, FOutputDevice& Ar );
-	bool HandleDumpAllocsCommand( const TCHAR* Cmd, FOutputDevice& Ar );
+	bool HandleDumpAllocatorStats( const TCHAR* Cmd, FOutputDevice& Ar );
 	bool HandleHeapCheckCommand( const TCHAR* Cmd, FOutputDevice& Ar );
 	bool HandleToggleOnscreenDebugMessageDisplayCommand( const TCHAR* Cmd, FOutputDevice& Ar );
 	bool HandleToggleOnscreenDebugMessageSystemCommand( const TCHAR* Cmd, FOutputDevice& Ar );	
diff --git a/Engine/Source/Runtime/Engine/Private/GameEngine.cpp b/Engine/Source/Runtime/Engine/Private/GameEngine.cpp
--- a/Engine/Source/Runtime/Engine/Private/GameEngine.cpp
+++ b/Engine/Source/Runtime/Engine/Private/GameEngine.cpp
@@ -819,12 +819,6 @@ void UGameEngine::Tick( float DeltaSeconds, bool bIdleMode )
 #endif
 	}
 
-	// Tick allocator
-	if( GMalloc != NULL )
-	{
-		GMalloc->Tick( DeltaSeconds );
-	}
-
 	// Tick the module manager
 	FModuleManager::Get().Tick();
 
diff --git a/Engine/Source/Runtime/Engine/Private/MallocProfilerEx.cpp b/Engine/Source/Runtime/Engine/Private/MallocProfilerEx.cpp
--- a/Engine/Source/Runtime/Engine/Private/MallocProfilerEx.cpp
+++ b/Engine/Source/Runtime/Engine/Private/MallocProfilerEx.cpp
@@ -6,10 +6,11 @@
 
 #include "EnginePrivate.h"
 
+#if USE_MALLOC_PROFILER
+
 #include "MallocProfiler.h"
 #include "MallocProfilerEx.h"
-
-#if USE_MALLOC_PROFILER
+#include "MemoryMisc.h"
 
 // These functions are here because FMallocProfiler is in the Core
 // project, and therefore can't access most of the classes needed by these functions.
@@ -73,7 +74,7 @@ void FMallocProfilerEx::WriteLoadedLevels( UWorld* InWorld )
 /** 
  * Gather texture memory stats. 
  */
-void FMallocProfilerEx::GetTexturePoolSize( FMemoryAllocationStats_DEPRECATED& MemoryStats )
+void FMallocProfilerEx::GetTexturePoolSize( FGenericMemoryStats& out_Stats )
 {
 	FTextureMemoryStats Stats;
 
@@ -82,8 +83,8 @@ void FMallocProfilerEx::GetTexturePoolSize( FMemoryAllocationStats_DEPRECATED& M
 		RHIGetTextureMemoryStats(Stats);
 	}
 
-	MemoryStats.AllocatedTextureMemorySize = Stats.AllocatedMemorySize;
-	MemoryStats.AvailableTextureMemorySize = Stats.AllocatedMemorySize;
+	static const FName NAME_TextureAllocatedMemorySize = TEXT("Texture Allocated Memory Size");
+	out_Stats.Add( NAME_TextureAllocatedMemorySize, Stats.AllocatedMemorySize );
 }
 
 #endif // USE_MALLOC_PROFILER
\ No newline at end of file
diff --git a/Engine/Source/Runtime/Engine/Private/UnrealEngine.cpp b/Engine/Source/Runtime/Engine/Private/UnrealEngine.cpp
--- a/Engine/Source/Runtime/Engine/Private/UnrealEngine.cpp
+++ b/Engine/Source/Runtime/Engine/Private/UnrealEngine.cpp
@@ -618,9 +618,13 @@ namespace EngineDefs
   */
 void EngineMemoryWarningHandler(const FGenericMemoryWarningContext& GenericContext)
 {
-	FMemoryAllocationStats_DEPRECATED MemStats;
-	GMalloc->GetAllocationInfo( MemStats );
-	FPlatformMisc::LowLevelOutputDebugStringf(TEXT("EngineMemoryWarningHandler: Mem Used %.2f MB, Texture Memory %.2f MB, Render Target memory %.2f MB, OS Free %.2f MB\n"), MemStats.TotalUsed / 1048576.0f, GCurrentTextureMemorySize / 1048576.0f, GCurrentRendertargetMemorySize / 1048576.0f, MemStats.OSReportedFree / 1048576.0f);
+	FPlatformMemoryStats Stats = FPlatformMemory::GetStats();
+
+	FPlatformMisc::LowLevelOutputDebugStringf(TEXT("EngineMemoryWarningHandler: Mem Used %.2f MB, Texture Memory %.2f MB, Render Target memory %.2f MB, OS Free %.2f MB\n"), 
+		Stats.UsedPhysical / 1048576.0f, 
+		GCurrentTextureMemorySize / 1048576.0f, 
+		GCurrentRendertargetMemorySize / 1048576.0f, 
+		Stats.AvailablePhysical / 1048576.0f);
 
 #if !UE_BUILD_SHIPPING && !UE_BUILD_TEST
 	const auto OOMMemReportVar = IConsoleManager::Get().FindTConsoleVariableDataInt(TEXT("Debug.OOMMemReport")); 
@@ -829,7 +833,7 @@ void UEngine::Init(IEngineLoop* InEngineLoop)
 	// register screenshot capture if we are dumping a movie
 	if(GIsDumpingMovie)
 	{
-		UGameViewportClient::OnScreenshotCaptured().AddUObject(this, &UEngine::HandleScreenshotCaptured);
+		//UGameViewportClient::OnScreenshotCaptured().AddUObject(this, &UEngine::HandleScreenshotCaptured);
 	}
 #endif
 }
@@ -854,7 +858,7 @@ void UEngine::PreExit()
 	FEngineAnalytics::Shutdown();
 
 #if WITH_EDITOR
-	UGameViewportClient::OnScreenshotCaptured().RemoveUObject(this, &UEngine::HandleScreenshotCaptured);
+	//UGameViewportClient::OnScreenshotCaptured().RemoveUObject(this, &UEngine::HandleScreenshotCaptured);
 #endif
 }
 
@@ -2320,7 +2324,7 @@ bool UEngine::Exec( UWorld* InWorld, const TCHAR* Cmd, FOutputDevice& Ar )
 	}
 	else if ( FParse::Command(&Cmd,TEXT("DUMPALLOCS")) )
 	{
-		return HandleDumpAllocsCommand( Cmd, Ar );
+		return HandleDumpAllocatorStats( Cmd, Ar );
 	}
 	else if ( FParse::Command(&Cmd,TEXT("HEAPCHECK")) )
 	{
@@ -3824,14 +3828,14 @@ bool UEngine::HandleMemCommand( const TCHAR* Cmd, FOutputDevice& Ar )
 	}
 
 	FPlatformMemory::DumpStats( Ar );
+	Ar.CategorizedLogf( LogMemory.GetCategoryName(), ELogVerbosity::Log, TEXT("") );
+	GMalloc->DumpAllocatorStats( Ar );
 
 	if( bDetailed || bReport)
 	{
-		GMalloc->DumpAllocations( Ar );
-
-		Ar.Logf( TEXT("Memory Stats:") );
-		Ar.Logf( TEXT("FMemStack (gamethread) allocation size [used/ unused] = [%.2f / %.2f] MB"), FMemStack::Get().GetByteCount() / (1024.0f * 1024.0f), FMemStack::Get().GetUnusedByteCount() / (1024.0f * 1024.0f)  );
-		Ar.Logf( TEXT("Nametable memory usage = %.2f MB"), FName::GetNameTableMemorySize() / (1024.0f * 1024.0f) );
+		Ar.CategorizedLogf( LogMemory.GetCategoryName(), ELogVerbosity::Log, TEXT("Memory Stats:") );
+		Ar.CategorizedLogf( LogMemory.GetCategoryName(), ELogVerbosity::Log, TEXT("FMemStack (gamethread) allocation size [used/ unused] = [%.2f / %.2f] MB"), FMemStack::Get().GetByteCount() / (1024.0f * 1024.0f), FMemStack::Get().GetUnusedByteCount() / (1024.0f * 1024.0f)  );
+		Ar.CategorizedLogf( LogMemory.GetCategoryName(), ELogVerbosity::Log, TEXT("Nametable memory usage = %.2f MB"), FName::GetNameTableMemorySize() / (1024.0f * 1024.0f) );
 
 #if STATS
 		TArray<FStatMessage> Stats;
@@ -3848,33 +3852,12 @@ bool UEngine::HandleMemCommand( const TCHAR* Cmd, FOutputDevice& Ar )
 			FName LastGroup = Meta.NameAndInfo.GetGroupName();
 			if ((LastGroup == NAME_STATGROUP_SceneMemory || LastGroup == NAME_STATGROUP_Memory || LastGroup == NAME_STATGROUP_TextureGroup || LastGroup == NAME_STATGROUP_RHI)  && Meta.NameAndInfo.GetFlag(EStatMetaFlags::IsMemory))
 			{
-				Ar.Logf(TEXT("%s"), *FStatsUtils::DebugPrint(Meta));
+				Ar.CategorizedLogf( LogMemory.GetCategoryName(), ELogVerbosity::Log, TEXT("%s"), *FStatsUtils::DebugPrint(Meta));
 			}
 		}
 #endif
 	}
 
-	FMemoryAllocationStats_DEPRECATED MemStats;
-	GMalloc->GetAllocationInfo( MemStats );
-
-	Ar.Logf( TEXT("") );
-
-	if( MemStats.TotalAllocated || MemStats.TotalUsed )
-	{
-		float AllocatedMB = MemStats.TotalAllocated / ( 1024.0f * 1024.0f );
-		float UsedMB = MemStats.TotalUsed / ( 1024.0f * 1024.0f );
-		float OverheadMB = ( MemStats.TotalAllocated - MemStats.TotalUsed ) / ( 1024.0f * 1024.0f );
-		Ar.Logf(TEXT("Memory allocations reported by the OS: %.2f MB (with %.2f MB used and %.2f MB overhead)"), AllocatedMB, UsedMB, OverheadMB );
-	}
-
-	if( MemStats.CPUUsed )
-	{
-		float UsedMB = MemStats.CPUUsed / ( 1024.0f * 1024.0f );
-		float WasteMB = MemStats.CPUWaste / ( 1024.0f * 1024.0f );
-		float SlackMB = MemStats.CPUSlack / ( 1024.0f * 1024.0f );
-		Ar.Logf(TEXT("Virtual memory tracked in the allocators: %.2f MB (with %.2f MB used, %.2f MB slack and %.2f MB waste)"), UsedMB + SlackMB + WasteMB, UsedMB, SlackMB, WasteMB );
-	}
-
 	return true;
 }
 
@@ -4866,9 +4849,9 @@ bool UEngine::HandleDumpParticleFrameRenderingStatsCommand( const TCHAR* Cmd, FO
 	return 1;
 }
 
-bool UEngine::HandleDumpAllocsCommand( const TCHAR* Cmd, FOutputDevice& Ar )
+bool UEngine::HandleDumpAllocatorStats( const TCHAR* Cmd, FOutputDevice& Ar )
 {
-	GMalloc->DumpAllocations(Ar);
+	GMalloc->DumpAllocatorStats(Ar);
 	return true;
 }
 
@@ -6240,7 +6223,7 @@ int32 DrawMemorySummaryStats( FViewport* Viewport, FCanvas* Canvas, int32 X, int
 
 	// Retrieve allocation info.
 	FPlatformMemoryStats MemoryStats = FPlatformMemory::GetStats();
-	float MemoryInMByte = MemoryStats.WorkingSetSize / 1024.f / 1024.f;
+	float MemoryInMByte = MemoryStats.UsedPhysical / 1024.f / 1024.f;
 
 	// Draw the memory summary stats.
 	Canvas->DrawShadowedString(
diff --git a/Engine/Source/Runtime/Engine/Public/MallocProfilerEx.h b/Engine/Source/Runtime/Engine/Public/MallocProfilerEx.h
--- a/Engine/Source/Runtime/Engine/Public/MallocProfilerEx.h
+++ b/Engine/Source/Runtime/Engine/Public/MallocProfilerEx.h
@@ -35,7 +35,7 @@ public:
 	/** 
 	 * Gather texture memory stats. 
 	 */
-	virtual void GetTexturePoolSize( FMemoryAllocationStats_DEPRECATED& MemoryStats ) OVERRIDE;
+	virtual void GetTexturePoolSize( FGenericMemoryStats& out_Stats ) OVERRIDE;
 };
 
 #endif // USE_MALLOC_PROFILER
diff --git a/Engine/Source/Runtime/Launch/Private/LaunchEngineLoop.cpp b/Engine/Source/Runtime/Launch/Private/LaunchEngineLoop.cpp
--- a/Engine/Source/Runtime/Launch/Private/LaunchEngineLoop.cpp
+++ b/Engine/Source/Runtime/Launch/Private/LaunchEngineLoop.cpp
@@ -1928,31 +1928,9 @@ void FEngineLoop::Tick()
 
 		GEngine->TickFPSChart( GDeltaTime );
 
-#if STATS	// Can't do this within the malloc classes as they get initialized before the stats
-
-		// @todo FPlatformMemory::UpdateStats();
-		FPlatformMemoryStats MemoryStats = FPlatformMemory::GetStats();
-		SET_DWORD_STAT(STAT_VirtualAllocSize,MemoryStats.PagefileUsage);
-		SET_DWORD_STAT(STAT_PhysicalAllocSize,MemoryStats.WorkingSetSize);
-
-		static uint64 LastMallocCalls		= 0;
-		static uint64 LastReallocCalls		= 0;
-		static uint64 LastFreeCalls			= 0;
-
-		uint32 CurrentFrameMallocCalls		= FMalloc::TotalMallocCalls - LastMallocCalls;
-		uint32 CurrentFrameReallocCalls		= FMalloc::TotalReallocCalls - LastReallocCalls;
-		uint32 CurrentFrameFreeCalls			= FMalloc::TotalFreeCalls - LastFreeCalls;
-		uint32 CurrentFrameAllocatorCalls	= CurrentFrameMallocCalls + CurrentFrameReallocCalls + CurrentFrameFreeCalls;
-
-		SET_DWORD_STAT( STAT_MallocCalls, CurrentFrameMallocCalls );
-		SET_DWORD_STAT( STAT_ReallocCalls, CurrentFrameReallocCalls );
-		SET_DWORD_STAT( STAT_FreeCalls, CurrentFrameFreeCalls );
-		SET_DWORD_STAT( STAT_TotalAllocatorCalls, CurrentFrameAllocatorCalls );
-
-		LastMallocCalls			= FMalloc::TotalMallocCalls;
-		LastReallocCalls		= FMalloc::TotalReallocCalls;
-		LastFreeCalls			= FMalloc::TotalFreeCalls;
-#endif
+		// Update platform memory and memory allocator stats.
+		FPlatformMemory::UpdateStats();
+		GMalloc->UpdateStats();
 	} 
 
 	FlushStatsFrame(false);
