#UE4
- Added FScopeLogTime to log time spent in the specified scope, can be used without stats, usage:

{
   SCOPE_LOG_TIME(TEXT("MyTimer"),[&GlobalTotalTimeAndCount]);
   .....
}

diff --git a/Engine/Source/Runtime/Core/Private/Stats/StatsMisc.cpp b/Engine/Source/Runtime/Core/Private/Stats/StatsMisc.cpp
new file mode 100644
--- /dev/null
+++ b/Engine/Source/Runtime/Core/Private/Stats/StatsMisc.cpp
@@ -0,0 +1,26 @@
+// Copyright 1998-2014 Epic Games, Inc. All Rights Reserved.
+
+#include "CorePrivate.h"
+
+FScopeLogTime::FScopeLogTime( const TCHAR* InName, FTotalTimeAndCount* InGlobal /*= nullptr */ ) 
+: StartTime( FPlatformTime::Seconds() )
+, Name( InName )
+, Global( InGlobal )
+{}
+
+FScopeLogTime::~FScopeLogTime()
+{
+	const double ScopedTime = FPlatformTime::Seconds() - StartTime;
+	if( Global )
+	{
+		Global->Key += ScopedTime;
+		Global->Value++;
+
+		const double Average = Global->Key / (double)Global->Value;
+		UE_LOG( LogStats, Log, TEXT( "%32s - %6.3f ms - Total %6.2f s / %5u / %6.3f ms" ), *Name, ScopedTime*1000.0f, Global->Key, Global->Value, Average*1000.0f );
+	}
+	else
+	{
+		UE_LOG( LogStats, Log, TEXT( "%32s - %6.3f ms" ), *Name, ScopedTime*1000.0f );
+	}
+}
diff --git a/Engine/Source/Runtime/Core/Public/Stats/StatsMisc.h b/Engine/Source/Runtime/Core/Public/Stats/StatsMisc.h
--- a/Engine/Source/Runtime/Core/Public/Stats/StatsMisc.h
+++ b/Engine/Source/Runtime/Core/Public/Stats/StatsMisc.h
@@ -39,3 +39,31 @@ private:
 #define SCOPE_SECONDS_COUNTER(Seconds)
 #endif 
 
+/** Key contains total time, value contains total count. */
+typedef TKeyValuePair<double, uint32> FTotalTimeAndCount;
+
+/**
+ *	Utility class to log time passed in seconds, adding cumulative stats to passed in variable. 
+ */
+struct FScopeLogTime
+{
+	/**
+	 * Initialization constructor.
+	 *
+	 * @param InName - String that will be displayed in the log
+	 * @param InGlobal - Pointer to the variable that holds the cumulative stats
+	 *
+	 */
+	FScopeLogTime( const TCHAR* InName, FTotalTimeAndCount* InGlobal = nullptr );
+
+	/** Destructor. */
+	~FScopeLogTime();
+
+protected:
+	const double StartTime;
+	const FString Name;
+	FTotalTimeAndCount* Global;
+};
+
+#define SCOPE_LOG_TIME(Name,Global) \
+	FScopeLogTime ScopeLogTime(Name,Global);
\ No newline at end of file
