#UE4 CIS fix. C2719 problem that appears only on 32-bit build when passing by value. Fixed.

diff --git a/Engine/Source/Runtime/Core/Private/Misc/ConfigCacheIni.cpp b/Engine/Source/Runtime/Core/Private/Misc/ConfigCacheIni.cpp
--- a/Engine/Source/Runtime/Core/Private/Misc/ConfigCacheIni.cpp
+++ b/Engine/Source/Runtime/Core/Private/Misc/ConfigCacheIni.cpp
@@ -2070,7 +2070,7 @@ void FConfigCacheIni::SetVector4
 (
  const TCHAR*		Section,
  const TCHAR*		Key,
- const FVector4		 Value,
+ const FVector4&	 Value,
  const FString&	Filename
 )
 {
diff --git a/Engine/Source/Runtime/Core/Public/Misc/ConfigCacheIni.h b/Engine/Source/Runtime/Core/Public/Misc/ConfigCacheIni.h
--- a/Engine/Source/Runtime/Core/Public/Misc/ConfigCacheIni.h
+++ b/Engine/Source/Runtime/Core/Public/Misc/ConfigCacheIni.h
@@ -439,7 +439,7 @@ public:
 	(
 		const TCHAR*		Section,
 		const TCHAR*		Key,
-		FVector4			Value,
+		const FVector4&		Value,
 		const FString&	Filename
 	);
 	void SetRotator
