#UE4
- Fixed a crash that may happen if stats is used on the other thread

diff --git a/Engine/Source/Runtime/Core/Private/Stats/StatsData.cpp b/Engine/Source/Runtime/Core/Private/Stats/StatsData.cpp
--- a/Engine/Source/Runtime/Core/Private/Stats/StatsData.cpp
+++ b/Engine/Source/Runtime/Core/Private/Stats/StatsData.cpp
@@ -1007,7 +1007,7 @@ void FStatsThreadState::GetRawStackStats(int64 TargetFrame, FRawStatStackNode& R
 	TMap<FName, FStatMessage> ThisFrameNonStackStats;
 
 	for (int32 PacketIndex = 0; PacketIndex < Frame.Packets.Num(); PacketIndex++)
-				{
+	{
 		FStatPacket const& Packet = *Frame.Packets[PacketIndex];
 		const FName ThreadName = GetStatThreadName( Packet );
 
@@ -1215,7 +1215,14 @@ FName FStatsThreadState::GetStatThreadName( const FStatPacket& Packet ) const
 	}
 	else if( Packet.ThreadType == EThreadType::Other )
 	{
-		ThreadName = Threads.FindChecked( Packet.ThreadId );
+		TMap<uint32, FName>& MutableThreads = const_cast<TMap<uint32, FName>&>(Threads);
+		FName& NewThreadName = MutableThreads.FindOrAdd( Packet.ThreadId );
+		if( NewThreadName == NAME_None )
+		{
+			static const FName NAME_UnknownThread = TEXT( "UnknownThread" );
+			NewThreadName = FName( *FStatsUtils::BuildUniqueThreadName( NAME_UnknownThread, Packet.ThreadId ) );
+		}
+		ThreadName = NewThreadName;
 	}
 
 	check( ThreadName != NAME_None );
