#UE4
- Platform memory fixes

diff --git a/Engine/Source/Runtime/Core/Private/Android/AndroidMemory.cpp b/Engine/Source/Runtime/Core/Private/Android/AndroidMemory.cpp
--- a/Engine/Source/Runtime/Core/Private/Android/AndroidMemory.cpp
+++ b/Engine/Source/Runtime/Core/Private/Android/AndroidMemory.cpp
@@ -56,8 +56,8 @@ FPlatformMemoryStats FAndroidPlatformMemory::GetStats()
 
 	MemoryStats.AvailablePhysical = MemoryConstants.TotalPhysical - GetNativeHeapAllocatedSize();
 	MemoryStats.AvailableVirtual = 0;
-	MemoryStats.WorkingSetSize = 0;
-	MemoryStats.PagefileUsage = 0;
+	MemoryStats.UsedPhysical = 0;
+	MemoryStats.UsedVirtual = 0;
 
 	return MemoryStats;
 }
