#UE4
- Platform memory fixes

diff --git a/Engine/Source/Runtime/Core/Public/HAL/MallocBinned.h b/Engine/Source/Runtime/Core/Public/HAL/MallocBinned.h
--- a/Engine/Source/Runtime/Core/Public/HAL/MallocBinned.h
+++ b/Engine/Source/Runtime/Core/Public/HAL/MallocBinned.h
@@ -871,6 +871,9 @@ public:
 		check(MAX_POOLED_ALLOCATION_SIZE - 1 == PoolTable[POOL_COUNT - 1].BlockSize);
 	}
 	
+	virtual ~FMallocBinned()
+	{}
+
 	/**
 	 * Returns if the allocator is guaranteed to be thread-safe and therefore
 	 * doesn't need a unnecessary thread-safety wrapper around it.
