#UE4
- IOS memory fixes?

diff --git a/Engine/Source/Runtime/Core/Private/HAL/MallocBinned.cpp b/Engine/Source/Runtime/Core/Private/HAL/MallocBinned.cpp
--- a/Engine/Source/Runtime/Core/Private/HAL/MallocBinned.cpp
+++ b/Engine/Source/Runtime/Core/Private/HAL/MallocBinned.cpp
@@ -6,8 +6,6 @@
 
 #include "CorePrivate.h"
 
-#if	!FORCE_ANSI_ALLOCATOR
-
 #include "MallocBinned.h"
 #include "MemoryMisc.h"
 
@@ -42,6 +40,4 @@ void FMallocBinned::GetAllocatorStats( FGenericMemoryStats& out_Stats )
 	out_Stats.Add( GET_STATFNAME( STAT_Binned_CurrentAllocs ), CurrentAllocs );
 	out_Stats.Add( GET_STATFNAME( STAT_Binned_SlackCurrent ), SlackCurrent );
 #endif // STATS
-}
-
-#endif // !FORCE_ANSI_ALLOCATOR
\ No newline at end of file
+}
\ No newline at end of file
diff --git a/Engine/Source/Runtime/Core/Public/HAL/MallocBinned.h b/Engine/Source/Runtime/Core/Public/HAL/MallocBinned.h
--- a/Engine/Source/Runtime/Core/Public/HAL/MallocBinned.h
+++ b/Engine/Source/Runtime/Core/Public/HAL/MallocBinned.h
@@ -878,7 +878,7 @@ public:
 	 * Returns if the allocator is guaranteed to be thread-safe and therefore
 	 * doesn't need a unnecessary thread-safety wrapper around it.
 	 */
-	virtual bool IsInternallyThreadSafe() const 
+	virtual bool IsInternallyThreadSafe() const OVERRIDE
 	{ 
 #ifdef USE_INTERNAL_LOCKS
 		return true;
