#UE4.Core
- Added thread metadata for unknown threads

diff --git a/Engine/Source/Runtime/Core/Private/Stats/StatsData.cpp b/Engine/Source/Runtime/Core/Private/Stats/StatsData.cpp
--- a/Engine/Source/Runtime/Core/Private/Stats/StatsData.cpp
+++ b/Engine/Source/Runtime/Core/Private/Stats/StatsData.cpp
@@ -1221,6 +1221,8 @@ FName FStatsThreadState::GetStatThreadName( const FStatPacket& Packet ) const
 		{
 			static const FName NAME_UnknownThread = TEXT( "UnknownThread" );
 			NewThreadName = FName( *FStatsUtils::BuildUniqueThreadName( NAME_UnknownThread, Packet.ThreadId ) );
+			// This is an unknown thread, not created by UE, but still we need the metadata in the system.
+			FStartupMessages::Get().AddThreadMetadata( NAME_UnknownThread, Packet.ThreadId );
 		}
 		ThreadName = NewThreadName;
 	}
