cleanup pass on recent FDateTime changes

diff --git a/Engine/Source/Runtime/Core/Private/Misc/DateTime.cpp b/Engine/Source/Runtime/Core/Private/Misc/DateTime.cpp
--- a/Engine/Source/Runtime/Core/Private/Misc/DateTime.cpp
+++ b/Engine/Source/Runtime/Core/Private/Misc/DateTime.cpp
@@ -132,6 +132,12 @@
 }
 
 
+FString FDateTime::ToIso8601( ) const
+{
+	return ToString(TEXT("%Y-%m-%dT%H:%M:%S.%sZ"));
+}
+
+
 FString FDateTime::ToString( ) const
 {
 	return ToString(TEXT("%Y.%m.%d-%H.%M.%S"));
@@ -142,7 +148,7 @@
 {
 	FString Result;
 
-	if (Format != NULL)
+	if (Format != nullptr)
 	{
 		while (*Format != TCHAR('\0'))
 		{
@@ -240,20 +246,49 @@
 	return FDateTime(Date.Year, Date.Month, Date.Day, Time.Hour, Time.Minute, Time.Second, Time.Millisecond);
 }
 
-FString FDateTime::ToIso8601() const
+
+bool FDateTime::Parse( const FString& DateTimeString, FDateTime& OutDateTime )
 {
-	return ToString(TEXT("%Y-%m-%dT%H:%M:%S.%sZ"));
+	// first replace -, : and . with space
+	FString FixedString = DateTimeString.Replace(TEXT("-"), TEXT(" "));
+
+	FixedString.ReplaceInline(TEXT(":"), TEXT(" "));
+	FixedString.ReplaceInline(TEXT("."), TEXT(" "));
+
+	TArray<FString> Tokens;
+
+	// split up on a single delimiter
+	FixedString.ParseIntoArray(&Tokens, TEXT(" "), true);
+
+	// make sure it parsed it properly (within reason)
+	if ((Tokens.Num() < 6) || (Tokens.Num() > 7))
+	{
+		return false;
+	}
+
+	// convert the tokens to numbers
+	OutDateTime.Ticks = FDateTime(
+		FCString::Atoi(*Tokens[0]), // year
+		FCString::Atoi(*Tokens[1]), // month
+		FCString::Atoi(*Tokens[2]), // day
+		FCString::Atoi(*Tokens[3]), // hour
+		FCString::Atoi(*Tokens[4]), // minute
+		FCString::Atoi(*Tokens[5]), // second
+		Tokens.Num() > 6 ? FCString::Atoi(*Tokens[6]) : 0 // millisecond
+	).GetTicks();
+
+	// @todo gmp: need some better validation here
+	return true;
 }
 
-//static 
-bool FDateTime::FromIso8601( const TCHAR* DateTimeString, FDateTime* OutDateTime )
+
+bool FDateTime::ParseIso8601( const TCHAR* DateTimeString, FDateTime& OutDateTime )
 {
-	check(OutDateTime != NULL);
 	// DateOnly: YYYY-MM-DD
 	// DateTime: YYYY-mm-ddTHH:MM:SS(.ssss)(Z|+th:tm|-th:tm)
 
 	const TCHAR* Ptr = DateTimeString;
-	TCHAR* Next = NULL;
+	TCHAR* Next = nullptr;
 
 	int32 Year = 0, Month = 0, Day = 0;
 	int32 Hour = 0, Minute = 0, Second = 0, Millisecond = 0;
@@ -261,18 +296,23 @@
 
 	// get date
 	Year = FCString::Strtoi(Ptr, &Next, 10);
-	if (Next <= Ptr || *Next == TCHAR('\0'))
+
+	if ((Next <= Ptr) || (*Next == TCHAR('\0')))
 	{
 		return false;
 	}
+
 	Ptr = Next + 1; // skip separator
 	Month = FCString::Strtoi(Ptr, &Next, 10);
-	if (Next <= Ptr || *Next == TCHAR('\0'))
+
+	if ((Next <= Ptr) || (*Next == TCHAR('\0')))
 	{
 		return false;
 	}
+
 	Ptr = Next + 1; // skip separator
 	Day = FCString::Strtoi(Ptr, &Next, 10);
+
 	if (Next <= Ptr)
 	{
 		return false;
@@ -285,18 +325,23 @@
 
 		// parse time
 		Hour = FCString::Strtoi(Ptr, &Next, 10);
-		if (Next <= Ptr || *Next == TCHAR('\0'))
+
+		if ((Next <= Ptr) || (*Next == TCHAR('\0')))
 		{
 			return false;
 		}
+
 		Ptr = Next + 1; // skip separator
 		Minute = FCString::Strtoi(Ptr, &Next, 10);
-		if (Next <= Ptr || *Next == TCHAR('\0'))
+
+		if ((Next <= Ptr) || (*Next == TCHAR('\0')))
 		{
 			return false;
 		}
+
 		Ptr = Next + 1; // skip separator
 		Second = FCString::Strtoi(Ptr, &Next, 10);
+
 		if (Next <= Ptr)
 		{
 			return false;
@@ -307,10 +352,13 @@
 		{
 			Ptr = Next + 1;
 			Millisecond = FCString::Strtoi(Ptr, &Next, 10);
-			if (Next <= Ptr || Next > Ptr + 3) // should be no more than 3 digits
+
+			// should be no more than 3 digits
+			if ((Next <= Ptr) || (Next > Ptr + 3))
 			{
 				return false;
 			}
+
 			for (int32 Digits = Next - Ptr; Digits < 3; ++Digits)
 			{
 				Millisecond *= 10;
@@ -325,18 +373,21 @@
 
 			// parse the timezone offset
 			TzHour = FCString::Strtoi(Ptr, &Next, 10);
-			if (Next <= Ptr || *Next == TCHAR('\0'))
+
+			if ((Next <= Ptr) || (*Next == TCHAR('\0')))
 			{
 				return false;
 			}
+
 			Ptr = Next + 1; // skip separator
 			TzMinute = FCString::Strtoi(Ptr, &Next, 10);
+
 			if (Next <= Ptr)
 			{
 				return false;
 			}
 		}
-		else if (*Next != TCHAR('\0') && *Next != TCHAR('Z'))
+		else if ((*Next != TCHAR('\0')) && (*Next != TCHAR('Z')))
 		{
 			return false;
 		}
@@ -346,49 +397,13 @@
 		return false;
 	}
 
-	// convert these times into a DateTime
 	FDateTime Final(Year, Month, Day, Hour, Minute, Second, Millisecond);
 
 	// adjust for the timezone (bringing the DateTime into UTC)
 	int32 TzOffsetMinutes = (TzHour < 0) ? TzHour * 60 - TzMinute : TzHour * 60 + TzMinute;
 	Final -= FTimespan(0, TzOffsetMinutes, 0);
+	OutDateTime = Final;
 
-	*OutDateTime = Final;
 	return true;
 }
 
-
-bool FDateTime::Parse( const FString& DateTimeString, FDateTime& OutDateTime )
-{
-	// first replace -, : and . with space
-	FString FixedString = DateTimeString.Replace(TEXT("-"), TEXT(" "));
-
-	FixedString.ReplaceInline(TEXT(":"), TEXT(" "));
-	FixedString.ReplaceInline(TEXT("."), TEXT(" "));
-
-	TArray<FString> Tokens;
-
-	// split up on a single delimiter
-	FixedString.ParseIntoArray(&Tokens, TEXT(" "), true);
-
-	// make sure it parsed it properly (within reason)
-	if ((Tokens.Num() < 6) || (Tokens.Num() > 7))
-	{
-		return false;
-	}
-
-	// convert the tokens to numbers
-	OutDateTime.Ticks = FDateTime(
-		FCString::Atoi(*Tokens[0]), // year
-		FCString::Atoi(*Tokens[1]), // month
-		FCString::Atoi(*Tokens[2]), // day
-		FCString::Atoi(*Tokens[3]), // hour
-		FCString::Atoi(*Tokens[4]), // minute
-		FCString::Atoi(*Tokens[5]), // second
-		Tokens.Num() > 6 ? FCString::Atoi(*Tokens[6]) : 0 // millisecond
-	).GetTicks();
-
-	// @todo gmp: need some better validation here
-	return true;
-}
-
diff --git a/Engine/Source/Runtime/Core/Public/Misc/DateTime.h b/Engine/Source/Runtime/Core/Public/Misc/DateTime.h
--- a/Engine/Source/Runtime/Core/Public/Misc/DateTime.h
+++ b/Engine/Source/Runtime/Core/Public/Misc/DateTime.h
@@ -406,20 +406,20 @@ public:
 	 *
 	 * @return The month.
 	 */
-	 int32 GetMonth( ) const
-	 {
-		 return ToDate().Month;
-	 }
-
-	 /**
-	  * Gets the date's month of the year (January to December).
-	  *
-	  * @return Month of year.
-	  */
-	 EMonthOfYear::Type GetMonthOfYear( ) const
-	 {
-		 return static_cast<EMonthOfYear::Type>(GetMonth());
-	 }
+	int32 GetMonth( ) const
+	{
+		return ToDate().Month;
+	}
+
+	/**
+	 * Gets the date's month of the year (January to December).
+	 *
+	 * @return Month of year.
+	 */
+	EMonthOfYear::Type GetMonthOfYear( ) const
+	{
+		return static_cast<EMonthOfYear::Type>(GetMonth());
+	}
 
 	/**
 	 * Gets this date's second part.
@@ -486,6 +486,17 @@ public:
 	 */
 	CORE_API FDate ToDate( ) const;
 
+	/**
+	 * Returns the ISO-8601 string representation of the FDateTime.
+	 *
+	 * The resulting string assumes that the FDateTime is in UTC.
+	 * 
+	 * @return String representation.
+	 *
+	 * @see FromIso8601
+	 */
+	CORE_API FString ToIso8601( ) const;
+
 	/**
 	 * Returns the string representation of this date using a default format.
 	 *
@@ -496,15 +507,6 @@ public:
 	 */
 	CORE_API FString ToString( ) const;
 
-	/**
-	 * Returns the ISO-8601 string representation of the FDateTime.
-	 * The resulting string assumes that the FDateTime is in UTC since this class
-	 * doesn't contain timezone offset information.
-	 * 
-	 * @return String representation.
-	 */
-	CORE_API FString ToIso8601() const;
-
 	/**
 	 * Returns the string representation of this date.
 	 *
@@ -578,16 +580,6 @@ public:
 		return FDateTime(1970, 1, 1) + FTimespan(static_cast<int64>(UnixTime) * ETimespan::TicksPerSecond);
 	}
 
-	/**
-	 * Parses a stringified DateTime in ISO-8601 format
-	 * 
-	 * @param DateTimeString The string to be parsed
-	 * @param OutDateTime FDateTime object (in UTC) corresponding to the input string (which may have been in any timezone).
-	 *
-	 * @return success/fail
-	 */
-	static CORE_API bool FromIso8601( const TCHAR* DateTimeString, FDateTime* OutDateTime );
-
 	/**
 	 * Checks whether the given year is a leap year.
 	 *
@@ -652,6 +644,18 @@ public:
 	 */
 	static CORE_API bool Parse( const FString& DateTimeString, FDateTime& OutDateTime );
 
+	/**
+	 * Parses a date string in ISO-8601 format.
+	 * 
+	 * @param DateTimeString The string to be parsed
+	 * @param OutDateTime FDateTime object (in UTC) corresponding to the input string (which may have been in any timezone).
+	 *
+	 * @return true if the string was converted successfully, false otherwise.
+	 *
+	 * @see ToIso8601
+	 */
+	static CORE_API bool ParseIso8601( const TCHAR* DateTimeString, FDateTime& OutDateTime );
+
 	/**
 	 * Gets the local date on this computer.
 	 *
