TTP 300686

The default splash screen now displays a clearer indication that it's loading a game project by using a splash screen that says [Project Nam] Editor.  If we don't have a project name then the old splash screens are used that just says Unreal Editor, or Unreal Engine 4 when launching the game.  For launching the game it will say "Project Name" instead of "Unreal Engine 4".  I've updated the Windows and Mac splash screens to be roughly the same.

This change should help reduce confusion when the editor restarts to load a project, by clearly indicating what is loading.

diff --git a/Engine/Source/Runtime/Core/Private/Mac/MacPlatformSplash.cpp b/Engine/Source/Runtime/Core/Private/Mac/MacPlatformSplash.cpp
--- a/Engine/Source/Runtime/Core/Private/Mac/MacPlatformSplash.cpp
+++ b/Engine/Source/Runtime/Core/Private/Mac/MacPlatformSplash.cpp
@@ -38,32 +38,34 @@ static NSRect GSplashScreenTextRects[ SplashTextType::NumTextTypes ];
 
 /**
  * Finds a usable splash pathname for the given filename
- * 
+ *
  * @param SplashFilename Name of the desired splash name ("Splash.bmp")
  * @param OutPath String containing the path to the file, if this function returns true
  *
  * @return true if a splash screen was found
  */
-static bool GetSplashPath(const TCHAR* SplashFilename, FString& OutPath)
+static bool GetSplashPath(const TCHAR* SplashFilename, FString& OutPath, bool& OutIsCustom)
 {
 	// first look in game's splash directory
 	OutPath = FPaths::GameContentDir() + TEXT("Splash/") + SplashFilename;
+	OutIsCustom = true;
 	
 	// if this was found, then we're done
 	if (IFileManager::Get().FileSize(*OutPath) != -1)
 	{
 		return true;
 	}
-
+	
 	// next look in Engine/Splash
 	OutPath = FPaths::EngineContentDir() + TEXT("Splash/") + SplashFilename;
-
+	OutIsCustom = false;
+	
 	// if this was found, then we're done
 	if (IFileManager::Get().FileSize(*OutPath) != -1)
 	{
 		return true;
 	}
-
+	
 	// if not found yet, then return failure
 	return false;
 }
@@ -76,12 +78,8 @@ static bool GetSplashPath(const TCHAR* SplashFilename, FString& OutPath)
  */
 static void StartSetSplashText( const SplashTextType::Type InType, const FText& InText )
 {
-	// Only allow copyright text displayed while loading the game.  Editor displays all.
-	if( InType == SplashTextType::CopyrightInfo || GIsEditor )
-	{
-		// Update splash text
-		GSplashScreenText[ InType ] = InText;
-	}
+	// Update splash text
+	GSplashScreenText[ InType ] = InText;
 }
 
 @interface UE4SplashView : NSView
@@ -89,7 +87,7 @@ static void StartSetSplashText( const SplashTextType::Type InType, const FText&
 }
 
 - (void)drawRect: (NSRect)DirtyRect;
-- (void)drawText: (NSString *)Text inRect: (NSRect)Rect withColor: (NSColor *)Color fontSize: (int32)FontSize;
+- (void)drawText: (NSString *)Text inRect: (NSRect)Rect withAlignment: (NSTextAlignment)align withColor: (NSColor *)Color fontName: (NSString *)FontName fontSize: (int32)FontSize;
 
 @end
 
@@ -109,43 +107,74 @@ static void StartSetSplashText( const SplashTextType::Type InType, const FText&
 
 		if( SplashText.ToString().Len() > 0 )
 		{
-			int32 FontSize = (CurTypeIndex == SplashTextType::StartupProgress || CurTypeIndex == SplashTextType::VersionInfo1) ? 12 : 11;
+			int32 FontSize = 11;
+			switch ( CurTypeIndex )
+			{
+				case SplashTextType::StartupProgress:
+				case SplashTextType::VersionInfo1:
+					FontSize = 12;
+					break;
+				case SplashTextType::GameName:
+					FontSize = 34;
+					break;
+			}
 
-			float Brightness;
+			float Brightness = 160.0f / 255.0f;
 
-			if( CurTypeIndex == SplashTextType::StartupProgress )
+			switch ( CurTypeIndex )
 			{
-				Brightness = 180.0f / 255.0f;
+				case SplashTextType::StartupProgress:
+					Brightness = 180.0f / 255.0f;
+					break;
+				case SplashTextType::VersionInfo1:
+					Brightness = 240.0f / 255.0f;
+					break;
+				case SplashTextType::GameName:
+					Brightness = 240.0f / 255.0f;
+					break;
 			}
-			else if( CurTypeIndex == SplashTextType::VersionInfo1 )
+
+			NSColor *TextColor = [NSColor colorWithDeviceRed: Brightness green: Brightness blue: Brightness alpha: 1.0f];
+
+			NSString* FontName = @"Helvetica-Bold";
+			switch ( CurTypeIndex )
 			{
-				Brightness = 240.0f / 255.0f;
+			case SplashTextType::GameName:
+				FontName = @"Verdana-Bold";
+				break;
 			}
-			else
+
+			// Alignment
+			NSTextAlignment align = NSLeftTextAlignment;
+			switch ( CurTypeIndex )
 			{
-				Brightness = 160.0f / 255.0f;
+			case SplashTextType::GameName:
+				align = NSRightTextAlignment;
+				break;
 			}
 
-			NSColor *TextColor = [NSColor colorWithDeviceRed: Brightness green: Brightness blue: Brightness alpha: 1.0f];
-
 			NSString *Text = (NSString *)FPlatformString::TCHARToCFString(*SplashText.ToString());
-			[self drawText: Text inRect: TextRect withColor: TextColor fontSize: FontSize];
+			[self drawText: Text inRect: TextRect withAlignment: align withColor: TextColor fontName: FontName fontSize: FontSize];
 			[Text release];
 		}
 	}
 }
 
-- (void)drawText: (NSString *)Text inRect: (NSRect)Rect withColor: (NSColor *)Color fontSize: (int32)FontSize
+- (void)drawText: (NSString *)Text inRect: (NSRect)Rect withAlignment: (NSTextAlignment)align withColor: (NSColor *)Color fontName: (NSString *)FontName fontSize: (int32)FontSize
 {
 	SCOPED_AUTORELEASE_POOL;
 
+	NSMutableParagraphStyle *style =[[NSParagraphStyle defaultParagraphStyle] mutableCopy];
+	[style setAlignment : align];
+
 	NSDictionary* Dict =
 		[NSDictionary dictionaryWithObjects:
 			[NSArray arrayWithObjects:
 				Color,
-				[NSFont fontWithName:@"Helvetica-Bold" size: FontSize],
+				[NSFont fontWithName: FontName size: FontSize],
 				[NSColor colorWithDeviceRed: 0.0 green: 0.0 blue: 0.0 alpha: 1.0],
 				[NSNumber numberWithFloat:-4.0],
+				style,
 				nil ]
 		forKeys:
 			[NSArray arrayWithObjects:
@@ -153,6 +182,7 @@ static void StartSetSplashText( const SplashTextType::Type InType, const FText&
 				NSFontAttributeName,
 				NSStrokeColorAttributeName,
 				NSStrokeWidthAttributeName,
+				NSParagraphStyleAttributeName,
 				nil]
 		];
 	[Text drawInRect: Rect withAttributes: Dict];
@@ -166,13 +196,21 @@ void FMacPlatformSplash::Show()
 	{
 		SCOPED_AUTORELEASE_POOL;
 
-		const TCHAR* SplashImage =
-				( GIsEditor ? TEXT("EdSplash.bmp") : TEXT("Splash.bmp") );
+		const FText GameName = FText::FromString(FApp::GetGameName());
+
+		const TCHAR* SplashImage = GIsEditor ? ( GameName.IsEmpty() ? TEXT("EdSplashDefault.bmp") : TEXT("EdSplash.bmp") ) : ( GameName.IsEmpty() ? TEXT("SplashDefault.bmp") : TEXT("Splash.bmp") );
 
 		// make sure a splash was found
 		FString SplashPath;
-		if (GetSplashPath( SplashImage, SplashPath ) == true)
+		bool IsCustom;
+		if ( GetSplashPath(SplashImage, SplashPath, IsCustom) == true )
 		{
+			// Don't set the game name if the splash screen is custom.
+			if ( !IsCustom )
+			{
+				StartSetSplashText( SplashTextType::GameName, GameName );
+			}
+
 			// In the editor, we'll display loading info
 			if( GIsEditor )
 			{
@@ -229,6 +267,12 @@ void FMacPlatformSplash::Show()
 		int32 OriginY = 6;
 		int32 FontHeight = 14;
 
+		// Setup bounds for game name
+		GSplashScreenTextRects[ SplashTextType::GameName ].origin.x = 10;
+		GSplashScreenTextRects[ SplashTextType::GameName ].origin.y = 0;
+		GSplashScreenTextRects[ SplashTextType::GameName ].size.width = ImageWidth - 2 * 10;
+		GSplashScreenTextRects[ SplashTextType::GameName ].size.height = ImageHeight;
+
 		// Setup bounds for texts
 		GSplashScreenTextRects[ SplashTextType::VersionInfo1 ].origin.x =
 		GSplashScreenTextRects[ SplashTextType::CopyrightInfo ].origin.x =
diff --git a/Engine/Source/Runtime/Core/Private/Windows/WindowsPlatformSplash.cpp b/Engine/Source/Runtime/Core/Private/Windows/WindowsPlatformSplash.cpp
--- a/Engine/Source/Runtime/Core/Private/Windows/WindowsPlatformSplash.cpp
+++ b/Engine/Source/Runtime/Core/Private/Windows/WindowsPlatformSplash.cpp
@@ -10,6 +10,8 @@
 #include "AllowWindowsPlatformTypes.h"
 #include "EngineBuildSettings.h"
 
+#include <strsafe.h>
+
 #include "wincodec.h"
 
 #pragma comment( lib, "windowscodecs.lib" )
@@ -28,6 +30,7 @@ static FText GSplashScreenText[ SplashTextType::NumTextTypes ];
 static RECT GSplashScreenTextRects[ SplashTextType::NumTextTypes ];
 static HFONT GSplashScreenSmallTextFontHandle = NULL;
 static HFONT GSplashScreenNormalTextFontHandle = NULL;
+static HFONT GSplashScreenTitleTextFontHandle = NULL;
 static FCriticalSection GSplashScreenSynchronizationObject;
 
 
@@ -61,15 +64,29 @@ LRESULT CALLBACK SplashScreenWindowProc(HWND hWnd, uint32 message, WPARAM wParam
 
 						if( !SplashText.IsEmpty() )
 						{
-							if( CurTypeIndex == SplashTextType::StartupProgress ||
-								CurTypeIndex == SplashTextType::VersionInfo1 )
+							if ( CurTypeIndex == SplashTextType::VersionInfo1 || CurTypeIndex == SplashTextType::StartupProgress )
 							{
 								SelectObject( hdc, GSplashScreenNormalTextFontHandle );
 							}
+							else if ( CurTypeIndex == SplashTextType::GameName )
+							{
+								SelectObject( hdc, GSplashScreenTitleTextFontHandle );
+							}
 							else
 							{
 								SelectObject( hdc, GSplashScreenSmallTextFontHandle );
 							}
+
+							// Alignment
+							if ( CurTypeIndex == SplashTextType::GameName )
+							{
+								SetTextAlign( hdc, TA_RIGHT | TA_TOP | TA_NOUPDATECP );
+							}
+							else
+							{
+								SetTextAlign( hdc, TA_LEFT | TA_TOP | TA_NOUPDATECP );
+							}
+
 							SetBkColor( hdc, 0x00000000 );
 							SetBkMode( hdc, TRANSPARENT );
 
@@ -102,20 +119,25 @@ LRESULT CALLBACK SplashScreenWindowProc(HWND hWnd, uint32 message, WPARAM wParam
 									*SplashText.ToString(),
 									SplashText.ToString().Len() );
 							}
-
+							
 							// Draw foreground text pass
 							if( CurTypeIndex == SplashTextType::StartupProgress )
 							{
-								SetTextColor( hdc, RGB( 180, 180, 180 ) );
+								SetTextColor( hdc, RGB( 200, 200, 200 ) );
 							}
 							else if( CurTypeIndex == SplashTextType::VersionInfo1 )
 							{
 								SetTextColor( hdc, RGB( 240, 240, 240 ) );
 							}
+							else if ( CurTypeIndex == SplashTextType::GameName )
+							{
+								SetTextColor(hdc, RGB(240, 240, 240));
+							}
 							else
 							{
 								SetTextColor( hdc, RGB( 160, 160, 160 ) );
 							}
+
 							TextOut(
 								hdc,
 								TextRect.left,
@@ -229,7 +251,7 @@ HBITMAP LoadSplashBitmap()
 		UINT height = 0;
 
 		void *ImageBits = NULL;
-    
+	
 		// Check BitmapSource format
 		WICPixelFormatGUID pixelFormat;
 		hr = ToRenderBitmapSource->GetPixelFormat(&pixelFormat);
@@ -285,7 +307,7 @@ HBITMAP LoadSplashBitmap()
 			// Size of a scan line represented in bytes: 4 bytes each pixel
 			hr = UIntMult(width, sizeof(DWORD), &cbStride);
 		}
-    
+	
 		UINT cbImage = 0;
 		if (SUCCEEDED(hr))
 		{
@@ -425,7 +447,30 @@ uint32 WINAPI StartSplashScreenThread( LPVOID unused )
 					GSplashScreenNormalTextFontHandle = SystemFontHandle;
 				}
 			}
+
+			// Create title font
+			{
+				LOGFONT MyFont;
+				FMemory::Memzero(&MyFont, sizeof( MyFont ));
+				GetObject(SystemFontHandle, sizeof( MyFont ), &MyFont);
+				MyFont.lfHeight = 40;
+				MyFont.lfWeight = FW_BOLD;
+				MyFont.lfQuality = ANTIALIASED_QUALITY;
+				StringCchCopy(MyFont.lfFaceName, LF_FACESIZE, TEXT("Verdana"));
+				GSplashScreenTitleTextFontHandle = CreateFontIndirect(&MyFont);
+				if ( GSplashScreenTitleTextFontHandle == NULL )
+				{
+					// Couldn't create font, so just use a system font
+					GSplashScreenTitleTextFontHandle = SystemFontHandle;
+				}
+			}
 		}
+
+		// Setup bounds for game name
+		GSplashScreenTextRects[ SplashTextType::GameName ].top = 10;
+		GSplashScreenTextRects[ SplashTextType::GameName ].bottom = 60;
+		GSplashScreenTextRects[ SplashTextType::GameName ].left = bm.bmWidth - 12;
+		GSplashScreenTextRects[ SplashTextType::GameName ].right = 12;
 		
 		// Setup bounds for version info text 1
 		GSplashScreenTextRects[ SplashTextType::VersionInfo1 ].top = bm.bmHeight - 60;
@@ -518,10 +563,11 @@ uint32 WINAPI StartSplashScreenThread( LPVOID unused )
  *
  * @return true if a splash screen was found
  */
-static bool GetSplashPath(const TCHAR* SplashFilename, FString& OutPath)
+static bool GetSplashPath(const TCHAR* SplashFilename, FString& OutPath, bool& OutIsCustom)
 {
 	// first look in game's splash directory
 	OutPath = FPaths::GameContentDir() + TEXT("Splash/") + SplashFilename;
+	OutIsCustom = true;
 	
 	// if this was found, then we're done
 	if (IFileManager::Get().FileSize(*OutPath) != -1)
@@ -531,6 +577,7 @@ static bool GetSplashPath(const TCHAR* SplashFilename, FString& OutPath)
 
 	// next look in Engine/Splash
 	OutPath = FPaths::EngineContentDir() + TEXT("Splash/") + SplashFilename;
+	OutIsCustom = false;
 
 	// if this was found, then we're done
 	if (IFileManager::Get().FileSize(*OutPath) != -1)
@@ -551,23 +598,28 @@ static bool GetSplashPath(const TCHAR* SplashFilename, FString& OutPath)
 static void StartSetSplashText( const SplashTextType::Type InType, const TCHAR* InText )
 {
 	// Only allow copyright text displayed while loading the game.  Editor displays all.
-	if( InType == SplashTextType::CopyrightInfo || GIsEditor )
-	{
-		// Update splash text
-		GSplashScreenText[ InType ] = FText::FromString( InText );
-	}
+	GSplashScreenText[InType] = FText::FromString(InText);
 }
 
 void FWindowsPlatformSplash::Show()
 {
 	if( !GSplashScreenThread && FParse::Param(FCommandLine::Get(),TEXT("NOSPLASH")) != true )
 	{
-		const TCHAR* SplashImage = GIsEditor ? TEXT("EdSplash.bmp") : TEXT("Splash.bmp");
+		const FText GameName = FText::FromString( FApp::GetGameName() );
+
+		const TCHAR* SplashImage = GIsEditor ? ( GameName.IsEmpty() ? TEXT("EdSplashDefault.bmp") : TEXT("EdSplash.bmp") ) : ( GameName.IsEmpty() ? TEXT("SplashDefault.bmp") : TEXT("Splash.bmp") );
 
 		// make sure a splash was found
 		FString SplashPath;
-		if (GetSplashPath( SplashImage, SplashPath ) == true)
+		bool IsCustom;
+		if ( GetSplashPath(SplashImage, SplashPath, IsCustom ) == true )
 		{
+			// Don't set the game name if the splash screen is custom.
+			if ( !IsCustom )
+			{
+				StartSetSplashText(SplashTextType::GameName, *GameName.ToString());
+			}
+
 			// In the editor, we'll display loading info
 			if( GIsEditor )
 			{
@@ -585,8 +637,6 @@ void FWindowsPlatformSplash::Show()
 #else	//PLATFORM_64BITS
 					const FText PlatformBits = FText::FromString( TEXT( "32" ) );
 #endif	//PLATFORM_64BITS
-
-					const FText GameName = FText::FromString( FApp::GetGameName() );
 					
 					const FText Version = FText::FromString( GEngineVersion.ToString( FEngineBuildSettings::IsPerforceBuild() ? EVersionComponent::Branch : EVersionComponent::Patch ) );
 
diff --git a/Engine/Source/Runtime/Core/Public/GenericPlatform/GenericPlatformSplash.h b/Engine/Source/Runtime/Core/Public/GenericPlatform/GenericPlatformSplash.h
--- a/Engine/Source/Runtime/Core/Public/GenericPlatform/GenericPlatformSplash.h
+++ b/Engine/Source/Runtime/Core/Public/GenericPlatform/GenericPlatformSplash.h
@@ -22,6 +22,9 @@ namespace SplashTextType
 		/** Copyright information text */
 		CopyrightInfo,
 
+		/** Game Name */
+		GameName,
+
 		// ...
 
 		/** Number of text types (must be final enum value) */
