#UE4 - Don't allow a Unicode string with a 0xffff character to be un-archived. Fixes TTP #330668.

diff --git a/Engine/Source/Runtime/Core/Private/Containers/String.cpp b/Engine/Source/Runtime/Core/Private/Containers/String.cpp
--- a/Engine/Source/Runtime/Core/Private/Containers/String.cpp
+++ b/Engine/Source/Runtime/Core/Private/Containers/String.cpp
@@ -1148,6 +1148,16 @@ FArchive& operator<<( FArchive& Ar, FString& A )
 				Passthru.Apply();
 
 				INTEL_ORDER_TCHARARRAY(A.Data.GetTypedData())
+
+				// Since Microsoft's vsnwprintf implementation raises an invalid parameter warning
+				// with a character of 0xffff, scan for it and terminate the string there.
+				// 0xffff isn't an actual Unicode character anyway.
+				int Index = 0;
+				if(A.FindChar(0xffff, Index))
+				{
+					A[Index] = '\0';
+					A.TrimToNullTerminator();		
+				}
 			}
 			else
 			{
