#UE4 - Shipping build fix

diff --git a/Engine/Source/Runtime/Core/Private/GenericPlatform/GenericPlatformMemory.cpp b/Engine/Source/Runtime/Core/Private/GenericPlatform/GenericPlatformMemory.cpp
--- a/Engine/Source/Runtime/Core/Private/GenericPlatform/GenericPlatformMemory.cpp
+++ b/Engine/Source/Runtime/Core/Private/GenericPlatform/GenericPlatformMemory.cpp
@@ -128,12 +128,17 @@ void FGenericPlatformMemory::DumpStats( class FOutputDevice& Ar )
 {
 	const float InvMB = 1.0f / 1024.0f / 1024.0f;
 	FPlatformMemoryStats MemoryStats = FPlatformMemory::GetStats();
-	Ar.CategorizedLogf( LogMemory.GetCategoryName(), ELogVerbosity::Log, TEXT("Platform Memory Stats for %s"), ANSI_TO_TCHAR(FPlatformProperties::PlatformName()) );
-	Ar.CategorizedLogf( LogMemory.GetCategoryName(), ELogVerbosity::Log, TEXT("Process Physical Memory: %.2f MB used, %.2f MB peak"), MemoryStats.UsedPhysical*InvMB, MemoryStats.PeakUsedPhysical*InvMB );
-	Ar.CategorizedLogf( LogMemory.GetCategoryName(), ELogVerbosity::Log, TEXT("Process Virtual Memory: %.2f MB used, %.2f MB peak"), MemoryStats.UsedVirtual*InvMB, MemoryStats.PeakUsedVirtual*InvMB );
-
-	Ar.CategorizedLogf( LogMemory.GetCategoryName(), ELogVerbosity::Log, TEXT("Physical Memory: %.2f MB used, %.2f MB total"), (MemoryStats.TotalPhysical - MemoryStats.AvailablePhysical)*InvMB, MemoryStats.TotalPhysical*InvMB );
-	Ar.CategorizedLogf( LogMemory.GetCategoryName(), ELogVerbosity::Log, TEXT("Virtual Memory: %.2f MB used, %.2f MB total"), (MemoryStats.TotalVirtual - MemoryStats.AvailableVirtual)*InvMB, MemoryStats.TotalVirtual*InvMB );
+#if !NO_LOGGING
+	const FName CategoryName(LogMemory.GetCategoryName());
+#else
+	const FName CategoryName(TEXT("LogMemory"));
+#endif
+	Ar.CategorizedLogf(CategoryName, ELogVerbosity::Log, TEXT("Platform Memory Stats for %s"), ANSI_TO_TCHAR(FPlatformProperties::PlatformName()));
+	Ar.CategorizedLogf(CategoryName, ELogVerbosity::Log, TEXT("Process Physical Memory: %.2f MB used, %.2f MB peak"), MemoryStats.UsedPhysical*InvMB, MemoryStats.PeakUsedPhysical*InvMB);
+	Ar.CategorizedLogf(CategoryName, ELogVerbosity::Log, TEXT("Process Virtual Memory: %.2f MB used, %.2f MB peak"), MemoryStats.UsedVirtual*InvMB, MemoryStats.PeakUsedVirtual*InvMB);
+
+	Ar.CategorizedLogf(CategoryName, ELogVerbosity::Log, TEXT("Physical Memory: %.2f MB used, %.2f MB total"), (MemoryStats.TotalPhysical - MemoryStats.AvailablePhysical)*InvMB, MemoryStats.TotalPhysical*InvMB);
+	Ar.CategorizedLogf(CategoryName, ELogVerbosity::Log, TEXT("Virtual Memory: %.2f MB used, %.2f MB total"), (MemoryStats.TotalVirtual - MemoryStats.AvailableVirtual)*InvMB, MemoryStats.TotalVirtual*InvMB);
 
 }
 
diff --git a/Engine/Source/Runtime/SkillSystem/Private/SkillSystemLog.cpp b/Engine/Source/Runtime/SkillSystem/Private/SkillSystemLog.cpp
--- a/Engine/Source/Runtime/SkillSystem/Private/SkillSystemLog.cpp
+++ b/Engine/Source/Runtime/SkillSystem/Private/SkillSystemLog.cpp
@@ -26,6 +26,7 @@ FString SkillSystemLog::Log(ELogVerbosity::Type Verbosity, FString Log)
 {
 	SkillSystemLog * Instance = GetInstance();
 
+#if !NO_LOGGING
 	if (!LogSkillSystem.IsSuppressed(Verbosity))
 	{
 		for (int32 idx=0; idx < Instance->ScopeStack.Num(); ++idx)
@@ -46,6 +47,7 @@ FString SkillSystemLog::Log(ELogVerbosity::Type Verbosity, FString Log)
 			}
 		}
 	}
+#endif
 
 	FString IndentStr = FString::Printf(TEXT("%*s"), Instance->Indent, TEXT(""));
 	
