#ue4
- Prevented delayed success adverts from being shown if they succeed after the app has requested hiding them

diff --git a/Engine/Source/Runtime/Core/Private/IOS/IOSAppDelegate.cpp b/Engine/Source/Runtime/Core/Private/IOS/IOSAppDelegate.cpp
--- a/Engine/Source/Runtime/Core/Private/IOS/IOSAppDelegate.cpp
+++ b/Engine/Source/Runtime/Core/Private/IOS/IOSAppDelegate.cpp
@@ -456,8 +456,8 @@ void InstallSignalHandlers()
 /** TRUE if the iAd banner should be on the bottom of the screen */
 BOOL bDrawOnBottom;
 
-/** TRUE when the banner is onscreen */
-BOOL bIsBannerVisible = NO;
+/** true if the user wants the banner to be displayed */
+bool bWantVisibleBanner = false;
 
 /**
 * Will show an iAd on the top or bottom of screen, on top of the GL view (doesn't resize
@@ -468,6 +468,7 @@ BOOL bIsBannerVisible = NO;
 -(void)ShowAdBanner:(NSNumber*)bShowOnBottomOfScreen
 {
 	bDrawOnBottom = [bShowOnBottomOfScreen boolValue];
+	bWantVisibleBanner = true;
 
 	bool bNeedsToAddSubview = false;
 	if (self.BannerView == nil)
@@ -506,7 +507,7 @@ BOOL bIsBannerVisible = NO;
 {
 	NSLog(@"Ad loaded!");
     
-	if (self.BannerView.hidden)
+	if (self.BannerView.hidden && bWantVisibleBanner)
     {
 		self.BannerView.hidden = NO;
 		[UIView animateWithDuration:0.4f
@@ -585,12 +586,19 @@ BOOL bIsBannerVisible = NO;
     }
 }
 
+-(void)UserHideAdBanner
+{
+	bWantVisibleBanner = false;
+	[self HideAdBanner];
+}
+
 /**
 * Forces closed any displayed ad. Can lead to loss of revenue
 */
 -(void)CloseAd
 {
 	// boot user out of the ad
+	bWantVisibleBanner = false;
 	[self.BannerView cancelBannerViewAction];
 }
 
@@ -611,7 +619,7 @@ CORE_API void IOSShowAdBanner(bool bShowOnBottomOfScreen)
 */
 CORE_API void IOSHideAdBanner()
 {
-	[[IOSAppDelegate GetDelegate] performSelectorOnMainThread:@selector(HideAdBanner) withObject:nil waitUntilDone : NO];
+	[[IOSAppDelegate GetDelegate] performSelectorOnMainThread:@selector(UserHideAdBanner) withObject:nil waitUntilDone : NO];
 }
 
 /**
