#ue4
- Wrapped some NSLog calls in if !NO_LOGGING

diff --git a/Engine/Source/Runtime/Core/Private/IOS/IOSAppDelegate.cpp b/Engine/Source/Runtime/Core/Private/IOS/IOSAppDelegate.cpp
--- a/Engine/Source/Runtime/Core/Private/IOS/IOSAppDelegate.cpp
+++ b/Engine/Source/Runtime/Core/Private/IOS/IOSAppDelegate.cpp
@@ -357,7 +357,9 @@ void InstallSignalHandlers()
 
 - (BOOL)application:(UIApplication *)application handleOpenURL:(NSURL *)url
 {
+#if !NO_LOGGING
 	NSLog(@"%s", "IOSAppDelegate handleOpenURL\n");
+#endif
 
 	NSString* EncdodedURLString = [url absoluteString];
 	NSString* URLString = [EncdodedURLString stringByReplacingPercentEscapesUsingEncoding:NSUTF8StringEncoding];
@@ -505,7 +507,9 @@ bool bWantVisibleBanner = false;
 
 -(void)bannerViewDidLoadAd:(ADBannerView*)Banner
 {
+#if !NO_LOGGING
 	NSLog(@"Ad loaded!");
+#endif
     
 	if (self.BannerView.hidden && bWantVisibleBanner)
     {
@@ -521,7 +525,9 @@ bool bWantVisibleBanner = false;
 
 -(void)bannerView:(ADBannerView *)Banner didFailToReceiveAdWithError : (NSError *)Error
 {
+#if !NO_LOGGING
 	NSLog(@"Ad failed to load: '%@'", Error);
+#endif
 
 	// if we get an error, hide the banner 
 	[self HideAdBanner];
