#ttp 330935: CORE: TWeakPtr does not implicitly convert to base types

#add Added constructors and assignment operators to allow TWeakPtr<Derived> to implicitly convert to TWeakPtr<Base>.
#add Added same for TWeakObjectPtr, except with static_asserts to ensure that the conversions are valid, because everything 'just works' when it shouldn't as it's all UObject under the hood.
#add Added == and != operators to allow comparison of compatible TWeakObjectPtrs, again with static_asserts to ensure that it's legal.
#add Added nullptr comparisons to TWeakPtr and TWeakObjectPtrs.

diff --git a/Engine/Source/Runtime/Core/Public/Templates/SharedPointer.h b/Engine/Source/Runtime/Core/Public/Templates/SharedPointer.h
--- a/Engine/Source/Runtime/Core/Public/Templates/SharedPointer.h
+++ b/Engine/Source/Runtime/Core/Public/Templates/SharedPointer.h
@@ -839,6 +839,20 @@ public:
 	}
 
 
+	/**
+	 * Constructs a weak pointer from a weak pointer of another type.
+	 * This constructor is intended to allow derived-to-base conversions.
+	 *
+	 * @param  InWeakPtr  The weak pointer to create a weak pointer from
+	 */
+	template< class OtherType >
+	FORCEINLINE TWeakPtr( TWeakPtr< OtherType, Mode > const& InWeakPtr )
+		: Object( InWeakPtr.Object ),
+		  WeakReferenceCount( InWeakPtr.WeakReferenceCount )
+	{
+	}
+
+
 	/**
 	 * Assignment to a NULL pointer.  Clears this weak pointer's reference.
 	 */
@@ -863,6 +877,21 @@ public:
 	}
 
 
+	/**
+	 * Assignment operator adds a weak reference to the object referenced by the specified weak pointer.
+	 * This assignment operator is intended to allow derived-to-base conversions.
+	 *
+	 * @param  InWeakPtr  The weak pointer for the object to assign
+	 */
+	template <typename OtherType>
+	FORCEINLINE TWeakPtr& operator=( TWeakPtr<OtherType, Mode> const& InWeakPtr )
+	{
+		Object = InWeakPtr.Pin().Get();
+		WeakReferenceCount = InWeakPtr.WeakReferenceCount;
+		return *this;
+	}
+
+
 	/**
 	 * Assignment operator sets this weak pointer from a shared reference
 	 *
@@ -952,6 +981,9 @@ private:
 	}
 
 
+	// We declare ourselves as a friend (templated using OtherType) so we can access members as needed
+    template< class OtherType, ESPMode::Type OtherMode > friend class TWeakPtr;
+
 	// Declare ourselves as a friend of TSharedPtr so we can access members as needed
     template< class OtherType, ESPMode::Type OtherMode > friend class TSharedPtr;
 
@@ -1251,6 +1283,30 @@ FORCEINLINE bool operator==( TWeakPtr< ObjectTypeA, Mode > const& InWeakPtrA, TW
 }
 
 
+/**
+ * Global equality operator for TWeakPtr
+ *
+ * @return  True if the weak pointer is null
+ */
+template< class ObjectTypeA, ESPMode::Type Mode >
+FORCEINLINE bool operator==( TWeakPtr< ObjectTypeA, Mode > const& InWeakPtrA, decltype(nullptr) )
+{
+	return !InWeakPtrA.IsValid();
+}
+
+
+/**
+ * Global equality operator for TWeakPtr
+ *
+ * @return  True if the weak pointer is null
+ */
+template< class ObjectTypeB, ESPMode::Type Mode >
+FORCEINLINE bool operator==( decltype(nullptr), TWeakPtr< ObjectTypeB, Mode > const& InWeakPtrB )
+{
+	return !InWeakPtrB.IsValid();
+}
+
+
 /**
  * Global inequality operator for TWeakPtr
  *
@@ -1263,6 +1319,30 @@ FORCEINLINE bool operator!=( TWeakPtr< ObjectTypeA, Mode > const& InWeakPtrA, TW
 }
 
 
+/**
+ * Global inequality operator for TWeakPtr
+ *
+ * @return  True if the weak pointer is not null
+ */
+template< class ObjectTypeA, ESPMode::Type Mode >
+FORCEINLINE bool operator!=( TWeakPtr< ObjectTypeA, Mode > const& InWeakPtrA, decltype(nullptr) )
+{
+	return InWeakPtrA.IsValid();
+}
+
+
+/**
+ * Global inequality operator for TWeakPtr
+ *
+ * @return  True if the weak pointer is not null
+ */
+template< class ObjectTypeB, ESPMode::Type Mode >
+FORCEINLINE bool operator!=( decltype(nullptr), TWeakPtr< ObjectTypeB, Mode > const& InWeakPtrB )
+{
+	return InWeakPtrB.IsValid();
+}
+
+
 /**
  * Casts a shared pointer of one type to another type. (static_cast)  Useful for downcasting.
  *
diff --git a/Engine/Source/Runtime/Core/Public/UObject/AutoPointer.h b/Engine/Source/Runtime/Core/Public/UObject/AutoPointer.h
--- a/Engine/Source/Runtime/Core/Public/UObject/AutoPointer.h
+++ b/Engine/Source/Runtime/Core/Public/UObject/AutoPointer.h
@@ -55,13 +55,13 @@ public:
 	/**  Compare for equality with a base pointer **/
 	FORCEINLINE bool operator==(const TBASE &Other) const
 	{
-		return TBASE::operator==(Other);
+		return (const TBASE&)*this == Other;
 	}
 	/**  Compare for equality with an auto pointer **/
 	template<class U, class UBASE>
 	FORCEINLINE bool operator==(const TAutoPointer<U,UBASE> &Other) const
 	{
-		return TBASE::operator==(Other);
+		return (const TBASE&)*this == Other;
 	}
 	/**  Compare for equality with a raw pointer **/
 	FORCEINLINE bool operator==(T* Other) const
@@ -82,13 +82,13 @@ public:
 	/**  Compare for inequality with a base pointer **/
 	FORCEINLINE bool operator!=(const TBASE &Other) const
 	{
-		return TBASE::operator!=(Other);
+		return (const TBASE&)*this != Other;
 	}
 	/**  Compare for equality with an auto pointer **/
 	template<class U, class UBASE>
 	FORCEINLINE bool operator!=(const TAutoPointer<U,UBASE> &Other) const
 	{
-		return TBASE::operator!=(Other);
+		return (const TBASE&)*this != Other;
 	}
 	/**  Compare for inequality with a raw pointer	**/
 	FORCEINLINE bool operator!=(T* Other) const
diff --git a/Engine/Source/Runtime/Core/Public/UObject/WeakObjectPtrTemplates.h b/Engine/Source/Runtime/Core/Public/UObject/WeakObjectPtrTemplates.h
--- a/Engine/Source/Runtime/Core/Public/UObject/WeakObjectPtrTemplates.h
+++ b/Engine/Source/Runtime/Core/Public/UObject/WeakObjectPtrTemplates.h
@@ -54,6 +54,19 @@ public:
 	{
 	}
 
+	/**  
+	 * Construct from another weak pointer of another type, intended for derived-to-base conversions
+	 * @param Other weak pointer to copy from
+	**/
+	template <typename OtherT>
+	FORCEINLINE TWeakObjectPtr(const TWeakObjectPtr<OtherT, TWeakObjectPtrBase, TUObjectArray> &Other) :
+		TWeakObjectPtrBase(Other)
+	{
+		// It's also possible that this static_assert may fail for valid conversions because
+		// one or both of the types have only been forward-declared.
+		static_assert(CanConvertPointerFromTo<OtherT, T>::Result, "Unable to convert TWeakObjectPtr - types are incompatible");
+	}
+
 	/**
 	 * Reset the weak pointer back to the NULL state
 	 */
@@ -83,23 +96,17 @@ public:
 	}
 
 	/**  
-	 * Compare weak pointers for equality
-	 * @param Other weak pointer to compare to
-	 * Caution: Two weak pointers might not be equal to each other, but they both might return NULL.
+	 * Assign from another weak pointer, intended for derived-to-base conversions
+	 * @param Other weak pointer to copy from
 	**/
-	FORCEINLINE bool operator==(const TWeakObjectPtr<T> &Other) const
+	template <typename OtherT>
+	FORCEINLINE void operator=(const TWeakObjectPtr<OtherT, TWeakObjectPtrBase, TUObjectArray> &Other)
 	{
-		return TWeakObjectPtrBase::operator==(Other);
-	}
+		// It's also possible that this static_assert may fail for valid conversions because
+		// one or both of the types have only been forward-declared.
+		static_assert(CanConvertPointerFromTo<OtherT, T>::Result, "Unable to convert TWeakObjectPtr - types are incompatible");
 
-	/**  
-	 * Compare weak pointers for equality
-	 * @param Other weak pointer to compare to
-	 * Caution: Two weak pointers might not be equal to each other, but they both might return NULL.
-	**/
-	FORCEINLINE bool operator!=(const TWeakObjectPtr<T> &Other) const
-	{
-		return TWeakObjectPtrBase::operator!=(Other);
+		TWeakObjectPtrBase::operator=(Other);
 	}
 
 	/**  
@@ -178,8 +185,105 @@ public:
 		Ar << static_cast<TWeakObjectPtrBase&>(WeakObjectPtr);
 		return Ar;
 	}
+
+	// We declare ourselves as a friend (templated using OtherType) so we can access members as needed
+	template<class OtherT, class OtherTWeakObjectPtrBase, class OtherTUObjectArray>
+	friend struct TWeakObjectPtr;
 };
 
+template <typename LhsT, typename RhsT, typename OtherTWeakObjectPtrBase, typename OtherTUObjectArray>
+FORCENOINLINE bool operator==(const TWeakObjectPtr<LhsT, OtherTWeakObjectPtrBase, OtherTUObjectArray> &Lhs, const TWeakObjectPtr<RhsT, OtherTWeakObjectPtrBase, OtherTUObjectArray> &Rhs)
+{
+	// It's also possible that this static_assert may fail for valid conversions because
+	// one or both of the types have only been forward-declared.
+	static_assert(CanConvertPointerFromTo<LhsT, RhsT>::Result || CanConvertPointerFromTo<RhsT, LhsT>::Result, "Unable to compare TWeakObjectPtrs - types are incompatible");
+
+	return (const OtherTWeakObjectPtrBase&)Lhs == (const OtherTWeakObjectPtrBase&)Rhs;
+}
+
+template <typename LhsT, typename RhsT, typename OtherTWeakObjectPtrBase, typename OtherTUObjectArray>
+FORCENOINLINE bool operator==(const TWeakObjectPtr<LhsT, OtherTWeakObjectPtrBase, OtherTUObjectArray> &Lhs, const RhsT* Rhs)
+{
+	// It's also possible that these static_asserts may fail for valid conversions because
+	// one or both of the types have only been forward-declared.
+	static_assert(CanConvertPointerFromTo<RhsT, UObject>::Result, "TWeakObjectPtr can only be compared with UObject types");
+	static_assert(CanConvertPointerFromTo<LhsT, RhsT>::Result || CanConvertPointerFromTo<RhsT, LhsT>::Result, "Unable to compare TWeakObjectPtr with raw pointer - types are incompatible");
+
+	// NOTE: this constructs a TWeakObjectPtrBase, which has some amount of overhead, so this may not be an efficient operation
+	return (const OtherTWeakObjectPtrBase&)Lhs == OtherTWeakObjectPtrBase(Rhs);
+}
+
+template <typename LhsT, typename RhsT, typename OtherTWeakObjectPtrBase, typename OtherTUObjectArray>
+FORCENOINLINE bool operator==(const LhsT* Lhs, const TWeakObjectPtr<RhsT, OtherTWeakObjectPtrBase, OtherTUObjectArray> &Rhs)
+{
+	// It's also possible that these static_asserts may fail for valid conversions because
+	// one or both of the types have only been forward-declared.
+	static_assert(CanConvertPointerFromTo<LhsT, UObject>::Result, "TWeakObjectPtr can only be compared with UObject types");
+	static_assert(CanConvertPointerFromTo<LhsT, RhsT>::Result || CanConvertPointerFromTo<RhsT, LhsT>::Result, "Unable to compare TWeakObjectPtr with raw pointer - types are incompatible");
+
+	// NOTE: this constructs a TWeakObjectPtrBase, which has some amount of overhead, so this may not be an efficient operation
+	return OtherTWeakObjectPtrBase(Lhs) == (const OtherTWeakObjectPtrBase&)Rhs;
+}
+
+template <typename LhsT, typename OtherTWeakObjectPtrBase, typename OtherTUObjectArray>
+FORCENOINLINE bool operator==(const TWeakObjectPtr<LhsT, OtherTWeakObjectPtrBase, OtherTUObjectArray> &Lhs, decltype(nullptr))
+{
+	return !Lhs.IsValid();
+}
+
+template <typename RhsT, typename OtherTWeakObjectPtrBase, typename OtherTUObjectArray>
+FORCENOINLINE bool operator==(decltype(nullptr), const TWeakObjectPtr<RhsT, OtherTWeakObjectPtrBase, OtherTUObjectArray> &Rhs)
+{
+	return !Rhs.IsValid();
+}
+
+template <typename LhsT, typename RhsT, typename OtherTWeakObjectPtrBase, typename OtherTUObjectArray>
+FORCENOINLINE bool operator!=(const TWeakObjectPtr<LhsT, OtherTWeakObjectPtrBase, OtherTUObjectArray> &Lhs, const TWeakObjectPtr<RhsT, OtherTWeakObjectPtrBase, OtherTUObjectArray> &Rhs)
+{
+	// It's also possible that this static_assert may fail for valid conversions because
+	// one or both of the types have only been forward-declared.
+	static_assert(CanConvertPointerFromTo<LhsT, RhsT>::Result || CanConvertPointerFromTo<RhsT, LhsT>::Result, "Unable to compare TWeakObjectPtrs - types are incompatible");
+
+	return (const OtherTWeakObjectPtrBase&)Lhs != (const OtherTWeakObjectPtrBase&)Rhs;
+}
+
+template <typename LhsT, typename RhsT, typename OtherTWeakObjectPtrBase, typename OtherTUObjectArray>
+FORCENOINLINE bool operator!=(const TWeakObjectPtr<LhsT, OtherTWeakObjectPtrBase, OtherTUObjectArray> &Lhs, const RhsT* Rhs)
+{
+	// It's also possible that these static_asserts may fail for valid conversions because
+	// one or both of the types have only been forward-declared.
+	static_assert(CanConvertPointerFromTo<RhsT, UObject>::Result, "TWeakObjectPtr can only be compared with UObject types");
+	static_assert(CanConvertPointerFromTo<LhsT, RhsT>::Result || CanConvertPointerFromTo<RhsT, LhsT>::Result, "Unable to compare TWeakObjectPtr with raw pointer - types are incompatible");
+
+	// NOTE: this constructs a TWeakObjectPtrBase, which has some amount of overhead, so this may not be an efficient operation
+	return (const OtherTWeakObjectPtrBase&)Lhs != OtherTWeakObjectPtrBase(Rhs);
+}
+
+template <typename LhsT, typename RhsT, typename OtherTWeakObjectPtrBase, typename OtherTUObjectArray>
+FORCENOINLINE bool operator!=(const RhsT* Lhs, const TWeakObjectPtr<LhsT, OtherTWeakObjectPtrBase, OtherTUObjectArray> &Rhs)
+{
+	// It's also possible that these static_asserts may fail for valid conversions because
+	// one or both of the types have only been forward-declared.
+	static_assert(CanConvertPointerFromTo<LhsT, UObject>::Result, "TWeakObjectPtr can only be compared with UObject types");
+	static_assert(CanConvertPointerFromTo<LhsT, RhsT>::Result || CanConvertPointerFromTo<RhsT, LhsT>::Result, "Unable to compare TWeakObjectPtr with raw pointer - types are incompatible");
+
+	// NOTE: this constructs a TWeakObjectPtrBase, which has some amount of overhead, so this may not be an efficient operation
+	return OtherTWeakObjectPtrBase(Lhs) != (const OtherTWeakObjectPtrBase&)Rhs;
+}
+
+template <typename LhsT, typename OtherTWeakObjectPtrBase, typename OtherTUObjectArray>
+FORCENOINLINE bool operator!=(const TWeakObjectPtr<LhsT, OtherTWeakObjectPtrBase, OtherTUObjectArray> &Lhs, decltype(nullptr))
+{
+	return Lhs.IsValid();
+}
+
+template <typename RhsT, typename OtherTWeakObjectPtrBase, typename OtherTUObjectArray>
+FORCENOINLINE bool operator!=(decltype(nullptr), const TWeakObjectPtr<RhsT, OtherTWeakObjectPtrBase, OtherTUObjectArray> &Rhs)
+{
+	return Rhs.IsValid();
+}
+
+
 template<class T> struct TIsPODType<TWeakObjectPtr<T> > { enum { Value = true }; };
 template<class T> struct TIsZeroConstructType<TWeakObjectPtr<T> > { enum { Value = true }; };
 template<class T> struct TIsWeakPointerType<TWeakObjectPtr<T> > { enum { Value = true }; };
