#UE4
- Added SCOPE_LOG_TIME_FUNC() and SCOPE_LOG_TIME_FUNC_WITH_GLOBAL(CumulativePtr) to log time spent in the specified scope, uses __FUNCTION__ as a display name

diff --git a/Engine/Source/Runtime/Core/Private/Stats/StatsMisc.cpp b/Engine/Source/Runtime/Core/Private/Stats/StatsMisc.cpp
--- a/Engine/Source/Runtime/Core/Private/Stats/StatsMisc.cpp
+++ b/Engine/Source/Runtime/Core/Private/Stats/StatsMisc.cpp
@@ -5,19 +5,19 @@
 FScopeLogTime::FScopeLogTime( const TCHAR* InName, FTotalTimeAndCount* InGlobal /*= nullptr */ ) 
 : StartTime( FPlatformTime::Seconds() )
 , Name( InName )
-, Global( InGlobal )
+, Cumulative( InGlobal )
 {}
 
 FScopeLogTime::~FScopeLogTime()
 {
 	const double ScopedTime = FPlatformTime::Seconds() - StartTime;
-	if( Global )
+	if( Cumulative )
 	{
-		Global->Key += ScopedTime;
-		Global->Value++;
+		Cumulative->Key += ScopedTime;
+		Cumulative->Value++;
 
-		const double Average = Global->Key / (double)Global->Value;
-		UE_LOG( LogStats, Log, TEXT( "%32s - %6.3f ms - Total %6.2f s / %5u / %6.3f ms" ), *Name, ScopedTime*1000.0f, Global->Key, Global->Value, Average*1000.0f );
+		const double Average = Cumulative->Key / (double)Cumulative->Value;
+		UE_LOG( LogStats, Log, TEXT( "%32s - %6.3f ms - Total %6.2f s / %5u / %6.3f ms" ), *Name, ScopedTime*1000.0f, Cumulative->Key, Cumulative->Value, Average*1000.0f );
 	}
 	else
 	{
diff --git a/Engine/Source/Runtime/Core/Public/Stats/StatsMisc.h b/Engine/Source/Runtime/Core/Public/Stats/StatsMisc.h
--- a/Engine/Source/Runtime/Core/Public/Stats/StatsMisc.h
+++ b/Engine/Source/Runtime/Core/Public/Stats/StatsMisc.h
@@ -54,16 +54,22 @@ struct FScopeLogTime
 	 * @param InGlobal - Pointer to the variable that holds the cumulative stats
 	 *
 	 */
-	FScopeLogTime( const TCHAR* InName, FTotalTimeAndCount* InGlobal = nullptr );
+	CORE_API FScopeLogTime( const TCHAR* InName, FTotalTimeAndCount* InCumulative = nullptr );
 
 	/** Destructor. */
-	~FScopeLogTime();
+	CORE_API ~FScopeLogTime();
 
 protected:
 	const double StartTime;
 	const FString Name;
-	FTotalTimeAndCount* Global;
+	FTotalTimeAndCount* Cumulative;
 };
 
-#define SCOPE_LOG_TIME(Name,Global) \
-	FScopeLogTime ScopeLogTime(Name,Global);
\ No newline at end of file
+#define SCOPE_LOG_TIME(Name,CumulativePtr) \
+	FScopeLogTime ScopeLogTime(Name,CumulativePtr);
+
+#define SCOPE_LOG_TIME_FUNC() \
+	FScopeLogTime ScopeLogTime(TEXT(__FUNCTION__));
+
+#define SCOPE_LOG_TIME_FUNC_WITH_GLOBAL(CumulativePtr) \
+	FScopeLogTime ScopeLogTime(TEXT(__FUNCTION__),CumulativePtr);
\ No newline at end of file
