#UE4
- Added thread metadata for the splash screen thread

diff --git a/Engine/Source/Runtime/Core/Private/Windows/WindowsPlatformSplash.cpp b/Engine/Source/Runtime/Core/Private/Windows/WindowsPlatformSplash.cpp
--- a/Engine/Source/Runtime/Core/Private/Windows/WindowsPlatformSplash.cpp
+++ b/Engine/Source/Runtime/Core/Private/Windows/WindowsPlatformSplash.cpp
@@ -677,7 +677,11 @@ void FWindowsPlatformSplash::Show()
 			}
 			
 			GSplashScreenFileName = FString( FPlatformProcess::BaseDir() ) / SplashPath;
-			GSplashScreenThread = CreateThread(NULL, 128 * 1024, (LPTHREAD_START_ROUTINE)StartSplashScreenThread, (LPVOID)NULL, STACK_SIZE_PARAM_IS_A_RESERVATION, NULL);
+			DWORD ThreadID = 0;
+			GSplashScreenThread = CreateThread(NULL, 128 * 1024, (LPTHREAD_START_ROUTINE)StartSplashScreenThread, (LPVOID)NULL, STACK_SIZE_PARAM_IS_A_RESERVATION, &ThreadID);
+#if	STATS
+			FStartupMessages::Get().AddThreadMetadata( FName( "SplashScreenThread" ), ThreadID );
+#endif // STATS
 		}
 	}
 }
