#UE4
- Added a warning if there is no metadata for the thread

diff --git a/Engine/Source/Runtime/Core/Private/Stats/StatsData.cpp b/Engine/Source/Runtime/Core/Private/Stats/StatsData.cpp
--- a/Engine/Source/Runtime/Core/Private/Stats/StatsData.cpp
+++ b/Engine/Source/Runtime/Core/Private/Stats/StatsData.cpp
@@ -1219,9 +1219,11 @@ FName FStatsThreadState::GetStatThreadName( const FStatPacket& Packet ) const
 		FName& NewThreadName = MutableThreads.FindOrAdd( Packet.ThreadId );
 		if( NewThreadName == NAME_None )
 		{
+			UE_LOG( LogStats, Warning, TEXT( "There is no thread with id: %u. Please add thread metadata for this thread." ), Packet.ThreadId );
+
 			static const FName NAME_UnknownThread = TEXT( "UnknownThread" );
 			NewThreadName = FName( *FStatsUtils::BuildUniqueThreadName( NAME_UnknownThread, Packet.ThreadId ) );
-			// This is an unknown thread, not created by UE, but still we need the metadata in the system.
+			// This is an unknown thread, but still we need the metadata in the system.
 			FStartupMessages::Get().AddThreadMetadata( NAME_UnknownThread, Packet.ThreadId );
 		}
 		ThreadName = NewThreadName;
