#ue4 nullptr fixes.

#add TYPE_OF_NULLPTR added to platform layer.
#add TYPE_OF_NULLPTR set to be decltype(__nullptr) on XB1, because it's compiled under C++/CX where nullptr has a different meaning.
#change TWeakObjectPtr operators changes to use TYPE_OF_NULLPTR.

diff --git a/Engine/Source/Runtime/Core/Public/GenericPlatform/GenericPlatform.h b/Engine/Source/Runtime/Core/Public/GenericPlatform/GenericPlatform.h
--- a/Engine/Source/Runtime/Core/Public/GenericPlatform/GenericPlatform.h
+++ b/Engine/Source/Runtime/Core/Public/GenericPlatform/GenericPlatform.h
@@ -59,5 +59,6 @@ struct FGenericPlatformTypes
 	typedef UPTRINT SIZE_T;															// signed int the same size as a pointer
 
 	typedef int32					TYPE_OF_NULL;
+	typedef decltype(nullptr)		TYPE_OF_NULLPTR;
 };
 
diff --git a/Engine/Source/Runtime/Core/Public/HAL/Platform.h b/Engine/Source/Runtime/Core/Public/HAL/Platform.h
--- a/Engine/Source/Runtime/Core/Public/HAL/Platform.h
+++ b/Engine/Source/Runtime/Core/Public/HAL/Platform.h
@@ -509,7 +509,8 @@ typedef FPlatformTypes::UPTRINT UPTRINT;	///< An unsigned integer the same size
 typedef FPlatformTypes::PTRINT PTRINT;		///< A signed integer the same size as a pointer
 typedef FPlatformTypes::SIZE_T SIZE_T;		///< A signed integer the same size as a pointer
 
-typedef FPlatformTypes::TYPE_OF_NULL	TYPE_OF_NULL; ///< The type of the NULL constant.
+typedef FPlatformTypes::TYPE_OF_NULL	TYPE_OF_NULL;	///< The type of the NULL constant.
+typedef FPlatformTypes::TYPE_OF_NULLPTR	TYPE_OF_NULLPTR; ///< The type of the C++ nullptr keyword.
 
 //------------------------------------------------------------------
 // Test the global types
diff --git a/Engine/Source/Runtime/Core/Public/UObject/WeakObjectPtrTemplates.h b/Engine/Source/Runtime/Core/Public/UObject/WeakObjectPtrTemplates.h
--- a/Engine/Source/Runtime/Core/Public/UObject/WeakObjectPtrTemplates.h
+++ b/Engine/Source/Runtime/Core/Public/UObject/WeakObjectPtrTemplates.h
@@ -226,13 +226,13 @@ FORCENOINLINE bool operator==(const LhsT* Lhs, const TWeakObjectPtr<RhsT, OtherT
 }
 
 template <typename LhsT, typename OtherTWeakObjectPtrBase, typename OtherTUObjectArray>
-FORCENOINLINE bool operator==(const TWeakObjectPtr<LhsT, OtherTWeakObjectPtrBase, OtherTUObjectArray> &Lhs, decltype(nullptr))
+FORCENOINLINE bool operator==(const TWeakObjectPtr<LhsT, OtherTWeakObjectPtrBase, OtherTUObjectArray> &Lhs, TYPE_OF_NULLPTR)
 {
 	return !Lhs.IsValid();
 }
 
 template <typename RhsT, typename OtherTWeakObjectPtrBase, typename OtherTUObjectArray>
-FORCENOINLINE bool operator==(decltype(nullptr), const TWeakObjectPtr<RhsT, OtherTWeakObjectPtrBase, OtherTUObjectArray> &Rhs)
+FORCENOINLINE bool operator==(TYPE_OF_NULLPTR, const TWeakObjectPtr<RhsT, OtherTWeakObjectPtrBase, OtherTUObjectArray> &Rhs)
 {
 	return !Rhs.IsValid();
 }
@@ -272,13 +272,13 @@ FORCENOINLINE bool operator!=(const RhsT* Lhs, const TWeakObjectPtr<LhsT, OtherT
 }
 
 template <typename LhsT, typename OtherTWeakObjectPtrBase, typename OtherTUObjectArray>
-FORCENOINLINE bool operator!=(const TWeakObjectPtr<LhsT, OtherTWeakObjectPtrBase, OtherTUObjectArray> &Lhs, decltype(nullptr))
+FORCENOINLINE bool operator!=(const TWeakObjectPtr<LhsT, OtherTWeakObjectPtrBase, OtherTUObjectArray> &Lhs, TYPE_OF_NULLPTR)
 {
 	return Lhs.IsValid();
 }
 
 template <typename RhsT, typename OtherTWeakObjectPtrBase, typename OtherTUObjectArray>
-FORCENOINLINE bool operator!=(decltype(nullptr), const TWeakObjectPtr<RhsT, OtherTWeakObjectPtrBase, OtherTUObjectArray> &Rhs)
+FORCENOINLINE bool operator!=(TYPE_OF_NULLPTR, const TWeakObjectPtr<RhsT, OtherTWeakObjectPtrBase, OtherTUObjectArray> &Rhs)
 {
 	return Rhs.IsValid();
 }
