Fix WinRT declaration of TYPE_OF_NULLPTR

diff --git a/Engine/Source/Runtime/Core/Public/WinRT/WinRTPlatform.h b/Engine/Source/Runtime/Core/Public/WinRT/WinRTPlatform.h
--- a/Engine/Source/Runtime/Core/Public/WinRT/WinRTPlatform.h
+++ b/Engine/Source/Runtime/Core/Public/WinRT/WinRTPlatform.h
@@ -17,6 +17,8 @@ struct FWinRTTypes : public FGenericPlatformTypes
 	// defined in windefs.h, even though this is equivalent, the compiler doesn't think so
 	typedef unsigned long       DWORD;
 	typedef unsigned __int64	SIZE_T;
+
+	typedef decltype(__nullptr) TYPE_OF_NULLPTR;
 };
 
 typedef FWinRTTypes FPlatformTypes;
