Moved Mac config files and logs location to ~/Library/Preferences and ~/Library/Logs subfolders, respectively.

diff --git a/Engine/Source/Programs/UnrealLightmass/Private/LightmassCore/LMCore.cpp b/Engine/Source/Programs/UnrealLightmass/Private/LightmassCore/LMCore.cpp
--- a/Engine/Source/Programs/UnrealLightmass/Private/LightmassCore/LMCore.cpp
+++ b/Engine/Source/Programs/UnrealLightmass/Private/LightmassCore/LMCore.cpp
@@ -31,7 +31,11 @@ FLightmassLog::FLightmassLog()
 	{
 		PathSeparatorPos--;
 	}
+#if PLATFORM_MAC
+	FString LogName = FPaths::Combine( FPlatformProcess::UserLogsDir(), *FString( ExeNameLen - PathSeparatorPos, *ExeName + PathSeparatorPos + 1 ) );
+#else
 	FString LogName( ExeNameLen - PathSeparatorPos - 5, *ExeName + PathSeparatorPos + 1 );
+#endif
 	LogName += TEXT("_");
 	LogName += FPlatformProcess::ComputerName();
 	LogName += TEXT("_");
diff --git a/Engine/Source/Runtime/Core/Private/Mac/MacPlatformProcess.cpp b/Engine/Source/Runtime/Core/Private/Mac/MacPlatformProcess.cpp
--- a/Engine/Source/Runtime/Core/Private/Mac/MacPlatformProcess.cpp
+++ b/Engine/Source/Runtime/Core/Private/Mac/MacPlatformProcess.cpp
@@ -589,6 +589,56 @@ const TCHAR* FMacPlatformProcess::UserSettingsDir()
 	return ApplicationSettingsDir();
 }
 
+static TCHAR* UserLibrarySubDirectory()
+{
+	static TCHAR Result[MAX_PATH] = TEXT("");
+	if (!Result[0])
+	{
+		FString SubDirectory = IsRunningGame() ? GGameName : FString(TEXT("Unreal Engine")) / GGameName;
+		if (IsRunningDedicatedServer())
+		{
+			SubDirectory += TEXT("Server");
+		}
+		else if (!IsRunningGame())
+		{
+#if WITH_EDITOR
+			SubDirectory += TEXT("Editor");
+#endif
+		}
+		SubDirectory += TEXT("/");
+		FCString::Strcpy(Result, *SubDirectory);
+	}
+	return Result;
+}
+
+const TCHAR* FMacPlatformProcess::UserPreferencesDir()
+{
+	static TCHAR Result[MAX_PATH] = TEXT("");
+	if (!Result[0])
+	{
+		SCOPED_AUTORELEASE_POOL;
+		NSString *UserLibraryDirectory = [NSSearchPathForDirectoriesInDomains(NSLibraryDirectory, NSUserDomainMask, YES) objectAtIndex: 0];
+		FPlatformString::CFStringToTCHAR((CFStringRef)UserLibraryDirectory, Result);
+		FCString::Strcat(Result, TEXT("/Preferences/"));
+		FCString::Strcat(Result, UserLibrarySubDirectory());
+	}
+	return Result;
+}
+
+const TCHAR* FMacPlatformProcess::UserLogsDir()
+{
+	static TCHAR Result[MAX_PATH] = TEXT("");
+	if (!Result[0])
+	{
+		SCOPED_AUTORELEASE_POOL;
+		NSString *UserLibraryDirectory = [NSSearchPathForDirectoriesInDomains(NSLibraryDirectory, NSUserDomainMask, YES) objectAtIndex: 0];
+		FPlatformString::CFStringToTCHAR((CFStringRef)UserLibraryDirectory, Result);
+		FCString::Strcat(Result, TEXT("/Logs/"));
+		FCString::Strcat(Result, UserLibrarySubDirectory());
+	}
+	return Result;
+}
+
 const TCHAR* FMacPlatformProcess::ApplicationSettingsDir()
 {
 	static TCHAR Result[MAX_PATH] = TEXT("");
diff --git a/Engine/Source/Runtime/Core/Private/Misc/Paths.cpp b/Engine/Source/Runtime/Core/Private/Misc/Paths.cpp
--- a/Engine/Source/Runtime/Core/Private/Misc/Paths.cpp
+++ b/Engine/Source/Runtime/Core/Private/Misc/Paths.cpp
@@ -109,7 +109,11 @@ FString FPaths::SourceConfigDir()
 
 FString FPaths::GeneratedConfigDir()
 {
+#if PLATFORM_MAC
+	return FPlatformProcess::UserPreferencesDir();
+#else
 	return FPaths::GameSavedDir() + TEXT("Config/");
+#endif
 }
 
 FString FPaths::SandboxesDir()
@@ -139,7 +143,11 @@ FString FPaths::VideoCaptureDir()
 
 FString FPaths::GameLogDir()
 {
+#if PLATFORM_MAC
+	return FPlatformProcess::UserLogsDir();
+#else
 	return FPaths::GameSavedDir() + TEXT("Logs/");
+#endif
 }
 
 FString FPaths::AutomationDir()
diff --git a/Engine/Source/Runtime/Core/Public/Mac/MacPlatformProcess.h b/Engine/Source/Runtime/Core/Public/Mac/MacPlatformProcess.h
--- a/Engine/Source/Runtime/Core/Public/Mac/MacPlatformProcess.h
+++ b/Engine/Source/Runtime/Core/Public/Mac/MacPlatformProcess.h
@@ -61,6 +61,10 @@ struct CORE_API FMacPlatformProcess : public FGenericPlatformProcess
 	static bool CreatePipe( void*& ReadPipe, void*& WritePipe );
 	static FString ReadPipe( void* ReadPipe );
 	static void ExecProcess( const TCHAR* URL, const TCHAR* Params, int32* OutReturnCode, FString* OutStdOut, FString* OutStdErr );
+
+	// Mac specific
+	static const TCHAR* UserPreferencesDir();
+	static const TCHAR* UserLogsDir();
 };
 
 typedef FMacPlatformProcess FPlatformProcess;
