Thread safe initialization of value returned by GetRelativePathToRoot().

diff --git a/Engine/Source/Runtime/Core/Private/Misc/Paths.cpp b/Engine/Source/Runtime/Core/Private/Misc/Paths.cpp
--- a/Engine/Source/Runtime/Core/Private/Misc/Paths.cpp
+++ b/Engine/Source/Runtime/Core/Private/Misc/Paths.cpp
@@ -796,25 +796,29 @@ void FPaths::Split( const FString& InPath, FString& PathPart, FString& FilenameP
 
 const FString& FPaths::GetRelativePathToRoot()
 {
-	static FString RelativePathToRoot;
-
-	// initialize static paths if needed
-	if (RelativePathToRoot.Len() == 0)
+	struct FRelativePathInitializer
 	{
-		FString RootDirectory = FPaths::RootDir();
-		FString BaseDirectory = FPlatformProcess::BaseDir();
-
-		// this is how to go from the base dir back to the root
-		RelativePathToRoot = RootDirectory;
-		FPaths::MakePathRelativeTo(RelativePathToRoot, *BaseDirectory);
+		FString RelativePathToRoot;
 
-		// Ensure that the path ends w/ '/'
-		if ((RelativePathToRoot.Len() > 0) && (RelativePathToRoot.EndsWith(TEXT("/")) == false) && (RelativePathToRoot.EndsWith(TEXT("\\")) == false))
+		FRelativePathInitializer()
 		{
-			RelativePathToRoot += TEXT("/");
+			FString RootDirectory = FPaths::RootDir();
+			FString BaseDirectory = FPlatformProcess::BaseDir();
+
+			// this is how to go from the base dir back to the root
+			RelativePathToRoot = RootDirectory;
+			FPaths::MakePathRelativeTo(RelativePathToRoot, *BaseDirectory);
+
+			// Ensure that the path ends w/ '/'
+			if ((RelativePathToRoot.Len() > 0) && (RelativePathToRoot.EndsWith(TEXT("/")) == false) && (RelativePathToRoot.EndsWith(TEXT("\\")) == false))
+			{
+				RelativePathToRoot += TEXT("/");
+			}
 		}
-	}
-	return RelativePathToRoot;
+	};
+
+	static FRelativePathInitializer StaticInstance;
+	return StaticInstance.RelativePathToRoot;
 }
 
 void FPaths::CombineInternal(FString& OutPath, const TCHAR** Pathes, int32 NumPathes)
