Return correct value from GetBuildConfiguration() when running DebugGame.

diff --git a/Engine/Source/Runtime/Core/Private/Misc/App.cpp b/Engine/Source/Runtime/Core/Private/Misc/App.cpp
--- a/Engine/Source/Runtime/Core/Private/Misc/App.cpp
+++ b/Engine/Source/Runtime/Core/Private/Misc/App.cpp
@@ -26,6 +26,25 @@ FString FApp::GetBranchName( )
 	return FString(TEXT(BRANCH_NAME));
 }
 
+EBuildConfigurations::Type FApp::GetBuildConfiguration()
+{
+#if UE_BUILD_DEBUG
+	return EBuildConfigurations::Debug;
+
+#elif UE_BUILD_DEVELOPMENT
+	static bool bUsingDebugGame = FParse::Param(FCommandLine::Get(), TEXT("debug"));
+	return bUsingDebugGame ? EBuildConfigurations::DebugGame : EBuildConfigurations::Development;
+
+#elif UE_BUILD_SHIPPING || UI_BUILD_SHIPPING_EDITOR
+	return EBuildConfigurations::Shipping;
+
+#elif UE_BUILD_TEST
+	return EBuildConfigurations::Test;
+
+#else
+	return EBuildConfigurations::Unknown;
+#endif
+}
 
 FString FApp::GetBuildDate( )
 {
diff --git a/Engine/Source/Runtime/Core/Public/Misc/App.h b/Engine/Source/Runtime/Core/Public/Misc/App.h
--- a/Engine/Source/Runtime/Core/Public/Misc/App.h
+++ b/Engine/Source/Runtime/Core/Public/Misc/App.h
@@ -36,27 +36,7 @@ public:
 	 *
 	 * @return The build configuration.
 	 */
-	static EBuildConfigurations::Type GetBuildConfiguration( )
-	{
-#if UE_BUILD_DEBUG
-		return EBuildConfigurations::Debug;
-
-#elif UE_BUILD_DEBUGGAME
-		return EBuildConfigurations::DebugGame;
-
-#elif UE_BUILD_DEVELOPMENT
-		return EBuildConfigurations::Development;
-
-#elif UE_BUILD_SHIPPING || UI_BUILD_SHIPPING_EDITOR
-		return EBuildConfigurations::Shipping;
-
-#elif UE_BUILD_TEST
-		return EBuildConfigurations::Test;
-
-#else
-		return EBuildConfigurations::Unknown;
-#endif
-	}
+	static EBuildConfigurations::Type GetBuildConfiguration();
 
 	/**
 	 * Gets the date at which this application was built.
