#ue4
#ios
* seems to fix the freeze when returning from sleep bug
* also fix crash when receiving applicationWillTerminate message

diff --git a/Engine/Source/Runtime/Core/Private/IOS/IOSAppDelegate.cpp b/Engine/Source/Runtime/Core/Private/IOS/IOSAppDelegate.cpp
--- a/Engine/Source/Runtime/Core/Private/IOS/IOSAppDelegate.cpp
+++ b/Engine/Source/Runtime/Core/Private/IOS/IOSAppDelegate.cpp
@@ -106,6 +106,7 @@ void InstallSignalHandlers()
 
 -(void)MainAppThread:(NSDictionary*)launchOptions
 {
+    self.bHasStarted = true;
 	GIsGuarded = false;
 	GStartTime = FPlatformTime::Seconds();
 
@@ -157,6 +158,8 @@ void InstallSignalHandlers()
 
 	[AutoreleasePool release];
 	FAppEntry::Shutdown();
+    
+    self.bHasStarted = false;
 }
 
 - (void)NoUrlCommandLine
@@ -441,6 +444,15 @@ void InstallSignalHandlers()
 	 See also applicationDidEnterBackground:.
 	 */
 	FCoreDelegates::ApplicationWillTerminateDelegate.Broadcast();
+    
+    // note that we are shutting down
+    GIsRequestingExit = true;
+    
+    // wait for the game thread to shut down
+    while (self.bHasStarted == true)
+    {
+        usleep(3);
+    }
 }
 
 - (void)applicationDidReceiveMemoryWarning:(UIApplication *)application
diff --git a/Engine/Source/Runtime/Core/Private/IOS/IOSPlatformProcess.cpp b/Engine/Source/Runtime/Core/Private/IOS/IOSPlatformProcess.cpp
--- a/Engine/Source/Runtime/Core/Private/IOS/IOSPlatformProcess.cpp
+++ b/Engine/Source/Runtime/Core/Private/IOS/IOSPlatformProcess.cpp
@@ -50,14 +50,14 @@ void FIOSPlatformProcess::SetRealTimeMode()
     mach_timebase_info_data_t TimeBaseInfo;
     mach_timebase_info( &TimeBaseInfo );
     double MsToAbs = ((double)TimeBaseInfo.denom / (double)TimeBaseInfo.numer) * 1000000.0;
-    uint32 NormalProcessingTimeMs = 20;
-    uint32 ConstraintProcessingTimeMs = 60;
+    uint32 NormalProcessingTimeMs = 33;
+    uint32 ConstraintProcessingTimeMs = 66;
     
     thread_time_constraint_policy_data_t Policy;
     Policy.period      = 0;
     Policy.computation = (uint32_t)(NormalProcessingTimeMs * MsToAbs);
     Policy.constraint  = (uint32_t)(ConstraintProcessingTimeMs * MsToAbs);
-    Policy.preemptible = true;
+    Policy.preemptible = false;
     
     thread_policy_set(pthread_mach_thread_np(pthread_self()), THREAD_TIME_CONSTRAINT_POLICY, (thread_policy_t)&Policy, THREAD_TIME_CONSTRAINT_POLICY_COUNT);
 }
\ No newline at end of file
diff --git a/Engine/Source/Runtime/Core/Public/IOS/IOSAppDelegate.h b/Engine/Source/Runtime/Core/Public/IOS/IOSAppDelegate.h
--- a/Engine/Source/Runtime/Core/Public/IOS/IOSAppDelegate.h
+++ b/Engine/Source/Runtime/Core/Public/IOS/IOSAppDelegate.h
@@ -121,6 +121,7 @@ UITextFieldDelegate, AVAudioSessionDelegate>
 
 @property (atomic) bool bIsSuspended;
 @property (atomic) bool bHasSuspended;
+@property (atomic) bool bHasStarted;
 - (void)ToggleSuspend:(bool)bSuspend;
 
 static void interruptionListener(void* ClientData, UInt32 Interruption);
diff --git a/Engine/Source/Runtime/Launch/Private/IOS/LaunchIOS.cpp b/Engine/Source/Runtime/Launch/Private/IOS/LaunchIOS.cpp
--- a/Engine/Source/Runtime/Launch/Private/IOS/LaunchIOS.cpp
+++ b/Engine/Source/Runtime/Launch/Private/IOS/LaunchIOS.cpp
@@ -160,6 +160,9 @@ void FAppEntry::Shutdown()
 {
 	NSLog(@"%s", "Shutting down Game ULD Communications\n");
 	GCommandSystem.Shutdown();
+    
+    // kill the engine
+    GEngineLoop.Exit();
 }
 
 FString GSavedCommandLine;
