UE4: No longer including SteamController module in Server-only builds. This is not ideal because it will still be included in client builds that do not use the steam controller, but it unblocks linux servers in games that do not use steam at all.

diff --git a/Engine/Source/Runtime/Core/Core.Build.cs b/Engine/Source/Runtime/Core/Core.Build.cs
--- a/Engine/Source/Runtime/Core/Core.Build.cs
+++ b/Engine/Source/Runtime/Core/Core.Build.cs
@@ -114,7 +114,7 @@ public class Core : ModuleRules
 				"SDL2"
                 );
 
-            if (UEBuildConfiguration.bCompileAgainstEngine == true)
+            if (UEBuildConfiguration.bCompileAgainstEngine == true && Target.Type != TargetRules.TargetType.Server)
             {
                 AddThirdPartyPrivateStaticDependencies(Target, "SteamController");
             }
