UE4: Added !UE_SERVER to usage of SteamController in core code to finish the change in CL#2040814

diff --git a/Engine/Source/Runtime/Core/Private/Linux/LinuxApplication.cpp b/Engine/Source/Runtime/Core/Private/Linux/LinuxApplication.cpp
--- a/Engine/Source/Runtime/Core/Private/Linux/LinuxApplication.cpp
+++ b/Engine/Source/Runtime/Core/Private/Linux/LinuxApplication.cpp
@@ -4,7 +4,7 @@
 #include "LinuxWindow.h"
 #include "LinuxCursor.h"
 #include "GenericApplicationMessageHandler.h"
-#if WITH_ENGINE
+#if WITH_ENGINE && !UE_SERVER
 #include "SteamControllerInterface.h"
 #endif
 
@@ -77,7 +77,7 @@ FLinuxApplication* FLinuxApplication::CreateLinuxApplication()
 
 
 FLinuxApplication::FLinuxApplication() : GenericApplication( MakeShareable( new FLinuxCursor() ) )
-#if WITH_ENGINE
+#if WITH_ENGINE && !UE_SERVER
 	, SteamInput( SteamControllerInterface::Create(MessageHandler) )
 #endif
 {
diff --git a/Engine/Source/Runtime/Core/Public/Linux/LinuxApplication.h b/Engine/Source/Runtime/Core/Public/Linux/LinuxApplication.h
--- a/Engine/Source/Runtime/Core/Public/Linux/LinuxApplication.h
+++ b/Engine/Source/Runtime/Core/Public/Linux/LinuxApplication.h
@@ -95,7 +95,7 @@ private:
 
 	SDLControllerState *ControllerStates;
 
-#if WITH_ENGINE
+#if WITH_ENGINE && !UE_SERVER
 	TSharedPtr< class SteamControllerInterface > SteamInput;
 #endif
 };
