#UE4 - Unitialized memory

diff --git a/Engine/Source/Runtime/Core/Public/Serialization/Bitstreams.h b/Engine/Source/Runtime/Core/Public/Serialization/Bitstreams.h
--- a/Engine/Source/Runtime/Core/Public/Serialization/Bitstreams.h
+++ b/Engine/Source/Runtime/Core/Public/Serialization/Bitstreams.h
@@ -138,7 +138,8 @@ struct CORE_API FBitWriterMark
 {
 public:
 	FBitWriterMark()
-	:	Num         ( 0 )
+	:	Overflowed	( false ),
+		Num			( 0 )
 	{}
 	FBitWriterMark( FBitWriter& Writer )
 	{
@@ -160,7 +161,7 @@ public:
 
 private:
 	bool			Overflowed;
-	int64				Num;
+	int64			Num;
 };
 
 /*-----------------------------------------------------------------------------
