Adding guards for Steam Controller support. Should be refactored and made a platform property.

diff --git a/Engine/Source/Runtime/Core/Private/Linux/LinuxApplication.cpp b/Engine/Source/Runtime/Core/Private/Linux/LinuxApplication.cpp
--- a/Engine/Source/Runtime/Core/Private/Linux/LinuxApplication.cpp
+++ b/Engine/Source/Runtime/Core/Private/Linux/LinuxApplication.cpp
@@ -4,9 +4,9 @@
 #include "LinuxWindow.h"
 #include "LinuxCursor.h"
 #include "GenericApplicationMessageHandler.h"
-#if WITH_ENGINE && !UE_SERVER
-#include "SteamControllerInterface.h"
-#endif
+#if STEAM_CONTROLLER_SUPPORT
+	#include "SteamControllerInterface.h"
+#endif // STEAM_CONTROLLER_SUPPORT
 
 //	todo:
 //	need to change to a linux one
@@ -77,9 +77,9 @@ FLinuxApplication* FLinuxApplication::CreateLinuxApplication()
 
 
 FLinuxApplication::FLinuxApplication() : GenericApplication( MakeShareable( new FLinuxCursor() ) )
-#if WITH_ENGINE && !UE_SERVER
+#if STEAM_CONTROLLER_SUPPORT
 	, SteamInput( SteamControllerInterface::Create(MessageHandler) )
-#endif
+#endif // STEAM_CONTROLLER_SUPPORT
 {
 	bUsingHighPrecisionMouseInput = false;
 	bAllowedToDeferMessageProcessing = false;
@@ -124,9 +124,9 @@ void FLinuxApplication::InitializeWindow(	const TSharedRef< FGenericWindow >& In
 void FLinuxApplication::SetMessageHandler( const TSharedRef< FGenericApplicationMessageHandler >& InMessageHandler )
 {
 	GenericApplication::SetMessageHandler(InMessageHandler);
-#if WITH_ENGINE
+#if STEAM_CONTROLLER_SUPPORT
 	SteamInput->SetMessageHandler(InMessageHandler);
-#endif
+#endif // STEAM_CONTROLLER_SUPPORT
 }
 
 static TSharedPtr< FLinuxWindow > FindWindowBySDLWindow( const TArray< TSharedRef< FLinuxWindow > >& WindowsToSearch, SDL_HWindow const WindowHandle )
@@ -686,10 +686,10 @@ void FLinuxApplication::ProcessDeferredEvents( const float TimeDelta )
 
 void FLinuxApplication::PollGameDeviceState( const float TimeDelta )
 {
-#if WITH_ENGINE
+#if STEAM_CONTROLLER_SUPPORT
 	// Poll game device states and send new events
 	SteamInput->SendControllerEvents();
-#endif
+#endif // STEAM_CONTROLLER_SUPPORT
 }
 
 TCHAR FLinuxApplication::ConvertChar( SDL_Keysym Keysym )
diff --git a/Engine/Source/Runtime/Core/Public/Linux/LinuxApplication.h b/Engine/Source/Runtime/Core/Public/Linux/LinuxApplication.h
--- a/Engine/Source/Runtime/Core/Public/Linux/LinuxApplication.h
+++ b/Engine/Source/Runtime/Core/Public/Linux/LinuxApplication.h
@@ -5,6 +5,8 @@
 #include "GenericApplication.h"
 #include "LinuxWindow.h"
 
+#define STEAM_CONTROLLER_SUPPORT				(WITH_ENGINE && !UE_SERVER)
+
 typedef SDL_GameController* SDL_HController;
 
 class FLinuxWindow;
@@ -95,9 +97,9 @@ private:
 
 	SDLControllerState *ControllerStates;
 
-#if WITH_ENGINE && !UE_SERVER
+#if STEAM_CONTROLLER_SUPPORT
 	TSharedPtr< class SteamControllerInterface > SteamInput;
-#endif
+#endif // STEAM_CONTROLLER_SUPPORT
 };
 
 
