#ttp 330906 CRITICAL: Regression: Mac: EDITOR: CRASH: Renaming Custom Event
#branch UE4
#project Core
#summary Fix crash renaming custom blueprint events.
#change In activateInputMethod deactivate the OS IME, but don't touch the cached UE4 IME context as it may be invalid.
#notes The inline textbox doesn't issue a DeactivateContext so the activateInputMethod selector was trying to operate on an invalid context pointer.

diff --git a/Engine/Source/Runtime/Core/Private/Mac/MacTextInputMethodSystem.cpp b/Engine/Source/Runtime/Core/Private/Mac/MacTextInputMethodSystem.cpp
--- a/Engine/Source/Runtime/Core/Private/Mac/MacTextInputMethodSystem.cpp
+++ b/Engine/Source/Runtime/Core/Private/Mac/MacTextInputMethodSystem.cpp
@@ -69,7 +69,8 @@ DEFINE_LOG_CATEGORY_STATIC(LogMacTextInputMethodSystem, Log, All);
 	check(InContext);
 	if(IMMContext)
 	{
-		[self deactivateInputMethod];
+		[[self inputContext] deactivate];
+		IMMContext = NULL;
 	}
 	IMMContext = InContext;
 	[[self inputContext] activate];
