#ttp 331239 - CORE: Remove unnecessary platform srand calls.
#proj core
#branch UE4
#summary Removed unnecessary platform srand calls.

diff --git a/Engine/Source/Runtime/Core/Private/HTML5/HTML5PlatformMisc.cpp b/Engine/Source/Runtime/Core/Private/HTML5/HTML5PlatformMisc.cpp
--- a/Engine/Source/Runtime/Core/Private/HTML5/HTML5PlatformMisc.cpp
+++ b/Engine/Source/Runtime/Core/Private/HTML5/HTML5PlatformMisc.cpp
@@ -10,16 +10,6 @@
 
 void FHTML5Misc::PlatformInit()
 {
-	// Randomize.
-	if (GIsBenchmarking)
-	{
-		srand(0);
-	}
-	else
-	{
-		srand((unsigned)time(NULL));
-	}
-
 	// Identity.
 	UE_LOG(LogInit, Log, TEXT("Computer: %s"), FPlatformProcess::ComputerName());
 	UE_LOG(LogInit, Log, TEXT("User: %s"), FPlatformProcess::UserName());
diff --git a/Engine/Source/Runtime/Core/Private/IOS/IOSPlatformMisc.cpp b/Engine/Source/Runtime/Core/Private/IOS/IOSPlatformMisc.cpp
--- a/Engine/Source/Runtime/Core/Private/IOS/IOSPlatformMisc.cpp
+++ b/Engine/Source/Runtime/Core/Private/IOS/IOSPlatformMisc.cpp
@@ -77,16 +77,6 @@ void FIOSPlatformMisc::PlatformInit()
 	int32 Result = setrlimit(RLIMIT_NOFILE, &Limit);
 	check(Result == 0);
 
-	// Randomize.
-	if( GIsBenchmarking )
-	{
-		srand( 0 );
-	}
-    else
-	{
-		srand( (unsigned)time( NULL ) );
-	}
-
 	// Identity.
 	UE_LOG(LogInit, Log, TEXT("Computer: %s"), FPlatformProcess::ComputerName() );
 	UE_LOG(LogInit, Log, TEXT("User: %s"), FPlatformProcess::UserName() );
diff --git a/Engine/Source/Runtime/Core/Private/Mac/MacPlatformMisc.cpp b/Engine/Source/Runtime/Core/Private/Mac/MacPlatformMisc.cpp
--- a/Engine/Source/Runtime/Core/Private/Mac/MacPlatformMisc.cpp
+++ b/Engine/Source/Runtime/Core/Private/Mac/MacPlatformMisc.cpp
@@ -265,16 +265,6 @@ void FMacPlatformMisc::PlatformInit()
 		UE_LOG(LogInit, Warning, TEXT("Failed to change open file limit, UE4 may be unstable."));
 	}
 
-	// Randomize.
-    if( GIsBenchmarking )
-	{
-		srand( 0 );
-	}
-    else
-	{
-		srand( (unsigned)time( NULL ) );
-	}
-
 	// Identity.
 	UE_LOG(LogInit, Log, TEXT("Computer: %s"), FPlatformProcess::ComputerName() );
 	UE_LOG(LogInit, Log, TEXT("User: %s"), FPlatformProcess::UserName() );
diff --git a/Engine/Source/Runtime/Core/Private/WinRT/WinRTMisc.cpp b/Engine/Source/Runtime/Core/Private/WinRT/WinRTMisc.cpp
--- a/Engine/Source/Runtime/Core/Private/WinRT/WinRTMisc.cpp
+++ b/Engine/Source/Runtime/Core/Private/WinRT/WinRTMisc.cpp
@@ -80,16 +80,6 @@ void FWinRTMisc::PlatformPreInit()
 
 void FWinRTMisc::PlatformInit()
 {
-	// Randomize.
-	if (GIsBenchmarking)
-	{
-		srand(0);
-	}
-	else
-	{
-		srand((unsigned)time(NULL));
-	}
-
 	// Identity.
 	UE_LOG(LogInit, Log, TEXT("Computer: %s"), FPlatformProcess::ComputerName());
 	UE_LOG(LogInit, Log, TEXT("User: %s"), FPlatformProcess::UserName());
diff --git a/Engine/Source/Runtime/Core/Private/Windows/WindowsPlatformMisc.cpp b/Engine/Source/Runtime/Core/Private/Windows/WindowsPlatformMisc.cpp
--- a/Engine/Source/Runtime/Core/Private/Windows/WindowsPlatformMisc.cpp
+++ b/Engine/Source/Runtime/Core/Private/Windows/WindowsPlatformMisc.cpp
@@ -197,16 +197,6 @@ void FWindowsPlatformMisc::PlatformInit()
 	_set_FMA3_enable(0);
 #endif
 
-	// Randomize.
-	if( GIsBenchmarking )
-	{
-		srand( 0 );
-	}
-	else
-	{
-		srand( (unsigned)time( NULL ) );
-	}
-
 	// Set granularity of sleep and such to 1 ms.
 	timeBeginPeriod( 1 );
 
