#UE4
- Better method of logging stat in the malloc binned, should solve a crash on PS4 using the streaming file server

diff --git a/Engine/Source/Runtime/Core/Public/HAL/MallocBinned.h b/Engine/Source/Runtime/Core/Public/HAL/MallocBinned.h
--- a/Engine/Source/Runtime/Core/Public/HAL/MallocBinned.h
+++ b/Engine/Source/Runtime/Core/Public/HAL/MallocBinned.h
@@ -1149,86 +1149,91 @@ public:
 	 */
 	virtual void DumpAllocatorStats( class FOutputDevice& Ar ) OVERRIDE
 	{
+		FBufferedOutputDevice BufferedOutput;
+		{
 #ifdef USE_INTERNAL_LOCKS
-		FScopeLock ScopedLock(&AccessGuard);
+			FScopeLock ScopedLock( &AccessGuard );
 #endif
-		ValidateHeap();
+			ValidateHeap();
 #if STATS
-		UpdateSlackStat();
-
-		// This is all of the memory including stuff too big for the pools
-		Ar.CategorizedLogf( LogMemory.GetCategoryName(), ELogVerbosity::Log, TEXT("Allocator Stats for %s:" ), GetDescriptiveName() );
-		// Waste is the total overhead of the memory system
-		Ar.CategorizedLogf( LogMemory.GetCategoryName(), ELogVerbosity::Log, TEXT("Current Memory %.2f MB used, plus %.2f MB waste"), UsedCurrent / ( 1024.0f * 1024.0f ), ( OsCurrent - UsedCurrent ) / ( 1024.0f * 1024.0f ) );
-		Ar.CategorizedLogf( LogMemory.GetCategoryName(), ELogVerbosity::Log, TEXT("Peak Memory %.2f MB used, plus %.2f MB waste"), UsedPeak / ( 1024.0f * 1024.0f ), ( OsPeak - UsedPeak ) / ( 1024.0f * 1024.0f ) );
-
-		Ar.CategorizedLogf( LogMemory.GetCategoryName(), ELogVerbosity::Log, TEXT("Current OS Memory %.2f MB, peak %.2f MB"), OsCurrent / ( 1024.0f * 1024.0f ), OsPeak / ( 1024.0f * 1024.0f ) );
-		Ar.CategorizedLogf( LogMemory.GetCategoryName(), ELogVerbosity::Log, TEXT("Current Waste %.2f MB, peak %.2f MB"), WasteCurrent / ( 1024.0f * 1024.0f ), WastePeak / ( 1024.0f * 1024.0f ) );
-		Ar.CategorizedLogf( LogMemory.GetCategoryName(), ELogVerbosity::Log, TEXT("Current Used %.2f MB, peak %.2f MB"), UsedCurrent / ( 1024.0f * 1024.0f ), UsedPeak / ( 1024.0f * 1024.0f ) );
-		Ar.CategorizedLogf( LogMemory.GetCategoryName(), ELogVerbosity::Log, TEXT("Current Slack %.2f MB"), SlackCurrent / ( 1024.0f * 1024.0f ) );
-
-		Ar.CategorizedLogf( LogMemory.GetCategoryName(), ELogVerbosity::Log, TEXT("Allocs      % 6i Current / % 6i Total"), CurrentAllocs, TotalAllocs );
-		MEM_TIME(Ar.CategorizedLogf( LogMemory.GetCategoryName(), ELogVerbosity::Log, TEXT( "Seconds     % 5.3f" ), MemTime ));
-		MEM_TIME(Ar.CategorizedLogf( LogMemory.GetCategoryName(), ELogVerbosity::Log, TEXT( "MSec/Allc   % 5.5f" ), 1000.0 * MemTime / MemAllocs ));
-
-		// This is the memory tracked inside individual allocation pools
-		Ar.CategorizedLogf( LogMemory.GetCategoryName(), ELogVerbosity::Log, TEXT("") );
-		Ar.CategorizedLogf( LogMemory.GetCategoryName(), ELogVerbosity::Log, TEXT("Block Size Num Pools Max Pools Cur Allocs Total Allocs Min Req Max Req Mem Used Mem Slack Mem Waste Efficiency") );
-		Ar.CategorizedLogf( LogMemory.GetCategoryName(), ELogVerbosity::Log, TEXT("---------- --------- --------- ---------- ------------ ------- ------- -------- --------- --------- ----------") );
-
-		uint32 TotalMemory = 0;
-		uint32 TotalWaste = 0;
-		uint32 TotalActiveRequests = 0;
-		uint32 TotalTotalRequests = 0;
-		uint32 TotalPools = 0;
-		uint32 TotalSlack = 0;
-
-		FPoolTable* Table = NULL;
-		for( int32 i = 0; i < BinnedSizeLimit+EXTENED_PAGE_POOL_ALLOCATION_COUNT; i++ )
-		{
-			if (Table == MemSizeToPoolTable[i] || MemSizeToPoolTable[i]->BlockSize == 0)
-				continue;
-
-			Table = MemSizeToPoolTable[i];
-			
-			uint32 TableAllocSize = (Table->BlockSize > BinnedSizeLimit ? (((3*(i-BinnedSizeLimit))+3)*BINNED_ALLOC_POOL_SIZE) : BINNED_ALLOC_POOL_SIZE);
-			// The amount of memory allocated from the OS
-			uint32 MemAllocated = ( Table->NumActivePools * TableAllocSize) / 1024;
-			// Amount of memory actually in use by allocations
-			uint32 MemUsed = ( Table->BlockSize * Table->ActiveRequests ) / 1024;
-			// Wasted memory due to pool size alignment
-			uint32 PoolMemWaste = Table->NumActivePools * ( TableAllocSize - ( ( TableAllocSize / Table->BlockSize ) * Table->BlockSize ) ) / 1024;
-			// Wasted memory due to individual allocation alignment. This is an estimate.
-			uint32 MemWaste = ( uint32 )( ( ( double )Table->TotalWaste / ( double )Table->TotalRequests ) * ( double )Table->ActiveRequests ) / 1024 + PoolMemWaste;
-			// Memory that is reserved in active pools and ready for future use
-			uint32 MemSlack = MemAllocated - MemUsed - PoolMemWaste;
-
-			Ar.CategorizedLogf( LogMemory.GetCategoryName(), ELogVerbosity::Log, TEXT("% 10i % 9i % 9i % 10i % 12i % 7i % 7i % 7iK % 8iK % 8iK % 9.2f%%"),
-				Table->BlockSize,
-				Table->NumActivePools,
-				Table->MaxActivePools,
-				Table->ActiveRequests,
-				( uint32 )Table->TotalRequests,
-				Table->MinRequest,
-				Table->MaxRequest,
-				MemUsed,
-				MemSlack,
-				MemWaste,
-				MemAllocated ? 100.0f * (MemAllocated - MemWaste) / MemAllocated : 100.0f );
-
-			TotalMemory += MemAllocated;
-			TotalWaste += MemWaste;
-			TotalSlack += MemSlack;
-			TotalActiveRequests += Table->ActiveRequests;
-			TotalTotalRequests += Table->TotalRequests;
-			TotalPools += Table->NumActivePools;
-		}
+			UpdateSlackStat();
+
+			// This is all of the memory including stuff too big for the pools
+			BufferedOutput.CategorizedLogf( LogMemory.GetCategoryName(), ELogVerbosity::Log, TEXT( "Allocator Stats for %s:" ), GetDescriptiveName() );
+			// Waste is the total overhead of the memory system
+			BufferedOutput.CategorizedLogf( LogMemory.GetCategoryName(), ELogVerbosity::Log, TEXT( "Current Memory %.2f MB used, plus %.2f MB waste" ), UsedCurrent / (1024.0f * 1024.0f), (OsCurrent - UsedCurrent) / (1024.0f * 1024.0f) );
+			BufferedOutput.CategorizedLogf( LogMemory.GetCategoryName(), ELogVerbosity::Log, TEXT( "Peak Memory %.2f MB used, plus %.2f MB waste" ), UsedPeak / (1024.0f * 1024.0f), (OsPeak - UsedPeak) / (1024.0f * 1024.0f) );
+
+			BufferedOutput.CategorizedLogf( LogMemory.GetCategoryName(), ELogVerbosity::Log, TEXT( "Current OS Memory %.2f MB, peak %.2f MB" ), OsCurrent / (1024.0f * 1024.0f), OsPeak / (1024.0f * 1024.0f) );
+			BufferedOutput.CategorizedLogf( LogMemory.GetCategoryName(), ELogVerbosity::Log, TEXT( "Current Waste %.2f MB, peak %.2f MB" ), WasteCurrent / (1024.0f * 1024.0f), WastePeak / (1024.0f * 1024.0f) );
+			BufferedOutput.CategorizedLogf( LogMemory.GetCategoryName(), ELogVerbosity::Log, TEXT( "Current Used %.2f MB, peak %.2f MB" ), UsedCurrent / (1024.0f * 1024.0f), UsedPeak / (1024.0f * 1024.0f) );
+			BufferedOutput.CategorizedLogf( LogMemory.GetCategoryName(), ELogVerbosity::Log, TEXT( "Current Slack %.2f MB" ), SlackCurrent / (1024.0f * 1024.0f) );
+
+			BufferedOutput.CategorizedLogf( LogMemory.GetCategoryName(), ELogVerbosity::Log, TEXT( "Allocs      % 6i Current / % 6i Total" ), CurrentAllocs, TotalAllocs );
+			MEM_TIME( BufferedOutput.CategorizedLogf( LogMemory.GetCategoryName(), ELogVerbosity::Log, TEXT( "Seconds     % 5.3f" ), MemTime ) );
+			MEM_TIME( BufferedOutput.CategorizedLogf( LogMemory.GetCategoryName(), ELogVerbosity::Log, TEXT( "MSec/Allc   % 5.5f" ), 1000.0 * MemTime / MemAllocs ) );
+
+			// This is the memory tracked inside individual allocation pools
+			BufferedOutput.CategorizedLogf( LogMemory.GetCategoryName(), ELogVerbosity::Log, TEXT( "" ) );
+			BufferedOutput.CategorizedLogf( LogMemory.GetCategoryName(), ELogVerbosity::Log, TEXT( "Block Size Num Pools Max Pools Cur Allocs Total Allocs Min Req Max Req Mem Used Mem Slack Mem Waste Efficiency" ) );
+			BufferedOutput.CategorizedLogf( LogMemory.GetCategoryName(), ELogVerbosity::Log, TEXT( "---------- --------- --------- ---------- ------------ ------- ------- -------- --------- --------- ----------" ) );
+
+			uint32 TotalMemory = 0;
+			uint32 TotalWaste = 0;
+			uint32 TotalActiveRequests = 0;
+			uint32 TotalTotalRequests = 0;
+			uint32 TotalPools = 0;
+			uint32 TotalSlack = 0;
+
+			FPoolTable* Table = NULL;
+			for( int32 i = 0; i < BinnedSizeLimit + EXTENED_PAGE_POOL_ALLOCATION_COUNT; i++ )
+			{
+				if( Table == MemSizeToPoolTable[i] || MemSizeToPoolTable[i]->BlockSize == 0 )
+					continue;
+
+				Table = MemSizeToPoolTable[i];
+
+				uint32 TableAllocSize = (Table->BlockSize > BinnedSizeLimit ? (((3 * (i - BinnedSizeLimit)) + 3)*BINNED_ALLOC_POOL_SIZE) : BINNED_ALLOC_POOL_SIZE);
+				// The amount of memory allocated from the OS
+				uint32 MemAllocated = (Table->NumActivePools * TableAllocSize) / 1024;
+				// Amount of memory actually in use by allocations
+				uint32 MemUsed = (Table->BlockSize * Table->ActiveRequests) / 1024;
+				// Wasted memory due to pool size alignment
+				uint32 PoolMemWaste = Table->NumActivePools * (TableAllocSize - ((TableAllocSize / Table->BlockSize) * Table->BlockSize)) / 1024;
+				// Wasted memory due to individual allocation alignment. This is an estimate.
+				uint32 MemWaste = (uint32)(((double)Table->TotalWaste / (double)Table->TotalRequests) * (double)Table->ActiveRequests) / 1024 + PoolMemWaste;
+				// Memory that is reserved in active pools and ready for future use
+				uint32 MemSlack = MemAllocated - MemUsed - PoolMemWaste;
+
+				BufferedOutput.CategorizedLogf( LogMemory.GetCategoryName(), ELogVerbosity::Log, TEXT( "% 10i % 9i % 9i % 10i % 12i % 7i % 7i % 7iK % 8iK % 8iK % 9.2f%%" ),
+											  Table->BlockSize,
+											  Table->NumActivePools,
+											  Table->MaxActivePools,
+											  Table->ActiveRequests,
+											  (uint32)Table->TotalRequests,
+											  Table->MinRequest,
+											  Table->MaxRequest,
+											  MemUsed,
+											  MemSlack,
+											  MemWaste,
+											  MemAllocated ? 100.0f * (MemAllocated - MemWaste) / MemAllocated : 100.0f );
+
+				TotalMemory += MemAllocated;
+				TotalWaste += MemWaste;
+				TotalSlack += MemSlack;
+				TotalActiveRequests += Table->ActiveRequests;
+				TotalTotalRequests += Table->TotalRequests;
+				TotalPools += Table->NumActivePools;
+			}
 
-		Ar.CategorizedLogf( LogMemory.GetCategoryName(), ELogVerbosity::Log, TEXT("") );
-		Ar.CategorizedLogf( LogMemory.GetCategoryName(), ELogVerbosity::Log, TEXT("%iK allocated in pools (with %iK slack and %iK waste). Efficiency %.2f%%"), TotalMemory, TotalSlack, TotalWaste, TotalMemory ? 100.0f * (TotalMemory - TotalWaste) / TotalMemory : 100.0f );
-		Ar.CategorizedLogf( LogMemory.GetCategoryName(), ELogVerbosity::Log, TEXT("Allocations %i Current / %i Total (in %i pools)"), TotalActiveRequests, TotalTotalRequests, TotalPools );
-		Ar.CategorizedLogf( LogMemory.GetCategoryName(), ELogVerbosity::Log, TEXT("") );
+			BufferedOutput.CategorizedLogf( LogMemory.GetCategoryName(), ELogVerbosity::Log, TEXT( "" ) );
+			BufferedOutput.CategorizedLogf( LogMemory.GetCategoryName(), ELogVerbosity::Log, TEXT( "%iK allocated in pools (with %iK slack and %iK waste). Efficiency %.2f%%" ), TotalMemory, TotalSlack, TotalWaste, TotalMemory ? 100.0f * (TotalMemory - TotalWaste) / TotalMemory : 100.0f );
+			BufferedOutput.CategorizedLogf( LogMemory.GetCategoryName(), ELogVerbosity::Log, TEXT( "Allocations %i Current / %i Total (in %i pools)" ), TotalActiveRequests, TotalTotalRequests, TotalPools );
+			BufferedOutput.CategorizedLogf( LogMemory.GetCategoryName(), ELogVerbosity::Log, TEXT( "" ) );
 #endif
+		}
+
+		BufferedOutput.RedirectTo( Ar );
 	}
 
 	virtual const TCHAR * GetDescriptiveName() OVERRIDE { return TEXT("binned"); }
diff --git a/Engine/Source/Runtime/Core/Public/HAL/OutputDevices.h b/Engine/Source/Runtime/Core/Public/HAL/OutputDevices.h
--- a/Engine/Source/Runtime/Core/Public/HAL/OutputDevices.h
+++ b/Engine/Source/Runtime/Core/Public/HAL/OutputDevices.h
@@ -241,28 +241,27 @@ public:
 	FOutputDeviceRedirector.
 -----------------------------------------------------------------------------*/
 
+/** The type of lines buffered by secondary threads. */
+struct FBufferedLine
+{
+	const FString Data;
+	const ELogVerbosity::Type Verbosity;
+	const FName Category;
+
+	/** Initialization constructor. */
+	FBufferedLine( const TCHAR* InData, ELogVerbosity::Type InVerbosity, const class FName& InCategory )
+		: Data( InData )
+		, Verbosity( InVerbosity )
+		, Category( InCategory )
+	{}
+};
+
 /**
  * Class used for output redirection to allow logs to show
  */
 class CORE_API FOutputDeviceRedirector : public FOutputDevice
 {
 private:
-
-	/** The type of lines buffered by secondary threads. */
-	struct FBufferedLine
-	{
-		const FString Data;
-		const ELogVerbosity::Type Verbosity;
-		const FName Category;
-
-		/** Initialization constructor. */
-		FBufferedLine(const TCHAR* InData, ELogVerbosity::Type InVerbosity, const class FName& InCategory):
-			Data(InData),
-			Verbosity(InVerbosity),
-			Category(InCategory)
-		{}
-	};
-
 	/** A FIFO of lines logged by non-master threads. */
 	TArray<FBufferedLine> BufferedLines;
 
@@ -455,6 +454,27 @@ public:
 	}
 };
 
+/** Buffered output device. */
+class FBufferedOutputDevice : public FOutputDevice
+{
+	TArray<FBufferedLine> BufferedLines;
+
+public:
+	virtual void Serialize( const TCHAR* InData, ELogVerbosity::Type Verbosity, const class FName& Category ) OVERRIDE
+	{
+		new(BufferedLines)FBufferedLine( InData, Verbosity, Category );
+	}
+
+	/** Pushes buffered lines into the specified output device. */
+	void RedirectTo( class FOutputDevice& Ar )
+	{
+		for( const FBufferedLine& BufferedLine : BufferedLines )
+		{
+			Ar.Serialize( *BufferedLine.Data, BufferedLine.Verbosity, BufferedLine.Category );
+		}
+	}
+};
+
 /*-----------------------------------------------------------------------------
 	FOutputDeviceError subclasses.
 -----------------------------------------------------------------------------*/
