#UE4 - Remove unneeded ensure when attempting to drag drop random stuff into the editor.

diff --git a/Engine/Source/Runtime/Core/Private/Windows/WindowsApplication.cpp b/Engine/Source/Runtime/Core/Private/Windows/WindowsApplication.cpp
--- a/Engine/Source/Runtime/Core/Private/Windows/WindowsApplication.cpp
+++ b/Engine/Source/Runtime/Core/Private/Windows/WindowsApplication.cpp
@@ -1584,7 +1584,6 @@ HRESULT FWindowsApplication::OnOLEDragEnter( const HWND HWnd, const FDragDropOLE
 			break;
 		case FDragDropOLEData::None:
 		default:
-			ensureMsgf(0, TEXT("Unhandled drag/drop OLE data type: %d"), OLEData.Type);
 			break;
 	}
 
