fixed CreateTempFilename() inserting an extra '/' into the path

diff --git a/Engine/Source/Runtime/Core/Private/Misc/Paths.cpp b/Engine/Source/Runtime/Core/Private/Misc/Paths.cpp
--- a/Engine/Source/Runtime/Core/Private/Misc/Paths.cpp
+++ b/Engine/Source/Runtime/Core/Private/Misc/Paths.cpp
@@ -777,17 +777,9 @@ FString FPaths::CreateTempFilename( const TCHAR* Path, const TCHAR* Prefix, cons
 	FString UniqueFilename;
 	do
 	{
-		const int32 PathLen = FCString::Strlen( Path );
-		if( PathLen > 0 && Path[ PathLen - 1 ] != TEXT('/') )
-		{
-			UniqueFilename = FString::Printf( TEXT("%s/%s%s%s"), Path, Prefix, *FGuid::NewGuid().ToString(), Extension );
-		}
-		else
-		{
-			UniqueFilename = FString::Printf( TEXT("%s/%s%s%s"), Path, Prefix, *FGuid::NewGuid().ToString(), Extension );
-		}
+		UniqueFilename = FPaths::Combine(Path, *FString::Printf(TEXT("%s%s%s"), Prefix, *FGuid::NewGuid().ToString(), Extension));
 	}
-	while( IFileManager::Get().FileSize( *UniqueFilename ) >= 0 );
+	while (IFileManager::Get().FileSize(*UniqueFilename) >= 0);
 	
 	return UniqueFilename;
 }
