[GitHub] - Integrated a pull request for windows return codewhen ExecProcess fails (https://github.com/EpicGames/UnrealEngine/pull/40 - by [REDACTED_PROFILE_URL])

diff --git a/Engine/Source/Runtime/Core/Private/Windows/WindowsPlatformProcess.cpp b/Engine/Source/Runtime/Core/Private/Windows/WindowsPlatformProcess.cpp
--- a/Engine/Source/Runtime/Core/Private/Windows/WindowsPlatformProcess.cpp
+++ b/Engine/Source/Runtime/Core/Private/Windows/WindowsPlatformProcess.cpp
@@ -517,6 +517,11 @@ bool FWindowsPlatformProcess::ExecProcess( const TCHAR* URL, const TCHAR* Params
 	}
 	else
 	{
+		// if CreateProcess failed, we should return a useful error code, which GetLastError will have
+		if (OutReturnCode)
+		{
+			*OutReturnCode = GetLastError();
+		}
 		if (bRedirectOutput)
 		{
 			for (int32 PipeIndex = 0; PipeIndex < MaxPipeCount; ++PipeIndex)
