Fix for console builds (missing PLATFORM_DESKTOP) #ifdef.

diff --git a/Engine/Source/Runtime/Core/Private/GenericPlatform/GenericPlatformMisc.cpp b/Engine/Source/Runtime/Core/Private/GenericPlatform/GenericPlatformMisc.cpp
--- a/Engine/Source/Runtime/Core/Private/GenericPlatform/GenericPlatformMisc.cpp
+++ b/Engine/Source/Runtime/Core/Private/GenericPlatform/GenericPlatformMisc.cpp
@@ -689,6 +689,8 @@ FText FGenericPlatformMisc::GetFileManagerName()
 	return NSLOCTEXT("GenericPlatform", "FileManagerName", "File Manager");
 }
 
+#if PLATFORM_DESKTOP
+
 void FGenericPlatformMisc::EnumerateEngineInstallations(TArray< TPair<FString, FString> > &OutInstallations)
 {
 	// Read the config file.
@@ -743,3 +745,4 @@ bool FGenericPlatformMisc::GetEngineIdentifierFromRootDir(const FString &RootDir
 	return false;
 }
 
+#endif
