#ue4
#ios
* disable setting the thread priority for now, it was actually slowing things down

diff --git a/Engine/Source/Runtime/Core/Private/IOS/IOSPlatformProcess.cpp b/Engine/Source/Runtime/Core/Private/IOS/IOSPlatformProcess.cpp
--- a/Engine/Source/Runtime/Core/Private/IOS/IOSPlatformProcess.cpp
+++ b/Engine/Source/Runtime/Core/Private/IOS/IOSPlatformProcess.cpp
@@ -59,5 +59,5 @@ void FIOSPlatformProcess::SetRealTimeMode()
     Policy.constraint  = (uint32_t)(ConstraintProcessingTimeMs * MsToAbs);
     Policy.preemptible = false;
     
-    thread_policy_set(pthread_mach_thread_np(pthread_self()), THREAD_TIME_CONSTRAINT_POLICY, (thread_policy_t)&Policy, THREAD_TIME_CONSTRAINT_POLICY_COUNT);
-}
\ No newline at end of file
+ //   thread_policy_set(pthread_mach_thread_np(pthread_self()), THREAD_TIME_CONSTRAINT_POLICY, (thread_policy_t)&Policy, THREAD_TIME_CONSTRAINT_POLICY_COUNT);
+}
