Remove PLATFORM_DESKTOP defines around EnumerateEngine functions, and just leave a blank implementation on non-desktop platforms.

diff --git a/Engine/Source/Runtime/Core/Public/GenericPlatform/GenericPlatformMisc.h b/Engine/Source/Runtime/Core/Public/GenericPlatform/GenericPlatformMisc.h
--- a/Engine/Source/Runtime/Core/Public/GenericPlatform/GenericPlatformMisc.h
+++ b/Engine/Source/Runtime/Core/Public/GenericPlatform/GenericPlatformMisc.h
@@ -723,7 +723,6 @@ public:
 		return PLATFORM_HAS_TOUCH_MAIN_SCREEN;
 	}
 
-#if PLATFORM_DESKTOP
 	/**
 	* Enumerates all the registered engine installations.
 	*
@@ -747,7 +746,6 @@ public:
 	* @param	OutIdentifier		Identifier used to refer to this installation.
 	*/
 	static bool GetEngineIdentifierFromRootDir(const FString &RootDir, FString &OutIdentifier);
-#endif
 
 #if !UE_BUILD_SHIPPING
 protected:
