Add a placeholder version number to allow fixing up broken material attributes.

diff --git a/Engine/Source/Runtime/Core/Public/UObject/ObjectVersion.h b/Engine/Source/Runtime/Core/Public/UObject/ObjectVersion.h
--- a/Engine/Source/Runtime/Core/Public/UObject/ObjectVersion.h
+++ b/Engine/Source/Runtime/Core/Public/UObject/ObjectVersion.h
@@ -514,6 +514,8 @@ enum EUnrealEngineObjectUE4Version
 	VER_UE4_REFLECTION_DATA_IN_PACKAGES,
 	// Fix up old foliage components to have movable mobility
 	VER_UE4_FOLIAGE_MOVABLE_MOBILITY,
+	// Undo BreakMaterialAttributes changes as it broke old content
+	VER_UE4_UNDO_BREAK_MATERIALATTRIBUTES_CHANGE,
 
 	// -----<new versions can be added before this line>-------------------------------------------------
 	// - this needs to be the last line (see note below)
