#ttp 331271 - EDITOR: New Project Dialog doesn't help inform users when they have an illegal character in their Windows user name
#proj UE4
#branch UE4
#summary Fixed some infinite loops in FPaths::ValidatePath if the path ended with a /, or if the path contained only a drive letter

diff --git a/Engine/Source/Runtime/Core/Private/Misc/Paths.cpp b/Engine/Source/Runtime/Core/Private/Misc/Paths.cpp
--- a/Engine/Source/Runtime/Core/Private/Misc/Paths.cpp
+++ b/Engine/Source/Runtime/Core/Private/Misc/Paths.cpp
@@ -798,6 +798,13 @@ bool FPaths::ValidatePath( const FString& InPath, FText* OutReason )
 	FString Standardized = InPath;
 	NormalizeFilename(Standardized);
 	CollapseRelativeDirectories(Standardized);
+	RemoveDuplicateSlashes(Standardized);
+
+	// The loop below requires that the path not end with a /
+	if(Standardized.EndsWith(TEXT("/")))
+	{
+		Standardized = Standardized.LeftChop(1);
+	}
 
 	// Walk each part of the path looking for name errors
 	for(int32 StartPos = 0, EndPos = Standardized.Find(TEXT("/")); ; 
@@ -810,6 +817,10 @@ bool FPaths::ValidatePath( const FString& InPath, FText* OutReason )
 		// If this is the first part of the path, it's possible for it to be a drive name and is allowed to contain a colon
 		if(StartPos == 0 && IsDrive(PathPart))
 		{
+			if(bIsLastPart)
+			{
+				break;
+			}
 			continue;
 		}
 
