Fixing compilation on a case-sensitive fs (yeah, the actual file in perforce is called WIndowsPlatform.h :-( ).

diff --git a/Engine/Source/Runtime/Core/Public/HAL/Platform.h b/Engine/Source/Runtime/Core/Public/HAL/Platform.h
--- a/Engine/Source/Runtime/Core/Public/HAL/Platform.h
+++ b/Engine/Source/Runtime/Core/Public/HAL/Platform.h
@@ -17,7 +17,7 @@
 
 //@port Identify the platform here and include the platform header to setup the platform types, etc
 #if PLATFORM_WINDOWS
-	#include "Windows/WindowsPlatform.h"
+	#include "Windows/WIndowsPlatform.h"
 #elif PLATFORM_PS4
 	#include "PS4/PS4Platform.h"
 #elif PLATFORM_XBOXONE
