Fix for TTP #332150 PORTAL: Crash when running with command arguments

diff --git a/Engine/Source/Runtime/Core/Private/Mac/MacPlatformProcess.cpp b/Engine/Source/Runtime/Core/Private/Mac/MacPlatformProcess.cpp
--- a/Engine/Source/Runtime/Core/Private/Mac/MacPlatformProcess.cpp
+++ b/Engine/Source/Runtime/Core/Private/Mac/MacPlatformProcess.cpp
@@ -553,7 +553,7 @@ const TCHAR* FMacPlatformProcess::BaseDir()
 			{
 				BundledBinariesPath = [BasePath stringByAppendingPathComponent: [NSString stringWithFormat: @"Contents/UE4/%s/Binaries/Mac", TCHAR_TO_UTF8(GGameName)]];
 			}
-			if (!BundledBinariesPath)
+			if (!BundledBinariesPath || ![FileManager fileExistsAtPath:BundledBinariesPath])
 			{
 				BundledBinariesPath = [BasePath stringByAppendingPathComponent: @"Contents/UE4/Engine/Binaries/Mac"];
 			}
