#UE4
- Fixed a deadlock in the malloc binned during updating stats

diff --git a/Engine/Source/Runtime/Core/Private/HAL/MallocBinned.cpp b/Engine/Source/Runtime/Core/Private/HAL/MallocBinned.cpp
--- a/Engine/Source/Runtime/Core/Private/HAL/MallocBinned.cpp
+++ b/Engine/Source/Runtime/Core/Private/HAL/MallocBinned.cpp
@@ -22,22 +22,47 @@ DEFINE_STAT(STAT_Binned_SlackCurrent);
 
 void FMallocBinned::GetAllocatorStats( FGenericMemoryStats& out_Stats )
 {
+	FMalloc::GetAllocatorStats( out_Stats );
+
+#if	STATS
+	SIZE_T	LocalOsCurrent = 0;
+	SIZE_T	LocalOsPeak = 0;
+	SIZE_T	LocalWasteCurrent = 0;
+	SIZE_T	LocalWastePeak = 0;
+	SIZE_T	LocalUsedCurrent = 0;
+	SIZE_T	LocalUsedPeak = 0;
+	SIZE_T	LocalCurrentAllocs = 0;
+	SIZE_T	LocalTotalAllocs = 0;
+	SIZE_T	LocalSlackCurrent = 0;
+
+	{
 #ifdef USE_INTERNAL_LOCKS
-	FScopeLock ScopedLock(&AccessGuard);
+		FScopeLock ScopedLock( &AccessGuard );
 #endif
 
-	FMalloc::GetAllocatorStats( out_Stats );
-	UpdateSlackStat();
+		UpdateSlackStat();
+
+		// Copy memory stats.
+		LocalOsCurrent = OsCurrent;
+		LocalOsPeak = OsPeak;
+		LocalWasteCurrent = WasteCurrent;
+		LocalWastePeak = WastePeak;
+		LocalUsedCurrent = UsedCurrent;
+		LocalUsedPeak = UsedPeak;
+		LocalCurrentAllocs = CurrentAllocs;
+		LocalTotalAllocs = TotalAllocs;
+		LocalSlackCurrent = SlackCurrent;
+	}
 
-#if	STATS
 	// Malloc binned stats.
-	out_Stats.Add( GET_STATFNAME( STAT_Binned_OsCurrent ), OsCurrent );
-	out_Stats.Add( GET_STATFNAME( STAT_Binned_OsPeak ), OsPeak );
-	out_Stats.Add( GET_STATFNAME( STAT_Binned_WasteCurrent ), WasteCurrent );
-	out_Stats.Add( GET_STATFNAME( STAT_Binned_WastePeak ), WastePeak );
-	out_Stats.Add( GET_STATFNAME( STAT_Binned_UsedCurrent ), UsedCurrent );
-	out_Stats.Add( GET_STATFNAME( STAT_Binned_UsedPeak ), UsedPeak );
-	out_Stats.Add( GET_STATFNAME( STAT_Binned_CurrentAllocs ), CurrentAllocs );
-	out_Stats.Add( GET_STATFNAME( STAT_Binned_SlackCurrent ), SlackCurrent );
+	out_Stats.Add( GET_STATFNAME( STAT_Binned_OsCurrent ), LocalOsCurrent );
+	out_Stats.Add( GET_STATFNAME( STAT_Binned_OsPeak ), LocalOsPeak );
+	out_Stats.Add( GET_STATFNAME( STAT_Binned_WasteCurrent ), LocalWasteCurrent );
+	out_Stats.Add( GET_STATFNAME( STAT_Binned_WastePeak ), LocalWastePeak );
+	out_Stats.Add( GET_STATFNAME( STAT_Binned_UsedCurrent ), LocalUsedCurrent );
+	out_Stats.Add( GET_STATFNAME( STAT_Binned_UsedPeak ), LocalUsedPeak );
+	out_Stats.Add( GET_STATFNAME( STAT_Binned_CurrentAllocs ), LocalCurrentAllocs );
+	out_Stats.Add( GET_STATFNAME( STAT_Binned_TotalAllocs ), LocalTotalAllocs );
+	out_Stats.Add( GET_STATFNAME( STAT_Binned_SlackCurrent ), LocalSlackCurrent );
 #endif // STATS
 }
\ No newline at end of file
diff --git a/Engine/Source/Runtime/Core/Public/HAL/MallocBinned.h b/Engine/Source/Runtime/Core/Public/HAL/MallocBinned.h
--- a/Engine/Source/Runtime/Core/Public/HAL/MallocBinned.h
+++ b/Engine/Source/Runtime/Core/Public/HAL/MallocBinned.h
@@ -1120,22 +1120,44 @@ public:
 	/** Called once per frame, gathers and sets all memory allocator statistics into the corresponding stats. */
 	virtual void UpdateStats() OVERRIDE
 	{
+#if STATS
+		SIZE_T	LocalOsCurrent = 0;
+		SIZE_T	LocalOsPeak = 0;
+		SIZE_T	LocalWasteCurrent = 0;
+		SIZE_T	LocalWastePeak = 0;
+		SIZE_T	LocalUsedCurrent = 0;
+		SIZE_T	LocalUsedPeak = 0;
+		SIZE_T	LocalCurrentAllocs = 0;
+		SIZE_T	LocalTotalAllocs = 0;
+		SIZE_T	LocalSlackCurrent = 0;
+
+		{
 #ifdef USE_INTERNAL_LOCKS
-		FScopeLock ScopedLock(&AccessGuard);
+			FScopeLock ScopedLock( &AccessGuard );
 #endif
+			UpdateSlackStat();
 
-		FMalloc::UpdateStats();
+			// Copy memory stats.
+			LocalOsCurrent = OsCurrent;
+			LocalOsPeak = OsPeak;
+			LocalWasteCurrent = WasteCurrent;
+			LocalWastePeak = WastePeak;
+			LocalUsedCurrent = UsedCurrent;
+			LocalUsedPeak = UsedPeak;
+			LocalCurrentAllocs = CurrentAllocs;
+			LocalTotalAllocs = TotalAllocs;
+			LocalSlackCurrent = SlackCurrent;
+		}
 
-#if STATS
-		SET_MEMORY_STAT(STAT_Binned_OsCurrent,OsCurrent);
-		SET_MEMORY_STAT(STAT_Binned_OsPeak,OsPeak);
-		SET_MEMORY_STAT(STAT_Binned_WasteCurrent,WasteCurrent);
-		SET_MEMORY_STAT(STAT_Binned_WastePeak,WastePeak);
-		SET_MEMORY_STAT(STAT_Binned_UsedCurrent,UsedCurrent);
-		SET_MEMORY_STAT(STAT_Binned_UsedPeak,UsedPeak);
-		SET_MEMORY_STAT(STAT_Binned_CurrentAllocs,CurrentAllocs);
-		SET_MEMORY_STAT(STAT_Binned_TotalAllocs,TotalAllocs);
-		SET_MEMORY_STAT(STAT_Binned_SlackCurrent,SlackCurrent);
+		SET_MEMORY_STAT( STAT_Binned_OsCurrent, LocalOsCurrent );
+		SET_MEMORY_STAT( STAT_Binned_OsPeak, LocalOsPeak );
+		SET_MEMORY_STAT( STAT_Binned_WasteCurrent, LocalWasteCurrent );
+		SET_MEMORY_STAT( STAT_Binned_WastePeak, LocalWastePeak );
+		SET_MEMORY_STAT( STAT_Binned_UsedCurrent, LocalUsedCurrent );
+		SET_MEMORY_STAT( STAT_Binned_UsedPeak, LocalUsedPeak );
+		SET_MEMORY_STAT( STAT_Binned_CurrentAllocs, LocalCurrentAllocs );
+		SET_MEMORY_STAT( STAT_Binned_TotalAllocs, LocalTotalAllocs );
+		SET_MEMORY_STAT( STAT_Binned_SlackCurrent, LocalSlackCurrent );
 #endif
 	}
 
