#UE4 Fix uninitialized memory crash in FTextBuilder

diff --git a/Engine/Source/Runtime/Core/Public/Internationalization/Text.h b/Engine/Source/Runtime/Core/Public/Internationalization/Text.h
--- a/Engine/Source/Runtime/Core/Public/Internationalization/Text.h
+++ b/Engine/Source/Runtime/Core/Public/Internationalization/Text.h
@@ -354,6 +354,7 @@ class CORE_API FTextBuilder
 public:
 
 	FTextBuilder()
+		: IndentCount(0)
 	{
 
 	}
