#summary Internationalization now operates as a singleton with non-static initialization. Should no longer have issues with initializing before file server when applicable. Resolved related previous double initialization issue. API-breaking change for calls to most FInternationalization methods.
#ttp 331776 - L10N: ICU: Executable size for engine and programs has regressed dramatically (ICU)

diff --git a/Engine/Plugins/Editor/EpicSurvey/Source/EpicSurvey/Private/EpicSurvey.cpp b/Engine/Plugins/Editor/EpicSurvey/Source/EpicSurvey/Private/EpicSurvey.cpp
--- a/Engine/Plugins/Editor/EpicSurvey/Source/EpicSurvey/Private/EpicSurvey.cpp
+++ b/Engine/Plugins/Editor/EpicSurvey/Source/EpicSurvey/Private/EpicSurvey.cpp
@@ -152,7 +152,7 @@ void FEpicSurvey::Initialize()
 
 void FEpicSurvey::InitializeTitleCloud()
 {
-	TSharedRef<FCulture> Culture = FInternationalization::GetCurrentCulture();
+	TSharedRef<FCulture> Culture = FInternationalization::Get().GetCurrentCulture();
 
 	FString CultureString;
 	switch (CurrentCulture)
diff --git a/Engine/Source/Developer/InternationalizationSettings/Private/InternationalizationSettingsModelDetails.cpp b/Engine/Source/Developer/InternationalizationSettings/Private/InternationalizationSettingsModelDetails.cpp
--- a/Engine/Source/Developer/InternationalizationSettings/Private/InternationalizationSettingsModelDetails.cpp
+++ b/Engine/Source/Developer/InternationalizationSettings/Private/InternationalizationSettingsModelDetails.cpp
@@ -69,15 +69,17 @@ FInternationalizationSettingsModelDetails::~FInternationalizationSettingsModelDe
 
 void FInternationalizationSettingsModelDetails::OnSettingsChanged()
 {
+	FInternationalization& I18N = FInternationalization::Get();
+
 	FString SavedCultureName = Model->GetCultureName();
-	if ( SavedCultureName != FInternationalization::GetCurrentCulture()->GetName() )
+	if ( SavedCultureName != I18N.GetCurrentCulture()->GetName() )
 	{
-		SelectedCulture = FInternationalization::GetCulture(SavedCultureName);
+		SelectedCulture = I18N.GetCulture(SavedCultureName);
 		RequiresRestart = true;
 	}
 	else
 	{
-		SelectedCulture = FInternationalization::GetCurrentCulture();
+		SelectedCulture = I18N.GetCurrentCulture();
 		RequiresRestart = false;
 	}
 	RefreshAvailableLanguages();
@@ -88,6 +90,8 @@ void FInternationalizationSettingsModelDetails::OnSettingsChanged()
 
 void FInternationalizationSettingsModelDetails::CustomizeDetails( IDetailLayoutBuilder& DetailBuilder )
 {
+	FInternationalization& I18N = FInternationalization::Get();
+
 	TArray< TWeakObjectPtr<UObject> > ObjectsBeingCustomized;
 	DetailBuilder.GetObjectsBeingCustomized(ObjectsBeingCustomized);
 	check(ObjectsBeingCustomized.Num() == 1);
@@ -102,14 +106,14 @@ void FInternationalizationSettingsModelDetails::CustomizeDetails( IDetailLayoutB
 
 	// If the saved culture is not the same as the actual current culture, a restart is needed to sync them fully and properly.
 	FString SavedCultureName = Model->GetCultureName();
-	if ( SavedCultureName != FInternationalization::GetCurrentCulture()->GetName() )
+	if ( SavedCultureName != I18N.GetCurrentCulture()->GetName() )
 	{
-			SelectedCulture = FInternationalization::GetCulture(SavedCultureName);
+			SelectedCulture = I18N.GetCulture(SavedCultureName);
 			RequiresRestart = true;
 	}
 	else
 	{
-		SelectedCulture = FInternationalization::GetCurrentCulture();
+		SelectedCulture = I18N.GetCurrentCulture();
 		RequiresRestart = false;
 	}
 
@@ -241,7 +245,7 @@ void FInternationalizationSettingsModelDetails::RefreshAvailableCultures()
 
 	IPlatformFile& PlatformFile = IPlatformFile::GetPlatformPhysical();
 	TArray<FString> AllCultureNames;
-	FInternationalization::GetCultureNames(AllCultureNames);
+	FInternationalization::Get().GetCultureNames(AllCultureNames);
 	const TArray<FString> LocalizationPaths = FPaths::GetEditorLocalizationPaths();
 	for(const auto& LocalizationPath : LocalizationPaths)
 	{
@@ -261,7 +265,7 @@ void FInternationalizationSettingsModelDetails::RefreshAvailableCultures()
 				{
 					for( const auto& CultureName : AllCultureNames )
 					{
-						TSharedPtr<FCulture> Culture = FInternationalization::GetCulture(CultureName);
+						TSharedPtr<FCulture> Culture = FInternationalization::Get().GetCulture(CultureName);
 						if(Culture.IsValid() && !AvailableCultures.Contains(Culture))
 						{
 							// UE localization resource folders use "en-US" style while ICU uses "en_US" style so we replace underscores with dashes here.
@@ -456,7 +460,7 @@ void FInternationalizationSettingsModelDetails::HandleShutdownPostPackagesSaved(
 		check(Model.IsValid());
 		Model->SetCultureName(SelectedCulture->GetName());
 		Model->ShouldLoadLocalizedPropertyNames(FieldNamesCheckBox->IsChecked());
-		if(SelectedCulture != FInternationalization::GetCurrentCulture())
+		if(SelectedCulture != FInternationalization::Get().GetCurrentCulture())
 		{
 			RequiresRestart = true;
 		}
diff --git a/Engine/Source/Developer/LauncherServices/Private/Profiles/LauncherProfile.h b/Engine/Source/Developer/LauncherServices/Private/Profiles/LauncherProfile.h
--- a/Engine/Source/Developer/LauncherServices/Private/Profiles/LauncherProfile.h
+++ b/Engine/Source/Developer/LauncherServices/Private/Profiles/LauncherProfile.h
@@ -438,6 +438,7 @@ public:
 		}
 		
 		BuildConfiguration = FApp::GetBuildConfiguration();
+		FInternationalization& I18N = FInternationalization::Get();
 
 		// default build settings
 		BuildGame = false;
@@ -449,7 +450,7 @@ public:
 		CookIncremental = true;
 		CookUnversioned = false;
 		CookedCultures.Reset();
-		CookedCultures.Add(FInternationalization::GetCurrentCulture()->GetName());
+		CookedCultures.Add(I18N.GetCurrentCulture()->GetName());
 		CookedMaps.Reset();
 		CookedPlatforms.Reset();
         ForceClose = false;
@@ -471,7 +472,7 @@ public:
 		// default launch settings
 		LaunchMode = ELauncherProfileLaunchModes::DoNotLaunch;
 		DefaultLaunchRole->SetCommandLine(FString());
-		DefaultLaunchRole->SetInitialCulture(FInternationalization::GetCurrentCulture()->GetName());
+		DefaultLaunchRole->SetInitialCulture(I18N.GetCurrentCulture()->GetName());
 		DefaultLaunchRole->SetInitialMap(FString());
 		DefaultLaunchRole->SetName(TEXT("Default Role"));
 		DefaultLaunchRole->SetInstanceType(ELauncherProfileRoleInstanceTypes::StandaloneClient);
diff --git a/Engine/Source/Developer/SessionLauncher/Private/Widgets/Cook/SSessionLauncherCookByTheBookSettings.cpp b/Engine/Source/Developer/SessionLauncher/Private/Widgets/Cook/SSessionLauncherCookByTheBookSettings.cpp
--- a/Engine/Source/Developer/SessionLauncher/Private/Widgets/Cook/SSessionLauncherCookByTheBookSettings.cpp
+++ b/Engine/Source/Developer/SessionLauncher/Private/Widgets/Cook/SSessionLauncherCookByTheBookSettings.cpp
@@ -679,7 +679,7 @@ void SSessionLauncherCookByTheBookSettings::RefreshCultureList()
 	CultureList.Reset();
 
 	TArray<FString> CultureNames;
-	FInternationalization::GetCultureNames(CultureNames);
+	FInternationalization::Get().GetCultureNames(CultureNames);
 
 	if (CultureNames.Num() > 0)
 	{
@@ -733,7 +733,7 @@ void SSessionLauncherCookByTheBookSettings::HandleAllCulturesHyperlinkNavigate(
 		if (AllPlatforms)
 		{
 			TArray<FString> CultureNames;
-			FInternationalization::GetCultureNames(CultureNames);
+			FInternationalization::Get().GetCultureNames(CultureNames);
 
 			for (int32 ExtensionIndex = 0; ExtensionIndex < CultureNames.Num(); ++ExtensionIndex)
 			{
@@ -751,7 +751,7 @@ void SSessionLauncherCookByTheBookSettings::HandleAllCulturesHyperlinkNavigate(
 EVisibility SSessionLauncherCookByTheBookSettings::HandleAllCulturesHyperlinkVisibility( ) const
 {
 	TArray<FString> CultureNames;
-	FInternationalization::GetCultureNames(CultureNames);
+	FInternationalization::Get().GetCultureNames(CultureNames);
 
 	if (CultureNames.Num() > 1)
 	{
diff --git a/Engine/Source/Developer/SessionLauncher/Private/Widgets/Launch/SSessionLauncherLaunchPage.cpp b/Engine/Source/Developer/SessionLauncher/Private/Widgets/Launch/SSessionLauncherLaunchPage.cpp
--- a/Engine/Source/Developer/SessionLauncher/Private/Widgets/Launch/SSessionLauncherLaunchPage.cpp
+++ b/Engine/Source/Developer/SessionLauncher/Private/Widgets/Launch/SSessionLauncherLaunchPage.cpp
@@ -162,7 +162,7 @@ void SSessionLauncherLaunchPage::Refresh( )
 		}
 		else
 		{
-			FInternationalization::GetCultureNames(AvailableCultures);
+			FInternationalization::Get().GetCultureNames(AvailableCultures);
 		}
 
 		AvailableMaps = FGameProjectHelper::GetAvailableMaps(SelectedProfile->GetProjectBasePath(), SelectedProfile->SupportsEngineMaps(), true);
diff --git a/Engine/Source/Editor/Documentation/Private/DocumentationLink.h b/Engine/Source/Editor/Documentation/Private/DocumentationLink.h
--- a/Engine/Source/Editor/Documentation/Private/DocumentationLink.h
+++ b/Engine/Source/Editor/Documentation/Private/DocumentationLink.h
@@ -20,7 +20,7 @@ public:
 
 	static FString GetHomeUrl()
 	{
-		return GetHomeUrl(FInternationalization::GetCurrentCulture());
+		return GetHomeUrl(FInternationalization::Get().GetCurrentCulture());
 	}
 
 	static FString GetHomeUrl(const TSharedRef<FCulture>& Culture)
@@ -34,7 +34,7 @@ public:
 
 	static FString ToUrl(const FString& Link)
 	{
-		return ToUrl(Link, FInternationalization::GetCurrentCulture());
+		return ToUrl(Link, FInternationalization::Get().GetCurrentCulture());
 	}
 
 	static FString ToUrl(const FString& Link, const TSharedRef<FCulture>& Culture)
@@ -50,11 +50,13 @@ public:
 
 	static FString ToFilePath( const FString& Link )
 	{
-		FString FilePath = ToFilePath(Link, FInternationalization::GetCurrentCulture());
+		FInternationalization& I18N = FInternationalization::Get();
+
+		FString FilePath = ToFilePath(Link, I18N.GetCurrentCulture());
 
 		if (!FPaths::FileExists(FilePath))
 		{
-			const TSharedPtr<FCulture> FallbackCulture = FInternationalization::GetCulture(TEXT("en"));
+			const TSharedPtr<FCulture> FallbackCulture = I18N.GetCulture(TEXT("en"));
 			if (FallbackCulture.IsValid())
 			{
 				const FString FallbackFilePath = ToFilePath(Link, FallbackCulture.ToSharedRef());
@@ -80,12 +82,14 @@ public:
 
 	static FString ToFileUrl( const FString& Link )
 	{
-		TSharedRef<FCulture> Culture = FInternationalization::GetCurrentCulture();
+		FInternationalization& I18N = FInternationalization::Get();
+
+		TSharedRef<FCulture> Culture = I18N.GetCurrentCulture();
 		FString FilePath = ToFilePath(Link, Culture);
 
 		if (!FPaths::FileExists(FilePath))
 		{
-			const TSharedPtr<FCulture> FallbackCulture = FInternationalization::GetCulture(TEXT("en"));
+			const TSharedPtr<FCulture> FallbackCulture = I18N.GetCulture(TEXT("en"));
 			if (FallbackCulture.IsValid())
 			{
 				const FString FallbackFilePath = ToFilePath(Link, FallbackCulture.ToSharedRef());
@@ -110,11 +114,13 @@ public:
 
 	static FString ToSourcePath(const FString& Link)
 	{
-		FString SourcePath = ToSourcePath(Link, FInternationalization::GetCurrentCulture());
+		FInternationalization& I18N = FInternationalization::Get();
+
+		FString SourcePath = ToSourcePath(Link, I18N.GetCurrentCulture());
 
 		if (!FPaths::FileExists(SourcePath))
 		{
-			const TSharedPtr<FCulture> FallbackCulture = FInternationalization::GetCulture(TEXT("en"));
+			const TSharedPtr<FCulture> FallbackCulture = I18N.GetCulture(TEXT("en"));
 			if (FallbackCulture.IsValid())
 			{
 				const FString FallbackSourcePath = ToSourcePath(Link, FallbackCulture.ToSharedRef());
diff --git a/Engine/Source/Editor/Documentation/Private/SDocumentationToolTip.cpp b/Engine/Source/Editor/Documentation/Private/SDocumentationToolTip.cpp
--- a/Engine/Source/Editor/Documentation/Private/SDocumentationToolTip.cpp
+++ b/Engine/Source/Editor/Documentation/Private/SDocumentationToolTip.cpp
@@ -131,7 +131,7 @@ void SDocumentationToolTip::ConstructSimpleTipContent()
 		{
 			if ( GEditor->EditorUserSettings->bDisplayDocumentationLink && FSlateApplication::Get().SupportsSourceAccess() )
 			{
-				FString DocPath = FDocumentationLink::ToSourcePath( DocumentationLink, FInternationalization::GetCurrentCulture() );
+				FString DocPath = FDocumentationLink::ToSourcePath( DocumentationLink, FInternationalization::Get().GetCurrentCulture() );
 				if ( !FPaths::FileExists(DocPath) )
 				{
 					DocPath = FPaths::ConvertRelativePathToFull(DocPath);
@@ -283,7 +283,7 @@ void SDocumentationToolTip::ConstructFullTipContent()
 				[
 					SNew( SHyperlink )
 						.Text( NSLOCTEXT( "SToolTip", "EditDocumentationMessage_Edit", "edit" ) )
-						.OnNavigate_Static(&Local::EditSource, FPaths::ConvertRelativePathToFull(FDocumentationLink::ToSourcePath(DocumentationLink, FInternationalization::GetCurrentCulture())) + TEXT("|") + FString::FromInt(Excerpts[ExcerptIndex].LineNumber))
+						.OnNavigate_Static(&Local::EditSource, FPaths::ConvertRelativePathToFull(FDocumentationLink::ToSourcePath(DocumentationLink, FInternationalization::Get().GetCurrentCulture())) + TEXT("|") + FString::FromInt(Excerpts[ExcerptIndex].LineNumber))
 				];
 			}
 		}
diff --git a/Engine/Source/Editor/GameProjectGeneration/Private/TemplateProjectDefs.cpp b/Engine/Source/Editor/GameProjectGeneration/Private/TemplateProjectDefs.cpp
--- a/Engine/Source/Editor/GameProjectGeneration/Private/TemplateProjectDefs.cpp
+++ b/Engine/Source/Editor/GameProjectGeneration/Private/TemplateProjectDefs.cpp
@@ -45,7 +45,7 @@ void UTemplateProjectDefs::FixupStrings(const FString& TemplateName, const FStri
 
 FText UTemplateProjectDefs::GetDisplayNameText()
 {
-	const TSharedRef< FCulture > CurrentCulture = FInternationalization::GetCurrentCulture();
+	const TSharedRef< FCulture > CurrentCulture = FInternationalization::Get().GetCurrentCulture();
 	for ( auto NameIt = LocalizedDisplayNames.CreateConstIterator(); NameIt; ++NameIt )
 	{
 		const FLocalizedTemplateString& Name = *NameIt;
@@ -60,7 +60,7 @@ FText UTemplateProjectDefs::GetDisplayNameText()
 
 FText UTemplateProjectDefs::GetLocalizedDescription()
 {
-	const TSharedRef< FCulture >  CurrentCulture = FInternationalization::GetCurrentCulture();
+	const TSharedRef< FCulture >  CurrentCulture = FInternationalization::Get().GetCurrentCulture();
 	for ( auto DescriptionIt = LocalizedDescriptions.CreateConstIterator(); DescriptionIt; ++DescriptionIt )
 	{
 		const FLocalizedTemplateString& Description = *DescriptionIt;
diff --git a/Engine/Source/Editor/Kismet/Private/SBlueprintDiff.cpp b/Engine/Source/Editor/Kismet/Private/SBlueprintDiff.cpp
--- a/Engine/Source/Editor/Kismet/Private/SBlueprintDiff.cpp
+++ b/Engine/Source/Editor/Kismet/Private/SBlueprintDiff.cpp
@@ -261,7 +261,7 @@ FText FListItemGraphToDiff::GetToolTip()
 
 		if ( Revision.Revision >= 0 )
 		{
-			RevisionText = FText::Format( LOCTEXT("Revision Number", "Revision {0}") , FText::AsNumber( Revision.Revision, NULL, FInternationalization::GetInvariantCulture() ) );
+			RevisionText = FText::Format( LOCTEXT("Revision Number", "Revision {0}") , FText::AsNumber( Revision.Revision, NULL, FInternationalization::Get().GetInvariantCulture() ) );
 		}
 
 		return FText::Format( LOCTEXT("MissingGraph", "Graph '{0}' missing from {1}"), FText::FromString( GoodGraph->GetName() ), RevisionText );
diff --git a/Engine/Source/Editor/Kismet/Private/SBlueprintEditorToolbar.cpp b/Engine/Source/Editor/Kismet/Private/SBlueprintEditorToolbar.cpp
--- a/Engine/Source/Editor/Kismet/Private/SBlueprintEditorToolbar.cpp
+++ b/Engine/Source/Editor/Kismet/Private/SBlueprintEditorToolbar.cpp
@@ -340,10 +340,12 @@ class SBlueprintDiffMenu : public SCompoundWidget
 					TSharedPtr<ISourceControlRevision, ESPMode::ThreadSafe> Revision = SourceControlState->GetHistoryItem(HistoryIndex);
 					if(Revision.IsValid())
 					{
-						FText Label = FText::Format( LOCTEXT("RevisionNumber", "Revision {0}"), FText::AsNumber( Revision->GetRevisionNumber(), NULL, FInternationalization::GetInvariantCulture() ) );
+						FInternationalization& I18N = FInternationalization::Get();
+
+						FText Label = FText::Format( LOCTEXT("RevisionNumber", "Revision {0}"), FText::AsNumber( Revision->GetRevisionNumber(), NULL, I18N.GetInvariantCulture() ) );
 
 						FFormatNamedArguments Args;
-						Args.Add( TEXT("CheckInNumber"), FText::AsNumber( Revision->GetCheckInIdentifier(), NULL, FInternationalization::GetInvariantCulture() ) );
+						Args.Add( TEXT("CheckInNumber"), FText::AsNumber( Revision->GetCheckInIdentifier(), NULL, I18N.GetInvariantCulture() ) );
 						Args.Add( TEXT("UserName"), FText::FromString( Revision->GetUserName() ) );
 						Args.Add( TEXT("DateTime"), FText::AsDate( Revision->GetDate().ToDate() ) );
 						Args.Add( TEXT("ChanglistDescription"), FText::FromString( Revision->GetDescription() ) );
diff --git a/Engine/Source/Editor/MainFrame/Private/Menus/TranslationEditorMenu.h b/Engine/Source/Editor/MainFrame/Private/Menus/TranslationEditorMenu.h
--- a/Engine/Source/Editor/MainFrame/Private/Menus/TranslationEditorMenu.h
+++ b/Engine/Source/Editor/MainFrame/Private/Menus/TranslationEditorMenu.h
@@ -137,7 +137,7 @@ public:
 				FString ArchiveName = FPaths::GetBaseFilename(FPaths::GetPath(ArchiveFileName));
 				FString ManifestName = FPaths::GetBaseFilename(ManifestFileName);
 
-				TSharedPtr<FCulture> Culture = FInternationalization::GetCulture(ArchiveName);
+				TSharedPtr<FCulture> Culture = FInternationalization::Get().GetCulture(ArchiveName);
 				if (Culture.IsValid())
 				{
 					ArchiveName = Culture->GetDisplayName();
diff --git a/Engine/Source/Editor/SourceControlWindows/Private/SSourceControlHistory.cpp b/Engine/Source/Editor/SourceControlWindows/Private/SSourceControlHistory.cpp
--- a/Engine/Source/Editor/SourceControlWindows/Private/SSourceControlHistory.cpp
+++ b/Engine/Source/Editor/SourceControlWindows/Private/SSourceControlHistory.cpp
@@ -510,14 +510,14 @@ public:
 				.AutoWidth()
 				[
 					SNew(STextBlock)
-					.Text( FText::AsNumber( RevisionListItem->RevisionNumber, NULL, FInternationalization::GetInvariantCulture() ) )
+					.Text( FText::AsNumber( RevisionListItem->RevisionNumber, NULL, FInternationalization::Get().GetInvariantCulture() ) )
 				];	
 		}
 		else if (ColumnName == TEXT("Changelist"))
 		{
 			return
 				SNew(STextBlock) 
-				.Text( FText::AsNumber( RevisionListItem->ChangelistNumber, NULL, FInternationalization::GetInvariantCulture() ) ) ;
+				.Text( FText::AsNumber( RevisionListItem->ChangelistNumber, NULL, FInternationalization::Get().GetInvariantCulture() ) ) ;
 		}
 		else if (ColumnName == TEXT("Date"))
 		{
diff --git a/Engine/Source/Editor/UnrealEd/Private/Commandlets/GenerateGatherArchiveCommandlet.cpp b/Engine/Source/Editor/UnrealEd/Private/Commandlets/GenerateGatherArchiveCommandlet.cpp
--- a/Engine/Source/Editor/UnrealEd/Private/Commandlets/GenerateGatherArchiveCommandlet.cpp
+++ b/Engine/Source/Editor/UnrealEd/Private/Commandlets/GenerateGatherArchiveCommandlet.cpp
@@ -22,6 +22,8 @@ UGenerateGatherArchiveCommandlet::UGenerateGatherArchiveCommandlet(const class F
 
 int32 UGenerateGatherArchiveCommandlet::Main( const FString& Params )
 {
+	FInternationalization& I18N = FInternationalization::Get();
+
 	// Parse command line - we're interested in the param vals
 	TArray<FString> Tokens;
 	TArray<FString> Switches;
@@ -76,7 +78,7 @@ int32 UGenerateGatherArchiveCommandlet::Main( const FString& Params )
 	FString SourceCulture;
 	if( GetConfigString( *SectionName, TEXT("SourceCulture"), SourceCulture, GatherTextConfigPath ) )
 	{
-		if( FInternationalization::GetCulture( SourceCulture ).IsValid() )
+		if( I18N.GetCulture( SourceCulture ).IsValid() )
 		{
 			UE_LOG(LogGenerateArchiveCommandlet, Verbose, TEXT("Specified culture is not a valid runtime culture, but may be a valid base language: %s"), *(SourceCulture) );
 		}
@@ -94,7 +96,7 @@ int32 UGenerateGatherArchiveCommandlet::Main( const FString& Params )
 
 	for(int32 i = 0; i < CulturesToGenerate.Num(); ++i)
 	{
-		if( FInternationalization::GetCulture( CulturesToGenerate[i] ).IsValid() )
+		if( I18N.GetCulture( CulturesToGenerate[i] ).IsValid() )
 		{
 			UE_LOG(LogGenerateArchiveCommandlet, Verbose, TEXT("Specified culture is not a valid runtime culture, but may be a valid base language: %s"), *(CulturesToGenerate[i]) );
 		}
diff --git a/Engine/Source/Editor/UnrealEd/Private/Commandlets/GenerateTextLocalizationReportCommandlet.cpp b/Engine/Source/Editor/UnrealEd/Private/Commandlets/GenerateTextLocalizationReportCommandlet.cpp
--- a/Engine/Source/Editor/UnrealEd/Private/Commandlets/GenerateTextLocalizationReportCommandlet.cpp
+++ b/Engine/Source/Editor/UnrealEd/Private/Commandlets/GenerateTextLocalizationReportCommandlet.cpp
@@ -160,7 +160,7 @@ bool UGenerateTextLocalizationReportCommandlet::ProcessWordCountReport(const FSt
 
 	for(int32 i = 0; i < CulturesToGenerate.Num(); ++i)
 	{
-		if( FInternationalization::GetCulture( CulturesToGenerate[i] ).IsValid() )
+		if( FInternationalization::Get().GetCulture( CulturesToGenerate[i] ).IsValid() )
 		{
 			UE_LOG(LogGenerateTextLocalizationReportCommandlet, Verbose, TEXT("Specified culture is not a valid runtime culture, but may be a valid base language: %s"), *(CulturesToGenerate[i]));
 		}
diff --git a/Engine/Source/Editor/UnrealEd/Private/Commandlets/GenerateTextLocalizationResourceCommandlet.cpp b/Engine/Source/Editor/UnrealEd/Private/Commandlets/GenerateTextLocalizationResourceCommandlet.cpp
--- a/Engine/Source/Editor/UnrealEd/Private/Commandlets/GenerateTextLocalizationResourceCommandlet.cpp
+++ b/Engine/Source/Editor/UnrealEd/Private/Commandlets/GenerateTextLocalizationResourceCommandlet.cpp
@@ -79,7 +79,7 @@ int32 UGenerateTextLocalizationResourceCommandlet::Main(const FString& Params)
 
 	for(int32 i = 0; i < CulturesToGenerate.Num(); ++i)
 	{
-		if( FInternationalization::GetCulture( CulturesToGenerate[i] ).IsValid() )
+		if( FInternationalization::Get().GetCulture( CulturesToGenerate[i] ).IsValid() )
 		{
 			UE_LOG(LogGenerateTextLocalizationResourceCommandlet, Verbose, TEXT("Specified culture is not a valid runtime culture, but may be a valid base language: %s"), *(CulturesToGenerate[i]));
 		}
diff --git a/Engine/Source/Editor/UnrealEd/Private/Settings/LevelEditorPlaySettingsCustomization.h b/Engine/Source/Editor/UnrealEd/Private/Settings/LevelEditorPlaySettingsCustomization.h
--- a/Engine/Source/Editor/UnrealEd/Private/Settings/LevelEditorPlaySettingsCustomization.h
+++ b/Engine/Source/Editor/UnrealEd/Private/Settings/LevelEditorPlaySettingsCustomization.h
@@ -224,9 +224,11 @@ protected:
 			{
 				FUIAction Action(FExecuteAction::CreateRaw(this, &SScreenResolutionCustomization::HandleCommonResolutionSelected, Iter->Width, Iter->Height));
 
+				FInternationalization& I18N = FInternationalization::Get();
+
 				FFormatNamedArguments Args;
-				Args.Add(TEXT("Width"), FText::AsNumber(Iter->Width, NULL, FInternationalization::GetInvariantCulture()));
-				Args.Add(TEXT("Height"), FText::AsNumber(Iter->Height, NULL, FInternationalization::GetInvariantCulture()));
+				Args.Add(TEXT("Width"), FText::AsNumber(Iter->Width, NULL, I18N.GetInvariantCulture()));
+				Args.Add(TEXT("Height"), FText::AsNumber(Iter->Height, NULL, I18N.GetInvariantCulture()));
 				Args.Add(TEXT("AspectRatio"), FText::FromString(Iter->AspectRatio));
 
 				MenuBuilder.AddMenuEntry(FText::FromString(Iter->Description), FText::Format(LOCTEXT("CommonResolutionFormat", "{Width} x {Height} (AspectRatio)"), Args), FSlateIcon(), Action);
diff --git a/Engine/Source/Editor/UnrealEd/Private/UnrealEdMisc.cpp b/Engine/Source/Editor/UnrealEd/Private/UnrealEdMisc.cpp
--- a/Engine/Source/Editor/UnrealEd/Private/UnrealEdMisc.cpp
+++ b/Engine/Source/Editor/UnrealEd/Private/UnrealEdMisc.cpp
@@ -1023,7 +1023,9 @@ FString FUnrealEdMisc::GenerateURL(const FString& InUDNPage)
 {
 	if( InUDNPage.Len() > 0 )
 	{
-		const FString PageURL = FString::Printf( TEXT( "%s/Editor/LevelEditing/MapErrors/index.html" ), *FInternationalization::GetCurrentCulture()->GetUnrealLegacyThreeLetterISOLanguageName() );
+		FInternationalization& I18N = FInternationalization::Get();
+
+		const FString PageURL = FString::Printf( TEXT( "%s/Editor/LevelEditing/MapErrors/index.html" ), *I18N.GetCurrentCulture()->GetUnrealLegacyThreeLetterISOLanguageName() );
 		const FString BookmarkURL = FString::Printf( TEXT( "#%s" ), *InUDNPage );
 
 		// Developers can browse documentation included with the engine distribution, check for file presence...
@@ -1035,7 +1037,7 @@ FString FUnrealEdMisc::GenerateURL(const FString& InUDNPage)
 		// ... if it's not present, fallback to using the online version, if the full URL is provided...
 		else if(FUnrealEdMisc::Get().GetURL( TEXT("MapErrorURL"), MapErrorURL, true ) && MapErrorURL.EndsWith( TEXT( ".html" ) ))
 		{	
-			MapErrorURL.ReplaceInline( TEXT( "/INT/" ), *FString::Printf( TEXT( "/%s/" ), *FInternationalization::GetCurrentCulture()->GetUnrealLegacyThreeLetterISOLanguageName() ) );
+			MapErrorURL.ReplaceInline( TEXT( "/INT/" ), *FString::Printf( TEXT( "/%s/" ), *I18N.GetCurrentCulture()->GetUnrealLegacyThreeLetterISOLanguageName() ) );
 			MapErrorURL += BookmarkURL;
 		}
 		// ...otherwise, attempt to create the URL from what we know here...
diff --git a/Engine/Source/Editor/UnrealEd/Private/UnrealEdSrv.cpp b/Engine/Source/Editor/UnrealEd/Private/UnrealEdSrv.cpp
--- a/Engine/Source/Editor/UnrealEd/Private/UnrealEdSrv.cpp
+++ b/Engine/Source/Editor/UnrealEd/Private/UnrealEdSrv.cpp
@@ -745,7 +745,7 @@ bool UUnrealEdEngine::Exec( UWorld* InWorld, const TCHAR* Stream, FOutputDevice&
 		FString CultureName;
 		if( FParse::Value(Str,TEXT("CULTURE="), CultureName) )
 		{
-			FInternationalization::SetCurrentCulture( CultureName );
+			FInternationalization::Get().SetCurrentCulture( CultureName );
 		}
 	}
 
@@ -2923,7 +2923,7 @@ bool UUnrealEdEngine::Exec_Mode( const TCHAR* Str, FOutputDevice& Ar )
 	FString CultureName;
 	if( FParse::Value(Str,TEXT("CULTURE="), CultureName) )
 	{
-		FInternationalization::SetCurrentCulture( CultureName );
+		FInternationalization::Get().SetCurrentCulture( CultureName );
 	}
 
 	FString ConfigFilePath;
diff --git a/Engine/Source/Runtime/Core/Private/Internationalization/ICUText.cpp b/Engine/Source/Runtime/Core/Private/Internationalization/ICUText.cpp
--- a/Engine/Source/Runtime/Core/Private/Internationalization/ICUText.cpp
+++ b/Engine/Source/Runtime/Core/Private/Internationalization/ICUText.cpp
@@ -13,8 +13,9 @@
 
 FText FText::AsDate(const FDateTime::FDate& Date, const EDateTimeStyle::Type DateStyle, const TSharedPtr<FCulture>& TargetCulture)
 {
-	checkf(FInternationalization::IsInitialized() == true, TEXT("FInternationalization is not initialized. An FText formatting method was likely used in static object initialization - this is not supported."));
-	TSharedPtr<FCulture> Culture = TargetCulture.IsValid() ? TargetCulture : FInternationalization::GetCurrentCulture();
+	FInternationalization& I18N = FInternationalization::Get();
+	checkf(I18N.IsInitialized() == true, TEXT("FInternationalization is not initialized. An FText formatting method was likely used in static object initialization - this is not supported."));
+	TSharedPtr<FCulture> Culture = TargetCulture.IsValid() ? TargetCulture : I18N.GetCurrentCulture();
 
 	FDateTime DateTime(Date.Year, Date.Month, Date.Day);
 	int32 UNIXTimestamp = DateTime.ToUnixTimestamp();
@@ -33,8 +34,9 @@ FText FText::AsDate(const FDateTime::FDate& Date, const EDateTimeStyle::Type Dat
 
 FText FText::AsTime(const FDateTime::FTime& Time, const EDateTimeStyle::Type TimeStyle, const FString& TimeZone, const TSharedPtr<FCulture>& TargetCulture)
 {
-	checkf(FInternationalization::IsInitialized() == true, TEXT("FInternationalization is not initialized. An FText formatting method was likely used in static object initialization - this is not supported."));
-	TSharedPtr<FCulture> Culture = TargetCulture.IsValid() ? TargetCulture : FInternationalization::GetCurrentCulture();
+	FInternationalization& I18N = FInternationalization::Get();
+	checkf(I18N.IsInitialized() == true, TEXT("FInternationalization is not initialized. An FText formatting method was likely used in static object initialization - this is not supported."));
+	TSharedPtr<FCulture> Culture = TargetCulture.IsValid() ? TargetCulture : I18N.GetCurrentCulture();
 
 	FDateTime DateTime(1970, 1, 1, Time.Hour, Time.Minute, Time.Second, Time.Millisecond);
 	int32 UNIXTimestamp = DateTime.ToUnixTimestamp();
@@ -53,8 +55,9 @@ FText FText::AsTime(const FDateTime::FTime& Time, const EDateTimeStyle::Type Tim
 
 FText FText::AsTime(const FTimespan& Time, const TSharedPtr<FCulture>& TargetCulture)
 {
-	checkf(FInternationalization::IsInitialized() == true, TEXT("FInternationalization is not initialized. An FText formatting method was likely used in static object initialization - this is not supported."));
-	TSharedPtr<FCulture> Culture = TargetCulture.IsValid() ? TargetCulture : FInternationalization::GetCurrentCulture();
+	FInternationalization& I18N = FInternationalization::Get();
+	checkf(I18N.IsInitialized() == true, TEXT("FInternationalization is not initialized. An FText formatting method was likely used in static object initialization - this is not supported."));
+	TSharedPtr<FCulture> Culture = TargetCulture.IsValid() ? TargetCulture : I18N.GetCurrentCulture();
 
 	FText TimespanFormatPattern = NSLOCTEXT("Timespan", "FormatPattern", "{Hours}:{Minutes}:{Seconds}");
 
@@ -76,8 +79,9 @@ FText FText::AsTime(const FTimespan& Time, const TSharedPtr<FCulture>& TargetCul
 
 FText FText::AsDateTime(const FDateTime& DateTime, const EDateTimeStyle::Type DateStyle, const EDateTimeStyle::Type TimeStyle, const FString& TimeZone, const TSharedPtr<FCulture>& TargetCulture)
 {
-	checkf(FInternationalization::IsInitialized() == true, TEXT("FInternationalization is not initialized. An FText formatting method was likely used in static object initialization - this is not supported."));
-	TSharedPtr<FCulture> Culture = TargetCulture.IsValid() ? TargetCulture : FInternationalization::GetCurrentCulture();
+	FInternationalization& I18N = FInternationalization::Get();
+	checkf(I18N.IsInitialized() == true, TEXT("FInternationalization is not initialized. An FText formatting method was likely used in static object initialization - this is not supported."));
+	TSharedPtr<FCulture> Culture = TargetCulture.IsValid() ? TargetCulture : I18N.GetCurrentCulture();
 
 	int32 UNIXTimestamp = DateTime.ToUnixTimestamp();
 	UDate ICUDate = static_cast<double>(UNIXTimestamp) * U_MILLIS_PER_SECOND;
@@ -95,7 +99,7 @@ FText FText::AsDateTime(const FDateTime& DateTime, const EDateTimeStyle::Type Da
 
 FText FText::AsMemory(SIZE_T NumBytes, const FNumberFormattingOptions* const Options, const TSharedPtr<FCulture>& TargetCulture)
 {
-	checkf(FInternationalization::IsInitialized() == true, TEXT("FInternationalization is not initialized. An FText formatting method was likely used in static object initialization - this is not supported."));
+	checkf(FInternationalization::Get().IsInitialized() == true, TEXT("FInternationalization is not initialized. An FText formatting method was likely used in static object initialization - this is not supported."));
 	FFormatNamedArguments Args;
 
 	if (NumBytes < 1024)
@@ -122,7 +126,7 @@ int32 FText::CompareTo( const FText& Other, const ETextComparisonLevel::Type Com
 {
 	UErrorCode ICUStatus = U_ZERO_ERROR;
 
-	const TSharedRef<const icu::Collator> Collator( FInternationalization::GetCurrentCulture()->Implementation->GetCollator(ComparisonLevel) );
+	const TSharedRef<const icu::Collator> Collator( FInternationalization::Get().GetCurrentCulture()->Implementation->GetCollator(ComparisonLevel) );
 
 	icu::UnicodeString A;
 	ICUUtilities::Convert(DisplayString.Get(), A);
@@ -142,7 +146,7 @@ bool FText::EqualTo( const FText& Other, const ETextComparisonLevel::Type Compar
 {
 	UErrorCode ICUStatus = U_ZERO_ERROR;
 
-	const TSharedRef<const icu::Collator> Collator( FInternationalization::GetCurrentCulture()->Implementation->GetCollator(ComparisonLevel) );
+	const TSharedRef<const icu::Collator> Collator( FInternationalization::Get().GetCurrentCulture()->Implementation->GetCollator(ComparisonLevel) );
 
 	icu::UnicodeString A;
 	ICUUtilities::Convert(DisplayString.Get(), A);
@@ -174,7 +178,7 @@ private:
 public:
 	FSortPredicateImplementation(const ETextComparisonLevel::Type InComparisonLevel)
 		: ComparisonLevel(InComparisonLevel)
-		, ICUCollator(FInternationalization::GetCurrentCulture()->Implementation->GetCollator(InComparisonLevel))
+		, ICUCollator(FInternationalization::Get().GetCurrentCulture()->Implementation->GetCollator(InComparisonLevel))
 	{
 	}
 
@@ -247,7 +251,7 @@ void FText::GetFormatPatternParameters(const FText& Pattern, TArray<FString>& Pa
 
 FText FText::Format(const FText& Pattern, const FFormatNamedArguments& Arguments)
 {
-	checkf(FInternationalization::IsInitialized() == true, TEXT("FInternationalization is not initialized. An FText formatting method was likely used in static object initialization - this is not supported."));
+	checkf(FInternationalization::Get().IsInitialized() == true, TEXT("FInternationalization is not initialized. An FText formatting method was likely used in static object initialization - this is not supported."));
 	//SCOPE_CYCLE_COUNTER( STAT_TextFormat );
 
 	const bool bEnableErrorResults = ENABLE_TEXT_ERROR_CHECKING_RESULTS && bEnableErrorCheckingResults;
@@ -326,7 +330,7 @@ FText FText::Format(const FText& Pattern, const FFormatNamedArguments& Arguments
 
 FText FText::Format(const FText& Pattern, const FFormatOrderedArguments& Arguments)
 {
-	checkf(FInternationalization::IsInitialized() == true, TEXT("FInternationalization is not initialized. An FText formatting method was likely used in static object initialization - this is not supported."));
+	checkf(FInternationalization::Get().IsInitialized() == true, TEXT("FInternationalization is not initialized. An FText formatting method was likely used in static object initialization - this is not supported."));
 	//SCOPE_CYCLE_COUNTER( STAT_TextFormat );
 
 	const bool bEnableErrorResults = ENABLE_TEXT_ERROR_CHECKING_RESULTS && bEnableErrorCheckingResults;
@@ -404,7 +408,7 @@ FText FText::Format(const FText& Pattern, const FFormatOrderedArguments& Argumen
 
 FText FText::Format(const FText& Pattern, const TArray< FFormatArgumentData > InArguments)
 {
-	checkf(FInternationalization::IsInitialized() == true, TEXT("FInternationalization is not initialized. An FText formatting method was likely used in static object initialization - this is not supported."));
+	checkf(FInternationalization::Get().IsInitialized() == true, TEXT("FInternationalization is not initialized. An FText formatting method was likely used in static object initialization - this is not supported."));
 	//SCOPE_CYCLE_COUNTER( STAT_TextFormat );
 
 	const bool bEnableErrorResults = ENABLE_TEXT_ERROR_CHECKING_RESULTS && bEnableErrorCheckingResults;
diff --git a/Engine/Source/Runtime/Core/Private/Internationalization/ICUText.h b/Engine/Source/Runtime/Core/Private/Internationalization/ICUText.h
--- a/Engine/Source/Runtime/Core/Private/Internationalization/ICUText.h
+++ b/Engine/Source/Runtime/Core/Private/Internationalization/ICUText.h
@@ -12,8 +12,9 @@
 template<typename T1, typename T2>
 FText FText::AsNumberTemplate(T1 Val, const FNumberFormattingOptions* const Options, const TSharedPtr<FCulture>& TargetCulture)
 {
-	checkf(FInternationalization::IsInitialized() == true, TEXT("FInternationalization is not initialized. An FText formatting method was likely used in static object initialization - this is not supported."));
-	const TSharedRef<FCulture> Culture = TargetCulture.IsValid() ? TargetCulture.ToSharedRef() : FInternationalization::GetCurrentCulture();
+	FInternationalization& I18N = FInternationalization::Get();
+	checkf(I18N.IsInitialized() == true, TEXT("FInternationalization is not initialized. An FText formatting method was likely used in static object initialization - this is not supported."));
+	const TSharedRef<FCulture> Culture = TargetCulture.IsValid() ? TargetCulture.ToSharedRef() : I18N.GetCurrentCulture();
 	UErrorCode ICUStatus = U_ZERO_ERROR;
 	const TSharedRef<const icu::DecimalFormat> ICUDecimalFormat( Culture->Implementation->GetDecimalFormatter(Options) );
 	icu::Formattable FormattableVal(static_cast<T2>(Val));
@@ -29,8 +30,9 @@ FText FText::AsNumberTemplate(T1 Val, const FNumberFormattingOptions* const Opti
 template<typename T1, typename T2>
 FText FText::AsCurrencyTemplate(T1 Val, const FNumberFormattingOptions* const Options, const TSharedPtr<FCulture>& TargetCulture)
 {
-	checkf(FInternationalization::IsInitialized() == true, TEXT("FInternationalization is not initialized. An FText formatting method was likely used in static object initialization - this is not supported."));
-	const TSharedRef<FCulture> Culture = TargetCulture.IsValid() ? TargetCulture.ToSharedRef() : FInternationalization::GetCurrentCulture();
+	FInternationalization& I18N = FInternationalization::Get();
+	checkf(I18N.IsInitialized() == true, TEXT("FInternationalization is not initialized. An FText formatting method was likely used in static object initialization - this is not supported."));
+	const TSharedRef<FCulture> Culture = TargetCulture.IsValid() ? TargetCulture.ToSharedRef() : I18N.GetCurrentCulture();
 	UErrorCode ICUStatus = U_ZERO_ERROR;
 	const TSharedRef<const icu::DecimalFormat> ICUDecimalFormat( Culture->Implementation->GetCurrencyFormatter(Options) );
 	icu::Formattable FormattableVal(static_cast<T2>(Val));
@@ -46,8 +48,9 @@ FText FText::AsCurrencyTemplate(T1 Val, const FNumberFormattingOptions* const Op
 template<typename T1, typename T2>
 FText FText::AsPercentTemplate(T1 Val, const FNumberFormattingOptions* const Options, const TSharedPtr<FCulture>& TargetCulture)
 {
-	checkf(FInternationalization::IsInitialized() == true, TEXT("FInternationalization is not initialized. An FText formatting method was likely used in static object initialization - this is not supported."));
-	const TSharedRef<FCulture> Culture = TargetCulture.IsValid() ? TargetCulture.ToSharedRef() : FInternationalization::GetCurrentCulture();
+	FInternationalization& I18N = FInternationalization::Get();
+	checkf(I18N.IsInitialized() == true, TEXT("FInternationalization is not initialized. An FText formatting method was likely used in static object initialization - this is not supported."));
+	const TSharedRef<FCulture> Culture = TargetCulture.IsValid() ? TargetCulture.ToSharedRef() : I18N.GetCurrentCulture();
 	UErrorCode ICUStatus = U_ZERO_ERROR;
 	const TSharedRef<const icu::DecimalFormat> ICUDecimalFormat( Culture->Implementation->GetPercentFormatter(Options) );
 	icu::Formattable FormattableVal(static_cast<T2>(Val));
diff --git a/Engine/Source/Runtime/Core/Private/Internationalization/Internationalization.cpp b/Engine/Source/Runtime/Core/Private/Internationalization/Internationalization.cpp
--- a/Engine/Source/Runtime/Core/Private/Internationalization/Internationalization.cpp
+++ b/Engine/Source/Runtime/Core/Private/Internationalization/Internationalization.cpp
@@ -23,19 +23,19 @@
 
 #define LOCTEXT_NAMESPACE "Internationalization"
 
-FInternationalization::FInternationalizationLifetimeObject FInternationalization::InternationalizationLifetimeObject;
-bool FInternationalization::bIsInitialized = false;
-TArray< TSharedRef< FCulture > > FInternationalization::AllCultures;
-int FInternationalization::CurrentCultureIndex = -1;
-TSharedPtr< FCulture > FInternationalization::DefaultCulture;
-TSharedPtr< FCulture > FInternationalization::InvariantCulture;
-
-FInternationalization::FInternationalizationLifetimeObject::FInternationalizationLifetimeObject()
+FInternationalization* FInternationalization::Instance;
+
+FInternationalization& FInternationalization::Get()
 {
-	if(!(FInternationalization::IsInitialized()))
+	if(!Instance)
+	{
+		Instance = new FInternationalization();
+	}
+	if(Instance && !Instance->IsInitialized())
 	{
-		FInternationalization::Initialize();
+		Instance->Initialize();
 	}
+	return *Instance;
 }
 
 void FInternationalization::GetTimeZonesIDs(TArray<FString>& TimeZonesIDs)
@@ -171,6 +171,11 @@ namespace
 
 void FInternationalization::Initialize()
 {
+	if(IsInitialized())
+	{
+		return;
+	}
+
 #if UE_ENABLE_ICU
 	UErrorCode ICUStatus = U_ZERO_ERROR;
 	u_setMemoryFunctions(NULL, &(FICUOverrides::Malloc), &(FICUOverrides::Realloc), &(FICUOverrides::Free), &(ICUStatus));
@@ -185,7 +190,7 @@ void FInternationalization::Initialize()
 	SetCurrentCulture( TEXT("") );
 #endif
 
-	bIsInitialized = true;
+	bIsInitialized = U_FAILURE(ICUStatus) ? true : false;
 }
 
 void FInternationalization::Terminate()
@@ -198,6 +203,9 @@ void FInternationalization::Terminate()
 #if UE_ENABLE_ICU
 	u_cleanup();
 #endif
+
+	delete Instance;
+	Instance = nullptr;
 }
 
 #if ENABLE_LOC_TESTING
@@ -362,6 +370,13 @@ void FInternationalization::GetCultureNames(TArray<FString>& CultureNames)
 #endif
 }
 
+FInternationalization::FInternationalization()
+	:	bIsInitialized(false)
+	,	CurrentCultureIndex(INDEX_NONE)
+{
+
+}
+
 void FInternationalization::PopulateAllCultures(void)
 {
 #if UE_ENABLE_ICU
diff --git a/Engine/Source/Runtime/Core/Private/Internationalization/LegacyText.cpp b/Engine/Source/Runtime/Core/Private/Internationalization/LegacyText.cpp
--- a/Engine/Source/Runtime/Core/Private/Internationalization/LegacyText.cpp
+++ b/Engine/Source/Runtime/Core/Private/Internationalization/LegacyText.cpp
@@ -7,34 +7,34 @@
 
 FText FText::AsDate(const FDateTime::FDate& Date, const EDateTimeStyle::Type DateStyle, const TSharedPtr<FCulture>& TargetCulture)
 {
-	checkf(FInternationalization::IsInitialized() == true, TEXT("FInternationalization is not initialized. An FText formatting method was likely used in static object initialization - this is not supported."));
+	checkf(FInternationalization::Get().IsInitialized() == true, TEXT("FInternationalization is not initialized. An FText formatting method was likely used in static object initialization - this is not supported."));
 	FDateTime DateTime(Date.Year, Date.Month, Date.Day);
 	return FText::FromString( DateTime.ToString( TEXT("%Y.%m.%d") ) );
 }
 
 FText FText::AsTime(const FDateTime::FTime& Time, const EDateTimeStyle::Type TimeStyle, const FString& TimeZone, const TSharedPtr<FCulture>& TargetCulture)
 {
-	checkf(FInternationalization::IsInitialized() == true, TEXT("FInternationalization is not initialized. An FText formatting method was likely used in static object initialization - this is not supported."));
+	checkf(FInternationalization::Get().IsInitialized() == true, TEXT("FInternationalization is not initialized. An FText formatting method was likely used in static object initialization - this is not supported."));
 	FDateTime DateTime(1, 1, 1, Time.Hour, Time.Minute, Time.Second, Time.Millisecond);
 	return FText::FromString( DateTime.ToString( TEXT("%H.%M.%S") ) );
 }
 
 FText FText::AsTime(const FTimespan& Time, const TSharedPtr<FCulture>& TargetCulture)
 {
-	checkf(FInternationalization::IsInitialized() == true, TEXT("FInternationalization is not initialized. An FText formatting method was likely used in static object initialization - this is not supported."));
+	checkf(FInternationalization::Get().IsInitialized() == true, TEXT("FInternationalization is not initialized. An FText formatting method was likely used in static object initialization - this is not supported."));
 	FDateTime DateTime(Time.GetTicks());
 	return FText::FromString( DateTime.ToString( TEXT("%H.%M.%S") ) );
 }
 
 FText FText::AsDateTime(const FDateTime& DateTime, const EDateTimeStyle::Type DateStyle, const EDateTimeStyle::Type TimeStyle, const FString& TimeZone, const TSharedPtr<FCulture>& TargetCulture)
 {
-	checkf(FInternationalization::IsInitialized() == true, TEXT("FInternationalization is not initialized. An FText formatting method was likely used in static object initialization - this is not supported."));
+	checkf(FInternationalization::Get().IsInitialized() == true, TEXT("FInternationalization is not initialized. An FText formatting method was likely used in static object initialization - this is not supported."));
 	return FText::FromString( DateTime.ToString( TEXT("%Y.%m.%d-%H.%M.%S") ) );
 }
 
 FText FText::AsMemory(SIZE_T NumBytes, const FNumberFormattingOptions* const Options, const TSharedPtr<FCulture>& TargetCulture)
 {
-	checkf(FInternationalization::IsInitialized() == true, TEXT("FInternationalization is not initialized. An FText formatting method was likely used in static object initialization - this is not supported."));
+	checkf(FInternationalization::Get().IsInitialized() == true, TEXT("FInternationalization is not initialized. An FText formatting method was likely used in static object initialization - this is not supported."));
 	FFormatNamedArguments Args;
 
 	if (NumBytes < 1024)
@@ -242,7 +242,7 @@ public:
 
 	static FText Format(const FText& Pattern, FGetArgumentValue GetArgumentValue)
 	{
-		checkf(FInternationalization::IsInitialized() == true, TEXT("FInternationalization is not initialized. An FText formatting method was likely used in static object initialization - this is not supported."));
+		checkf(FInternationalization::Get().IsInitialized() == true, TEXT("FInternationalization is not initialized. An FText formatting method was likely used in static object initialization - this is not supported."));
 		//SCOPE_CYCLE_COUNTER( STAT_TextFormat );
 
 		const FString& PatternString = Pattern.ToString();
@@ -414,7 +414,7 @@ void FText::GetFormatPatternParameters(const FText& Pattern, TArray<FString>& Pa
 
 FText FText::Format(const FText& Pattern, const FFormatNamedArguments& Arguments)
 {
-	checkf(FInternationalization::IsInitialized() == true, TEXT("FInternationalization is not initialized. An FText formatting method was likely used in static object initialization - this is not supported."));
+	checkf(FInternationalization::Get().IsInitialized() == true, TEXT("FInternationalization is not initialized. An FText formatting method was likely used in static object initialization - this is not supported."));
 	//SCOPE_CYCLE_COUNTER( STAT_TextFormat );
 
 	struct FArgumentGetter
@@ -433,7 +433,7 @@ FText FText::Format(const FText& Pattern, const FFormatNamedArguments& Arguments
 
 FText FText::Format(const FText& Pattern, const FFormatOrderedArguments& Arguments)
 {
-	checkf(FInternationalization::IsInitialized() == true, TEXT("FInternationalization is not initialized. An FText formatting method was likely used in static object initialization - this is not supported."));
+	checkf(FInternationalization::Get().IsInitialized() == true, TEXT("FInternationalization is not initialized. An FText formatting method was likely used in static object initialization - this is not supported."));
 	//SCOPE_CYCLE_COUNTER( STAT_TextFormat );
 
 	struct FArgumentGetter
@@ -457,7 +457,7 @@ FText FText::Format(const FText& Pattern, const FFormatOrderedArguments& Argumen
 
 FText FText::Format(const FText& Pattern, const TArray< FFormatArgumentData > InArguments)
 {
-	checkf(FInternationalization::IsInitialized() == true, TEXT("FInternationalization is not initialized. An FText formatting method was likely used in static object initialization - this is not supported."));
+	checkf(FInternationalization::Get().IsInitialized() == true, TEXT("FInternationalization is not initialized. An FText formatting method was likely used in static object initialization - this is not supported."));
 	//SCOPE_CYCLE_COUNTER( STAT_TextFormat );
 
 	struct FArgumentGetter
diff --git a/Engine/Source/Runtime/Core/Private/Internationalization/LegacyText.h b/Engine/Source/Runtime/Core/Private/Internationalization/LegacyText.h
--- a/Engine/Source/Runtime/Core/Private/Internationalization/LegacyText.h
+++ b/Engine/Source/Runtime/Core/Private/Internationalization/LegacyText.h
@@ -5,23 +5,26 @@
 template<typename T1, typename T2>
 FText FText::AsNumberTemplate(T1 Val, const FNumberFormattingOptions* const Options, const TSharedPtr<FCulture>& TargetCulture)
 {
-	checkf(FInternationalization::IsInitialized() == true, TEXT("FInternationalization is not initialized. An FText formatting method was likely used in static object initialization - this is not supported."));
-	const TSharedRef<FCulture> Culture = TargetCulture.IsValid() ? TargetCulture.ToSharedRef() : FInternationalization::GetCurrentCulture();
+	FInternationalization& I18N = FInternationalization::Get();
+	checkf(I18N.IsInitialized() == true, TEXT("FInternationalization is not initialized. An FText formatting method was likely used in static object initialization - this is not supported."));
+	const TSharedRef<FCulture> Culture = TargetCulture.IsValid() ? TargetCulture.ToSharedRef() : I18N.GetCurrentCulture();
 	return FText::CreateNumericalText( Culture->NumberFormattingRule.AsNumber(Val) );
 }
 
 template<typename T1, typename T2>
 FText FText::AsCurrencyTemplate(T1 Val, const FNumberFormattingOptions* const Options, const TSharedPtr<FCulture>& TargetCulture)
 {
-	checkf(FInternationalization::IsInitialized() == true, TEXT("FInternationalization is not initialized. An FText formatting method was likely used in static object initialization - this is not supported."));
-	const TSharedRef<FCulture> Culture = TargetCulture.IsValid() ? TargetCulture.ToSharedRef() : FInternationalization::GetCurrentCulture();
+	FInternationalization& I18N = FInternationalization::Get();
+	checkf(I18N.IsInitialized() == true, TEXT("FInternationalization is not initialized. An FText formatting method was likely used in static object initialization - this is not supported."));
+	const TSharedRef<FCulture> Culture = TargetCulture.IsValid() ? TargetCulture.ToSharedRef() : I18N.GetCurrentCulture();
 	return FText::CreateNumericalText( Culture->NumberFormattingRule.AsCurrency(Val) );
 }
 
 template<typename T1, typename T2>
 FText FText::AsPercentTemplate(T1 Val, const FNumberFormattingOptions* const Options, const TSharedPtr<FCulture>& TargetCulture)
 {
-	checkf(FInternationalization::IsInitialized() == true, TEXT("FInternationalization is not initialized. An FText formatting method was likely used in static object initialization - this is not supported."));
-	const TSharedRef<FCulture> Culture = TargetCulture.IsValid() ? TargetCulture.ToSharedRef() : FInternationalization::GetCurrentCulture();
+	FInternationalization& I18N = FInternationalization::Get();
+	checkf(I18N.IsInitialized() == true, TEXT("FInternationalization is not initialized. An FText formatting method was likely used in static object initialization - this is not supported."));
+	const TSharedRef<FCulture> Culture = TargetCulture.IsValid() ? TargetCulture.ToSharedRef() : I18N.GetCurrentCulture();
 	return FText::CreateNumericalText( Culture->NumberFormattingRule.AsPercent(Val) );
 }
\ No newline at end of file
diff --git a/Engine/Source/Runtime/Core/Private/Internationalization/TextLocalizationManager.cpp b/Engine/Source/Runtime/Core/Private/Internationalization/TextLocalizationManager.cpp
--- a/Engine/Source/Runtime/Core/Private/Internationalization/TextLocalizationManager.cpp
+++ b/Engine/Source/Runtime/Core/Private/Internationalization/TextLocalizationManager.cpp
@@ -165,8 +165,10 @@ void FTextLocalizationManager::FLocalizationEntryTracker::ReadFromArchive(FArchi
 
 void FTextLocalizationManager::LoadResources(const bool ShouldLoadEditor, const bool ShouldLoadGame)
 {
-	const FString& CultureName = FInternationalization::GetCurrentCulture()->GetName();
-	const FString& BaseLanguageName = FInternationalization::GetCurrentCulture()->GetTwoLetterISOLanguageName();
+	FInternationalization& I18N = FInternationalization::Get();
+
+	const FString& CultureName = I18N.GetCurrentCulture()->GetName();
+	const FString& BaseLanguageName = I18N.GetCurrentCulture()->GetTwoLetterISOLanguageName();
 
 #if ENABLE_LOC_TESTING
 	if(CultureName == TEXT("LEET"))
@@ -364,7 +366,7 @@ TSharedRef<FString> FTextLocalizationManager::GetString(const FString& Namespace
 	}
 
 #if ENABLE_LOC_TESTING
-	const bool bShouldLEETIFYAll = bIsInitialized && FInternationalization::GetCurrentCulture()->GetName() == TEXT("LEET");
+	const bool bShouldLEETIFYAll = bIsInitialized && FInternationalization::Get().GetCurrentCulture()->GetName() == TEXT("LEET");
 	static const bool bShouldLEETIFYUnlocalizedString = FCommandLine::IsInitialized() && FParse::Param(FCommandLine::Get(), TEXT("LEETIFYUnlocalized"));
 #endif
 
@@ -493,6 +495,8 @@ namespace
 
 void FTextLocalizationManager::RegenerateResources(const FString& ConfigFilePath)
 {
+	FInternationalization& I18N = FInternationalization::Get();
+
 	FString SectionName = TEXT("RegenerateResources");
 
 	// Get source path.
@@ -529,9 +533,9 @@ void FTextLocalizationManager::RegenerateResources(const FString& ConfigFilePath
 
 	TArray<FString> LocaleNames;
 	{
-		const FString CultureName = FInternationalization::GetCurrentCulture()->GetName();
+		const FString CultureName = I18N.GetCurrentCulture()->GetName();
 		LocaleNames.Add(CultureName);
-		const FString BaseLanguageName = FInternationalization::GetCurrentCulture()->GetTwoLetterISOLanguageName();
+		const FString BaseLanguageName = I18N.GetCurrentCulture()->GetTwoLetterISOLanguageName();
 		if(BaseLanguageName != CultureName)
 		{
 			LocaleNames.Add(BaseLanguageName);
diff --git a/Engine/Source/Runtime/Core/Private/Mac/MacPlatformMisc.cpp b/Engine/Source/Runtime/Core/Private/Mac/MacPlatformMisc.cpp
--- a/Engine/Source/Runtime/Core/Private/Mac/MacPlatformMisc.cpp
+++ b/Engine/Source/Runtime/Core/Private/Mac/MacPlatformMisc.cpp
@@ -158,7 +158,7 @@ struct MacApplicationInfo
 		
 		AppPath = FPlatformProcess::GenerateApplicationPath(FApp::GetName(), FApp::GetBuildConfiguration());
 		
-		LCID = FString::Printf(TEXT("%d"), FInternationalization::GetCurrentCulture()->GetLCID());
+		LCID = FString::Printf(TEXT("%d"), FInternationalization::Get().GetCurrentCulture()->GetLCID());
 	}
 	
 	char AppNameUTF8[PATH_MAX+1];
@@ -466,7 +466,7 @@ void FMacPlatformMisc::SubmitErrorReport( const TCHAR* InErrorHist, EErrorReport
 			FCString::Strcpy(PlatformName, TEXT("Mac 32-bit"));
 #endif	//PLATFORM_64BITS
 			TCHAR CultureName[10];
-			FCString::Strcpy(CultureName, *FInternationalization::GetCurrentCulture()->GetName());
+			FCString::Strcpy(CultureName, *FInternationalization::Get().GetCurrentCulture()->GetName());
 			TCHAR SystemTime[256];
 			FCString::Strcpy(SystemTime, *FDateTime::Now().ToString());
 			TCHAR EngineVersionStr[32];
diff --git a/Engine/Source/Runtime/Core/Private/Tests/Internationalization/DateTimeFormattingRulesTest.cpp b/Engine/Source/Runtime/Core/Private/Tests/Internationalization/DateTimeFormattingRulesTest.cpp
--- a/Engine/Source/Runtime/Core/Private/Tests/Internationalization/DateTimeFormattingRulesTest.cpp
+++ b/Engine/Source/Runtime/Core/Private/Tests/Internationalization/DateTimeFormattingRulesTest.cpp
@@ -21,7 +21,9 @@ namespace
 
 bool FDateTimeFormattingRulesTest::RunTest (const FString& Parameters)
 {
-	const FString OriginalCulture = FInternationalization::GetCurrentCulture()->GetName();
+	FInternationalization& I18N = FInternationalization::Get();
+
+	const FString OriginalCulture = I18N.GetCurrentCulture()->GetName();
 
 	//////////////////////////////////////////////////////////////////////////
 
@@ -34,7 +36,7 @@ bool FDateTimeFormattingRulesTest::RunTest (const FString& Parameters)
 	const FDateTime::FTime TestTime = { 12, 34, 56, 789 };
 	const FDateTime TestDateTime( TestDate.Year, TestDate.Month, TestDate.Day, TestTime.Hour, TestTime.Minute, TestTime.Second, TestTime.Millisecond );
 
-	FInternationalization::SetCurrentCulture("en-US");
+	I18N.SetCurrentCulture("en-US");
 
 	// Unix Time Values via Date Time
 	Test( this, TEXT("Testing Unix Epoch"), FText::AsDateTime(UnixEpoch, EDateTimeStyle::Short, EDateTimeStyle::Short, "GMT"), FText::FromString( TEXT("1/1/70, 12:00 AM") ) );
@@ -75,7 +77,7 @@ bool FDateTimeFormattingRulesTest::RunTest (const FString& Parameters)
 	Test( this, TEXT("Testing Date-Time"), FText::AsDateTime(TestDateTime, EDateTimeStyle::Long, EDateTimeStyle::Long, "GMT"), FText::FromString( TEXT("June 13, 1990 at 12:34:56 PM GMT") ) );
 	Test( this, TEXT("Testing Date-Time"), FText::AsDateTime(TestDateTime, EDateTimeStyle::Full, EDateTimeStyle::Full, "GMT"), FText::FromString( TEXT("Wednesday, June 13, 1990 at 12:34:56 PM GMT") ) );
 
-	FInternationalization::SetCurrentCulture("ja-JP");
+	I18N.SetCurrentCulture("ja-JP");
 
 	// Unix Time Values via Date Time
 	Test( this, TEXT("Testing Unix Epoch"), FText::AsDateTime(UnixEpoch, EDateTimeStyle::Short, EDateTimeStyle::Short, "GMT"), FText::FromString( TEXT("1970/01/01 0:00") ) );
@@ -120,7 +122,7 @@ bool FDateTimeFormattingRulesTest::RunTest (const FString& Parameters)
 #endif
 	//////////////////////////////////////////////////////////////////////////
 
-	FInternationalization::SetCurrentCulture(OriginalCulture);
+	I18N.SetCurrentCulture(OriginalCulture);
 
 	return true;
 }
diff --git a/Engine/Source/Runtime/Core/Private/Tests/Internationalization/NumberFormatingRulesTest.cpp b/Engine/Source/Runtime/Core/Private/Tests/Internationalization/NumberFormatingRulesTest.cpp
--- a/Engine/Source/Runtime/Core/Private/Tests/Internationalization/NumberFormatingRulesTest.cpp
+++ b/Engine/Source/Runtime/Core/Private/Tests/Internationalization/NumberFormatingRulesTest.cpp
@@ -71,9 +71,10 @@ bool FNumberFormattingRulesTest::RunTest (const FString& Parameters)
 	int32 Int32NegativeValue = -12345;
 	int64 Int64NegativeValue = -12345;
 
-	const FString OriginalCulture = FInternationalization::GetCurrentCulture()->GetName();
+	FInternationalization& I18N = FInternationalization::Get();
+	const FString OriginalCulture = I18N.GetCurrentCulture()->GetName();
 
-	FInternationalization::SetCurrentCulture("en-US");
+	I18N.SetCurrentCulture("en-US");
 
 	Test(this, TEXT("Convert a Double to a number formatted correct for en-US"),			FText::AsNumber(DoubleValue),			FText::FromString(TEXT("12,345,678.901")));
 	Test(this, TEXT("Convert a Float to a number formatted correct for en-US"),			FText::AsNumber(FloatValue),			FText::FromString(TEXT("1,234.567")));
@@ -231,7 +232,7 @@ bool FNumberFormattingRulesTest::RunTest (const FString& Parameters)
 	Test(this, TEXT("Convert a Negative Double to a percent formatted correct for en-US"),	FText::AsPercent(DoubleNegativeValue),	FText::FromString(TEXT("-1,234,567,890%")));
 	Test(this, TEXT("Convert a Negative Float to a percent formatted correct for en-US"),		FText::AsPercent(FloatNegativeValue),	FText::FromString(TEXT("-123,457%")));
 
-	FInternationalization::SetCurrentCulture("hi-IN");
+	I18N.SetCurrentCulture("hi-IN");
 
 	Test(this, TEXT("Convert a Double to a number formatted correct for hi-IN"),			FText::AsNumber(DoubleValue),			FText::FromString(TEXT("1,23,45,678.901")));
 	Test(this, TEXT("Convert a Float to a number formatted correct for hi-IN"),			FText::AsNumber(FloatValue),			FText::FromString(TEXT("1,234.567")));
@@ -252,7 +253,7 @@ bool FNumberFormattingRulesTest::RunTest (const FString& Parameters)
 
 #if UE_ENABLE_ICU
 	{
-		const TSharedPtr<FCulture>& InvariantCulture = FInternationalization::GetInvariantCulture();
+		const TSharedPtr<FCulture>& InvariantCulture = I18N.GetInvariantCulture();
 		Test(this, TEXT("Convert a Double to a number formatted correct for hi-IN but as invariant"),				FText::AsNumber(DoubleValue, NULL, InvariantCulture),			FText::FromString(TEXT("12345678.901")));
 		Test(this, TEXT("Convert a Float to a number formatted correct for hi-IN but as invariant"),				FText::AsNumber(FloatValue, NULL, InvariantCulture),			FText::FromString(TEXT("1234.567017")));
 		Test(this, TEXT("Convert a Negative Double to a number formatted correct for hi-IN but as invariant"),	FText::AsNumber(DoubleNegativeValue, NULL, InvariantCulture),	FText::FromString(TEXT("-12345678.901")));
@@ -296,7 +297,7 @@ bool FNumberFormattingRulesTest::RunTest (const FString& Parameters)
 	Test(this, TEXT("Convert a Negative Double to a percent formatted correct for hi-IN"),	FText::AsPercent(DoubleNegativeValue),	FText::FromString(TEXT("-1,23,45,67,890%")));
 	Test(this, TEXT("Convert a Negative Float to a percent formatted correct for hi-IN"),		FText::AsPercent(FloatNegativeValue),	FText::FromString(TEXT("-1,23,457%")));
 
-	FInternationalization::SetCurrentCulture( OriginalCulture );
+	I18N.SetCurrentCulture( OriginalCulture );
 
 	FText Number = FText::AsNumber(Int64NegativeValue);
 	FText Percent = FText::AsPercent(DoubleValue);
diff --git a/Engine/Source/Runtime/Core/Private/Tests/Internationalization/TextTest.cpp b/Engine/Source/Runtime/Core/Private/Tests/Internationalization/TextTest.cpp
--- a/Engine/Source/Runtime/Core/Private/Tests/Internationalization/TextTest.cpp
+++ b/Engine/Source/Runtime/Core/Private/Tests/Internationalization/TextTest.cpp
@@ -49,13 +49,14 @@ namespace
 
 bool FTextTest::RunTest (const FString& Parameters)
 {
+	FInternationalization& I18N = FInternationalization::Get();
 	const bool OriginalEnableErrorCheckingValue = FText::GetEnableErrorCheckingResults();
 	const bool OriginalSuppressWarningsValue = FText::GetSuppressWarnings();
-	const FString OriginalCulture = FInternationalization::GetCurrentCulture()->GetName();
+	const FString OriginalCulture = I18N.GetCurrentCulture()->GetName();
 
 	FText::SetEnableErrorCheckingResults(true);
 	FText::SetSuppressWarnings(true);
-	FInternationalization::SetCurrentCulture("en-US");
+	I18N.SetCurrentCulture("en-US");
 
 	FText ArgText0 = FText::FromString(TEXT("Arg0"));
 	FText ArgText1 = FText::FromString(TEXT("Arg1"));
@@ -322,7 +323,7 @@ bool FTextTest::RunTest (const FString& Parameters)
 	// Sort Testing
 	// French
 	{
-		FInternationalization::SetCurrentCulture("fr");
+		I18N.SetCurrentCulture("fr");
 		TArray<FText> CorrectlySortedValues;
 		CorrectlySortedValues.Add( FText::FromString( TEXT("cote") ) );
 		CorrectlySortedValues.Add( FText::FromString( TEXT("coté") ) );
@@ -354,14 +355,14 @@ bool FTextTest::RunTest (const FString& Parameters)
 			}
 			if( !Identical )
 			{
-				//currently failing AddError(FString::Printf(TEXT("Sort order is wrong for culture (%s)."), *FInternationalization::GetCurrentCulture()->GetEnglishName()));
+				//currently failing AddError(FString::Printf(TEXT("Sort order is wrong for culture (%s)."), *FInternationalization::Get().GetCurrentCulture()->GetEnglishName()));
 			}
 		}
 	}
 
 	// French Canadian
 	{
-		FInternationalization::SetCurrentCulture("fr-CA");
+		I18N.SetCurrentCulture("fr-CA");
 		TArray<FText> CorrectlySortedValues;
 		CorrectlySortedValues.Add( FText::FromString( TEXT("cote") ) );
 		CorrectlySortedValues.Add( FText::FromString( TEXT("côte") ) );
@@ -390,7 +391,7 @@ bool FTextTest::RunTest (const FString& Parameters)
 			}
 			if( !Identical )
 			{
-				//currently failing AddError(FString::Printf(TEXT("Sort order is wrong for culture (%s)."), *FInternationalization::GetCurrentCulture()->GetEnglishName()));
+				//currently failing AddError(FString::Printf(TEXT("Sort order is wrong for culture (%s)."), *FInternationalization::Get().GetCurrentCulture()->GetEnglishName()));
 			}
 		}
 	}
@@ -421,7 +422,7 @@ bool FTextTest::RunTest (const FString& Parameters)
 		AddError( TEXT("FromString should never produce a Transient Text") );
 	}
 
-	FInternationalization::SetCurrentCulture(OriginalCulture);
+	I18N.SetCurrentCulture(OriginalCulture);
 
 	return true;
 }
diff --git a/Engine/Source/Runtime/Core/Private/Windows/WindowsPlatformMisc.cpp b/Engine/Source/Runtime/Core/Private/Windows/WindowsPlatformMisc.cpp
--- a/Engine/Source/Runtime/Core/Private/Windows/WindowsPlatformMisc.cpp
+++ b/Engine/Source/Runtime/Core/Private/Windows/WindowsPlatformMisc.cpp
@@ -413,7 +413,7 @@ void FWindowsPlatformMisc::SubmitErrorReport( const TCHAR* InErrorHist, EErrorRe
 			FCString::Strcpy(PlatformName, TEXT("PC 32-bit"));
 #endif	//PLATFORM_64BITS
 			TCHAR CultureName[10];
-			FCString::Strcpy(CultureName, *FInternationalization::GetDefaultCulture()->GetName());
+			FCString::Strcpy(CultureName, *FInternationalization::Get().GetDefaultCulture()->GetName());
 			TCHAR SystemTime[256];
 			FCString::Strcpy(SystemTime, *FDateTime::Now().ToString());
 			TCHAR EngineVersionStr[32];
diff --git a/Engine/Source/Runtime/Core/Public/Internationalization/Internationalization.h b/Engine/Source/Runtime/Core/Public/Internationalization/Internationalization.h
--- a/Engine/Source/Runtime/Core/Public/Internationalization/Internationalization.h
+++ b/Engine/Source/Runtime/Core/Public/Internationalization/Internationalization.h
@@ -27,67 +27,66 @@
 class FInternationalization
 {
 public:
+	static CORE_API FInternationalization& Get();
 
 	static CORE_API FText ForUseOnlyByLocMacroAndGraphNodeTextLiterals_CreateText( const TCHAR* InTextLiteral, const TCHAR* Namespace, const TCHAR* Key )
 	{
 		return FText( InTextLiteral, Namespace, Key );
 	}
 
-	static CORE_API void GetTimeZonesIDs(TArray<FString>& TimeZonesIDs);
+	CORE_API void GetTimeZonesIDs(TArray<FString>& TimeZonesIDs);
 
 	//Set the current culture by name
-	static CORE_API void SetCurrentCulture(const FString& Name);
+	CORE_API void SetCurrentCulture(const FString& Name);
 
 	//@return the current culture
-	static CORE_API TSharedRef< FCulture > GetCurrentCulture();
+	CORE_API TSharedRef< FCulture > GetCurrentCulture();
 
 	//@return culture object by given name, or NULL if not found
-	static CORE_API TSharedPtr< FCulture > GetCulture(const FString& Name);
+	CORE_API TSharedPtr< FCulture > GetCulture(const FString& Name);
 
 	//@return the default culture
-	static CORE_API TSharedRef< FCulture > GetDefaultCulture()
+	CORE_API TSharedRef< FCulture > GetDefaultCulture()
 	{
 		return DefaultCulture.ToSharedRef();
 	}
 
 	//@return the invariant culture
-	static CORE_API TSharedRef< FCulture > GetInvariantCulture()
+	CORE_API TSharedRef< FCulture > GetInvariantCulture()
 	{
 		return InvariantCulture.ToSharedRef();
 	}
 
-	static CORE_API void Initialize();
-	static CORE_API void Terminate();
+	CORE_API void Initialize();
+	CORE_API void Terminate();
 
-	static CORE_API bool IsInitialized() {return bIsInitialized;}
+	CORE_API bool IsInitialized() {return bIsInitialized;}
 
 #if ENABLE_LOC_TESTING
 	static CORE_API FString& Leetify(FString& SourceString);
 #endif
 
-	static CORE_API void GetCultureNames(TArray<FString>& CultureNames);
+	CORE_API void GetCultureNames(TArray<FString>& CultureNames);
 
 private:
+	FInternationalization();
+
 	//@todo how are we going to initialize the list of cultures? below is temporary
-	static CORE_API void PopulateAllCultures(void);
+	CORE_API void PopulateAllCultures(void);
 
 	//@return index of a culture by given name, or -1 if not found
-	static CORE_API int32 GetCultureIndex(const FString& Name);
+	CORE_API int32 GetCultureIndex(const FString& Name);
 
 private:
-	static bool bIsInitialized;
-
-	struct FInternationalizationLifetimeObject
-	{
-		FInternationalizationLifetimeObject();
-	} static InternationalizationLifetimeObject;
+	static FInternationalization* Instance;
+	bool bIsInitialized;
 
-	static CORE_API TArray< TSharedRef< FCulture > > AllCultures;
+	TArray< TSharedRef< FCulture > > AllCultures;
 
-	static CORE_API int CurrentCultureIndex;
+	int CurrentCultureIndex;
 
-	static CORE_API TSharedPtr< FCulture > DefaultCulture;
-	static CORE_API TSharedPtr< FCulture > InvariantCulture;
+	TSharedPtr< FCulture > DefaultCulture;
+	TSharedPtr< FCulture > InvariantCulture;
 };
 
 #undef LOC_DEFINE_REGION
diff --git a/Engine/Source/Runtime/CoreUObject/Private/UObject/Class.cpp b/Engine/Source/Runtime/CoreUObject/Private/UObject/Class.cpp
--- a/Engine/Source/Runtime/CoreUObject/Private/UObject/Class.cpp
+++ b/Engine/Source/Runtime/CoreUObject/Private/UObject/Class.cpp
@@ -2559,7 +2559,7 @@ FString UClass::GetDescription() const
 #endif
 
 	// Look up the the classes name in the legacy int file and return the class name if there is no match.
-	//Description = Localize( TEXT("Objects"), *GetName(), *(FInternationalization::GetCurrentCulture()->GetName()), true );
+	//Description = Localize( TEXT("Objects"), *GetName(), *(FInternationalization::Get().GetCurrentCulture()->GetName()), true );
 	//if (Description.Len())
 	//{
 	//	return Description;
diff --git a/Engine/Source/Runtime/Engine/Private/UnrealEngine.cpp b/Engine/Source/Runtime/Engine/Private/UnrealEngine.cpp
--- a/Engine/Source/Runtime/Engine/Private/UnrealEngine.cpp
+++ b/Engine/Source/Runtime/Engine/Private/UnrealEngine.cpp
@@ -9613,7 +9613,7 @@ void LoadGametypeContent_Helper(const FString& ContentStr,
 								FLoadPackageAsyncDelegate CompletionCallback, 
 								FLoadPackageAsyncDelegate LocalizedCompletionCallback)
 {
-	//const TCHAR* Language = *(FInternationalization::GetCurrentCulture()->Name);
+	//const TCHAR* Language = *(FInternationalization::Get().GetCurrentCulture()->Name);
 	//const FString LocalizedPreloadName(ContentStr + LOCALIZED_SEEKFREE_SUFFIX + TEXT("_") + Language);
 	//FString LocalizedPreloadFilename;
 	//if (FPackageName::DoesPackageExist(*LocalizedPreloadName, NULL, &LocalizedPreloadFilename))
diff --git a/Engine/Source/Runtime/Launch/Private/LaunchEngineLoop.cpp b/Engine/Source/Runtime/Launch/Private/LaunchEngineLoop.cpp
--- a/Engine/Source/Runtime/Launch/Private/LaunchEngineLoop.cpp
+++ b/Engine/Source/Runtime/Launch/Private/LaunchEngineLoop.cpp
@@ -2152,7 +2152,6 @@ void FEngineLoop::AppInit( )
 	GWarn = FPlatformOutputDevices::GetWarn();
 
 	BeginInitTextLocalization();
-	FInternationalization::Initialize();
 
 	// Avoiding potential exploits by not exposing command line overrides in the shipping games.
 #if !UE_BUILD_SHIPPING && WITH_EDITORONLY_DATA
@@ -2235,11 +2234,13 @@ void FEngineLoop::AppInit( )
 	// after the above has run we now have the REQUIRED set of engine .INIs  (all of the other .INIs)
 	// that are gotten from .h files' config() are not requires and are dynamically loaded when the .u files are loaded
 
+	FInternationalization& I18N = FInternationalization::Get();
+
 	// Set culture according to configuration now that configs are available.
 #if ENABLE_LOC_TESTING
 	if( FCommandLine::IsInitialized() && FParse::Param(FCommandLine::Get(), TEXT("LEET")) )
 	{
-		FInternationalization::SetCurrentCulture(TEXT("LEET"));
+		I18N.SetCurrentCulture(TEXT("LEET"));
 	}
 	else
 #endif
@@ -2268,10 +2269,10 @@ void FEngineLoop::AppInit( )
 		}
 		else
 		{
-			CultureName = FInternationalization::GetDefaultCulture()->GetName();
+			CultureName = I18N.GetDefaultCulture()->GetName();
 		}
 
-		FInternationalization::SetCurrentCulture(CultureName);
+		I18N.SetCurrentCulture(CultureName);
 	}
 
 #if !UE_BUILD_SHIPPING
@@ -2418,7 +2419,7 @@ void FEngineLoop::AppInit( )
 	
 	if (FParse::Value(FCommandLine::Get(), TEXT("CULTUREFORCOOKING="), CookerCulture, ARRAY_COUNT(CookerCulture)))
 	{
-		FInternationalization::SetCurrentCulture( CookerCulture );
+		FInternationalization::Get().SetCurrentCulture( CookerCulture );
 
 		// Write the culture passed in if first install...
 		if (FParse::Param(FCommandLine::Get(), TEXT("firstinstall")))
@@ -2480,7 +2481,7 @@ void FEngineLoop::AppExit( )
 		GLog->TearDown();
 	}
 
-	FInternationalization::Terminate();
+	FInternationalization::Get().Terminate();
 }
 
 #undef LOCTEXT_NAMESPACE
