#summary Fixed ICU related build error due to changes to internationalization initialization.

diff --git a/Engine/Source/Runtime/Core/Private/Internationalization/Internationalization.cpp b/Engine/Source/Runtime/Core/Private/Internationalization/Internationalization.cpp
--- a/Engine/Source/Runtime/Core/Private/Internationalization/Internationalization.cpp
+++ b/Engine/Source/Runtime/Core/Private/Internationalization/Internationalization.cpp
@@ -190,7 +190,11 @@ void FInternationalization::Initialize()
 	SetCurrentCulture( TEXT("") );
 #endif
 
+#if UE_ENABLE_ICU
 	bIsInitialized = U_FAILURE(ICUStatus) ? true : false;
+#else
+	bIsInitialized = true;
+#endif
 }
 
 void FInternationalization::Terminate()
