Added support for NULL params to FMacPlatformMisc::MessageBoxExt()

diff --git a/Engine/Source/Runtime/Core/Private/Mac/MacPlatformMisc.cpp b/Engine/Source/Runtime/Core/Private/Mac/MacPlatformMisc.cpp
--- a/Engine/Source/Runtime/Core/Private/Mac/MacPlatformMisc.cpp
+++ b/Engine/Source/Runtime/Core/Private/Mac/MacPlatformMisc.cpp
@@ -810,8 +810,8 @@ EAppReturnType::Type FMacPlatformMisc::MessageBoxExt( EAppMsgType::Type MsgType,
 {
 	SCOPED_AUTORELEASE_POOL;
 
-	NSString* CocoaText = (NSString*)FPlatformString::TCHARToCFString(Text);
-	NSString* CocoaCaption = (NSString*)FPlatformString::TCHARToCFString(Caption);
+	NSString* CocoaText = Text ? (NSString*)FPlatformString::TCHARToCFString(Text) : NULL;
+	NSString* CocoaCaption = Caption ? (NSString*)FPlatformString::TCHARToCFString(Caption) : NULL;
 
 	EAppReturnType::Type RetValue = EAppReturnType::Cancel;
 	NSInteger Result;
@@ -932,8 +932,14 @@ EAppReturnType::Type FMacPlatformMisc::MessageBoxExt( EAppMsgType::Type MsgType,
 		MacApplication->UseMouseCaptureWindow(true);
 	}
 
-	CFRelease((CFStringRef)CocoaCaption);
-	CFRelease((CFStringRef)CocoaText);
+	if (CocoaCaption)
+	{
+		CFRelease((CFStringRef)CocoaCaption);
+	}
+	if (CocoaText)
+	{
+		CFRelease((CFStringRef)CocoaText);
+	}
 
 	return RetValue;
 }
