More robust support for foreign engine directories, now that we have multiple engine versions installed. Always use the directory passed in by UBT, but store it in the game module to allow it to work for binary Rocket too.

diff --git a/Engine/Source/Runtime/Core/Private/GenericPlatform/GenericPlatformMisc.cpp b/Engine/Source/Runtime/Core/Private/GenericPlatform/GenericPlatformMisc.cpp
--- a/Engine/Source/Runtime/Core/Private/GenericPlatform/GenericPlatformMisc.cpp
+++ b/Engine/Source/Runtime/Core/Private/GenericPlatform/GenericPlatformMisc.cpp
@@ -369,42 +369,27 @@ const TCHAR* FGenericPlatformMisc::EngineDir()
 	if (EngineDirectory.Len() == 0)
 	{
 		// See if we are a root-level project
-		FString CheckDir = TEXT("../../../Engine/");
+		FString DefaultEngineDir = TEXT("../../../Engine/");
 #if PLATFORM_DESKTOP
 		//@todo. Need to have a define specific for this scenario??
-		if (FPlatformFileManager::Get().GetPlatformFile().DirectoryExists(*(CheckDir / TEXT("Binaries"))))
+		if (FPlatformFileManager::Get().GetPlatformFile().DirectoryExists(*(DefaultEngineDir / TEXT("Binaries"))))
 		{
-			EngineDirectory = CheckDir;
+			EngineDirectory = DefaultEngineDir;
 		}
-		else
+		else if (GForeignEngineDir != NULL && FPlatformFileManager::Get().GetPlatformFile().DirectoryExists(*(FString(GForeignEngineDir) / TEXT("Binaries"))))
 		{
-#ifdef UE_ENGINE_DIRECTORY
-			CheckDir = TEXT( PREPROCESSOR_TO_STRING(UE_ENGINE_DIRECTORY) );
-#endif
-			if ((CheckDir.Len() == 0) || 
-				(FPlatformFileManager::Get().GetPlatformFile().DirectoryExists(*(CheckDir / TEXT("Binaries"))) == false))
-			{
-				// Try using the registry to find the 'installed' version
-				CheckDir = FRocketSupport::GetInstalledEngineDirectory();
-				if ((CheckDir.Len() != 0) && (FPlatformFileManager::Get().GetPlatformFile().DirectoryExists(*(CheckDir / TEXT("Binaries")))))
-				{
-					EngineDirectory = CheckDir;
-				}
-			}
-			else
-			{
-				EngineDirectory = CheckDir;
-			}
+			EngineDirectory = GForeignEngineDir;
 		}
-#else
-		EngineDirectory = CheckDir;
-#endif
+
 		if (EngineDirectory.Len() == 0)
 		{
 			// Temporary work-around for legacy dependency on ../../../ (re Lightmass)
-			EngineDirectory = TEXT("../../../Engine/");
+			EngineDirectory = DefaultEngineDir;
 			UE_LOG(LogGenericPlatformMisc, Warning, TEXT("Failed to determine engine directory: Defaulting to %s"), *EngineDirectory);
 		}
+#else
+		EngineDirectory = DefaultEngineDir;
+#endif
 	}
 	return *EngineDirectory;
 }
@@ -693,4 +678,3 @@ FText FGenericPlatformMisc::GetFileManagerName()
 {
 	return NSLOCTEXT("GenericPlatform", "FileManagerName", "File Manager");
 }
-
diff --git a/Engine/Source/Runtime/Core/Private/Misc/Core.cpp b/Engine/Source/Runtime/Core/Private/Misc/Core.cpp
--- a/Engine/Source/Runtime/Core/Private/Misc/Core.cpp
+++ b/Engine/Source/Runtime/Core/Private/Misc/Core.cpp
@@ -147,6 +147,11 @@ TCHAR					GGameName[64];
 // macro for the game's main game module.
 #endif
 
+// Foreign engine directory. This is required to projects built outside the engine root to reference their engine directory.
+#if !IS_MONOLITHIC
+IMPLEMENT_FOREIGN_ENGINE_DIR()
+#endif
+
 /** Exec handler for game debugging tool, allowing commands like "editactor", ...							*/
 FExec*					GDebugToolExec					= NULL;
 /** Whether we're currently in the async loading codepath or not											*/
diff --git a/Engine/Source/Runtime/Core/Private/Misc/RocketSupport.cpp b/Engine/Source/Runtime/Core/Private/Misc/RocketSupport.cpp
--- a/Engine/Source/Runtime/Core/Private/Misc/RocketSupport.cpp
+++ b/Engine/Source/Runtime/Core/Private/Misc/RocketSupport.cpp
@@ -39,51 +39,3 @@ bool FRocketSupport::IsRocket( const TCHAR* CmdLine )
 	return (RocketState == 1) ? true : false;
 #endif
 }
-
-FString FRocketSupport::GetInstalledEngineDirectory()
-{
-#if PLATFORM_WINDOWS
-
-	FString RegistryEntry = FString::Printf(TEXT("EpicGames/Unreal Engine/%s"), *GEngineVersion.ToString(EVersionComponent::Minor));
-
-	FString InstalledDirectory;
-	if (FPlatformMisc::GetRegistryString(RegistryEntry, TEXT("InstalledDirectory"), true, InstalledDirectory))
-	{
-		return InstalledDirectory.Replace(TEXT("\\"), TEXT("/")) / TEXT("Engine/");
-	}
-	if (FPlatformMisc::GetRegistryString(RegistryEntry, TEXT("InstalledDirectory"), false, InstalledDirectory))
-	{
-		return InstalledDirectory.Replace(TEXT("\\"), TEXT("/")) / TEXT("Engine/");
-	}
-
-#elif PLATFORM_MAC
-
-	NSString* AppName = @"UE4Editor";
-	NSFileManager* FileManager = [NSFileManager defaultManager];
-
-	// First, check if there is a self-contained editor bundle in /Applications folder
-	NSString* InstalledPath = [NSString stringWithFormat: @"/Applications/%@.app/Contents/UE4/Engine/", AppName];
-	if ([FileManager fileExistsAtPath: InstalledPath])
-	{
-		return InstalledPath;
-	}
-
-	// No editor found in /Applications. In that case, we ask OS X to find the editor app bundle for us
-	NSURL* AppURL = [[NSWorkspace sharedWorkspace] URLForApplicationWithBundleIdentifier: [@"com.epicgames." stringByAppendingString: AppName]];
-	if (AppURL)
-	{
-		if ([FileManager fileExistsAtPath: [[AppURL path] stringByAppendingPathComponent: @"/Contents/UE4/Engine"]])
-		{
-			// It's a self-contained app
-			return FString([AppURL path]) / TEXT("Contents/UE4/Engine/");
-		}
-		else
-		{
-			return FString([[AppURL path] stringByDeletingLastPathComponent]) / TEXT("../../Engine/");
-		}
-	}
-
-#endif // PLATFORM
-
-	return TEXT("");
-}
\ No newline at end of file
diff --git a/Engine/Source/Runtime/Core/Public/Core.h b/Engine/Source/Runtime/Core/Public/Core.h
--- a/Engine/Source/Runtime/Core/Public/Core.h
+++ b/Engine/Source/Runtime/Core/Public/Core.h
@@ -252,6 +252,7 @@ extern CORE_API float GNearClippingPlane;
 
 extern CORE_API bool GExitPurge;
 extern CORE_API TCHAR GGameName[64];
+extern CORE_API const TCHAR* GForeignEngineDir;
 
 /** Exec handler for game debugging tool, allowing commands like "editactor" */
 extern CORE_API FExec* GDebugToolExec;
diff --git a/Engine/Source/Runtime/Core/Public/Misc/RocketSupport.h b/Engine/Source/Runtime/Core/Public/Misc/RocketSupport.h
--- a/Engine/Source/Runtime/Core/Public/Misc/RocketSupport.h
+++ b/Engine/Source/Runtime/Core/Public/Misc/RocketSupport.h
@@ -20,11 +20,4 @@ public:
 	 * @return	True if we are in rocket mode, otherwise false
 	 */
 	static bool IsRocket( const TCHAR* CmdLine = FCommandLine::Get() );
-
-	/**
-	 *	Gets the engine directory for an 'installed' UE4/Rocket
-	 *
-	 *	@return	FString		The EngineDir, empty string if not found
-	 */
-	static FString GetInstalledEngineDirectory();
 };
diff --git a/Engine/Source/Runtime/Core/Public/Modules/ModuleManager.h b/Engine/Source/Runtime/Core/Public/Modules/ModuleManager.h
--- a/Engine/Source/Runtime/Core/Public/Modules/ModuleManager.h
+++ b/Engine/Source/Runtime/Core/Public/Modules/ModuleManager.h
@@ -782,6 +782,17 @@ class FDefaultGameModuleImpl
 #define IMPLEMENT_GAME_MODULE( ModuleImplClass, ModuleName ) \
 	IMPLEMENT_MODULE( ModuleImplClass, ModuleName )
 
+/** IMPLEMENT_ENGINE_DIR declares the engine directory to check for foreign or nested projects */
+#if PLATFORM_DESKTOP
+	#ifdef UE_ENGINE_DIRECTORY
+		#define IMPLEMENT_FOREIGN_ENGINE_DIR() const TCHAR *GForeignEngineDir = TEXT( PREPROCESSOR_TO_STRING(UE_ENGINE_DIRECTORY) );
+	#else
+		#define IMPLEMENT_FOREIGN_ENGINE_DIR() const TCHAR *GForeignEngineDir = NULL;
+	#endif
+#else
+	#define IMPLEMENT_FOREIGN_ENGINE_DIR() 
+#endif
+
 #if IS_PROGRAM
 
 	#if IS_MONOLITHIC
@@ -789,7 +800,8 @@ class FDefaultGameModuleImpl
 		#define IMPLEMENT_APPLICATION( ModuleName, GameName ) \
 			/* For monolithic builds, we must statically define the game's name string (See Core.h) */ \
 			TCHAR GGameName[64] = TEXT( GameName ); \
-			IMPLEMENT_GAME_MODULE( FDefaultGameModuleImpl, ModuleName ) \
+			IMPLEMENT_FOREIGN_ENGINE_DIR() \
+			IMPLEMENT_GAME_MODULE(FDefaultGameModuleImpl, ModuleName) \
 			PER_MODULE_BOILERPLATE \
 			FEngineLoop GEngineLoop;
 
@@ -804,6 +816,7 @@ class FDefaultGameModuleImpl
 					FCString::Strncpy(GGameName, TEXT( GameName ), ARRAY_COUNT(GGameName)); \
 				} \
 			} AutoSet##ModuleName; \
+			IMPLEMENT_FOREIGN_ENGINE_DIR() \
 			PER_MODULE_BOILERPLATE \
 			PER_MODULE_BOILERPLATE_ANYLINK(FDefaultGameModuleImpl, ModuleName) \
 			FEngineLoop GEngineLoop;
@@ -821,6 +834,7 @@ class FDefaultGameModuleImpl
 			TCHAR GGameName[64] = TEXT( GameName ); \
 			/* Implement the GIsGameAgnosticExe variable (See Core.h). */ \
 			bool GIsGameAgnosticExe = false; \
+			IMPLEMENT_FOREIGN_ENGINE_DIR() \
 			IMPLEMENT_GAME_MODULE( ModuleImplClass, ModuleName ) \
 			PER_MODULE_BOILERPLATE \
 			void UELinkerFixupCheat() \
@@ -835,6 +849,7 @@ class FDefaultGameModuleImpl
 			/* For monolithic builds, we must statically define the game's name string (See Core.h) */ \
 			TCHAR GGameName[64] = TEXT( GameName ); \
 			PER_MODULE_BOILERPLATE \
+			IMPLEMENT_FOREIGN_ENGINE_DIR() \
 			IMPLEMENT_GAME_MODULE( ModuleImplClass, ModuleName ) \
 			/* Implement the GIsGameAgnosticExe variable (See Core.h). */ \
 			bool GIsGameAgnosticExe = false; \
diff --git a/Engine/Source/Runtime/UE4Game/Private/UE4Game.cpp b/Engine/Source/Runtime/UE4Game/Private/UE4Game.cpp
--- a/Engine/Source/Runtime/UE4Game/Private/UE4Game.cpp
+++ b/Engine/Source/Runtime/UE4Game/Private/UE4Game.cpp
@@ -9,4 +9,5 @@ IMPLEMENT_MODULE(FDefaultModuleImpl, UE4Game);
 PER_MODULE_BOILERPLATE
 bool GIsGameAgnosticExe = true;
 TCHAR GGameName[64] = TEXT("");
+IMPLEMENT_FOREIGN_ENGINE_DIR()
 #endif
