CIS fix; don't declare foreign engine directory for non-monolithic programs - it's already defined in Core.

diff --git a/Engine/Source/Runtime/Core/Public/Modules/ModuleManager.h b/Engine/Source/Runtime/Core/Public/Modules/ModuleManager.h
--- a/Engine/Source/Runtime/Core/Public/Modules/ModuleManager.h
+++ b/Engine/Source/Runtime/Core/Public/Modules/ModuleManager.h
@@ -816,7 +816,6 @@ class FDefaultGameModuleImpl
 					FCString::Strncpy(GGameName, TEXT( GameName ), ARRAY_COUNT(GGameName)); \
 				} \
 			} AutoSet##ModuleName; \
-			IMPLEMENT_FOREIGN_ENGINE_DIR() \
 			PER_MODULE_BOILERPLATE \
 			PER_MODULE_BOILERPLATE_ANYLINK(FDefaultGameModuleImpl, ModuleName) \
 			FEngineLoop GEngineLoop;
