Removing Linux-specific CopyDirectoryTree() and replacing with generic one written by Rama.

Details: https://github.com/EpicGames/UnrealEngine/pull/27

diff --git a/Engine/Source/Runtime/Core/Private/GenericPlatform/GenericPlatformFile.cpp b/Engine/Source/Runtime/Core/Private/GenericPlatform/GenericPlatformFile.cpp
--- a/Engine/Source/Runtime/Core/Private/GenericPlatform/GenericPlatformFile.cpp
+++ b/Engine/Source/Runtime/Core/Private/GenericPlatform/GenericPlatformFile.cpp
@@ -119,8 +119,79 @@ bool IPlatformFile::CopyDirectoryTree(const TCHAR* DestinationDirectory, const T
 	check(DestinationDirectory);
 	check(Source);
 
-	// TODO: can be implemented using existing functions, all bits seem to be there.
-	return false;
+	FString DestDir(DestinationDirectory);
+	FPaths::NormalizeDirectoryName(DestDir);
+
+	FString SourceDir(Source);
+	FPaths::NormalizeDirectoryName(SourceDir);
+
+	// Does Source dir exist?
+	if (!DirectoryExists(*SourceDir))
+	{
+		return false;
+	}
+
+	// Destination directory exists already or can be created ?
+	if (!DirectoryExists(*DestDir) &&
+		!CreateDirectory(*DestDir))
+	{
+		return false;
+	}
+
+	// Copy all files and directories
+	struct FCopyFilesAndDirs : public FDirectoryVisitor
+	{
+		IPlatformFile & PlatformFile;
+		const TCHAR * SourceRoot;
+		const TCHAR * DestRoot;
+		bool bOverwrite;
+
+		FCopyFilesAndDirs(IPlatformFile& InPlatformFile, const TCHAR * InSourceRoot, const TCHAR * InDestRoot, bool bInOverwrite)
+			: PlatformFile(InPlatformFile)
+			, SourceRoot(InSourceRoot)
+			, DestRoot(InDestRoot)
+			, bOverwrite(bInOverwrite)
+		{
+		}
+
+		virtual bool Visit(const TCHAR* FilenameOrDirectory, bool bIsDirectory)
+		{
+			FString NewName(FilenameOrDirectory);
+			// change the root
+			NewName = NewName.Replace(SourceRoot, DestRoot);
+
+			if (bIsDirectory)
+			{
+				// create new directory structure
+				if (!PlatformFile.CreateDirectoryTree(*NewName) && !PlatformFile.DirectoryExists(*NewName))
+				{
+					return false;
+				}
+			}
+			else
+			{
+				// Delete destination file if it exists and we are overwriting
+				if (PlatformFile.FileExists(*NewName) && bOverwrite)
+				{
+					PlatformFile.DeleteFile(*NewName);
+				}
+
+				// Copy file from source
+				if (!PlatformFile.CopyFile(*NewName, FilenameOrDirectory))
+				{
+					// Not all files could be copied
+					return false;
+				}
+			}
+			return true; // continue searching
+		}
+	};
+
+	// copy files and directories visitor
+	FCopyFilesAndDirs CopyFilesAndDirs(*this, *SourceDir, *DestDir, bOverwriteAllExisting);
+
+	// create all files subdirectories and files in subdirectories!
+	return IterateDirectoryRecursively(*SourceDir, CopyFilesAndDirs);
 }
 
 FString IPlatformFile::ConvertToAbsolutePathForExternalAppForRead( const TCHAR* Filename )
diff --git a/Engine/Source/Runtime/Core/Private/Linux/LinuxPlatformFile.cpp b/Engine/Source/Runtime/Core/Private/Linux/LinuxPlatformFile.cpp
--- a/Engine/Source/Runtime/Core/Private/Linux/LinuxPlatformFile.cpp
+++ b/Engine/Source/Runtime/Core/Private/Linux/LinuxPlatformFile.cpp
@@ -329,27 +329,6 @@ bool FLinuxPlatformFile::IterateDirectory(const TCHAR* Directory, FDirectoryVisi
 	return Result;
 }
 
-bool FLinuxPlatformFile::CopyDirectoryTree(const TCHAR* DestinationDirectory, const TCHAR* Source, bool bOverwriteAllExisting)
-{
-	check(DestinationDirectory);
-	check(Source);
-	
-	// TODO: use own implementation
-
-	if (system(NULL) == 0)
-	{
-		UE_LOG(LogLinuxPlatformFile, Warning, TEXT("CopyDirectoryTree() failed: command processor is not available."));
-		return false;
-	}
-
-	FString SourceNormalized = NormalizeFilename(Source);
-	FString DestNormalized = NormalizeFilename(DestinationDirectory);
-
-	FString SysCall = FString::Printf(TEXT("cp -R \"%s\" \"%s\""), *SourceNormalized, *DestNormalized);
-	UE_LOG(LogLinuxPlatformFile, Log, TEXT("CopyDirectoryTree(): %s"), *SysCall);
-	return system(TCHAR_TO_ANSI(*SysCall)) == 0;
-}
-
 IPlatformFile& IPlatformFile::GetPlatformPhysical()
 {
 	static FLinuxPlatformFile Singleton;
diff --git a/Engine/Source/Runtime/Core/Public/Linux/LinuxPlatformFile.h b/Engine/Source/Runtime/Core/Public/Linux/LinuxPlatformFile.h
--- a/Engine/Source/Runtime/Core/Public/Linux/LinuxPlatformFile.h
+++ b/Engine/Source/Runtime/Core/Public/Linux/LinuxPlatformFile.h
@@ -35,6 +35,4 @@ public:
 	virtual bool CreateDirectory(const TCHAR* Directory) OVERRIDE;
 	virtual bool DeleteDirectory(const TCHAR* Directory) OVERRIDE;
 	bool IterateDirectory(const TCHAR* Directory, FDirectoryVisitor& Visitor);
-
-	virtual bool CopyDirectoryTree(const TCHAR* DestinationDirectory, const TCHAR* Source, bool bOverwriteAllExisting);
 };
