#ttp 331817 - CORE: Make sure GetChacheLineSize is initialized before first use.
#proj UE4
#branch UE4
#summary Refactored functions that were querying __cpuid to be pre-cached at startup.

diff --git a/Engine/Source/Runtime/Core/Private/Windows/WindowsPlatformMisc.cpp b/Engine/Source/Runtime/Core/Private/Windows/WindowsPlatformMisc.cpp
--- a/Engine/Source/Runtime/Core/Private/Windows/WindowsPlatformMisc.cpp
+++ b/Engine/Source/Runtime/Core/Private/Windows/WindowsPlatformMisc.cpp
@@ -1722,41 +1722,93 @@ FLinearColor FWindowsPlatformMisc::GetScreenPixelColor(const FVector2D& InScreen
 	return ScreenColor;
 }
 
-bool FWindowsPlatformMisc::HasCPUIDInstruction()
+/**
+* Class that caches __cpuid queried data.
+*/
+class FCPUIDQueriedData
 {
-#if PLATFORM_SEH_EXCEPTIONS_DISABLED
-	return false;
-#else
-	__try
+public:
+	FCPUIDQueriedData()
+		: bHasCPUIDInstruction(CheckForCPUIDInstruction()), Vendor(), CPUInfo(0), CacheLineSize(1)
 	{
-		int Args[4];
-		__cpuid(Args, 0);
+		if(bHasCPUIDInstruction)
+		{
+			Vendor = QueryCPUVendor();
+			CPUInfo = QueryCPUInfo();
+			CacheLineSize = QueryCacheLineSize();
+		}
 	}
-	__except (EXCEPTION_EXECUTE_HANDLER)
+
+	/** 
+	 * Checks if this CPU supports __cpuid instruction.
+	 *
+	 * @returns True if this CPU supports __cpuid instruction. False otherwise.
+	 */
+	static bool HasCPUIDInstruction() { return CPUIDStaticCache.bHasCPUIDInstruction; }
+
+	/**
+	 * Gets pre-cached CPU vendor name.
+	 *
+	 * @returns CPU vendor name.
+	 */
+	static const FString& GetVendor() { return CPUIDStaticCache.Vendor; }
+
+	/**
+	 * Gets __cpuid CPU info.
+	 *
+	 * @returns CPU info unsigned int queried using __cpuid.
+	 */
+	static uint32 GetCPUInfo() { return CPUIDStaticCache.CPUInfo; }
+
+	/** 
+	 * Gets cache line size.
+	 *
+	 * @returns Cache line size.
+	 */
+	static int32 GetCacheLineSize() { return CPUIDStaticCache.CacheLineSize; }
+
+private:
+	/**
+	 * Checks if __cpuid instruction is present on current machine.
+	 *
+	 * @returns True if this CPU supports __cpuid instruction. False otherwise.
+	 */
+	static bool CheckForCPUIDInstruction()
 	{
+#if PLATFORM_SEH_EXCEPTIONS_DISABLED
 		return false;
-	}
-	return true;
+#else
+		__try
+		{
+			int Args[4];
+			__cpuid(Args, 0);
+		}
+		__except (EXCEPTION_EXECUTE_HANDLER)
+		{
+			return false;
+		}
+		return true;
 #endif
-}
+	}
 
-FString FWindowsPlatformMisc::GetCPUVendor()
-{
-	union
+	/**
+	 * Queries Vendor name using __cpuid instruction.
+	 *
+	 * @returns CPU vendor name.
+	 */
+	static FString QueryCPUVendor()
 	{
-		char Buffer[12+1];
-		struct
+		union
 		{
-			int dw0;
-			int dw1;
-			int dw2;
-		} Dw;
-	} VendorResult;
-	
+			char Buffer[12 + 1];
+			struct
+			{
+				int dw0;
+				int dw1;
+				int dw2;
+			} Dw;
+		} VendorResult;
 
-	FString Vendor;
-	if (HasCPUIDInstruction())
-	{
 		int Args[4];
 		__cpuid(Args, 0);
 
@@ -1765,22 +1817,81 @@ FString FWindowsPlatformMisc::GetCPUVendor()
 		VendorResult.Dw.dw2 = Args[2];
 		VendorResult.Buffer[12] = 0;
 
-		Vendor = ANSI_TO_TCHAR(VendorResult.Buffer);
+		return ANSI_TO_TCHAR(VendorResult.Buffer);
 	}
-	return Vendor;
-}
 
-uint32 FWindowsPlatformMisc::GetCPUInfo()
-{
-	uint32 Info = 0;
-	if (HasCPUIDInstruction())
+	/**
+	 * Queries CPU info using __cpuid instruction.
+	 *
+	 * @returns CPU info unsigned int queried using __cpuid.
+	 */
+	static uint32 QueryCPUInfo()
 	{
+		uint32 Info = 0;
+
 		int Args[4];
 		__cpuid(Args, 1);
 
 		Info = Args[0];
+
+		return Info;
+	}
+
+	/**
+	 * Queries cache line size using __cpuid instruction.
+	 *
+	 * @returns Cache line size.
+	 */
+	static int32 QueryCacheLineSize()
+	{
+		int32 Result = 1;
+
+		int Args[4];
+		__cpuid(Args, 0x80000006);
+
+		Result = Args[2] & 0xFF;
+		check(Result && !(Result & (Result - 1))); // assumed to be a power of two
+
+		return Result;
 	}
-	return Info;
+
+	/** Static field with pre-cached __cpuid data. */
+	static FCPUIDQueriedData CPUIDStaticCache;
+
+	/** If machine has CPUID instruction. */
+	bool bHasCPUIDInstruction;
+
+	/** Vendor of the CPU. */
+	FString Vendor;
+
+	/** CPU info from __cpuid. */
+	uint32 CPUInfo;
+
+	/** CPU cache line size. */
+	int32 CacheLineSize;
+};
+
+/** Static initialization of data to pre-cache __cpuid queries. */
+FCPUIDQueriedData FCPUIDQueriedData::CPUIDStaticCache;
+
+bool FWindowsPlatformMisc::HasCPUIDInstruction()
+{
+	return FCPUIDQueriedData::HasCPUIDInstruction();
+}
+
+FString FWindowsPlatformMisc::GetCPUVendor()
+{
+	return FCPUIDQueriedData::GetVendor();
+}
+
+uint32 FWindowsPlatformMisc::GetCPUInfo()
+{
+	return FCPUIDQueriedData::GetCPUInfo();
+}
+
+int32 FWindowsPlatformMisc::GetCacheLineSize()
+{
+	return FCPUIDQueriedData::GetCacheLineSize();
 }
 
 bool FWindowsPlatformMisc::GetRegistryString(const FString& InRegistryKey, const FString& InValueName, bool bPerUserSetting, FString& OutValue)
@@ -1829,20 +1940,6 @@ bool FWindowsPlatformMisc::QueryRegKey( const HKEY InKey, const TCHAR* InSubKey,
 	return bSuccess;
 }
 
-int32 FWindowsPlatformMisc::GetCacheLineSize()
-{
-	int32 Result = 1;
-	if (HasCPUIDInstruction())
-	{
-		int Args[4];
-		__cpuid(Args, 0x80000006);
-
-		Result = Args[2] & 0xFF;
-		check(Result && !(Result & (Result - 1))); // assumed to be a power of two
-	}
-	return Result;
-}
-
 const TCHAR* FWindowsPlatformMisc::GetDefaultPathSeparator()
 {
 	return TEXT( "\\" );
