#UE4
- Added helper function WithAlpha

diff --git a/Engine/Source/Runtime/Core/Public/Math/Color.h b/Engine/Source/Runtime/Core/Public/Math/Color.h
--- a/Engine/Source/Runtime/Core/Public/Math/Color.h
+++ b/Engine/Source/Runtime/Core/Public/Math/Color.h
@@ -412,6 +412,15 @@ public:
 	 */
 	static CORE_API FColor MakeRedToGreenColorFromScalar(float Scalar);
 
+	/**
+	 *	@return a new FColor based of this color with the new alpha value.
+	 *	Usage: const FColor& MyColor = FColorList::Green.WithAlpha(128);
+	 */
+	FColor WithAlpha( uint8 Alpha ) const
+	{
+		return FColor( R, G, B, Alpha );
+	}
+
 	/**
 	 * Reinterprets the color as a linear color.
 	 *
