#UE4
- Improved code used to finalize the stats file

diff --git a/Engine/Source/Runtime/Core/Private/Stats/StatsFile.cpp b/Engine/Source/Runtime/Core/Private/Stats/StatsFile.cpp
--- a/Engine/Source/Runtime/Core/Private/Stats/StatsFile.cpp
+++ b/Engine/Source/Runtime/Core/Private/Stats/StatsFile.cpp
@@ -286,20 +286,26 @@ void FStatsWriteFile::Finalize( FArchive* File )
 
 	const FStatsThreadState& Stats = FStatsThreadState::GetLocalState();
 
-	// Create copy of names.
+	// Add FNames from the stats metadata.
+	for( const auto& It : Stats.ShortNameToLongName )
+	{
+		const FName ShortName = It.Key;
+		const FStatMessage& StatMessage = It.Value;
+		FNamesSent.Add( StatMessage.NameAndInfo.GetRawName().GetIndex() );
+	}
+
+	// Create a copy of names.
 	TSet<int32> FNamesToSent = FNamesSent;
 	FNamesSent.Empty( FNamesSent.Num() );
 
 	// Serialize FNames.
 	Header.FNameTableOffset = Ar.Tell();
 	Header.NumFNames = FNamesToSent.Num();
-	for( auto& It : FNamesToSent )
+	for( const int32 It : FNamesToSent )
 	{
 		WriteFName( Ar, FStatNameAndInfo(FName(EName(It)),false) );
 	}
 
-	//STATGROUP_UObjects
-
 	// Serialize metadata messages.
 	Header.MetadataMessagesOffset = Ar.Tell();
 	Header.NumMetadataMessages = Stats.ShortNameToLongName.Num();
@@ -308,12 +314,21 @@ void FStatsWriteFile::Finalize( FArchive* File )
 		WriteMessage( Ar, It.Value() );
 	}
 
-	// Find the difference?
-	TSet<int32> Union = FNamesSent.Union( FNamesToSent );
-	TSet<int32> Intersect = FNamesSent.Intersect( FNamesToSent );
-
-	TSet<int32> AMinB = FNamesToSent.Difference( FNamesSent );
+	// Verify data.
 	TSet<int32> BMinA = FNamesSent.Difference( FNamesToSent );
+	struct FLocal
+	{
+		static TArray<FName> GetFNameArray( const TSet<int32>& NameIndices )
+		{
+			TArray<FName> Result;
+			for( const int32 NameIndex : NameIndices )
+			{
+				new(Result) FName( EName( NameIndex ) );
+			}
+			return Result;
+		}
+	};
+	auto BMinANames = FLocal::GetFNameArray( BMinA );
 
 	// Seek to the position just after a magic value of the file and write out proper header.
 	Ar.Seek( sizeof(uint32) );
