UE4 hack fix for cook hang

diff --git a/Engine/Source/Runtime/Core/Private/Misc/OutputDevice.cpp b/Engine/Source/Runtime/Core/Private/Misc/OutputDevice.cpp
--- a/Engine/Source/Runtime/Core/Private/Misc/OutputDevice.cpp
+++ b/Engine/Source/Runtime/Core/Private/Misc/OutputDevice.cpp
@@ -487,24 +487,47 @@ VARARG_BODY(void, FMsg::Logf, const TCHAR*, VARARG_EXTRA(const ANSICHAR* File) V
 	}
 	else
 	{
-		// We're using one shared static buffer here, so guard against re-entry
-		FScopeLock MsgLock(&MsgLogfReEntryGuard);
+		if (IsInGameThread())
+		{
+			static TCHAR MsgLogfStaticBuffer[8192]; // Increased from 4096 to fix crashes in the renderthread without autoreporter
 
-		// Flush logs queued by threads when we hit an assert because they will not make it to the log otherwise
-		GLog->PanicFlushThreadedLogs();
+			// Flush logs queued by threads when we hit an assert because they will not make it to the log otherwise
+			GLog->PanicFlushThreadedLogs();
 
-		// Print to a large static buffer so we can keep the stack allocation below 16K
-		GET_VARARGS(MsgLogfStaticBuffer, ARRAY_COUNT(MsgLogfStaticBuffer), ARRAY_COUNT(MsgLogfStaticBuffer) - 1, Fmt, Fmt);
+			// Print to a large static buffer so we can keep the stack allocation below 16K
+			GET_VARARGS(MsgLogfStaticBuffer, ARRAY_COUNT(MsgLogfStaticBuffer), ARRAY_COUNT(MsgLogfStaticBuffer) - 1, Fmt, Fmt);
 
-		FailDebug(TEXT("Fatal error:"), File, Line, MsgLogfStaticBuffer);
-		if (!FPlatformMisc::IsDebuggerPresent())
-		{
-			FPlatformMisc::PromptForRemoteDebugging(false);
+			FailDebug(TEXT("Fatal error:"), File, Line, MsgLogfStaticBuffer);
+			if (!FPlatformMisc::IsDebuggerPresent())
+			{
+				FPlatformMisc::PromptForRemoteDebugging(false);
+			}
+
+			FPlatformMisc::DebugBreak();
+
+			GError->Log(Category, Verbosity, MsgLogfStaticBuffer);
 		}
+		else
+		{
+			// We're using one shared static buffer here, so guard against re-entry
+			FScopeLock MsgLock(&MsgLogfReEntryGuard);
 
-		FPlatformMisc::DebugBreak();
+			// Flush logs queued by threads when we hit an assert because they will not make it to the log otherwise
+			GLog->PanicFlushThreadedLogs();
+
+			// Print to a large static buffer so we can keep the stack allocation below 16K
+			GET_VARARGS(MsgLogfStaticBuffer, ARRAY_COUNT(MsgLogfStaticBuffer), ARRAY_COUNT(MsgLogfStaticBuffer) - 1, Fmt, Fmt);
+
+			FailDebug(TEXT("Fatal error:"), File, Line, MsgLogfStaticBuffer);
+			if (!FPlatformMisc::IsDebuggerPresent())
+			{
+				FPlatformMisc::PromptForRemoteDebugging(false);
+			}
+
+			FPlatformMisc::DebugBreak();
 
-		GError->Log(Category, Verbosity, MsgLogfStaticBuffer);
+			GError->Log(Category, Verbosity, MsgLogfStaticBuffer);
+		}
 	}
 #endif
 }
