- Fixed controller buttons in Android
- Added some buttons inspired by kindle fire tv

diff --git a/Engine/Source/Runtime/Core/Private/Android/AndroidInputInterface.cpp b/Engine/Source/Runtime/Core/Private/Android/AndroidInputInterface.cpp
--- a/Engine/Source/Runtime/Core/Private/Android/AndroidInputInterface.cpp
+++ b/Engine/Source/Runtime/Core/Private/Android/AndroidInputInterface.cpp
@@ -237,14 +237,17 @@ void FAndroidInputInterface::JoystickButtonEvent(int32 deviceId, int32 buttonId,
 
 	switch (buttonId)
 	{
-		case AKEYCODE_BUTTON_A:      NewControllerData[deviceId].ButtonStates[ 0] = buttonDown; break;
+		case AKEYCODE_BUTTON_A:
+		case AKEYCODE_DPAD_CENTER:   NewControllerData[deviceId].ButtonStates[ 0] = buttonDown; break;
 		case AKEYCODE_BUTTON_B:      NewControllerData[deviceId].ButtonStates[ 1] = buttonDown; break;
 		case AKEYCODE_BUTTON_X:      NewControllerData[deviceId].ButtonStates[ 2] = buttonDown; break;
 		case AKEYCODE_BUTTON_Y:      NewControllerData[deviceId].ButtonStates[ 3] = buttonDown; break;
 		case AKEYCODE_BUTTON_L1:     NewControllerData[deviceId].ButtonStates[ 4] = buttonDown; break;
 		case AKEYCODE_BUTTON_R1:     NewControllerData[deviceId].ButtonStates[ 5] = buttonDown; break;
-		case AKEYCODE_BUTTON_START:  NewControllerData[deviceId].ButtonStates[ 6] = buttonDown; break;
-		case AKEYCODE_BUTTON_SELECT: NewControllerData[deviceId].ButtonStates[ 7] = buttonDown; break;
+		case AKEYCODE_BUTTON_START:
+		case AKEYCODE_MENU:          NewControllerData[deviceId].ButtonStates[ 6] = buttonDown; break;
+		case AKEYCODE_BUTTON_SELECT: 
+		case AKEYCODE_BACK:          NewControllerData[deviceId].ButtonStates[ 7] = buttonDown; break;
 		case AKEYCODE_BUTTON_THUMBL: NewControllerData[deviceId].ButtonStates[ 8] = buttonDown; break;
 		case AKEYCODE_BUTTON_THUMBR: NewControllerData[deviceId].ButtonStates[ 9] = buttonDown; break;
 		case AKEYCODE_BUTTON_L2:     NewControllerData[deviceId].ButtonStates[10] = buttonDown; break;
diff --git a/Engine/Source/Runtime/Launch/Private/Android/LaunchAndroid.cpp b/Engine/Source/Runtime/Launch/Private/Android/LaunchAndroid.cpp
--- a/Engine/Source/Runtime/Launch/Private/Android/LaunchAndroid.cpp
+++ b/Engine/Source/Runtime/Launch/Private/Android/LaunchAndroid.cpp
@@ -331,6 +331,9 @@ void android_main(struct android_app* state)
 //Called from the event process thread
 static int32_t HandleInputCB(struct android_app* app, AInputEvent* event)
 {
+//	FPlatformMisc::LowLevelOutputDebugStringf(L"INPUT - type: %x, action: %x, source: %x, keycode: %x, buttons: %x", AInputEvent_getType(event), 
+//		AMotionEvent_getAction(event), AInputEvent_getSource(event), AKeyEvent_getKeyCode(event), AMotionEvent_getButtonState(event));
+
 	if (AInputEvent_getType(event) == AINPUT_EVENT_TYPE_MOTION)
 	{
 		int action = AMotionEvent_getAction(event);
@@ -459,7 +462,7 @@ static int32_t HandleInputCB(struct android_app* app, AInputEvent* event)
 		int keyCode = AKeyEvent_getKeyCode(event);
 
 		//Trap Joystick events first, with fallthrough if there is no joystick support
-		if (((AInputEvent_getSource(event) & AINPUT_SOURCE_CLASS_JOYSTICK) != 0) && (GetAxes != NULL))
+		if (((AInputEvent_getSource(event) & (AINPUT_SOURCE_GAMEPAD | AINPUT_SOURCE_DPAD)) != 0) && (GetAxes != NULL))
 		{
 			int device = AInputEvent_getDeviceId(event);
 			bool down = AKeyEvent_getAction(event) != AKEY_EVENT_ACTION_UP;
