#UE4 - Prevent FMsg::Logf deadlocks, lock the static buffer only when doing sprintf

diff --git a/Engine/Source/Runtime/Core/Private/Misc/OutputDevice.cpp b/Engine/Source/Runtime/Core/Private/Misc/OutputDevice.cpp
--- a/Engine/Source/Runtime/Core/Private/Misc/OutputDevice.cpp
+++ b/Engine/Source/Runtime/Core/Private/Misc/OutputDevice.cpp
@@ -455,7 +455,7 @@ CORE_API FOutputDeviceError*			GError							= NULL;						/* Critical errors */
 CORE_API FOutputDevice*					GThrow							= &ThrowOut;				/* Exception thrower */
 
 // Statics to prevent FMsg::Logf from allocating too much stack memory
-static FCriticalSection					MsgLogfReEntryGuard;
+static FCriticalSection					MsgLogfStaticBufferGuard;
 static TCHAR							MsgLogfStaticBuffer[8192]; // Increased from 4096 to fix crashes in the renderthread without autoreporter
 
 VARARG_BODY(void, FMsg::Logf, const TCHAR*, VARARG_EXTRA(const ANSICHAR* File) VARARG_EXTRA(int32 Line) VARARG_EXTRA(const class FName& Category) VARARG_EXTRA(ELogVerbosity::Type Verbosity))
@@ -487,47 +487,32 @@ VARARG_BODY(void, FMsg::Logf, const TCHAR*, VARARG_EXTRA(const ANSICHAR* File) V
 	}
 	else
 	{
-		if (IsInGameThread())
-		{
-			static TCHAR MsgLogfStaticBuffer[8192]; // Increased from 4096 to fix crashes in the renderthread without autoreporter
-
-			// Flush logs queued by threads when we hit an assert because they will not make it to the log otherwise
-			GLog->PanicFlushThreadedLogs();
-
-			// Print to a large static buffer so we can keep the stack allocation below 16K
-			GET_VARARGS(MsgLogfStaticBuffer, ARRAY_COUNT(MsgLogfStaticBuffer), ARRAY_COUNT(MsgLogfStaticBuffer) - 1, Fmt, Fmt);
+		// Flush logs queued by threads when we hit an assert because they will not make it to the log otherwise
+		GLog->PanicFlushThreadedLogs();
 
-			FailDebug(TEXT("Fatal error:"), File, Line, MsgLogfStaticBuffer);
-			if (!FPlatformMisc::IsDebuggerPresent())
-			{
-				FPlatformMisc::PromptForRemoteDebugging(false);
-			}
-
-			FPlatformMisc::DebugBreak();
-
-			GError->Log(Category, Verbosity, MsgLogfStaticBuffer);
-		}
-		else
+		// Keep Message buffer small, in some cases, this code is executed with 16KB stack.
+		TCHAR Message[4096];
 		{
-			// We're using one shared static buffer here, so guard against re-entry
-			FScopeLock MsgLock(&MsgLogfReEntryGuard);
-
-			// Flush logs queued by threads when we hit an assert because they will not make it to the log otherwise
-			GLog->PanicFlushThreadedLogs();
-
+			// Simulate Sprintf_s
+			// @todo: implement platform independent sprintf_S
+			// We're using one big shared static buffer here, so guard against re-entry
+			FScopeLock MsgLock(&MsgLogfStaticBufferGuard);
 			// Print to a large static buffer so we can keep the stack allocation below 16K
 			GET_VARARGS(MsgLogfStaticBuffer, ARRAY_COUNT(MsgLogfStaticBuffer), ARRAY_COUNT(MsgLogfStaticBuffer) - 1, Fmt, Fmt);
+			// Copy the message to the stack-allocated buffer)
+			FCString::Strncpy(Message, MsgLogfStaticBuffer, ARRAY_COUNT(Message) - 1);
+			Message[ARRAY_COUNT(Message) - 1] = '\0';
+		}
 
-			FailDebug(TEXT("Fatal error:"), File, Line, MsgLogfStaticBuffer);
-			if (!FPlatformMisc::IsDebuggerPresent())
-			{
-				FPlatformMisc::PromptForRemoteDebugging(false);
-			}
+		FailDebug(TEXT("Fatal error:"), File, Line, Message);
+		if (!FPlatformMisc::IsDebuggerPresent())
+		{
+			FPlatformMisc::PromptForRemoteDebugging(false);
+		}
 
-			FPlatformMisc::DebugBreak();
+		FPlatformMisc::DebugBreak();
 
-			GError->Log(Category, Verbosity, MsgLogfStaticBuffer);
-		}
+		GError->Log(Category, Verbosity, Message);
 	}
 #endif
 }
