#ttp 331016 -Android: Replace hardcoded /mnt/sdcard with getExternalStorageLocation() or whateve
#proj UnrealEd
#branch UE4
#summary On startup external storage location is queried and cached for subsequent file ops

diff --git a/Engine/Source/Runtime/Core/Private/Android/AndroidFile.cpp b/Engine/Source/Runtime/Core/Private/Android/AndroidFile.cpp
--- a/Engine/Source/Runtime/Core/Private/Android/AndroidFile.cpp
+++ b/Engine/Source/Runtime/Core/Private/Android/AndroidFile.cpp
@@ -17,6 +17,9 @@ FString GPackageName;
 static int32 GPackageVersion = 0;
 static int32 GPackagePatchVersion = 0;
 
+// External File Path base - setup during load
+FString GFilePathBase;
+
 //This function is declared in the Java-defined class, GameActivity.java: "public native void nativeSetObbInfo(String PackageName, int Version, int PatchVersion);"
 extern "C" void Java_com_epicgames_ue4_GameActivity_nativeSetObbInfo(JNIEnv* jenv, jobject thiz, jstring PackageName, jint Version, jint PatchVersion)
 {
@@ -386,7 +389,7 @@ public:
 			// look in the downloader path first, then manual location as a fallback
 			for (int32 PassIndex = 0; PassIndex < 2 && ObbHijackedPakDirectory == TEXT(""); PassIndex++)
 			{
-				FString ObbPath = ((PassIndex == 0) ? FString(TEXT("/mnt/sdcard/Android/obb/")) : FString(TEXT("/mnt/sdcard/obb/"))) + GPackageName;
+				FString ObbPath = GFilePathBase + ((PassIndex == 0) ? FString(TEXT("/Android/obb/")) : FString(TEXT("/obb/"))) + GPackageName;
 
 				DIR* Handle = opendir(TCHAR_TO_UTF8(*ObbPath));
 				if (Handle)
@@ -462,7 +465,7 @@ public:
 		Result.ReplaceInline(TEXT(".."), TEXT(""));
 		Result.ReplaceInline(FPlatformProcess::BaseDir(), TEXT(""));
 
-		static FString BasePath = FString("/mnt/sdcard/") + GGameName + FString("/");
+		static FString BasePath = GFilePathBase + FString("/") + GGameName + FString("/");
 		Result =  BasePath + Result;
 		return Result;
 	}
diff --git a/Engine/Source/Runtime/Core/Private/Android/AndroidMisc.cpp b/Engine/Source/Runtime/Core/Private/Android/AndroidMisc.cpp
--- a/Engine/Source/Runtime/Core/Private/Android/AndroidMisc.cpp
+++ b/Engine/Source/Runtime/Core/Private/Android/AndroidMisc.cpp
@@ -138,23 +138,26 @@ int32 FAndroidMisc::NumberOfCores()
 	return NumberOfCores;
 }
 
+extern FString GFilePathBase;
+
 class FTestUtime
 {
 public:
 	FTestUtime()
 		: Supported(false)
 	{
-		static const char * TestFilePath = "/mnt/sdcard/UE4UtimeTest.txt";
-		FILE * FileHandle = fopen(TestFilePath, "w");
+		static FString TestFilePath = GFilePathBase + FString(TEXT("/UE4UtimeTest.txt"));
+		static const char * TestFilePathChar = StringCast<ANSICHAR>(*TestFilePath).Get();
+		FILE * FileHandle = fopen(TestFilePathChar, "w");
 		if(FileHandle)
 		{
 			fclose(FileHandle);
 			struct utimbuf Times;
 			Times.actime = 0;
 			Times.modtime = 0;
-			int Result = utime(TestFilePath, &Times);
+			int Result = utime(TestFilePathChar, &Times);
 			Supported = -1 != Result;
-			unlink(TestFilePath);
+			unlink(TestFilePathChar);
 		}
 
 		FPlatformMisc::LowLevelOutputDebugStringf(TEXT("Is Local Caching Supported? %d\n"), Supported);
diff --git a/Engine/Source/Runtime/Launch/Private/Android/AndroidJNI.cpp b/Engine/Source/Runtime/Launch/Private/Android/AndroidJNI.cpp
--- a/Engine/Source/Runtime/Launch/Private/Android/AndroidJNI.cpp
+++ b/Engine/Source/Runtime/Launch/Private/Android/AndroidJNI.cpp
@@ -10,6 +10,8 @@
 JavaVM* GJavaVM;
 jobject GJavaGlobalThis = NULL;
 
+extern FString GFilePathBase;
+
 //////////////////////////////////////////////////////////////////////////
 // FJNIHelper
 
@@ -193,6 +195,21 @@ JNIEXPORT jint JNI_OnLoad(JavaVM* InJavaVM, void* InReserved)
 		FPlatformMisc::SetCrashHandler(EngineCrashHandler);
 	}
 
+	// Cache path to external storage
+	jclass EnvClass = env->FindClass("android/os/Environment");
+	jmethodID getExternalStorageDir = env->GetStaticMethodID(EnvClass, "getExternalStorageDirectory", "()Ljava/io/File;");
+	jobject externalStoragePath = env->CallStaticObjectMethod(EnvClass, getExternalStorageDir, nullptr);
+	jmethodID getFilePath = env->GetMethodID(env->FindClass("java/io/File"), "getPath", "()Ljava/lang/String;");
+	jstring pathString = (jstring)env->CallObjectMethod(externalStoragePath, getFilePath, nullptr);
+	const char *nativePathString = env->GetStringUTFChars(pathString, 0);
+	// Copy that somewhere safe 
+	GFilePathBase = FString(nativePathString);
+	
+	// then release...
+	env->ReleaseStringUTFChars(pathString, nativePathString);
+	FPlatformMisc::LowLevelOutputDebugStringf(TEXT("Path found as '%s'\n"), *GFilePathBase);
+
+
 	// Wire up to core delegates, so core code can call out to Java
 	DECLARE_DELEGATE_OneParam(FAndroidLaunchURLDelegate, const FString&);
 	extern CORE_API FAndroidLaunchURLDelegate OnAndroidLaunchURL;
diff --git a/Engine/Source/Runtime/Launch/Private/Android/LaunchAndroid.cpp b/Engine/Source/Runtime/Launch/Private/Android/LaunchAndroid.cpp
--- a/Engine/Source/Runtime/Launch/Private/Android/LaunchAndroid.cpp
+++ b/Engine/Source/Runtime/Launch/Private/Android/LaunchAndroid.cpp
@@ -53,6 +53,9 @@ extern "C"
 
 extern void AndroidThunkCpp_ShowConsoleWindow();
 
+// Base path for file accesses
+extern FString GFilePathBase;
+
 /** The global EngineLoop instance */
 FEngineLoop	GEngineLoop;
 
@@ -122,7 +125,7 @@ static void InitCommandLine()
 	FCommandLine::Set(TEXT(""));
 
 	// read in the command line text file from the sdcard if it exists
-	FString CommandLineFilePath = FString("/mnt/sdcard/") + (GGameName[0] ? GGameName : TEXT("UE4Game")) + FString("/UE4CommandLine.txt");
+	FString CommandLineFilePath = GFilePathBase + FString("/") + (GGameName[0] ? GGameName : TEXT("UE4Game")) + FString("/UE4CommandLine.txt");
 	FILE* CommandLineFile = fopen(TCHAR_TO_UTF8(*CommandLineFilePath), "r");
 	if(CommandLineFile == NULL)
 	{
