#UE4
- Replaced magic value 'Thread_' with ThreadNameMarker, improved a bit versioning of stats files
- Added an option to inject the full stats metadata into stream, required by the profiler client manager to maintain the continuity of the stat messages
- Simplified thread metadata
- Added full thread name to the raw callstack

diff --git a/Engine/Source/Runtime/Core/Private/Stats/Stats2.cpp b/Engine/Source/Runtime/Core/Private/Stats/Stats2.cpp
--- a/Engine/Source/Runtime/Core/Private/Stats/Stats2.cpp
+++ b/Engine/Source/Runtime/Core/Private/Stats/Stats2.cpp
@@ -89,9 +89,9 @@ FName TStatId::TStatId_NAME_None;
 void FStartupMessages::AddThreadMetadata( const FName InThreadName, uint32 InThreadID )
 {
 	// Make unique name.
-	const FString ThreadName = FStatsUtils::BuildUniqueThreadName( InThreadName, InThreadID );
+	const FString ThreadName = FStatsUtils::BuildUniqueThreadName( InThreadID );
 
-	FStartupMessages::AddMetadata( FName(*ThreadName), *ThreadName, STAT_GROUP_TO_FStatGroup( STATGROUP_Threads )::GetGroupName(), STAT_GROUP_TO_FStatGroup( STATGROUP_Threads )::GetDescription(), true, EStatDataType::ST_int64, true );
+	FStartupMessages::AddMetadata( InThreadName, *ThreadName, STAT_GROUP_TO_FStatGroup( STATGROUP_Threads )::GetGroupName(), STAT_GROUP_TO_FStatGroup( STATGROUP_Threads )::GetDescription(), true, EStatDataType::ST_int64, true );
 }
 
 
diff --git a/Engine/Source/Runtime/Core/Private/Stats/StatsCommand.cpp b/Engine/Source/Runtime/Core/Private/Stats/StatsCommand.cpp
--- a/Engine/Source/Runtime/Core/Private/Stats/StatsCommand.cpp
+++ b/Engine/Source/Runtime/Core/Private/Stats/StatsCommand.cpp
@@ -235,7 +235,7 @@ void DumpCPUSummary(FStatsThreadState const& StatsData, int64 TargetFrame)
 
 		FString Desc;
 		Item.NameAndInfo.GetDescription(Desc);
-		bool bIsThread = Desc.StartsWith("Thread_");
+		bool bIsThread = Desc.StartsWith( FStatConstants::ThreadNameMarker );
 		bool bIsStall = !bIsThread && Desc.StartsWith("CPU Stall");
 		
 		EStatOperation::Type Op = Item.NameAndInfo.GetField<EStatOperation>();
diff --git a/Engine/Source/Runtime/Core/Private/Stats/StatsData.cpp b/Engine/Source/Runtime/Core/Private/Stats/StatsData.cpp
--- a/Engine/Source/Runtime/Core/Private/Stats/StatsData.cpp
+++ b/Engine/Source/Runtime/Core/Private/Stats/StatsData.cpp
@@ -24,6 +24,8 @@ const FName FStatConstants::NAME_ThreadGroup = FStatConstants::ThreadGroupName;
 const FString FStatConstants::StatsFileExtension = TEXT( ".ue4stats" );
 const FString FStatConstants::StatsFileRawExtension = TEXT( ".ue4statsraw" );
 
+const FString FStatConstants::ThreadNameMarker = TEXT( "Thread_" );
+
 /**
 * Magic numbers for stats streams
 */
@@ -1014,8 +1016,8 @@ void FStatsThreadState::GetRawStackStats(int64 TargetFrame, FRawStatStackNode& R
 		FRawStatStackNode* ThreadRoot = Root.Children.FindRef(ThreadName);
 		if (!ThreadRoot)
 		{
-			FString ThreadIdName = FString::Printf(TEXT("Thread_%x_0"), Packet.ThreadId);
-			ThreadRoot = Root.Children.Add(ThreadName, new FRawStatStackNode(FStatMessage(ThreadName, EStatDataType::ST_int64, NULL, *ThreadIdName, true, true)));
+			FString ThreadIdName = FStatsUtils::BuildUniqueThreadName( Packet.ThreadId );
+			ThreadRoot = Root.Children.Add( ThreadName, new FRawStatStackNode( FStatMessage( ThreadName, EStatDataType::ST_int64, STAT_GROUP_TO_FStatGroup( STATGROUP_Threads )::GetGroupName(), *ThreadIdName, true, true ) ) );
 			ThreadRoot->Meta.NameAndInfo.SetFlag(EStatMetaFlags::IsPackedCCAndDuration, true);
 			ThreadRoot->Meta.Clear();
 		}
@@ -1222,7 +1224,7 @@ FName FStatsThreadState::GetStatThreadName( const FStatPacket& Packet ) const
 			UE_LOG( LogStats, Warning, TEXT( "There is no thread with id: %u. Please add thread metadata for this thread." ), Packet.ThreadId );
 
 			static const FName NAME_UnknownThread = TEXT( "UnknownThread" );
-			NewThreadName = FName( *FStatsUtils::BuildUniqueThreadName( NAME_UnknownThread, Packet.ThreadId ) );
+			NewThreadName = FName( *FStatsUtils::BuildUniqueThreadName( Packet.ThreadId ) );
 			// This is an unknown thread, but still we need the metadata in the system.
 			FStartupMessages::Get().AddThreadMetadata( NAME_UnknownThread, Packet.ThreadId );
 		}
@@ -1262,7 +1264,9 @@ void FStatsThreadState::FindOrAddMetaData(FStatMessage const& Item)
 		const bool bIsThread = FStatConstants::NAME_ThreadGroup == GroupName;
 		if( bIsThread )
 		{
-			Threads.Add( FStatsUtils::ParseThreadID( ShortName ), ShortName );
+			FString Desc;
+			Item.NameAndInfo.GetDescription( Desc );
+			Threads.Add( FStatsUtils::ParseThreadID( Desc ), ShortName );
 		}
 
 		ShortNameToLongName.Add(ShortName, AsSet); // we want this to be a clear, but it should be a SetLongName
diff --git a/Engine/Source/Runtime/Core/Private/Stats/StatsFile.cpp b/Engine/Source/Runtime/Core/Private/Stats/StatsFile.cpp
--- a/Engine/Source/Runtime/Core/Private/Stats/StatsFile.cpp
+++ b/Engine/Source/Runtime/Core/Private/Stats/StatsFile.cpp
@@ -227,24 +227,27 @@ void FStatsWriteFile::WriteHeader()
 	Ar << Magic;
 
 	// Serialize dummy header, overwritten in Finalize.
+	Header.Version = EStatMagicWithHeader::VERSION_2;
+	Header.PlatformName = FPlatformProperties::PlatformName();
 	Header.bRawStatFile = false;
 	Ar << Header;
 
 	// Serialize metadata.
-	FStatsThreadState const& Stats = FStatsThreadState::GetLocalState();
-	for( auto It = Stats.ShortNameToLongName.CreateConstIterator(); It; ++It )
-	{
-		WriteMessage( Ar, It.Value() );
-	}
+	WriteMetadata( Ar );
 
 	Filepos += Ar.Tell();
 }
 
-void FStatsWriteFile::WriteFrame( int64 TargetFrame )
+void FStatsWriteFile::WriteFrame( int64 TargetFrame, bool bNeedFullMetadata /*= false*/ )
 {
 	FMemoryWriter Ar( OutData, false, true );
 	const int64 PrevArPos = Ar.Tell();
 
+	if( bNeedFullMetadata )
+	{
+		WriteMetadata( Ar );
+	}
+
 	FStatsThreadState const& Stats = FStatsThreadState::GetLocalState();
 	TArray<FStatMessage> const& Data = Stats.GetCondensedHistory( TargetFrame );
 	for( auto It = Data.CreateConstIterator(); It; ++It )
@@ -277,8 +280,6 @@ void FStatsWriteFile::Finalize( FArchive* File )
 	WriteMessage( Ar, SpecialMessage );
 
 	// Real header, written at start of the file, but written out right before we close the file.
-	Header.Version = EStatMagicWithHeader::VERSION;
-	Header.PlatformName = FPlatformProperties::PlatformName();
 
 	// Write out frame table and update header with offset and count.
 	Header.FrameTableOffset = Ar.Tell();
@@ -289,7 +290,6 @@ void FStatsWriteFile::Finalize( FArchive* File )
 	// Add FNames from the stats metadata.
 	for( const auto& It : Stats.ShortNameToLongName )
 	{
-		const FName ShortName = It.Key;
 		const FStatMessage& StatMessage = It.Value;
 		FNamesSent.Add( StatMessage.NameAndInfo.GetRawName().GetIndex() );
 	}
@@ -309,10 +309,7 @@ void FStatsWriteFile::Finalize( FArchive* File )
 	// Serialize metadata messages.
 	Header.MetadataMessagesOffset = Ar.Tell();
 	Header.NumMetadataMessages = Stats.ShortNameToLongName.Num();
-	for( auto It = Stats.ShortNameToLongName.CreateConstIterator(); It; ++It )
-	{
-		WriteMessage( Ar, It.Value() );
-	}
+	WriteMetadata( Ar );
 
 	// Verify data.
 	TSet<int32> BMinA = FNamesSent.Difference( FNamesToSent );
@@ -411,6 +408,8 @@ void FRawStatsWriteFile::WriteHeader()
 	Ar << Magic;
 
 	// Serialize dummy header, overwritten in Finalize.
+	Header.Version = EStatMagicWithHeader::VERSION_2;
+	Header.PlatformName = FPlatformProperties::PlatformName();
 	Header.bRawStatFile = true;
 	Ar << Header;
 
@@ -418,15 +417,12 @@ void FRawStatsWriteFile::WriteHeader()
 	FStatsThreadState const& Stats = FStatsThreadState::GetLocalState();
 	int32 NumMetadataMessages = Stats.ShortNameToLongName.Num();
 	Ar << NumMetadataMessages;
-	for( auto It = Stats.ShortNameToLongName.CreateConstIterator(); It; ++It )
-	{
-		WriteMessage( Ar, It.Value() );
-	}
+	WriteMetadata( Ar );
 
 	Filepos += Ar.Tell();
 }
 
-void FRawStatsWriteFile::WriteFrame( int64 TargetFrame )
+void FRawStatsWriteFile::WriteFrame( int64 TargetFrame, bool bNeedFullMetadata /*= false */)
 {
 	FMemoryWriter Ar( OutData, false, true );
 	const int64 PrevArPos = Ar.Tell();
diff --git a/Engine/Source/Runtime/Core/Public/Stats/StatsData.h b/Engine/Source/Runtime/Core/Public/Stats/StatsData.h
--- a/Engine/Source/Runtime/Core/Public/Stats/StatsData.h
+++ b/Engine/Source/Runtime/Core/Public/Stats/StatsData.h
@@ -30,6 +30,9 @@ struct CORE_API FStatConstants
 
 	/** Extension used to save a raw stats file, may be changed to the same as a regular stats file. */
 	static const FString StatsFileRawExtension;
+
+	/** Indicates that the item is a thread. */
+	static const FString ThreadNameMarker;
 };
 
 /** Enumerates stat compare types. */
@@ -598,21 +601,28 @@ struct CORE_API FStatsUtils
 	static FString ToEscapedFString(const TCHAR* Source);
 	static FString FromEscapedFString(const TCHAR* Escaped);
 
-	static FString BuildUniqueThreadName( const FName InThreadName, uint32 InThreadID )
+	static FString BuildUniqueThreadName( uint32 InThreadID )
 	{
 		// Make unique name.
-		return FString::Printf( TEXT( "%s_%x_0" ), *InThreadName.ToString(), InThreadID );
+		return FString::Printf( TEXT( "%s%x_0" ), *FStatConstants::ThreadNameMarker, InThreadID );
 	}
 
-	static int32 ParseThreadID( const FName ThreadFName )
+	static int32 ParseThreadID( const FString& InThreadName, FString* out_ThreadName = nullptr )
 	{
 		// Extract the thread id.
-		const FString ThreadName = ThreadFName.ToString().Replace(TEXT("_0"),TEXT(""));
+		const FString ThreadName = InThreadName.Replace( TEXT( "_0" ), TEXT( "" ) );
 		int32 Index = 0;
 		ThreadName.FindLastChar(TEXT('_'),Index);
 
+		// Parse the thread name if requested.
+		if( out_ThreadName )
+		{
+			*out_ThreadName = ThreadName.Left( Index );
+		}
+
 		const FString ThreadIDStr = ThreadName.RightChop(Index+1);
 		const uint32 ThreadID = FParse::HexNumber( *ThreadIDStr );
+
 		return ThreadID;
 	}
 };
