#UE4
- Added missing file

diff --git a/Engine/Source/Runtime/Core/Public/Stats/StatsFile.h b/Engine/Source/Runtime/Core/Public/Stats/StatsFile.h
--- a/Engine/Source/Runtime/Core/Public/Stats/StatsFile.h
+++ b/Engine/Source/Runtime/Core/Public/Stats/StatsFile.h
@@ -17,6 +17,7 @@ namespace EStatMagicNoHeader
 	{
 		MAGIC = 0x7E1B83C1,
 		MAGIC_SWAPPED = 0xC1831B7E,
+		NO_VERSION = 0,
 	};
 }
 
@@ -30,7 +31,8 @@ namespace EStatMagicWithHeader
 	{
 		MAGIC = 0x10293847,
 		MAGIC_SWAPPED = 0x47382910,
-		VERSION = 2,
+		VERSION_2 = 2,
+		VERSION_3 = 3,
 	};
 }
 
@@ -248,7 +250,7 @@ public:
 	virtual void WriteHeader();
 
 	/** Grabs a frame from the local FStatsThreadState and adds it to the output **/
-	virtual void WriteFrame( int64 TargetFrame );
+	virtual void WriteFrame( int64 TargetFrame, bool bNeedFullMetadata = false );
 
 protected:
 	/** Finalizes writing to the file, called directly form the same thread. */
@@ -260,6 +262,17 @@ protected:
 	void AddNewFrameDelegate();
 	void RemoveNewFrameDelegate();
 
+	/** Writes metadata messages into the stream. */
+	void WriteMetadata( FArchive& Ar )
+	{
+		const FStatsThreadState& Stats = FStatsThreadState::GetLocalState();
+		for( const auto& It : Stats.ShortNameToLongName )
+		{
+			const FStatMessage& StatMessage = It.Value;
+			WriteMessage( Ar, StatMessage );
+		}
+	}
+
 	/** Sends an FName, and the string it represents if we have not sent that string before. **/
 	FORCEINLINE_STATS void WriteFName( FArchive& Ar, FStatNameAndInfo NameAndInfo )
 	{
@@ -319,7 +332,7 @@ protected:
 	virtual void WriteHeader() OVERRIDE;
 
 	/** Grabs a frame of raw message from the local FStatsThreadState and adds it to the output. */
-	virtual void WriteFrame( int64 TargetFrame ) OVERRIDE;
+	virtual void WriteFrame( int64 TargetFrame, bool bNeedFullMetadata = false ) OVERRIDE;
 
 	/** Write a stat packed into the specified archive. */
 	void WriteStatPacket( FArchive& Ar, /*const*/ FStatPacket& StatPacket )
