#UE4 - Remove built-in ad banner support from engine in favor of using modules

diff --git a/Engine/Source/Runtime/Core/Core.Build.cs b/Engine/Source/Runtime/Core/Core.Build.cs
--- a/Engine/Source/Runtime/Core/Core.Build.cs
+++ b/Engine/Source/Runtime/Core/Core.Build.cs
@@ -90,8 +90,6 @@ public class Core : ModuleRules
 
 			bool bSupportAdvertising = true;
 
-			Definitions.Add("UE_WITH_IAD=" + (bSupportAdvertising ? "1" : "0"));
-			
 			if (bSupportAdvertising)
 			{
 				PublicFrameworks.AddRange(new string[] { "iAD", "CoreGraphics" });
diff --git a/Engine/Source/Runtime/Core/Private/IOS/IOSAppDelegate.cpp b/Engine/Source/Runtime/Core/Private/IOS/IOSAppDelegate.cpp
--- a/Engine/Source/Runtime/Core/Private/IOS/IOSAppDelegate.cpp
+++ b/Engine/Source/Runtime/Core/Private/IOS/IOSAppDelegate.cpp
@@ -463,193 +463,6 @@ void InstallSignalHandlers()
 	FPlatformMisc::HandleLowMemoryWarning();
 }
 
-#if UE_WITH_IAD
-
-@synthesize BannerView;
-
-/** TRUE if the iAd banner should be on the bottom of the screen */
-BOOL bDrawOnBottom;
-
-/** true if the user wants the banner to be displayed */
-bool bWantVisibleBanner = false;
-
-/**
-* Will show an iAd on the top or bottom of screen, on top of the GL view (doesn't resize
-* the view)
-*
-* @param bShowOnBottomOfScreen If true, the iAd will be shown at the bottom of the screen, top otherwise
-*/
--(void)ShowAdBanner:(NSNumber*)bShowOnBottomOfScreen
-{
-	bDrawOnBottom = [bShowOnBottomOfScreen boolValue];
-	bWantVisibleBanner = true;
-
-	bool bNeedsToAddSubview = false;
-	if (self.BannerView == nil)
-	{
-		self.BannerView = [[ADBannerView alloc] initWithAdType:ADAdTypeBanner];
-		self.BannerView.delegate = self;
-		bNeedsToAddSubview = true;
-	}
-
-	CGRect BannerFrame = CGRectZero;
-	BannerFrame.size = [self.BannerView sizeThatFits : self.RootView.bounds.size];
-
-	if (bDrawOnBottom)
-	{
-		// move to off the bottom
-		BannerFrame.origin.y = self.RootView.bounds.size.height - BannerFrame.size.height;
-	}
-
-	self.BannerView.frame = BannerFrame;
-
-	// start out hidden, will fade in when ad loads
-	self.BannerView.hidden = YES;
-	self.BannerView.alpha = 0.0f;
-
-	if (bNeedsToAddSubview)
-	{
-		[self.RootView addSubview : self.BannerView];
-	}
-	else// if (self.BannerView.bannerLoaded)
-	{
-		[self bannerViewDidLoadAd : self.BannerView];
-	}
-}
-
--(void)bannerViewDidLoadAd:(ADBannerView*)Banner
-{
-#if !NO_LOGGING
-	NSLog(@"Ad loaded!");
-#endif
-    
-	if (self.BannerView.hidden && bWantVisibleBanner)
-    {
-		self.BannerView.hidden = NO;
-		[UIView animateWithDuration:0.4f
-			animations:^
-			{
-				self.BannerView.alpha = 1.0f;
-			}
-		];
-    }
-}
-
--(void)bannerView:(ADBannerView *)Banner didFailToReceiveAdWithError : (NSError *)Error
-{
-#if !NO_LOGGING
-	NSLog(@"Ad failed to load: '%@'", Error);
-#endif
-
-	// if we get an error, hide the banner 
-	[self HideAdBanner];
-}
-
-/**
-* Callback when the user clicks on an ad
-*/
--(BOOL)bannerViewActionShouldBegin:(ADBannerView*)Banner willLeaveApplication : (BOOL)bWillLeave
-{
-	// if we aren't about to swap out the app, tell the game to pause (or whatever)
-	if (!bWillLeave)
-	{
-		FIOSAsyncTask* AsyncTask = [[FIOSAsyncTask alloc] init];
-		AsyncTask.GameThreadCallback = ^ bool(void)
-		{
-			// tell the ad manager the user clicked on the banner
-//@TODO: IAD:			UPlatformInterfaceBase::GetInGameAdManagerSingleton()->OnUserClickedBanner();
-
-			return true;
-		};
-		[AsyncTask FinishedTask];
-	}
-
-	return YES;
-}
-
-/**
-* Callback when an ad is closed
-*/
--(void)bannerViewActionDidFinish:(ADBannerView*)Banner
-{
-	FIOSAsyncTask* AsyncTask = [[FIOSAsyncTask alloc] init];
-	AsyncTask.GameThreadCallback = ^ bool(void)
-	{
-		// tell ad singleton we closed the ad
-//@TODO: IAD:		UPlatformInterfaceBase::GetInGameAdManagerSingleton()->OnUserClosedAd();
-
-		return true;
-	};
-	[AsyncTask FinishedTask];
-}
-
-/**
-* Hides the iAd banner shows with ShowAdBanner. Will force close the ad if it's open
-*/
--(void)HideAdBanner
-{
-	// fade it out
-	if (!self.BannerView.hidden)
-    {
-		[UIView animateWithDuration:0.4f
-			animations:^
-			{
-				self.BannerView.alpha = 0.0f;
-			}
-			completion:^(BOOL finished) 
-			{
-				self.BannerView.hidden = YES;
-			}
-		];
-    }
-}
-
--(void)UserHideAdBanner
-{
-	bWantVisibleBanner = false;
-	[self HideAdBanner];
-}
-
-/**
-* Forces closed any displayed ad. Can lead to loss of revenue
-*/
--(void)CloseAd
-{
-	// boot user out of the ad
-	bWantVisibleBanner = false;
-	[self.BannerView cancelBannerViewAction];
-}
-
-
-/**
-* Will show an iAd on the top or bottom of screen, on top of the GL view (doesn't resize
-* the view)
-*
-* @param bShowOnBottomOfScreen If true, the iAd will be shown at the bottom of the screen, top otherwise
-*/
-CORE_API void IOSShowAdBanner(bool bShowOnBottomOfScreen)
-{
-	[[IOSAppDelegate GetDelegate] performSelectorOnMainThread:@selector(ShowAdBanner:) withObject:[NSNumber numberWithBool : bShowOnBottomOfScreen] waitUntilDone : NO];
-}
-
-/**
-* Hides the iAd banner shows with IPhoneShowAdBanner. Will force close the ad if it's open
-*/
-CORE_API void IOSHideAdBanner()
-{
-	[[IOSAppDelegate GetDelegate] performSelectorOnMainThread:@selector(UserHideAdBanner) withObject:nil waitUntilDone : NO];
-}
-
-/**
-* Forces closed any displayed ad. Can lead to loss of revenue
-*/
-CORE_API void IOSCloseAd()
-{
-	[[IOSAppDelegate GetDelegate] performSelectorOnMainThread:@selector(CloseAd) withObject:nil waitUntilDone : NO];
-}
-
-#endif
-
 /**
  * Shows the given Game Center supplied controller on the screen
  *
diff --git a/Engine/Source/Runtime/Core/Public/IOS/IOSAppDelegate.h b/Engine/Source/Runtime/Core/Public/IOS/IOSAppDelegate.h
--- a/Engine/Source/Runtime/Core/Public/IOS/IOSAppDelegate.h
+++ b/Engine/Source/Runtime/Core/Public/IOS/IOSAppDelegate.h
@@ -4,15 +4,6 @@
 
 #import <UIKit/UIKit.h>
 #import <AVFoundation/AVAudioSession.h>
-
-#if !defined(UE_WITH_IAD)
-#define UE_WITH_IAD 1
-#endif
-
-#if UE_WITH_IAD
-#import <iAd/ADBannerView.h>
-#endif
-
 #import <GameKit/GKLeaderboardViewController.h>
 
 @class EAGLView;
@@ -35,9 +26,6 @@ namespace FAppEntry
 @interface IOSAppDelegate : UIResponder <UIApplicationDelegate,
 #if !UE_BUILD_SHIPPING
 	UIGestureRecognizerDelegate,
-#endif
-#if UE_WITH_IAD
-	ADBannerViewDelegate,
 #endif
 	GKLeaderboardViewControllerDelegate,
 UITextFieldDelegate, AVAudioSessionDelegate>
@@ -72,29 +60,6 @@ UITextFieldDelegate, AVAudioSessionDelegate>
 	@property (nonatomic, assign) int				ConsoleHistoryValuesIndex;
 #endif
 
-#if UE_WITH_IAD
-	/** iAd banner view, if open */
-	@property(retain) ADBannerView* BannerView;
-
-	/**
-	* Will show an iAd on the top or bottom of screen, on top of the GL view (doesn't resize
-	* the view)
-	*
-	* @param bShowOnBottomOfScreen If true, the iAd will be shown at the bottom of the screen, top otherwise
-	*/
-	-(void)ShowAdBanner:(NSNumber*)bShowOnBottomOfScreen;
-
-	/**
-	* Hides the iAd banner shows with ShowAdBanner. Will force close the ad if it's open
-	*/
-	-(void)HideAdBanner;
-
-	/**
-	* Forces closed any displayed ad. Can lead to loss of revenue
-	*/
-	-(void)CloseAd;
-#endif
-
 /** True if the engine has been initialized */
 @property (readonly) bool bEngineInit;
 
diff --git a/Engine/Source/Runtime/Engine/Private/KismetSystemLibrary.cpp b/Engine/Source/Runtime/Engine/Private/KismetSystemLibrary.cpp
--- a/Engine/Source/Runtime/Engine/Private/KismetSystemLibrary.cpp
+++ b/Engine/Source/Runtime/Engine/Private/KismetSystemLibrary.cpp
@@ -6,6 +6,7 @@
 #include "LatentActions.h"
 #include "DelayAction.h"
 #include "InterpolateComponentToAction.h"
+#include "Advertising.h"
 
 //////////////////////////////////////////////////////////////////////////
 // UKismetSystemLibrary
@@ -2217,31 +2218,37 @@ void UKismetSystemLibrary::CollectGarbage()
 }
 
 #if PLATFORM_IOS
-extern CORE_API void IOSShowAdBanner(bool bShowOnBottomOfScreen);
-extern CORE_API void IOSHideAdBanner();
-extern CORE_API void IOSCloseAd();
 extern CORE_API void IOSShowLeaderboardUI(const FString& CategoryName);
 #endif
 
 void UKismetSystemLibrary::EXPERIMENTAL_ShowAdBanner(bool bShowOnBottomOfScreen)
 {
-#if PLATFORM_IOS
-	IOSShowAdBanner(bShowOnBottomOfScreen);
-#endif
+	IAdvertisingProvider * Provider = FAdvertising::Get().GetDefaultProvider();
+
+	if ( Provider != NULL )
+	{
+		Provider->ShowAdBanner( bShowOnBottomOfScreen );
+	}
 }
 
 void UKismetSystemLibrary::EXPERIMENTAL_HideAdBanner()
 {
-#if PLATFORM_IOS
-	IOSHideAdBanner();
-#endif
+	IAdvertisingProvider * Provider = FAdvertising::Get().GetDefaultProvider();
+
+	if ( Provider != NULL )
+	{
+		Provider->HideAdBanner();
+	}
 }
 
 void UKismetSystemLibrary::EXPERIMENTAL_CloseAdBanner()
 {
-#if PLATFORM_IOS
-	IOSCloseAd();
-#endif
+	IAdvertisingProvider * Provider = FAdvertising::Get().GetDefaultProvider();
+
+	if ( Provider != NULL )
+	{
+		Provider->CloseAdBanner();
+	}
 }
 
 void UKismetSystemLibrary::EXPERIMENTAL_ShowGameCenterLeaderboard(const FString& CategoryName)
