#ttp 332719 - FText::AsMemory is returning strings where all decimal places have the digit 0
#summary Fixed formating memory sizes for decimal places
#proj Core

diff --git a/Engine/Source/Runtime/Core/Private/Internationalization/ICUText.cpp b/Engine/Source/Runtime/Core/Private/Internationalization/ICUText.cpp
--- a/Engine/Source/Runtime/Core/Private/Internationalization/ICUText.cpp
+++ b/Engine/Source/Runtime/Core/Private/Internationalization/ICUText.cpp
@@ -117,7 +117,8 @@ FText FText::AsMemory(SIZE_T NumBytes, const FNumberFormattingOptions* const Opt
 		++Prefix;
 	}
 
-	Args.Add( TEXT("Number"), FText::AsNumber( (uint64)(NumBytes / 1024), Options, TargetCulture) );
+	const double MemorySizeAsDouble = (double)NumBytes / 1024.0;
+	Args.Add( TEXT("Number"), FText::AsNumber( MemorySizeAsDouble, Options, TargetCulture) );
 	Args.Add( TEXT("Unit"), FText::FromString( FString( 1, &Prefixes[Prefix] ) + TEXT("B") ) );
 	return FText::Format( NSLOCTEXT("Internationalization", "ComputerMemoryFormatting", "{Number} {Unit}"), Args);
 }
diff --git a/Engine/Source/Runtime/Core/Private/Internationalization/LegacyText.cpp b/Engine/Source/Runtime/Core/Private/Internationalization/LegacyText.cpp
--- a/Engine/Source/Runtime/Core/Private/Internationalization/LegacyText.cpp
+++ b/Engine/Source/Runtime/Core/Private/Internationalization/LegacyText.cpp
@@ -52,7 +52,8 @@ FText FText::AsMemory(SIZE_T NumBytes, const FNumberFormattingOptions* const Opt
 		++Prefix;
 	}
 
-	Args.Add( TEXT("Number"), FText::AsNumber( (uint64)(NumBytes / 1024) , Options, TargetCulture) );
+	const double MemorySizeAsDouble = (double)NumBytes / 1024.0;
+	Args.Add( TEXT("Number"), FText::AsNumber( MemorySizeAsDouble, Options, TargetCulture) );
 	Args.Add( TEXT("Unit"), FText::FromString( FString( 1, &Prefixes[Prefix] ) + TEXT("B") ) );
 	return FText::Format( NSLOCTEXT("Internationalization", "ComputerMemoryFormatting", "{Number} {Unit}"), Args);
 }
