Workaround for Mac crash when using PS3 controller

diff --git a/Engine/Source/Runtime/Core/Private/Mac/HIDInputInterface.cpp b/Engine/Source/Runtime/Core/Private/Mac/HIDInputInterface.cpp
--- a/Engine/Source/Runtime/Core/Private/Mac/HIDInputInterface.cpp
+++ b/Engine/Source/Runtime/Core/Private/Mac/HIDInputInterface.cpp
@@ -452,7 +452,8 @@ void HIDInputInterface::OnNewHIDController( IOReturn Result, IOHIDDeviceRef Devi
 			Button.Usage = Usage;
 			DeviceInfo.Buttons.Push( Button );
 		}
-		else if( ( ( ElementType == kIOHIDElementTypeInput_Axis ) || ( ElementType == kIOHIDElementTypeInput_Misc ) ) && UsagePage == kHIDPage_GenericDesktop && Usage != kHIDUsage_GD_Hatswitch )
+		else if( ( ( ElementType == kIOHIDElementTypeInput_Axis ) || ( ElementType == kIOHIDElementTypeInput_Misc ) ) && UsagePage == kHIDPage_GenericDesktop
+				&& Usage != kHIDUsage_GD_Hatswitch && Usage != kHIDUsage_GD_Pointer )
 		{
 			IOHIDValueRef Value;
 			IOHIDDeviceGetValue( DeviceRef, Element, &Value );
