Added stub implementation for LaunchURL on Linux to alleviate TTP #331907.

diff --git a/Engine/Source/Runtime/Core/Private/Linux/LinuxPlatformProcess.cpp b/Engine/Source/Runtime/Core/Private/Linux/LinuxPlatformProcess.cpp
--- a/Engine/Source/Runtime/Core/Private/Linux/LinuxPlatformProcess.cpp
+++ b/Engine/Source/Runtime/Core/Private/Linux/LinuxPlatformProcess.cpp
@@ -206,7 +206,13 @@ FRunnableThread* FLinuxPlatformProcess::CreateRunnableThread()
 	return new FRunnableThreadLinux();
 }
 
-FProcHandle FLinuxPlatformProcess::CreateProc( const TCHAR* URL, const TCHAR* Parms, bool bLaunchDetached, bool bLaunchHidden, bool bLaunchReallyHidden, uint32* OutProcessID, int32 PriorityModifier, const TCHAR* OptionalWorkingDirectory, void* PipeWrite )
+void FLinuxPlatformProcess::LaunchURL(const TCHAR* URL, const TCHAR* Parms, FString* Error)
+{
+	// stub implementation for now
+	// TODO: consider looking for gnome-open, xdg-open, sensible-browser or just hardcoded names...
+}
+
+FProcHandle FLinuxPlatformProcess::CreateProc(const TCHAR* URL, const TCHAR* Parms, bool bLaunchDetached, bool bLaunchHidden, bool bLaunchReallyHidden, uint32* OutProcessID, int32 PriorityModifier, const TCHAR* OptionalWorkingDirectory, void* PipeWrite)
 {
 	FString AbsolutePath = FPaths::ConvertRelativePathToFull(URL);
 	FString Commandline = AbsolutePath;
diff --git a/Engine/Source/Runtime/Core/Public/Linux/LinuxPlatformProcess.h b/Engine/Source/Runtime/Core/Public/Linux/LinuxPlatformProcess.h
--- a/Engine/Source/Runtime/Core/Public/Linux/LinuxPlatformProcess.h
+++ b/Engine/Source/Runtime/Core/Public/Linux/LinuxPlatformProcess.h
@@ -137,7 +137,8 @@ struct CORE_API FLinuxPlatformProcess : public FGenericPlatformProcess
 	static bool CreatePipe( void*& ReadPipe, void*& WritePipe );
 	static FString ReadPipe( void* ReadPipe );
 	static class FRunnableThread* CreateRunnableThread();
-	static FProcHandle CreateProc( const TCHAR* URL, const TCHAR* Parms, bool bLaunchDetached, bool bLaunchHidden, bool bLaunchReallyHidden, uint32* OutProcessID, int32 PriorityModifier, const TCHAR* OptionalWorkingDirectory, void* PipeWrite );
+	static void LaunchURL(const TCHAR* URL, const TCHAR* Parms, FString* Error);
+	static FProcHandle CreateProc(const TCHAR* URL, const TCHAR* Parms, bool bLaunchDetached, bool bLaunchHidden, bool bLaunchReallyHidden, uint32* OutProcessID, int32 PriorityModifier, const TCHAR* OptionalWorkingDirectory, void* PipeWrite);
 	static bool IsProcRunning( FProcHandle & ProcessHandle );
 	static void WaitForProc( FProcHandle & ProcessHandle );
 	static void TerminateProc( FProcHandle & ProcessHandle, bool KillTree = false );
