#summary Potential solution for crashing applications without appropriate current working directories.
#ttp 332492 - CRASH: BLOCKER: MAIN: Client crashes on launch

diff --git a/Engine/Source/Runtime/Core/Private/Internationalization/Internationalization.cpp b/Engine/Source/Runtime/Core/Private/Internationalization/Internationalization.cpp
--- a/Engine/Source/Runtime/Core/Private/Internationalization/Internationalization.cpp
+++ b/Engine/Source/Runtime/Core/Private/Internationalization/Internationalization.cpp
@@ -179,7 +179,13 @@ void FInternationalization::Initialize()
 #if UE_ENABLE_ICU
 	UErrorCode ICUStatus = U_ZERO_ERROR;
 	u_setMemoryFunctions(NULL, &(FICUOverrides::Malloc), &(FICUOverrides::Realloc), &(FICUOverrides::Free), &(ICUStatus));
-	FString DataDirectory = FPaths::EngineContentDir() / TEXT("Localization");
+	FString DataDirectory;
+	DataDirectory = FPaths::EngineContentDir() / TEXT("Localization");
+	if( !FPaths::FileExists(DataDirectory / TEXT("icudt51l.dat")) )
+	{
+		DataDirectory = FString(FPlatformProcess::BaseDir()) / FPaths::EngineContentDir() / TEXT("Localization");
+	}
+	check( FPaths::FileExists(DataDirectory / TEXT("icudt51l.dat")) );
 	u_setDataDirectory(StringCast<char>(*DataDirectory).Get());
 	u_init(&(ICUStatus));
 #endif
@@ -432,4 +438,4 @@ void FInternationalization::PopulateAllCultures(void)
 	InvariantCulture = GetCulture(TEXT(""));
 }
 
-#undef LOCTEXT_NAMESPACE
+#undef LOCTEXT_NAMESPACE
\ No newline at end of file
