#summary "Installed" flag is now forwarded to CrashReportClient
#branch main
#proj UE4
#note This ensure CrashReportClient's logs are saved to the user's directory when run from an installed UE4.

diff --git a/Engine/Source/Runtime/Core/Private/Windows/WindowsPlatformExceptionHandling.cpp b/Engine/Source/Runtime/Core/Private/Windows/WindowsPlatformExceptionHandling.cpp
--- a/Engine/Source/Runtime/Core/Private/Windows/WindowsPlatformExceptionHandling.cpp
+++ b/Engine/Source/Runtime/Core/Private/Windows/WindowsPlatformExceptionHandling.cpp
@@ -297,6 +297,13 @@ int32 NewReportCrash( EXCEPTION_POINTERS* ExceptionInfo, const TCHAR* ErrorMessa
 			CrashReportUploaderArguments += TEXT( " -Unattended" );
 		}
 
+		if ( FApp::IsInstalled() )
+		{
+			// Temporary workaround for CrashReportClient being built in Development, not Shipping (TTP328030). The
+			// following ensures that logs are saved to the user directory when UE4 is installed.
+			CrashReportUploaderArguments += TEXT( " -Installed" );
+		}
+
 		const TCHAR* CrashClientPath = GUseCrashReportClient ?
 			TEXT("CrashReportClient.exe") :
 			TEXT("..\\..\\..\\Engine\\Binaries\\DotNET\\CrashReportUploader.exe");
