#ttp 332144 - COLLISION: Bug with DefaultChannelResponses ini saving

#Summary Updated the config system to handle single element arrays as arrays, and not as single properties.

#Note This is only applicable to default*.ini files, as these require array syntax when saving

diff --git a/Engine/Source/Runtime/Core/Private/Misc/ConfigCacheIni.cpp b/Engine/Source/Runtime/Core/Private/Misc/ConfigCacheIni.cpp
--- a/Engine/Source/Runtime/Core/Private/Misc/ConfigCacheIni.cpp
+++ b/Engine/Source/Runtime/Core/Private/Misc/ConfigCacheIni.cpp
@@ -761,9 +761,9 @@ bool FConfigFile::Write( const FString& Filename, bool bDoRemoteWrite/* = true*/
 						&& !Filename.Contains( FPaths::EngineSavedDir() ) 
 						&& FPaths::GetBaseFilename(Filename).StartsWith( TEXT( "Default" ) );
 
-					if( bIsADefaultIniWrite && CompletePropertyToWrite.Num() > 1 )
+					if( bIsADefaultIniWrite )
 					{
-						ProcessArrayForSaveToDefaults(CompletePropertyToWrite, Text, SectionName, PropertyName.ToString());
+						ProcessPropertyAndWriteForDefaults(CompletePropertyToWrite, Text, SectionName, PropertyName.ToString());
 					}
 					else
 					{
@@ -1052,12 +1052,12 @@ void FConfigFile::ProcessSourceAndCheckAgainstBackup()
 }
 
 
-void FConfigFile::ProcessArrayForSaveToDefaults( const TArray< FString >& ArrayToProcess, FString& OutText, const FString& SectionName, const FString& PropertyName )
+void FConfigFile::ProcessPropertyAndWriteForDefaults( const TArray< FString >& InCompletePropertyToProcess, FString& OutText, const FString& SectionName, const FString& PropertyName )
 {
 	// Get a list of the array elements we have processed, 
 	// we will remove any which are in the base hierarchy
 	// any which are left will be written to the default config
-	TArray<FString> UnprocessedPropertyValues = ArrayToProcess;
+	TArray<FString> UnprocessedPropertyValues = InCompletePropertyToProcess;
 
 	// Iterate over the base hierarchy of the config file and check that our element did not originate from a further up the tree file
 	// if it did we need to ensure it is not added again.
@@ -1071,7 +1071,7 @@ void FConfigFile::ProcessArrayForSaveToDefaults( const TArray< FString >& ArrayT
 		{
 			// As we are dealing with default configs, we must process array elements correct syntax syntax. I.e. + or -
 			TArray< FString > HierearchysArrayContribution;
-			FoundSection->MultiFind( *(FString(TEXT("+"))+PropertyName), HierearchysArrayContribution, true );
+			FoundSection->MultiFind( *PropertyName, HierearchysArrayContribution, true );
 
 			// Find array elements which should be removed.
 			for( TArray<FString>::TIterator PropertyIt(HierearchysArrayContribution); PropertyIt; ++PropertyIt )
@@ -1096,7 +1096,8 @@ void FConfigFile::ProcessArrayForSaveToDefaults( const TArray< FString >& ArrayT
 				//			+ArrEl=1.f
 				if( UnprocessedPropertyValues.Contains( PropertyValue ) == false )
 				{
-					OutText += FString::Printf( TEXT("-%s=%s") LINE_TERMINATOR, *PropertyName, *PropertyValue);
+					FString PropertyNameWithRemoveOp = PropertyName.Replace(TEXT("+"), TEXT("-"));
+					OutText += FString::Printf(TEXT("%s=%s") LINE_TERMINATOR, *PropertyNameWithRemoveOp, *PropertyValue);
 				}
 
 				// We need to remove this element from unprocessed if it exists on the list.
@@ -1106,16 +1107,16 @@ void FConfigFile::ProcessArrayForSaveToDefaults( const TArray< FString >& ArrayT
 	}
 
 	// Add all elements which were not in the hierarchy to the default with the correct '+' syntax
-	for( TArray<FString>::TIterator PropertyIt(UnprocessedPropertyValues); PropertyIt; ++PropertyIt )
+	for (TArray<FString>::TIterator PropertyIt(UnprocessedPropertyValues); PropertyIt; ++PropertyIt)
 	{
 		FString PropertyValue = *PropertyIt;
 
 		if (ShouldExportQuotedString(PropertyValue))
 		{
-			PropertyValue = FString::Printf( TEXT( "\"%s\"" ), *PropertyValue );
+			PropertyValue = FString::Printf(TEXT("\"%s\""), *PropertyValue);
 		}
 
-		OutText += FString::Printf( TEXT("+%s=%s") LINE_TERMINATOR, *PropertyName, *PropertyValue);
+		OutText += FString::Printf(TEXT("%s=%s") LINE_TERMINATOR, *PropertyName, *PropertyValue);
 	}
 }
 
diff --git a/Engine/Source/Runtime/Core/Public/Misc/ConfigCacheIni.h b/Engine/Source/Runtime/Core/Public/Misc/ConfigCacheIni.h
--- a/Engine/Source/Runtime/Core/Public/Misc/ConfigCacheIni.h
+++ b/Engine/Source/Runtime/Core/Public/Misc/ConfigCacheIni.h
@@ -124,14 +124,15 @@ private:
 	void SaveSourceToBackupFile();
 
 	/** 
-	 * Process an array which is being written to a default config and write it to the out stream
+	 * Process the property for Writing to a default file. We will need to check for array operations, as default ini's rely on this
+	 * being correct to function properly
 	 *
-	 * @param ArrayPropertyToProcess - The complete property array which we need to process for saving.
+	 * @param InCompletePropertyToProcess - The complete property which we need to process for saving.
 	 * @param OutText - The stream we are processing the array to
 	 * @param SectionName - The section name the array property is being written to
 	 * @param PropertyName - The property name of the array
 	 */
-	void ProcessArrayForSaveToDefaults( const TArray< FString >& ArrayPropertyToProcess, FString& OutText, const FString& SectionName, const FString& PropertyName );
+	void ProcessPropertyAndWriteForDefaults(const TArray< FString >& InCompletePropertyToProcess, FString& OutText, const FString& SectionName, const FString& PropertyName);
 
 };
 
diff --git a/Engine/Source/Runtime/CoreUObject/Private/UObject/Obj.cpp b/Engine/Source/Runtime/CoreUObject/Private/UObject/Obj.cpp
--- a/Engine/Source/Runtime/CoreUObject/Private/UObject/Obj.cpp
+++ b/Engine/Source/Runtime/CoreUObject/Private/UObject/Obj.cpp
@@ -1546,6 +1546,11 @@ void UObject::SaveConfig( uint64 Flags, const TCHAR* InFilename, FConfigCacheIni
 		? InFilename
 		: GetConfigFilename(this);
 
+	// Determine whether the file we are writing is a default file config.
+	const bool bIsADefaultIniWrite = !Filename.Contains(FPaths::GameSavedDir())
+		&& !Filename.Contains(FPaths::EngineSavedDir())
+		&& FPaths::GetBaseFilename(Filename).StartsWith(TEXT("Default"));
+
 	const bool bPerObject = UsesPerObjectConfig(this);
 	FString Section;
 	if ( bPerObject == true )
@@ -1619,13 +1624,16 @@ void UObject::SaveConfig( uint64 Flags, const TCHAR* InFilename, FConfigCacheIni
 					FConfigSection* Sec = Config->GetSectionPrivate( *Section, 1, 0, *PropFileName );
 					check(Sec);
 					Sec->Remove( *Key );
-				
+
+					// Default ini's require the array syntax to be applied to the property name
+					FString CompleteKey = FString::Printf(TEXT("%s%s"), bIsADefaultIniWrite ? TEXT("+") : TEXT(""), *Key);
+
 					FScriptArrayHelper_InContainer ArrayHelper(Array, this);
 					for( int32 i=0; i<ArrayHelper.Num(); i++ )
 					{
 						FString	Buffer;
 						Array->Inner->ExportTextItem( Buffer, ArrayHelper.GetRawPtr(i), ArrayHelper.GetRawPtr(i), this, PortFlags );
-						Sec->Add( *Key, *Buffer );
+						Sec->Add(*CompleteKey, *Buffer);
 					}
 				}
 				else if( Property->Identical_InContainer(this, SuperClassDefaultObject) )
