#ttp 331741 - UE4: CORE: TaskGraph:  Move GetTaskName functionality into GetStatId
#proj UE4
#branch UE4
#summary The interface of FGenericTask is now changed. Removed GetTaskName() and changed GetStatId() to non-static and const.

diff --git a/Engine/Source/Editor/MainFrame/Private/Frame/MainFrameActions.cpp b/Engine/Source/Editor/MainFrame/Private/Frame/MainFrameActions.cpp
--- a/Engine/Source/Editor/MainFrame/Private/Frame/MainFrameActions.cpp
+++ b/Engine/Source/Editor/MainFrame/Private/Frame/MainFrameActions.cpp
@@ -1040,8 +1040,7 @@ public:
 
 	static ESubsequentsMode::Type GetSubsequentsMode() { return ESubsequentsMode::TrackSubsequents; }
 	ENamedThreads::Type GetDesiredThread( ) { return ENamedThreads::GameThread; }
-	static const TCHAR* GetTaskName( ) { return TEXT("FMainFrameActionsNotificationTask"); }
-	FORCEINLINE static TStatId GetStatId()
+	FORCEINLINE TStatId GetStatId() const
 	{
 		RETURN_QUICK_DECLARE_CYCLE_STAT(FMainFrameActionsNotificationTask, STATGROUP_TaskGraphTasks);
 	}
diff --git a/Engine/Source/Editor/UnrealEd/Private/PlayLevel.cpp b/Engine/Source/Editor/UnrealEd/Private/PlayLevel.cpp
--- a/Engine/Source/Editor/UnrealEd/Private/PlayLevel.cpp
+++ b/Engine/Source/Editor/UnrealEd/Private/PlayLevel.cpp
@@ -1047,8 +1047,7 @@ public:
 
 	static ESubsequentsMode::Type GetSubsequentsMode() { return ESubsequentsMode::TrackSubsequents; }
 	ENamedThreads::Type GetDesiredThread( ) { return ENamedThreads::GameThread; }
-	static const TCHAR* GetTaskName( ) { return TEXT("FLauncherNotificationTask"); }
-	FORCEINLINE static TStatId GetStatId()
+	FORCEINLINE TStatId GetStatId() const
 	{
 		RETURN_QUICK_DECLARE_CYCLE_STAT(FLauncherNotificationTask, STATGROUP_TaskGraphTasks);
 	}
diff --git a/Engine/Source/Runtime/Core/Private/Async/TaskGraph.cpp b/Engine/Source/Runtime/Core/Private/Async/TaskGraph.cpp
--- a/Engine/Source/Runtime/Core/Private/Async/TaskGraph.cpp
+++ b/Engine/Source/Runtime/Core/Private/Async/TaskGraph.cpp
@@ -7,10 +7,7 @@
 DEFINE_LOG_CATEGORY_STATIC(LogTaskGraph, Log, All);
 
 DEFINE_STAT(STAT_FReturnGraphTask);
-DEFINE_STAT(STAT_FNullGraphTask);
 DEFINE_STAT(STAT_FTriggerEventGraphTask);
-DEFINE_STAT(STAT_FSimpleDelegateGraphTask);
-DEFINE_STAT(STAT_FDelegateGraphTask);
 
 namespace ENamedThreads
 {
diff --git a/Engine/Source/Runtime/Core/Private/Stats/Stats2.cpp b/Engine/Source/Runtime/Core/Private/Stats/Stats2.cpp
--- a/Engine/Source/Runtime/Core/Private/Stats/Stats2.cpp
+++ b/Engine/Source/Runtime/Core/Private/Stats/Stats2.cpp
@@ -641,11 +641,7 @@ public:
 		: Packet(InPacket)
 	{
 	}
-	static const TCHAR* GetTaskName()
-	{
-		return TEXT("FStatMessagesTask");
-	}
-	FORCEINLINE static TStatId GetStatId()
+	FORCEINLINE TStatId GetStatId() const
 	{
 		return TStatId(); // we don't want to record this or it spams the stat system; we cover this time when we tick the stats system
 	}
diff --git a/Engine/Source/Runtime/Core/Public/Async/TaskGraphInterfaces.h b/Engine/Source/Runtime/Core/Public/Async/TaskGraphInterfaces.h
--- a/Engine/Source/Runtime/Core/Public/Async/TaskGraphInterfaces.h
+++ b/Engine/Source/Runtime/Core/Public/Async/TaskGraphInterfaces.h
@@ -19,10 +19,7 @@
 DECLARE_STATS_GROUP(TEXT("Task Graph Tasks"), STATGROUP_TaskGraphTasks);
 
 DECLARE_CYCLE_STAT_EXTERN(TEXT("FReturnGraphTask"), STAT_FReturnGraphTask, STATGROUP_TaskGraphTasks, CORE_API);
-DECLARE_CYCLE_STAT_EXTERN(TEXT("FNullGraphTask"), STAT_FNullGraphTask, STATGROUP_TaskGraphTasks, CORE_API);
 DECLARE_CYCLE_STAT_EXTERN(TEXT("FTriggerEventGraphTask"), STAT_FTriggerEventGraphTask, STATGROUP_TaskGraphTasks, CORE_API);
-DECLARE_CYCLE_STAT_EXTERN(TEXT("FSimpleDelegateGraphTask"), STAT_FSimpleDelegateGraphTask, STATGROUP_TaskGraphTasks, CORE_API);
-DECLARE_CYCLE_STAT_EXTERN(TEXT("FSimpleDelegateGraphTask"), STAT_FDelegateGraphTask, STATGROUP_TaskGraphTasks, CORE_API);
 
 namespace ENamedThreads
 {
@@ -272,13 +269,6 @@ private:
 	 **/
 	virtual void ExecuteTask(TArray<FBaseGraphTask*>& NewTasks, ENamedThreads::Type CurrentThread)=0;
 
-	/** 
-	 *	Virtual call to retrieve the task name.
-	 *	@return the name of the task. This returned string should be static memory, usable until the end of time.
-	 **/
-	virtual const TCHAR* GetTaskName()=0;
-
-
 	// API called from other parts of the system
 
 	/** 
@@ -304,16 +294,6 @@ private:
 		ExecuteTask(NewTasks, CurrentThread);
 	}
 
-	/** 
-	 *	Called by the system to retrieve a task name for debugging or profilers. Just calls the virtual version.
-	 *	@return the name of the task. This returned string should be static memory, usable until the end of time.
-	 **/
-	const TCHAR* GetDebugName()
-	{
-		return GetTaskName();
-	}
-
-
 	// Internal Use
 
 	/** 
@@ -500,11 +480,7 @@ public:
 	{
 		// you will be destroyed immediately after you execute. Might as well do cleanup in DoWork, but you could also use a destructor.
 	}
-	static const TCHAR* GetTaskName()
-	{
-		return TEXT("FGenericTask"); // this must return a pointer to memory usable until the end of time.
-	}
-	FORCEINLINE static TStatId GetStatId()
+	FORCEINLINE TStatId GetStatId() const
 	{
 		RETURN_QUICK_DECLARE_CYCLE_STAT(FGenericTask, STATGROUP_TaskGraphTasks);
 	}
@@ -653,14 +629,20 @@ private:
 		
 		TTask& Task = *(TTask*)&TaskStorage;
 		{
-			FScopeCycleCounter Scope(TTask::GetStatId()); 
+			FScopeCycleCounter Scope(Task.GetStatId()); 
 			Task.DoTask(CurrentThread, Subsequents);
 #if !(UE_BUILD_SHIPPING || UE_BUILD_TEST)
 			if (TTask::GetSubsequentsMode() == ESubsequentsMode::TrackSubsequents)
 			{
 				if (!Subsequents->CheckSubsequents())
 				{
-					LogPossiblyInvalidSubsequentsTask(Task.GetTaskName());					
+					LogPossiblyInvalidSubsequentsTask(
+#if STATS
+						*Task.GetStatId()->ToString()
+#else
+						TEXT("Unknown")
+#endif
+					);
 				}
 			}	
 #endif
@@ -686,17 +668,6 @@ private:
 			delete this;
 		}
 	}
-	/** 
-	 *	Virtual call to retrieve the task name.
-	 *	@return the name of the task. This returned string should be static memory, usable until the end of time.
-	 **/
-	virtual const TCHAR* GetTaskName()
-	{
-		checkThreadGraph(TaskConstructed);
-		TTask& Task = *(TTask*)&TaskStorage;
-		return Task.GetTaskName();
-	}
-
 
 	// Internals 
 
@@ -796,15 +767,8 @@ public:
 	{
 		checkThreadGraph(ThreadToReturnFrom != ENamedThreads::AnyThread); // doesn't make any sense to return from any thread
 	}
-	/** 
-	 *	Retrieve the task name.
-	 *	@return the name of the task. This returned string should be static memory, usable until the end of time.
-	 **/
-	static const TCHAR* GetTaskName()
-	{
-		return TEXT("FReturnGraphTask");
-	}
-	FORCEINLINE static TStatId GetStatId()
+
+	FORCEINLINE TStatId GetStatId() const
 	{
 		return GET_STATID(STAT_FReturnGraphTask);
 	}
@@ -837,10 +801,58 @@ private:
 	ENamedThreads::Type ThreadToReturnFrom;
 };
 
+/**
+ * Class that generalizes functionality of creating, storing and exposing stat id for custom name.
+ * Should only be used as a base of a task graph class with custom stat name.
+ */
+class FCustomNameGraphTaskBase
+{
+public:
+	/** 
+	 * Creates the custom stat id.
+	 *
+	 * @param TaskClassName The name of generic class name for stats e.g. FNullTaskGraph
+	 * @param InTaskName The name of this task.
+	 */
+	FCustomNameGraphTaskBase(const TCHAR* TaskClassName, const TCHAR* TaskName)
+	{
+#if STATS
+		const FName StatName = FName(*(FString(TaskClassName) + TEXT(".") + FString(TaskName)));
+
+		FStartupMessages::Get().AddMetadata(StatName, TEXT(""),
+			STAT_GROUP_TO_FStatGroup(STATGROUP_TaskGraphTasks)::GetGroupName(),
+			STAT_GROUP_TO_FStatGroup(STATGROUP_TaskGraphTasks)::GetDescription(),
+			true, EStatDataType::ST_int64, true);
+
+		StatID = IStatGroupEnableManager::Get().GetHighPerformanceEnableForStat(StatName,
+			STAT_GROUP_TO_FStatGroup(STATGROUP_TaskGraphTasks)::GetGroupName(),
+			STAT_GROUP_TO_FStatGroup(STATGROUP_TaskGraphTasks)::DefaultEnable,
+			true, EStatDataType::ST_int64, TEXT(""), true);
+#endif
+	}
+
+	/** 
+	 * Gets stat id for this task.
+	 *
+	 * @return Stat id.
+	 */
+	FORCEINLINE TStatId GetStatId() const
+	{
+#if STATS
+		return StatID;
+#endif
+		return TStatId();
+	}
+
+private:
+	/** Stat id of this object. */
+	STAT(TStatId StatID;)
+};
+
 /** 
  *	FNullGraphTask is a task that does nothing. It can be used to "gather" tasks into one prerequisite.
  **/
-class FNullGraphTask
+class FNullGraphTask : public FCustomNameGraphTaskBase
 {
 public:
 	/** 
@@ -849,22 +861,9 @@ public:
 	 *	@param InDesiredThread; Thread to run on, can be ENamedThreads::AnyThread
 	**/
 	FNullGraphTask(const TCHAR*	InTaskName, ENamedThreads::Type InDesiredThread = ENamedThreads::AnyThread)
-		: DesiredThread(InDesiredThread)
-		, TaskName(InTaskName)
-	{
-	}
-
-	/** 
-	 *	Retrieve the task name.
-	 *	@return the name of the task. This returned string should be static memory, usable until the end of time.
-	 **/
-	const TCHAR* GetTaskName()
-	{
-		return TaskName;
-	}
-	FORCEINLINE static TStatId GetStatId()
+		: FCustomNameGraphTaskBase(TEXT("FNullGraphTask"), InTaskName)
+		, DesiredThread(InDesiredThread)
 	{
-		return GET_STATID(STAT_FNullGraphTask);
 	}
 
 	/** 
@@ -890,8 +889,6 @@ public:
 private:
 	/** Thread to run on, can be ENamedThreads::AnyThread **/
 	ENamedThreads::Type DesiredThread;
-	/** Name of the task **/
-	const TCHAR*		TaskName;
 };
 
 /** 
@@ -910,11 +907,7 @@ public:
 		check(Event);
 	}
 
-	static const TCHAR* GetTaskName()
-	{
-		return TEXT("FTriggerEventGraphTask");
-	}
-	FORCEINLINE static TStatId GetStatId()
+	FORCEINLINE TStatId GetStatId() const
 	{
 		return GET_STATID(STAT_FTriggerEventGraphTask);
 	}
@@ -935,15 +928,13 @@ private:
 };
 
 /** Task class for simple delegate based tasks. This is less efficient than a custom task, doesn't provide the task arguments, doesn't allow specification of the current thread, etc. **/
-class FSimpleDelegateGraphTask
+class FSimpleDelegateGraphTask : public FCustomNameGraphTaskBase
 {
 public:
 	DECLARE_DELEGATE(FDelegate);
 
 	/** Delegate to fire when task runs **/
 	FDelegate							TaskDelegate;
-	/** Name of the task, static memory **/
-	const TCHAR*						TaskName;
 	/** Thread to run delegate on **/
 	const ENamedThreads::Type			DesiredThread;
 public:
@@ -953,14 +944,6 @@ public:
 	}
 	static ESubsequentsMode::Type GetSubsequentsMode() { return ESubsequentsMode::TrackSubsequents; }
 
-	const TCHAR* GetTaskName()
-	{
-		return TaskName;
-	}
-	FORCEINLINE static TStatId GetStatId()
-	{
-		return GET_STATID(STAT_FSimpleDelegateGraphTask);
-	}
 	void DoTask(ENamedThreads::Type CurrentThread, const FGraphEventRef& MyCompletionGraphEvent)
 	{
 		TaskDelegate.ExecuteIfBound();
@@ -972,8 +955,8 @@ public:
 	  * @param InDesiredThread - Thread to run on
 	**/
 	FSimpleDelegateGraphTask(const FDelegate& InTaskDeletegate, const TCHAR* InTaskName, ENamedThreads::Type InDesiredThread)
-		: TaskDelegate(InTaskDeletegate)
-		, TaskName(InTaskName)
+		: FCustomNameGraphTaskBase(TEXT("FSimpleDelegateGraphTask"), InTaskName)
+		, TaskDelegate(InTaskDeletegate)
 		, DesiredThread(InDesiredThread)
 	{
 	}
@@ -1008,15 +991,13 @@ public:
 };
 
 /** Task class for more full featured delegate based tasks. Still less efficient than a custom task, but provides all of the args **/
-class FDelegateGraphTask
+class FDelegateGraphTask : public FCustomNameGraphTaskBase
 {
 public:
 	DECLARE_DELEGATE_TwoParams( FDelegate,ENamedThreads::Type, const  FGraphEventRef& );
 
 	/** Delegate to fire when task runs **/
 	FDelegate							TaskDelegate;
-	/** Name of the task, static memory **/
-	const TCHAR*						TaskName;
 	/** Thread to run delegate on **/
 	const ENamedThreads::Type			DesiredThread;
 public:
@@ -1028,14 +1009,6 @@ public:
 	{ 
 		return ESubsequentsMode::TrackSubsequents; 
 	}
-	const TCHAR* GetTaskName()
-	{
-		return TaskName;
-	}
-	FORCEINLINE static TStatId GetStatId()
-	{
-		return GET_STATID(STAT_FDelegateGraphTask);
-	}
 	void DoTask(ENamedThreads::Type CurrentThread, const FGraphEventRef& MyCompletionGraphEvent)
 	{
 		TaskDelegate.ExecuteIfBound(CurrentThread, MyCompletionGraphEvent);
@@ -1047,8 +1020,8 @@ public:
 	  * @param InDesiredThread - Thread to run on
 	**/
 	FDelegateGraphTask(const FDelegate& InTaskDeletegate, const TCHAR* InTaskName, ENamedThreads::Type InDesiredThread)
-		: TaskDelegate(InTaskDeletegate)
-		, TaskName(InTaskName)
+		: FCustomNameGraphTaskBase(TEXT("FDelegateGraphTask"), InTaskName)
+		, TaskDelegate(InTaskDeletegate)
 		, DesiredThread(InDesiredThread)
 	{
 	}
diff --git a/Engine/Source/Runtime/CoreUObject/Private/UObject/GarbageCollection.cpp b/Engine/Source/Runtime/CoreUObject/Private/UObject/GarbageCollection.cpp
--- a/Engine/Source/Runtime/CoreUObject/Private/UObject/GarbageCollection.cpp
+++ b/Engine/Source/Runtime/CoreUObject/Private/UObject/GarbageCollection.cpp
@@ -372,11 +372,7 @@ private:
 				NumObjects--;
 			}
 		}
-		static const TCHAR* GetTaskName()
-		{
-			return TEXT("FGCTask");
-		}
-		FORCEINLINE static TStatId GetStatId()
+		FORCEINLINE TStatId GetStatId() const
 		{
 			RETURN_QUICK_DECLARE_CYCLE_STAT(FGCTask, STATGROUP_TaskGraphTasks);
 		}
diff --git a/Engine/Source/Runtime/Engine/Private/Collision/WorldCollisionAsync.cpp b/Engine/Source/Runtime/Engine/Private/Collision/WorldCollisionAsync.cpp
--- a/Engine/Source/Runtime/Engine/Private/Collision/WorldCollisionAsync.cpp
+++ b/Engine/Source/Runtime/Engine/Private/Collision/WorldCollisionAsync.cpp
@@ -338,12 +338,7 @@ namespace
 			DataCount   = InDataCount;
 		}
 
-		/** return the name of the task **/
-		static const TCHAR* GetTaskName()
-		{
-			return TEXT("FAsyncTraceTask");
-		}
-		FORCEINLINE static TStatId GetStatId()
+		FORCEINLINE TStatId GetStatId() const
 		{
 			RETURN_QUICK_DECLARE_CYCLE_STAT(FAsyncTraceTask, STATGROUP_TaskGraphTasks);
 		}
diff --git a/Engine/Source/Runtime/Engine/Private/LevelTick.cpp b/Engine/Source/Runtime/Engine/Private/LevelTick.cpp
--- a/Engine/Source/Runtime/Engine/Private/LevelTick.cpp
+++ b/Engine/Source/Runtime/Engine/Private/LevelTick.cpp
@@ -784,12 +784,7 @@ void UWorld::SendAllEndOfFrameUpdates(FGraphEventArray* OutCompletion)
 				: Components(InComponents)
 			{
 			}
-
-			static const TCHAR* GetTaskName()
-			{
-				return TEXT("DoRenderthreadUpdatesTask");
-			}
-			FORCEINLINE static TStatId GetStatId()
+			FORCEINLINE TStatId GetStatId() const
 			{
 				RETURN_QUICK_DECLARE_CYCLE_STAT(DoRenderthreadUpdatesTask, STATGROUP_TaskGraphTasks);
 			}
diff --git a/Engine/Source/Runtime/Engine/Private/PhysicsEngine/PhysXSupport.cpp b/Engine/Source/Runtime/Engine/Private/PhysicsEngine/PhysXSupport.cpp
--- a/Engine/Source/Runtime/Engine/Private/PhysicsEngine/PhysXSupport.cpp
+++ b/Engine/Source/Runtime/Engine/Private/PhysicsEngine/PhysXSupport.cpp
@@ -540,12 +540,7 @@ public:
 		Task.release();
 	}
 
-	static const TCHAR* GetTaskName()
-	{
-		return TEXT("FPhysXTask"); // @TODO Use Task.getName()
-	}
-
-	FORCEINLINE static TStatId GetStatId()
+	FORCEINLINE TStatId GetStatId() const
 	{
 		RETURN_QUICK_DECLARE_CYCLE_STAT(FPhysXTask, STATGROUP_TaskGraphTasks);
 	}
diff --git a/Engine/Source/Runtime/Engine/Private/TickTaskManager.cpp b/Engine/Source/Runtime/Engine/Private/TickTaskManager.cpp
--- a/Engine/Source/Runtime/Engine/Private/TickTaskManager.cpp
+++ b/Engine/Source/Runtime/Engine/Private/TickTaskManager.cpp
@@ -291,12 +291,7 @@ private:
 			, bLogTick(InbLogTick)
 		{
 		}
-		/** return the name of the task **/
-		static const TCHAR* GetTaskName()
-		{
-			return TEXT("FTickFunctionTask");
-		}
-		FORCEINLINE static TStatId GetStatId()
+		FORCEINLINE TStatId GetStatId() const
 		{
 			RETURN_QUICK_DECLARE_CYCLE_STAT(FTickFunctionTask, STATGROUP_TaskGraphTasks);
 		}
@@ -831,11 +826,7 @@ private:
 			, Context(InContext)
 		{
 		}
-		static const TCHAR* GetTaskName()
-		{
-			return TEXT("QueueTickTasks");
-		}
-		FORCEINLINE static TStatId GetStatId()
+		FORCEINLINE TStatId GetStatId() const
 		{
 			RETURN_QUICK_DECLARE_CYCLE_STAT(QueueTickTasks, STATGROUP_TaskGraphTasks);
 		}
@@ -875,11 +866,7 @@ private:
 			, NumTickFunctions(InNumTickFunctions)
 		{
 		}
-		static const TCHAR* GetTaskName()
-		{
-			return TEXT("PostTickTasks");
-		}
-		FORCEINLINE static TStatId GetStatId()
+		FORCEINLINE TStatId GetStatId() const
 		{
 			RETURN_QUICK_DECLARE_CYCLE_STAT(PostTickTasks, STATGROUP_TaskGraphTasks);
 		}
diff --git a/Engine/Source/Runtime/Engine/Public/ParticleHelper.h b/Engine/Source/Runtime/Engine/Public/ParticleHelper.h
--- a/Engine/Source/Runtime/Engine/Public/ParticleHelper.h
+++ b/Engine/Source/Runtime/Engine/Public/ParticleHelper.h
@@ -1074,14 +1074,7 @@ struct FAsyncParticleFill
 	/** Work function, just forwards the request to the parent  */
 	void DoTask(ENamedThreads::Type CurrentThread, const FGraphEventRef& MyCompletionGraphEvent);
 
-	/** Give the name for external event viewers
-	* @return	the name to display in external event viewers
-	*/
-	static const TCHAR *GetTaskName()
-	{
-		return TEXT("FAsyncParticleFill");
-	}
-	FORCEINLINE static TStatId GetStatId()
+	FORCEINLINE TStatId GetStatId() const
 	{
 		RETURN_QUICK_DECLARE_CYCLE_STAT(FAsyncParticleFill, STATGROUP_TaskGraphTasks);
 	}
diff --git a/Engine/Source/Runtime/Messaging/Private/Bridge/MessageDeserializeTask.h b/Engine/Source/Runtime/Messaging/Private/Bridge/MessageDeserializeTask.h
--- a/Engine/Source/Runtime/Messaging/Private/Bridge/MessageDeserializeTask.h
+++ b/Engine/Source/Runtime/Messaging/Private/Bridge/MessageDeserializeTask.h
@@ -71,7 +71,7 @@ public:
 	 *
 	 * @return Stats identifier.
 	 */
-	static TStatId GetStatId( );
+	TStatId GetStatId() const;
 
 	/**
 	 * Gets the mode for tracking subsequent tasks.
@@ -83,16 +83,6 @@ public:
 		return ESubsequentsMode::TrackSubsequents; 
 	}
 
-	/**
-	 * Gets the name of this task.
-	 *
-	 * @return Task name.
-	 */
-	static const TCHAR* GetTaskName( )
-	{
-		return TEXT("FMessageDeserializeTask");
-	}
-
 private:
 
 	// Holds the message's type information.
diff --git a/Engine/Source/Runtime/Messaging/Private/Bridge/MessageForwardTask.h b/Engine/Source/Runtime/Messaging/Private/Bridge/MessageForwardTask.h
--- a/Engine/Source/Runtime/Messaging/Private/Bridge/MessageForwardTask.h
+++ b/Engine/Source/Runtime/Messaging/Private/Bridge/MessageForwardTask.h
@@ -62,7 +62,7 @@ public:
 	 *
 	 * @return Stats identifier.
 	 */
-	static TStatId GetStatId( );
+	TStatId GetStatId() const;
 
 	/**
 	 * Gets the mode for tracking subsequent tasks.
@@ -74,16 +74,6 @@ public:
 		return ESubsequentsMode::TrackSubsequents; 
 	}
 
-	/**
-	 * Gets the name of this task.
-	 *
-	 * @return Task name.
-	 */
-	static const TCHAR* GetTaskName( )
-	{
-		return TEXT("FMessageForwardTask");
-	}
-
 private:
 
 	// Holds a pointer to the message bus to forward the message to.
diff --git a/Engine/Source/Runtime/Messaging/Private/Bridge/MessageSerializeTask.h b/Engine/Source/Runtime/Messaging/Private/Bridge/MessageSerializeTask.h
--- a/Engine/Source/Runtime/Messaging/Private/Bridge/MessageSerializeTask.h
+++ b/Engine/Source/Runtime/Messaging/Private/Bridge/MessageSerializeTask.h
@@ -69,7 +69,7 @@ public:
 	 *
 	 * @return Stats identifier.
 	 */
-	static TStatId GetStatId( );
+	TStatId GetStatId() const;
 
 	/**
 	 * Gets the mode for tracking subsequent tasks.
@@ -81,16 +81,6 @@ public:
 		return ESubsequentsMode::TrackSubsequents; 
 	}
 
-	/**
-	 * Gets the name of this task.
-	 *
-	 * @return Task name.
-	 */
-	static const TCHAR* GetTaskName( )
-	{
-		return TEXT("FMessageSerializeTask");
-	}
-
 private:
 
 	// Holds the context of the message to serialize.
diff --git a/Engine/Source/Runtime/Messaging/Private/Bus/MessageDispatchTask.h b/Engine/Source/Runtime/Messaging/Private/Bus/MessageDispatchTask.h
--- a/Engine/Source/Runtime/Messaging/Private/Bus/MessageDispatchTask.h
+++ b/Engine/Source/Runtime/Messaging/Private/Bus/MessageDispatchTask.h
@@ -74,7 +74,7 @@ public:
 	 *
 	 * @return Stats identifier.
 	 */
-	static TStatId GetStatId( );
+	TStatId GetStatId() const;
 
 	/**
 	 * Gets the mode for tracking subsequent tasks.
@@ -86,16 +86,6 @@ public:
 		return ESubsequentsMode::TrackSubsequents; 
 	}
 
-	/**
-	 * Gets the name of this task.
-	 *
-	 * @return Task name.
-	 */
-	static const TCHAR* GetTaskName( )
-	{
-		return TEXT("FMessageDispatchTask");
-	}
-
 private:
 
 	// Holds the message context.
diff --git a/Engine/Source/Runtime/Messaging/Private/MessagingModule.cpp b/Engine/Source/Runtime/Messaging/Private/MessagingModule.cpp
--- a/Engine/Source/Runtime/Messaging/Private/MessagingModule.cpp
+++ b/Engine/Source/Runtime/Messaging/Private/MessagingModule.cpp
@@ -150,25 +150,25 @@ private:
 /* Misc static functions
  *****************************************************************************/
 
-TStatId FMessageForwardTask::GetStatId( )
+TStatId FMessageForwardTask::GetStatId() const
 {
 	RETURN_QUICK_DECLARE_CYCLE_STAT(FMessageForwardTask, STATGROUP_TaskGraphTasks);
 }
 
 
-TStatId FMessageDispatchTask::GetStatId( )
+TStatId FMessageDispatchTask::GetStatId() const
 {
 	RETURN_QUICK_DECLARE_CYCLE_STAT(FMessageDispatchTask, STATGROUP_TaskGraphTasks);
 }
 
 
-TStatId FMessageSerializeTask::GetStatId( )
+TStatId FMessageSerializeTask::GetStatId() const
 {
 	RETURN_QUICK_DECLARE_CYCLE_STAT(FMessageSerializeTask, STATGROUP_TaskGraphTasks);
 }
 
 
-TStatId FMessageDeserializeTask::GetStatId( )
+TStatId FMessageDeserializeTask::GetStatId() const
 {
 	RETURN_QUICK_DECLARE_CYCLE_STAT(FMessageDeserializeTask, STATGROUP_TaskGraphTasks);
 }
diff --git a/Engine/Source/Runtime/RenderCore/Public/RenderingThread.h b/Engine/Source/Runtime/RenderCore/Public/RenderingThread.h
--- a/Engine/Source/Runtime/RenderCore/Public/RenderingThread.h
+++ b/Engine/Source/Runtime/RenderCore/Public/RenderingThread.h
@@ -167,11 +167,7 @@ DECLARE_STATS_GROUP(TEXT("Render Thread Commands"), STATGROUP_RenderThreadComman
 		} 
 
 #define TASKNAME_FUNCTION(TypeName) \
-		FORCEINLINE static const TCHAR* GetTaskName() \
-		{ \
-			return TEXT( #TypeName ); \
-		} \
-		FORCEINLINE static TStatId GetStatId() \
+		FORCEINLINE TStatId GetStatId() const \
 		{ \
 			RETURN_QUICK_DECLARE_CYCLE_STAT(TypeName, STATGROUP_RenderThreadCommands); \
 		}
@@ -195,7 +191,7 @@ DECLARE_STATS_GROUP(TEXT("Render Thread Commands"), STATGROUP_RenderThreadComman
 		else \
 		{ \
 			EURCMacro_##TypeName TempCommand; \
-			FScopeCycleCounter EURCMacro_Scope(EURCMacro_##TypeName::GetStatId()); \
+			FScopeCycleCounter EURCMacro_Scope(TempCommand.GetStatId()); \
 			TempCommand.DoTask(ENamedThreads::GameThread, FGraphEventRef() ); \
 		} \
 	}
@@ -227,7 +223,7 @@ DECLARE_STATS_GROUP(TEXT("Render Thread Commands"), STATGROUP_RenderThreadComman
 		else \
 		{ \
 			EURCMacro_##TypeName TempCommand(ParamValue1); \
-			FScopeCycleCounter EURCMacro_Scope(EURCMacro_##TypeName::GetStatId()); \
+			FScopeCycleCounter EURCMacro_Scope(TempCommand.GetStatId()); \
 			TempCommand.DoTask(ENamedThreads::GameThread, FGraphEventRef() ); \
 		} \
 	}
@@ -273,7 +269,7 @@ DECLARE_STATS_GROUP(TEXT("Render Thread Commands"), STATGROUP_RenderThreadComman
 		else \
 		{ \
 			EURCMacro_##TypeName TempCommand(ParamValue1,ParamValue2); \
-			FScopeCycleCounter EURCMacro_Scope(EURCMacro_##TypeName::GetStatId()); \
+			FScopeCycleCounter EURCMacro_Scope(TempCommand.GetStatId()); \
 			TempCommand.DoTask(ENamedThreads::GameThread, FGraphEventRef() ); \
 		} \
 	}
@@ -322,7 +318,7 @@ DECLARE_STATS_GROUP(TEXT("Render Thread Commands"), STATGROUP_RenderThreadComman
 		else \
 		{ \
 			EURCMacro_##TypeName TempCommand(ParamValue1,ParamValue2,ParamValue3); \
-			FScopeCycleCounter EURCMacro_Scope(EURCMacro_##TypeName::GetStatId()); \
+			FScopeCycleCounter EURCMacro_Scope(TempCommand.GetStatId()); \
 			TempCommand.DoTask(ENamedThreads::GameThread, FGraphEventRef() ); \
 		} \
 	}
@@ -374,7 +370,7 @@ DECLARE_STATS_GROUP(TEXT("Render Thread Commands"), STATGROUP_RenderThreadComman
 		else \
 		{ \
 			EURCMacro_##TypeName TempCommand(ParamValue1,ParamValue2,ParamValue3,ParamValue4); \
-			FScopeCycleCounter EURCMacro_Scope(EURCMacro_##TypeName::GetStatId()); \
+			FScopeCycleCounter EURCMacro_Scope(TempCommand.GetStatId()); \
 			TempCommand.DoTask(ENamedThreads::GameThread, FGraphEventRef() ); \
 		} \
 	}
