#UE4
- Temporary limited raw stats file to contain only 120 frames

diff --git a/Engine/Source/Runtime/Core/Private/Stats/StatsFile.cpp b/Engine/Source/Runtime/Core/Private/Stats/StatsFile.cpp
--- a/Engine/Source/Runtime/Core/Private/Stats/StatsFile.cpp
+++ b/Engine/Source/Runtime/Core/Private/Stats/StatsFile.cpp
@@ -336,23 +336,24 @@ void FStatsWriteFile::NewFrame( int64 TargetFrame )
 {
 	SCOPE_CYCLE_COUNTER( STAT_StreamFile );
 
-	// @TODO yrx 2014-03-24 add -num=frames to stat startfile|startfileraw
-#if	0
-	if( FCommandStatsFile::FirstFrame == -1 )
+	// Currently raw stat files are limited to 120 frames.
+	enum
 	{
-		FCommandStatsFile::FirstFrame = TargetFrame;
-	}
-	else if( TargetFrame > FCommandStatsFile::FirstFrame + 600 )
+		MAX_NUM_RAWFRAMES = 120,
+	};
+	if( Header.bRawStatFile )
 	{
-		if( FCommandStatsFile::CurrentStatFile.IsValid() )
+		if( FCommandStatsFile::FirstFrame == -1 )
 		{
-			FCommandStatsFile::CurrentStatFile->Stop();
+			FCommandStatsFile::FirstFrame = TargetFrame;
+		}
+		else if( TargetFrame > FCommandStatsFile::FirstFrame + MAX_NUM_RAWFRAMES )
+		{
+			FCommandStatsFile::Stop();
+			FCommandStatsFile::FirstFrame = -1;
+			return;
 		}
-		FCommandStatsFile::CurrentStatFile = nullptr;
-		FCommandStatsFile::FirstFrame = -1;
-		return;
 	}
-#endif // 0
 
 	WriteFrame( TargetFrame );
 	if( OutData.Num() > 1024 * 1024 )
