Removed unused mouse code from Mac HIDInputInterface

diff --git a/Engine/Source/Runtime/Core/Private/Mac/HIDInputInterface.cpp b/Engine/Source/Runtime/Core/Private/Mac/HIDInputInterface.cpp
--- a/Engine/Source/Runtime/Core/Private/Mac/HIDInputInterface.cpp
+++ b/Engine/Source/Runtime/Core/Private/Mac/HIDInputInterface.cpp
@@ -18,11 +18,10 @@ HIDInputInterface::HIDInputInterface( const TSharedRef< FGenericApplicationMessa
 		
 		ControllerState.ControllerId = ControllerIndex;
 	}
-	
+
 	InitialButtonRepeatDelay = 0.2f;
 	ButtonRepeatDelay = 0.1f;
-	
-	
+
 	HIDToXboxControllerMapping[0] = -1;
 	HIDToXboxControllerMapping[1] = 0;		// A
 	HIDToXboxControllerMapping[2] = 1;		// B
@@ -39,7 +38,7 @@ HIDInputInterface::HIDInputInterface( const TSharedRef< FGenericApplicationMessa
 	HIDToXboxControllerMapping[13] = -1;
 	HIDToXboxControllerMapping[14] = -1;
 	HIDToXboxControllerMapping[15] = -1;
-	
+
 	Buttons[0] = EControllerButtons::FaceButtonBottom;
 	Buttons[1] = EControllerButtons::FaceButtonRight;
 	Buttons[2] = EControllerButtons::FaceButtonLeft;
@@ -105,18 +104,6 @@ HIDInputInterface::HIDInputInterface( const TSharedRef< FGenericApplicationMessa
 	CFArrayAppendValue( MatchingArray, MatchingGamepads );
 	CFRelease( MatchingGamepads );
 
-	CFDictionaryRef MatchingMouses = CreateDeviceMatchingDictionary( kHIDPage_GenericDesktop, kHIDUsage_GD_Mouse );
-	if( !MatchingMouses )
-	{
-		CFRelease( MatchingArray );
-		CFRelease( HIDManager );
-		HIDManager = 0;
-		return;
-	}
-
-	CFArrayAppendValue( MatchingArray, MatchingMouses );
-	CFRelease( MatchingMouses );
-
 	IOHIDManagerSetDeviceMatchingMultiple( HIDManager, MatchingArray );
 	CFRelease( MatchingArray );
 
@@ -130,38 +117,6 @@ HIDInputInterface::HIDInputInterface( const TSharedRef< FGenericApplicationMessa
 
 void HIDInputInterface::Tick( float DeltaTime )
 {
-	const float DeltaScale = 2.1f;
-
-	MouseDelta = FVector2D( 0, 0 );
-
-	for( int32 Index = 0; Index < MouseStates.Num(); ++Index )
-	{
-		int32 AxisXValue = 0;
-		int32 AxisYValue = 0;
-
-		IOHIDValueRef Value;
-		IOReturn Result = IOHIDDeviceGetValue( MouseStates[Index].DeviceRef, MouseStates[Index].AxisXElement, &Value );
-		if( Result == kIOReturnSuccess )
-		{
-			AxisXValue = IOHIDValueGetIntegerValue( Value );
-		}
-
-		Result = IOHIDDeviceGetValue( MouseStates[Index].DeviceRef, MouseStates[Index].AxisYElement, &Value );
-		if( Result == kIOReturnSuccess )
-		{
-			AxisYValue = IOHIDValueGetIntegerValue( Value );
-		}
-
-		if( FMath::Abs( AxisXValue ) > 1 )
-		{
-			MouseDelta.X += AxisXValue * DeltaScale;
-		}
-
-		if( FMath::Abs( AxisYValue ) > 1 )
-		{
-			MouseDelta.Y += AxisYValue * DeltaScale;
-		}
-	}
 }
 
 void HIDInputInterface::SendControllerEvents()
@@ -354,17 +309,7 @@ CFMutableDictionaryRef HIDInputInterface::CreateDeviceMatchingDictionary( UInt32
 void HIDInputInterface::HIDDeviceMatchingCallback( void* Context, IOReturn Result, void* Sender, IOHIDDeviceRef DeviceRef )
 {
 	HIDInputInterface* HIDInput = ( HIDInputInterface* )Context;
-
-	const bool bIsMouse = IOHIDDeviceConformsTo( DeviceRef, kHIDPage_GenericDesktop, kHIDUsage_GD_Mouse );
-
-	if( bIsMouse )
-	{
-		HIDInput->OnNewHIDMouse( Result, DeviceRef );
-	}
-	else
-	{
-		HIDInput->OnNewHIDController( Result, DeviceRef );
-	}
+	HIDInput->OnNewHIDController( Result, DeviceRef );
 }
 
 void HIDInputInterface::HIDDeviceRemovalCallback( void* Context, IOReturn Result, void* Sender, IOHIDDeviceRef DeviceRef )
@@ -378,21 +323,6 @@ void HIDInputInterface::HIDDeviceRemovalCallback( void* Context, IOReturn Result
 			break;
 		}
 	}
-
-	int32 MouseToRemove = -1;
-	for( int32 Index = 0; Index < HIDInput->MouseStates.Num(); ++Index )
-	{
-		if( HIDInput->MouseStates[Index].DeviceRef == DeviceRef )
-		{
-			MouseToRemove = Index;
-			break;
-		}
-	}
-
-	if( MouseToRemove != -1 )
-	{
-		HIDInput->MouseStates.RemoveAtSwap( MouseToRemove );
-	}
 }
 
 void HIDInputInterface::OnNewHIDController( IOReturn Result, IOHIDDeviceRef DeviceRef )
@@ -478,83 +408,6 @@ void HIDInputInterface::OnNewHIDController( IOReturn Result, IOHIDDeviceRef Devi
 	CFRelease( Elements );
 }
 
-void HIDInputInterface::OnNewHIDMouse( IOReturn Result, IOHIDDeviceRef DeviceRef )
-{
-	if( IsMouseAlreadyKnown( DeviceRef ) )
-	{
-		return;
-	}
-
-	Result = IOHIDDeviceOpen( DeviceRef, kIOHIDOptionsTypeNone );
-	if( Result != kIOReturnSuccess )
-	{
-		return;	// couldn't open the device
-	}
-
-	// Get all elements from the device
-	CFArrayRef Elements = IOHIDDeviceCopyMatchingElements( DeviceRef, NULL, 0 );
-	if( !Elements )
-	{
-		IOHIDDeviceClose( DeviceRef, kIOHIDOptionsTypeNone );
-		return;
-	}
-
-	CFIndex ElementsCount = CFArrayGetCount( Elements );
-
-	FMouseState MouseState;
-	MouseState.DeviceRef = DeviceRef;
-
-	for( CFIndex ElementIndex = 0; ElementIndex < ElementsCount; ++ElementIndex )
-	{
-		IOHIDElementRef Element = (IOHIDElementRef)CFArrayGetValueAtIndex( Elements, ElementIndex );
-
-		IOHIDElementType ElementType = IOHIDElementGetType( Element );
-		uint32 UsagePage = IOHIDElementGetUsagePage( Element );
-		uint32 Usage = IOHIDElementGetUsage( Element );
-
-		if( ( ElementType == kIOHIDElementTypeInput_Axis || ElementType == kIOHIDElementTypeInput_Misc ) && UsagePage == kHIDPage_GenericDesktop )
-		{
-			if( Usage == kHIDUsage_GD_X )
-			{
-				MouseState.AxisXElement = Element;
-			}
-			else if( Usage == kHIDUsage_GD_Y )
-			{
-				MouseState.AxisYElement = Element;
-			}
-		}
-	}
-
-	if( MouseState.AxisXElement && MouseState.AxisYElement )
-	{
-		MouseStates.Add( MouseState );
-	}
-
-	CFRelease( Elements );
-}
-
-bool HIDInputInterface::IsMouseAlreadyKnown( IOHIDDeviceRef DeviceRef )
-{
-	const CFNumberRef NewVendor = (CFNumberRef)IOHIDDeviceGetProperty(DeviceRef, CFSTR(kIOHIDVendorIDKey));
-	const CFNumberRef NewProduct = (CFNumberRef)IOHIDDeviceGetProperty(DeviceRef, CFSTR(kIOHIDProductIDKey));
-	const CFNumberRef NewLocation = (CFNumberRef)IOHIDDeviceGetProperty(DeviceRef, CFSTR(kIOHIDLocationIDKey));
-
-	for( int32 Index = 0; Index < MouseStates.Num(); ++Index )
-	{
-		const CFNumberRef Vendor = (CFNumberRef)IOHIDDeviceGetProperty(MouseStates[Index].DeviceRef, CFSTR(kIOHIDVendorIDKey));
-		const CFNumberRef Product = (CFNumberRef)IOHIDDeviceGetProperty(MouseStates[Index].DeviceRef, CFSTR(kIOHIDProductIDKey));
-		const CFNumberRef Location = (CFNumberRef)IOHIDDeviceGetProperty(MouseStates[Index].DeviceRef, CFSTR(kIOHIDLocationIDKey));
-		if( CFNumberCompare( NewVendor, Vendor, NULL ) == kCFCompareEqualTo
-		   && CFNumberCompare( NewProduct, Product, NULL ) == kCFCompareEqualTo
-		   && CFNumberCompare( NewLocation, Location, NULL ) == kCFCompareEqualTo )
-		{
-			return true;
-		}
-	}
-
-	return false;
-}
-
 void HIDInputInterface::SetMessageHandler( const TSharedRef< FGenericApplicationMessageHandler >& InMessageHandler )
 {
 	MessageHandler = InMessageHandler;
diff --git a/Engine/Source/Runtime/Core/Private/Mac/HIDInputInterface.h b/Engine/Source/Runtime/Core/Private/Mac/HIDInputInterface.h
--- a/Engine/Source/Runtime/Core/Private/Mac/HIDInputInterface.h
+++ b/Engine/Source/Runtime/Core/Private/Mac/HIDInputInterface.h
@@ -24,8 +24,6 @@ public:
 
 	void SetMessageHandler( const TSharedRef< FGenericApplicationMessageHandler >& InMessageHandler );
 
-	const FVector2D& GetMouseDelta() const { return MouseDelta; }
-
 private:
 
 	HIDInputInterface( const TSharedRef< FGenericApplicationMessageHandler >& MessageHandler );
@@ -38,10 +36,6 @@ private:
 
 	void OnNewHIDController( IOReturn Result, IOHIDDeviceRef DeviceRef );
 
-	void OnNewHIDMouse( IOReturn Result, IOHIDDeviceRef DeviceRef );
-
-	bool IsMouseAlreadyKnown( IOHIDDeviceRef DeviceRef );
-
 private:
 
 	struct FHIDButtonInfo
@@ -98,15 +92,6 @@ private:
 		FHIDDeviceInfo Device;
 	};
 
-	struct FMouseState
-	{
-		IOHIDDeviceRef DeviceRef;
-		IOHIDElementRef AxisXElement;
-		IOHIDElementRef AxisYElement;
-
-		FMouseState() : DeviceRef(NULL), AxisXElement(NULL), AxisYElement(NULL) {}
-	};
-
 private:
 
 	/** In the engine, all controllers map to xbox controllers for consistency */
@@ -118,11 +103,6 @@ private:
 	/** Controller states */
 	FControllerState ControllerStates[MAX_NUM_HIDINPUT_CONTROLLERS];
 
-	/** Mouse states */
-	TArray<FMouseState> MouseStates;
-
-	FVector2D MouseDelta;
-
 	/** Delay before sending a repeat message after a button was first pressed */
 	float InitialButtonRepeatDelay;
 
