Fixed Mac non-unity build errors

diff --git a/Engine/Source/Programs/SlateViewer/Private/Mac/SlateViewerMainMac.cpp b/Engine/Source/Programs/SlateViewer/Private/Mac/SlateViewerMainMac.cpp
--- a/Engine/Source/Programs/SlateViewer/Private/Mac/SlateViewerMainMac.cpp
+++ b/Engine/Source/Programs/SlateViewer/Private/Mac/SlateViewerMainMac.cpp
@@ -1,6 +1,7 @@
 // Copyright 1998-2014 Epic Games, Inc. All Rights Reserved.
 
 #include "SlateViewerApp.h"
+#include "ExceptionHandling.h"
 
 static FString GSavedCommandLine;
 
diff --git a/Engine/Source/Runtime/Core/Private/Mac/MacPlatformOutputDevices.cpp b/Engine/Source/Runtime/Core/Private/Mac/MacPlatformOutputDevices.cpp
--- a/Engine/Source/Runtime/Core/Private/Mac/MacPlatformOutputDevices.cpp
+++ b/Engine/Source/Runtime/Core/Private/Mac/MacPlatformOutputDevices.cpp
@@ -7,8 +7,9 @@
 #include "CorePrivate.h"
 
 #include "FeedbackContextAnsi.h"
-#include "../Private/Mac/MacPlatformOutputDevicesPrivate.h"
-#include "../Private/Mac/MacPlatformFeedbackContextPrivate.h"
+#include "MacApplication.h"
+#include "MacPlatformOutputDevicesPrivate.h"
+#include "MacPlatformFeedbackContextPrivate.h"
 
 //////////////////////////////////
 // FMacPlatformOutputDevices
diff --git a/Engine/Source/Runtime/Core/Private/Mac/MacTextInputMethodSystem.cpp b/Engine/Source/Runtime/Core/Private/Mac/MacTextInputMethodSystem.cpp
--- a/Engine/Source/Runtime/Core/Private/Mac/MacTextInputMethodSystem.cpp
+++ b/Engine/Source/Runtime/Core/Private/Mac/MacTextInputMethodSystem.cpp
@@ -2,6 +2,7 @@
 
 #include "CorePrivate.h"
 
+#include "MacWindow.h"
 #include "MacTextInputMethodSystem.h"
 
 DEFINE_LOG_CATEGORY_STATIC(LogMacTextInputMethodSystem, Log, All);
