Expose whether we're building with Steam through WITH_STEAMWORKS define and compile SteamController conditionally (TTP #332911).

diff --git a/Engine/Source/Programs/UnrealBuildTool/Configuration/UEBuildConfiguration.cs b/Engine/Source/Programs/UnrealBuildTool/Configuration/UEBuildConfiguration.cs
--- a/Engine/Source/Programs/UnrealBuildTool/Configuration/UEBuildConfiguration.cs
+++ b/Engine/Source/Programs/UnrealBuildTool/Configuration/UEBuildConfiguration.cs
@@ -43,7 +43,7 @@ namespace UnrealBuildTool
 		/** Whether we should compile in support for Simplygon or not. */
 		public static bool bCompileSimplygon;
 
-		/** Whether we should compile in support for Steam OnlineSubsystem or not. */
+        /** Whether we should compile in support for Steam OnlineSubsystem or not. [RCL] FIXME 2014-Apr-17: bCompileSteamOSS means "bHasSteamworksInstalled" for some code, these meanings need to be untangled */
 		public static bool bCompileSteamOSS;
 
 		/** Whether we should compile in support for Mcp OnlineSubsystem or not. */
diff --git a/Engine/Source/Runtime/Core/Core.Build.cs b/Engine/Source/Runtime/Core/Core.Build.cs
--- a/Engine/Source/Runtime/Core/Core.Build.cs
+++ b/Engine/Source/Runtime/Core/Core.Build.cs
@@ -112,7 +112,7 @@ public class Core : ModuleRules
 				"SDL2"
                 );
 
-            if (UEBuildConfiguration.bCompileAgainstEngine == true && Target.Type != TargetRules.TargetType.Server)
+            if (UEBuildConfiguration.bCompileAgainstEngine == true && Target.Type != TargetRules.TargetType.Server && UEBuildConfiguration.bCompileSteamOSS)
             {
                 AddThirdPartyPrivateStaticDependencies(Target, "SteamController");
             }
@@ -136,7 +136,9 @@ public class Core : ModuleRules
         }
         Definitions.Add("UE_ENABLE_ICU=" + (UEBuildConfiguration.bCompileICU ? "1" : "0")); // Enable/disable (=1/=0) ICU usage in the codebase. NOTE: This flag is for use while integrating ICU and will be removed afterward.
 
-		// If we're compiling with the engine, then add Core's engine dependencies
+        Definitions.Add("WITH_STEAMWORKS=" + (UEBuildConfiguration.bCompileSteamOSS ? "1" : "0"));
+
+        // If we're compiling with the engine, then add Core's engine dependencies
 		if (UEBuildConfiguration.bCompileAgainstEngine == true)
 		{
 			if (!UEBuildConfiguration.bBuildRequiresCookedData)
diff --git a/Engine/Source/Runtime/Core/Public/Linux/LinuxApplication.h b/Engine/Source/Runtime/Core/Public/Linux/LinuxApplication.h
--- a/Engine/Source/Runtime/Core/Public/Linux/LinuxApplication.h
+++ b/Engine/Source/Runtime/Core/Public/Linux/LinuxApplication.h
@@ -5,7 +5,7 @@
 #include "GenericApplication.h"
 #include "LinuxWindow.h"
 
-#define STEAM_CONTROLLER_SUPPORT				(WITH_ENGINE && !UE_SERVER)
+#define STEAM_CONTROLLER_SUPPORT				(WITH_STEAMWORKS && WITH_ENGINE && !UE_SERVER)
 
 typedef SDL_GameController* SDL_HController;
 
