Fixed FApplePlatformFile::MoveFile() to work with files on different file systems

diff --git a/Engine/Source/Runtime/Core/Private/Apple/ApplePlatformFile.cpp b/Engine/Source/Runtime/Core/Private/Apple/ApplePlatformFile.cpp
--- a/Engine/Source/Runtime/Core/Private/Apple/ApplePlatformFile.cpp
+++ b/Engine/Source/Runtime/Core/Private/Apple/ApplePlatformFile.cpp
@@ -339,7 +339,17 @@ bool FApplePlatformFile::IsReadOnly(const TCHAR* Filename)
 }
 bool FApplePlatformFile::MoveFile(const TCHAR* To, const TCHAR* From)
 {
-	return rename(TCHAR_TO_UTF8(*NormalizeFilename(From)), TCHAR_TO_UTF8(*NormalizeFilename(To))) != -1;
+	int32 Result = rename(TCHAR_TO_UTF8(*NormalizeFilename(From)), TCHAR_TO_UTF8(*NormalizeFilename(To)));
+	if (Result == -1 && errno == EXDEV)
+	{
+		// Copy the file if rename failed because To and From are on different file systems
+		if (CopyFile(To, From))
+		{
+			DeleteFile(From);
+			Result = 0;
+		}
+	}
+	return Result != -1;
 }
 bool FApplePlatformFile::SetReadOnly(const TCHAR* Filename, bool bNewReadOnlyValue)
 {
