Fix RInterpTo doing wrong early out when given very small interpolation value #codereview Jeff.Farris

diff --git a/Engine/Source/Runtime/Core/Private/Math/UnrealMath.cpp b/Engine/Source/Runtime/Core/Private/Math/UnrealMath.cpp
--- a/Engine/Source/Runtime/Core/Private/Math/UnrealMath.cpp
+++ b/Engine/Source/Runtime/Core/Private/Math/UnrealMath.cpp
@@ -1921,15 +1921,16 @@ CORE_API FRotator FMath::RInterpTo( const FRotator& Current, const FRotator& Tar
 
 	const float DeltaInterpSpeed = InterpSpeed * DeltaTime;
 
-	// Delta Move, Clamp so we do not over shoot.
-	const FRotator DeltaMove = (Target - Current).GetNormalized() * FMath::Clamp<float>(DeltaInterpSpeed, 0.f, 1.f);
-
+	const FRotator Delta = (Target - Current).GetNormalized();
+	
 	// If steps are too small, just return Target and assume we have reached our destination.
-	if( DeltaMove.IsNearlyZero() )
+	if (Delta.IsNearlyZero())
 	{
 		return Target;
 	}
 
+	// Delta Move, Clamp so we do not over shoot.
+	const FRotator DeltaMove = Delta * FMath::Clamp<float>(DeltaInterpSpeed, 0.f, 1.f);
 	return (Current + DeltaMove).GetNormalized();
 }
 
