#UE4 - Moved TAttribute and TOptional to core

diff --git a/Engine/Source/Runtime/Core/Public/Misc/Attribute.h b/Engine/Source/Runtime/Core/Public/Misc/Attribute.h
new file mode 100644
--- /dev/null
+++ b/Engine/Source/Runtime/Core/Public/Misc/Attribute.h
@@ -0,0 +1,321 @@
+// Copyright 1998-2014 Epic Games, Inc. All Rights Reserved.
+
+#pragma once
+
+
+#include "Delegate.h"
+
+/**
+ * Attribute object
+ */
+template< typename ObjectType >
+class TAttribute
+{
+
+public:
+
+	/**
+	 * Attribute 'getter' delegate
+	 *
+	 * ObjectType GetValue() const
+	 *
+	 * @return  The attribute's value
+	 */
+	DECLARE_DELEGATE_RetVal( ObjectType, FGetter );
+
+
+	/**
+	 * Default constructor
+	 */
+	TAttribute()
+		: Getter()
+		, Value()         // NOTE: Potentially uninitialized for atomics!!
+		, bIsSet(false)
+	{
+	}
+
+
+	/**
+	 * Construct implicitly from an initial value
+	 *
+	 * @param  InInitialValue  The value for this attribute
+	 */
+	template< typename OtherType >
+	TAttribute( const OtherType& InInitialValue )
+		: Getter()
+		, Value( InInitialValue )
+		, bIsSet(true)
+	{
+	}
+
+
+	/**
+	 * Constructs by binding an arbitrary function that will be called to generate this attribute's value on demand. 
+	 * After binding, the attribute will no longer have a value that can be accessed directly, and instead the bound
+	 * function will always be called to generate the value.
+	 *
+	 * @param  InUserObject  Shared Pointer to the instance of the class that contains the member function you want to bind.  The attribute will only retain a weak pointer to this class.
+	 * @param  InMethodPtr Member function to bind.  The function's structure (return value, arguments, etc) must match IBoundAttributeDelegate's definition.
+	 */
+	template< class SourceType >	
+	TAttribute( TSharedRef< SourceType > InUserObject, typename FGetter::template TSPMethodDelegate_Const< SourceType >::FMethodPtr InMethodPtr )
+		: Getter( FGetter::CreateSP( InUserObject, InMethodPtr ) )
+		, Value()
+		, bIsSet(true)
+	{
+	}
+
+	
+	/**
+	 * Constructs by binding an arbitrary function that will be called to generate this attribute's value on demand. 
+	 * After binding, the attribute will no longer have a value that can be accessed directly, and instead the bound
+	 * function will always be called to generate the value.
+	 *
+	 * @param  InUserObject  Shared Pointer to the instance of the class that contains the member function you want to bind.  The attribute will only retain a weak pointer to this class.
+	 * @param  InMethodPtr Member function to bind.  The function's structure (return value, arguments, etc) must match IBoundAttributeDelegate's definition.
+	 */
+	template< class SourceType >	
+	TAttribute( SourceType* InUserObject, typename FGetter::template TSPMethodDelegate_Const< SourceType >::FMethodPtr InMethodPtr )
+		: Getter( FGetter::CreateSP( InUserObject, InMethodPtr ) )
+		, Value()
+		, bIsSet(true)
+	{
+	}
+
+
+	/**
+	 * Static: Creates an attribute that's pre-bound to the specified 'getter' delegate
+	 *
+	 * @param  InGetter		Delegate to bind
+	 */
+	static TAttribute< ObjectType > Create( const FGetter& InGetter )
+	{
+		const bool bExplicitConstructor = true;
+		return TAttribute< ObjectType >( InGetter, bExplicitConstructor );
+	}
+
+	
+	/**
+	 * Creates an attribute by binding an arbitrary function that will be called to generate this attribute's value on demand. 
+	 * After binding, the attribute will no longer have a value that can be accessed directly, and instead the bound
+	 * function will always be called to generate the value.
+	 *
+	 * @param  InFuncPtr Member function to bind.  The function's structure (return value, arguments, etc) must match IBoundAttributeDelegate's definition.
+	 */
+	static TAttribute< ObjectType > Create( typename FGetter::FStaticDelegate::FFuncPtr InFuncPtr )
+	{
+		const bool bExplicitConstructor = true;
+		return TAttribute< ObjectType >( InFuncPtr, bExplicitConstructor );
+	}
+
+
+	/**
+	 * Sets the attribute's value
+	 *
+	 * @param  InNewValue  The value to set the attribute to
+	 */
+	template< typename OtherType >
+	void Set( const OtherType& InNewValue )
+	{
+		Getter.Unbind();
+		Value = InNewValue;
+		bIsSet = true;
+	}
+
+	/** Was this TAttribute ever assigned? */
+	bool IsSet() const
+	{
+		return bIsSet;
+	}
+
+	/**
+	 * Gets the attribute's current value.
+	 * Assumes that the attribute is set.
+	 *
+	 * @return  The attribute's value
+	 */
+ 	const ObjectType& Get() const
+ 	{
+		// If we have a getter delegate, then we'll call that to generate the value
+		if( Getter.IsBound() )
+		{
+			// Call the delegate to get the value.  Note that this will assert if the delegate is not currently
+			// safe to call (e.g. object was deleted and we could detect that)
+
+			// NOTE: We purposely overwrite our value copy here so that we can return the value by address in
+			// the most common case, which is an attribute that doesn't have a delegate bound to it at all.
+			Value = Getter.Execute();
+		}
+
+		// Return the stored value
+		return Value;
+ 	}
+
+	/**
+	 * Gets the attribute's current value. The attribute may not be set, in which case use the default value provided.
+	 * Shorthand for the boilerplate code: MyAttribute.IsSet() ? MyAttribute.Get() : DefaultValue
+	 */
+	const ObjectType& Get( const ObjectType& DefaultValue ) const
+	{
+		return bIsSet ? Get() : DefaultValue;
+	}
+
+
+	/**
+	 * Binds an arbitrary function that will be called to generate this attribute's value on demand.  After
+	 * binding, the attribute will no longer have a value that can be accessed directly, and instead the bound
+	 * function will always be called to generate the value.
+	 *
+	 * @param  InGetter  The delegate object with your function binding
+	 */
+	void Bind( const FGetter& InGetter )
+	{
+		bIsSet = true;
+		Getter = InGetter;
+	}
+
+
+	/**
+	 * Binds an arbitrary function that will be called to generate this attribute's value on demand. 
+	 * After binding, the attribute will no longer have a value that can be accessed directly, and instead the bound
+	 * function will always be called to generate the value.
+	 *
+	 * @param  InFuncPtr	Function to bind.  The function's structure (return value, arguments, etc) must match IBoundAttributeDelegate's definition.
+	 */
+	void BindStatic( typename FGetter::FStaticDelegate::FFuncPtr InFuncPtr )
+	{
+		bIsSet = true;
+		Getter.BindStatic( InFuncPtr );
+	}
+
+	
+	/**
+	 * Binds an arbitrary function that will be called to generate this attribute's value on demand. 
+	 * After binding, the attribute will no longer have a value that can be accessed directly, and instead the bound
+	 * function will always be called to generate the value.
+	 *
+	 * @param  InUserObject  Instance of the class that contains the member function you want to bind.
+	 * @param  InMethodPtr Member function to bind.  The function's structure (return value, arguments, etc) must match IBoundAttributeDelegate's definition.
+	 */
+	template< class SourceType >	
+	void BindRaw( SourceType* InUserObject, typename FGetter::template TRawMethodDelegate_Const< SourceType >::FMethodPtr InMethodPtr )
+	{
+		bIsSet = true;
+		Getter.BindRaw( InUserObject, InMethodPtr );
+	}
+
+
+	/**
+	 * Binds an arbitrary function that will be called to generate this attribute's value on demand. 
+	 * After binding, the attribute will no longer have a value that can be accessed directly, and instead the bound
+	 * function will always be called to generate the value.
+	 *
+	 * @param  InUserObject  Shared Pointer to the instance of the class that contains the member function you want to bind.  The attribute will only retain a weak pointer to this class.
+	 * @param  InMethodPtr Member function to bind.  The function's structure (return value, arguments, etc) must match IBoundAttributeDelegate's definition.
+	 */
+	template< class SourceType >	
+	void Bind( TSharedRef< SourceType > InUserObject, typename FGetter::template TSPMethodDelegate_Const< SourceType >::FMethodPtr InMethodPtr )
+	{
+		bIsSet = true;
+		Getter.BindSP( InUserObject, InMethodPtr );
+	}
+
+	
+	/**
+	 * Binds an arbitrary function that will be called to generate this attribute's value on demand. 
+	 * After binding, the attribute will no longer have a value that can be accessed directly, and instead the bound
+	 * function will always be called to generate the value.
+	 *
+	 * @param  InUserObject  Shared Pointer to the instance of the class that contains the member function you want to bind.  The attribute will only retain a weak pointer to this class.
+	 * @param  InMethodPtr Member function to bind.  The function's structure (return value, arguments, etc) must match IBoundAttributeDelegate's definition.
+	 */
+	template< class SourceType >	
+	void Bind( SourceType* InUserObject, typename FGetter::template TSPMethodDelegate_Const< SourceType >::FMethodPtr InMethodPtr )
+	{
+		bIsSet = true;
+		Getter.BindSP( InUserObject, InMethodPtr );
+	}
+
+	
+	/**
+	 * Binds an arbitrary function that will be called to generate this attribute's value on demand. 
+	 * After binding, the attribute will no longer have a value that can be accessed directly, and instead the bound
+	 * function will always be called to generate the value.
+	 *
+	 * @param  InUserObject  Instance of the class that contains the member function you want to bind.
+	 * @param  InMethodPtr Member function to bind.  The function's structure (return value, arguments, etc) must match IBoundAttributeDelegate's definition.
+	 */
+	template< class SourceType >	
+	void BindUObject( SourceType* InUserObject, typename FGetter::template TUObjectMethodDelegate_Const< SourceType >::FMethodPtr InMethodPtr )
+	{
+		bIsSet = true;
+		Getter.BindUObject( InUserObject, InMethodPtr );
+	}
+
+
+	/**
+	 * Checks to see if this attribute has a 'getter' function bound
+	 *
+	 * @return  True if attribute is bound to a getter function
+	 */
+	bool IsBound() const
+	{
+		return Getter.IsBound();
+	}
+
+
+	/**
+	 * Equality operator
+	 */
+	bool operator==( const TAttribute& InOther ) const
+	{
+		return InOther.Get() == Get();
+	}
+
+
+	/**
+	 * Not equal operator
+	 */
+	bool operator!=( const TAttribute& InOther ) const
+	{
+		return InOther.Get() != Get();
+	}
+
+
+	/** Special explicit constructor for TAttribute::Create() */
+	TAttribute( typename FGetter::FStaticDelegate::FFuncPtr InFuncPtr, bool bExplicitConstructor )
+		: Getter( FGetter::CreateStatic( InFuncPtr ) )
+		, Value()
+		, bIsSet( true )
+	{
+	}
+
+	
+	/** Special explicit constructor for TAttribute::Create() */
+	TAttribute( const FGetter& InGetter, bool bExplicitConstructor )
+		: Getter( InGetter )
+		, Value()
+		, bIsSet( true )
+	{
+	}
+
+
+private:
+
+	// We declare ourselves as a friend (templated using OtherType) so we can access members as needed
+    template< class OtherType > friend class TAttribute;
+
+	/** Bound member function for this attribute (may be NULL if no function is bound.)  When set, all attempts
+	    to read the attribute's value will instead call this delegate to generate the value. */
+	/** Our attribute's 'getter' delegate */
+	FGetter Getter;
+
+	/** Current value.  Mutable so that we can cache the value locally when using a bound Getter (allows const ref return value.) */
+	mutable ObjectType Value;
+
+	/** true when this attribute was explicitly set by a consumer, false when the attribute's value is set to the default*/
+	bool bIsSet;
+
+};
+
+
diff --git a/Engine/Source/Runtime/Core/Public/Misc/Optional.h b/Engine/Source/Runtime/Core/Public/Misc/Optional.h
new file mode 100644
--- /dev/null
+++ b/Engine/Source/Runtime/Core/Public/Misc/Optional.h
@@ -0,0 +1,39 @@
+// Copyright 1998-2014 Epic Games, Inc. All Rights Reserved.
+
+#pragma once
+
+/**
+ * When we have an optional value IsSet() returns true, and GetValue() is meaningful.
+ * Otherwise GetValue() is not meaningful.
+ */
+template<typename OptionalType>
+struct TOptional
+{
+public:
+	/** Construct an OptionaType with a valid value. */
+	TOptional( OptionalType InValue )
+		: bIsSet( true )
+		, Value( InValue )
+	{
+	}
+
+	/** Construct an OptionalType with no value; i.e. NULL */
+	TOptional()
+		: bIsSet( false )
+		, Value()
+	{
+	}
+
+	/** @return true when the value is meaningulf; false if calling GetValue() is undefined. */
+	bool IsSet() const { return bIsSet; }
+
+	/** @return The optional value; undefined when IsSet() returns false. */
+	const OptionalType& GetValue() const { return Value; }
+
+	/** @return The optional value when set; DefaultValue otherwise. */
+	const OptionalType& Get( const OptionalType& DefaultValue ) const { return IsSet() ? Value : DefaultValue; }
+
+private:
+	bool bIsSet;
+	OptionalType Value;
+};
\ No newline at end of file
diff --git a/Engine/Source/Runtime/Slate/Public/Attribute.h b/Engine/Source/Runtime/Slate/Public/Attribute.h
deleted file mode 100644
diff --git a/Engine/Source/Runtime/Slate/Public/PrimitiveTypes.h b/Engine/Source/Runtime/Slate/Public/PrimitiveTypes.h
--- a/Engine/Source/Runtime/Slate/Public/PrimitiveTypes.h
+++ b/Engine/Source/Runtime/Slate/Public/PrimitiveTypes.h
@@ -863,40 +863,4 @@ class SLATE_API FArrangedChildren
 		Array.Empty();
 	}
 
-};
-
-/**
- * When we have an optional value IsSet() returns true, and GetValue() is meaningful.
- * Otherwise GetValue() is not meaningful.
- */
-template<typename OptionalType>
-struct TOptional
-{
-public:
-	/** Construct an OptionaType with a valid value. */
-	TOptional( OptionalType InValue )
-		: bIsSet( true )
-		, Value( InValue )
-	{
-	}
-
-	/** Construct an OptionalType with no value; i.e. NULL */
-	TOptional()
-		: bIsSet( false )
-		, Value()
-	{
-	}
-
-	/** @return true when the value is meaningulf; false if calling GetValue() is undefined. */
-	bool IsSet() const { return bIsSet; }
-
-	/** @return The optional value; undefined when IsSet() returns false. */
-	const OptionalType& GetValue() const { return Value; }
-
-	/** @return The optional value when set; DefaultValue otherwise. */
-	const OptionalType& Get( const OptionalType& DefaultValue ) const { return IsSet() ? Value : DefaultValue; }
-
-private:
-	bool bIsSet;
-	OptionalType Value;
 };
\ No newline at end of file
