#UE4 - Adding missing changes to core.h

diff --git a/Engine/Source/Runtime/Core/Public/Core.h b/Engine/Source/Runtime/Core/Public/Core.h
--- a/Engine/Source/Runtime/Core/Public/Core.h
+++ b/Engine/Source/Runtime/Core/Public/Core.h
@@ -467,6 +467,8 @@ extern CORE_API double GBlueprintCompileTime;
 #include "App.h"
 #include "OutputDeviceConsole.h"
 #include "MonitoredProcess.h"
+#include "Attribute.h"
+#include "Optional.h"
 
 #ifdef TRUE
 	#undef TRUE
