UE4: Disabled NaN checking in Main branch because it currently causes thousands of errors in Orion cook. This the NaN issue is under investigation at the moment and NaN checking will be re-enabled once this bug is fixed.

diff --git a/Engine/Source/Runtime/Core/Public/Math/TransformVectorized.h b/Engine/Source/Runtime/Core/Public/Math/TransformVectorized.h
--- a/Engine/Source/Runtime/Core/Public/Math/TransformVectorized.h
+++ b/Engine/Source/Runtime/Core/Public/Math/TransformVectorized.h
@@ -2,7 +2,7 @@
 
 
 #pragma once
-#define TRANSFORM_TRACK_NAN_ISSUES (1 && !UE_BUILD_SHIPPING && !UE_BUILD_TEST)
+#define TRANSFORM_TRACK_NAN_ISSUES (0 && !UE_BUILD_SHIPPING && !UE_BUILD_TEST)
 
 /** This returns Quaternion Inverse of X **/
 #define MAKE_QUATINV_VECTORREGISTER(X) VectorMultiply(GlobalVectorConstants::QINV_SIGN_MASK, X)
