#UE4
 - Fixed problem with æstat dumphitchesÆ not found game thread, switched to use the short name instead of looking at the beginning of the full name
 - Fixed problem with æstat dumphitchesÆ causes chain reaction of hitch reports (if a hitch is detected, the stack is logged, but the next hitch can be detected after 4 frames, when the stats thread gets synchronized with the game thread, assuming that this command is used in the long term testing, it shouldnÆt be a problem if it misses a hitch due to 4 frames freeze)

diff --git a/Engine/Source/Runtime/Core/Private/Stats/Stats2.cpp b/Engine/Source/Runtime/Core/Private/Stats/Stats2.cpp
--- a/Engine/Source/Runtime/Core/Private/Stats/Stats2.cpp
+++ b/Engine/Source/Runtime/Core/Private/Stats/Stats2.cpp
@@ -829,10 +829,6 @@ void FThreadStats::StartThread()
 
 }
 
-enum 
-{
-	MAX_STAT_LAG = 4,
-};
 static FGraphEventRef LastFramesEvents[MAX_STAT_LAG];
 static int32 CurrentEventIndex = 0;
 
diff --git a/Engine/Source/Runtime/Core/Private/Stats/StatsCommand.cpp b/Engine/Source/Runtime/Core/Private/Stats/StatsCommand.cpp
--- a/Engine/Source/Runtime/Core/Private/Stats/StatsCommand.cpp
+++ b/Engine/Source/Runtime/Core/Private/Stats/StatsCommand.cpp
@@ -80,7 +80,17 @@ private:
 	T Value;
 };
 
-
+struct FGroupFilter : public IItemFiler
+{
+	TSet<FName> const& EnabledItems;
+	FGroupFilter( TSet<FName> const& InEnabledItems )
+		: EnabledItems( InEnabledItems )
+	{}
+	virtual bool Keep( FStatMessage const& Item )
+	{
+		return EnabledItems.Contains( Item.NameAndInfo.GetRawName() );
+	}
+};
 
 /** Holds parameters used by the 'stat hier' or 'stat group ##' command. */
 struct FStatGroupParams
@@ -363,12 +373,24 @@ void DumpCPUSummary(FStatsThreadState const& StatsData, int64 TargetFrame)
 
 static void DumpHitch(int64 Frame)
 {
+	// !!!CAUTION!!! 
+	// Due to chain reaction of hitch reports after detecting the first hitch, the hitch detector is disabled for the next 4 frames.
+	// There is no other safe method to detect if the next hitch is a real hitch or just waiting for flushing the threaded logs or waiting for the stats. 
+	// So, the best way is to just wait until stats gets synchronized with the game thread.
+
+	static int64 LastHitchFrame = -(MAX_STAT_LAG + STAT_FRAME_SLOP);
+	if( LastHitchFrame + (MAX_STAT_LAG + STAT_FRAME_SLOP) > Frame )
+	{
+		return;
+	}
+
 	FStatsThreadState& Stats = FStatsThreadState::GetLocalState();
 	SCOPE_CYCLE_COUNTER(STAT_HitchScan);
 
-	float GameThreadTime = FPlatformTime::ToSeconds(Stats.GetFastThreadFrameTime(Frame, EThreadType::Game));
-	float RenderThreadTime = FPlatformTime::ToSeconds(Stats.GetFastThreadFrameTime(Frame, EThreadType::Renderer));
-	if (GameThreadTime > GHitchThreshold || RenderThreadTime > GHitchThreshold)
+	const float GameThreadTime = FPlatformTime::ToSeconds(Stats.GetFastThreadFrameTime(Frame, EThreadType::Game));
+	const float RenderThreadTime = FPlatformTime::ToSeconds(Stats.GetFastThreadFrameTime(Frame, EThreadType::Renderer));
+
+	if( GameThreadTime > GHitchThreshold || RenderThreadTime > GHitchThreshold )
 	{
 		UE_LOG(LogStats, Log, TEXT("------------------Thread Hitch, Frame %lld  %6.1fms ---------------"), Frame, FMath::Max<float>(GameThreadTime, RenderThreadTime) * 1000.0f );
 		FRawStatStackNode Stack;
@@ -376,44 +398,42 @@ static void DumpHitch(int64 Frame)
 		Stack.AddNameHierarchy();
 		Stack.AddSelf();
 
-		int64 MinCycles = int64(FMath::Min<float>(FMath::Max<float>(GHitchThreshold - 33.3f/1000.0f, 1.0f/1000.0f), 1.0f/1000.0f) / FPlatformTime::GetSecondsPerCycle());
+		int64 MinCycles = int64( FMath::Min<float>( FMath::Max<float>( GHitchThreshold - 33.3f / 1000.0f, 1.0f / 1000.0f ), 1.0f / 1000.0f ) / FPlatformTime::GetSecondsPerCycle() );
 		FRawStatStackNode* GameThread = NULL;
 		FRawStatStackNode* RenderThread = NULL;
-		const FString GameThreadString = FName(NAME_GameThread).ToString();
-		const FString RenderThreadString = FName(NAME_RenderThread).ToString();
 
-		for (auto ChildIter = Stack.Children.CreateConstIterator(); ChildIter; ++ChildIter)
+		for( auto ChildIter = Stack.Children.CreateConstIterator(); ChildIter; ++ChildIter )
 		{
-			const FName ChildName = ChildIter.Key();
-			const FString ChildString = ChildName.ToString();
+			const FName ThreadName = ChildIter.Value()->Meta.NameAndInfo.GetShortName();
 
-			if (ChildString.StartsWith(GameThreadString))
+			if( ThreadName == FName( NAME_GameThread ) )
 			{
 				GameThread = ChildIter.Value();
-				UE_LOG(LogStats, Log, TEXT("------------------ Game Thread %.2fms"), GameThreadTime * 1000.0f);
-				GameThread->Cull(MinCycles);
+				UE_LOG( LogStats, Log, TEXT( "------------------ Game Thread %.2fms" ), GameThreadTime * 1000.0f );
+				GameThread->Cull( MinCycles );
 				GameThread->DebugPrint();
 			}
-			else if (ChildString.StartsWith(RenderThreadString))
+			else if( ThreadName == FName( NAME_RenderThread ) )
 			{
 				RenderThread = ChildIter.Value();
-				UE_LOG(LogStats, Log, TEXT("------------------ Render Thread (%s) %.2fms"), *RenderThread->Meta.NameAndInfo.GetRawName().ToString(), RenderThreadTime * 1000.0f);
-				RenderThread->Cull(MinCycles);
+				UE_LOG( LogStats, Log, TEXT( "------------------ Render Thread (%s) %.2fms" ), *RenderThread->Meta.NameAndInfo.GetRawName().ToString(), RenderThreadTime * 1000.0f );
+				RenderThread->Cull( MinCycles );
 				RenderThread->DebugPrint();
 			}
 		}
 
-		if (!GameThread)
+		if( !GameThread )
 		{
-			UE_LOG(LogStats, Warning, TEXT("No game thread?!"));
+			UE_LOG( LogStats, Warning, TEXT( "No game thread?!" ) );
 		}
 
-		if (!RenderThread)
+		if( !RenderThread )
 		{
-			UE_LOG(LogStats, Warning, TEXT("No render thread."));
+			UE_LOG( LogStats, Warning, TEXT( "No render thread." ) );
 		}
-	}
 
+		LastHitchFrame = Frame;
+	}
 }
 
 FHUDGroupGameThreadRenderer& FHUDGroupGameThreadRenderer::Get()
@@ -422,19 +442,6 @@ FHUDGroupGameThreadRenderer& FHUDGroupGameThreadRenderer::Get()
 	return Singleton;
 }
 
-struct FGroupFilter : public IItemFiler
-{
-	TSet<FName> const& EnabledItems;
-	FGroupFilter(TSet<FName> const& InEnabledItems)
-		: EnabledItems(InEnabledItems)
-	{
-	}
-	virtual bool Keep(FStatMessage const& Item)
-	{
-		return EnabledItems.Contains(Item.NameAndInfo.GetRawName());
-	}
-};
-
 struct FInternalGroup
 {
 	/** Initialization constructor. */
diff --git a/Engine/Source/Runtime/Core/Public/Stats/StatsData.h b/Engine/Source/Runtime/Core/Public/Stats/StatsData.h
--- a/Engine/Source/Runtime/Core/Public/Stats/StatsData.h
+++ b/Engine/Source/Runtime/Core/Public/Stats/StatsData.h
@@ -12,7 +12,8 @@
 */
 enum
 {
-	STAT_FRAME_SLOP = 3 
+	STAT_FRAME_SLOP = 3,
+	MAX_STAT_LAG = 4,
 };
 
 /** Holds stats related constants. */
