#UE4:
* fixed some cvar names to be consistent

diff --git a/Engine/Source/Runtime/Core/Private/HAL/ConsoleManager.cpp b/Engine/Source/Runtime/Core/Private/HAL/ConsoleManager.cpp
--- a/Engine/Source/Runtime/Core/Private/HAL/ConsoleManager.cpp
+++ b/Engine/Source/Runtime/Core/Private/HAL/ConsoleManager.cpp
@@ -1941,7 +1941,7 @@ void CreateConsoleVariables()
 		TEXT("1: VSync is enabled."),
 		ECVF_Scalability | ECVF_RenderThreadSafe);
 
-	IConsoleManager::Get().RegisterConsoleVariable(TEXT("R.FinishCurrentFrame"),
+	IConsoleManager::Get().RegisterConsoleVariable(TEXT("r.FinishCurrentFrame"),
 		0,
 		TEXT("If on, the current frame will be forced to finish and render to the screen instead of being buffered.  This will improve latency, but slow down overall performance."),
 		ECVF_Scalability | ECVF_RenderThreadSafe);
diff --git a/Engine/Source/Runtime/Renderer/Private/RendererScene.cpp b/Engine/Source/Runtime/Renderer/Private/RendererScene.cpp
--- a/Engine/Source/Runtime/Renderer/Private/RendererScene.cpp
+++ b/Engine/Source/Runtime/Renderer/Private/RendererScene.cpp
@@ -1772,7 +1772,7 @@ static void DumpDrawListStats()
 }
 
 static FAutoConsoleCommand GDumpDrawListStatsCmd(
-	TEXT("R.DumpDrawListStats"),
+	TEXT("r.DumpDrawListStats"),
 	TEXT("Dumps static mesh draw list statistics for all scenes associated with ")
 	TEXT("world objects."),
 	FConsoleCommandDelegate::CreateStatic(&DumpDrawListStats)
