#ue4
- Added a method to show the platform-specific Achievements screen on iOS

diff --git a/Engine/Source/Runtime/Core/Private/IOS/IOSAppDelegate.cpp b/Engine/Source/Runtime/Core/Private/IOS/IOSAppDelegate.cpp
--- a/Engine/Source/Runtime/Core/Private/IOS/IOSAppDelegate.cpp
+++ b/Engine/Source/Runtime/Core/Private/IOS/IOSAppDelegate.cpp
@@ -504,10 +504,10 @@ void InstallSignalHandlers()
 	[self HideController : Controller Animated : YES];
 }
 
--(void)leaderboardViewControllerDidFinish:(GKLeaderboardViewController*)LeaderboardDisplay
+-(void)gameCenterViewControllerDidFinish:(GKGameCenterViewController*)GameCenterDisplay
 {
 	// close the view 
-	[self HideController : LeaderboardDisplay];
+	[self HideController : GameCenterDisplay];
 }
 
 /**
@@ -516,14 +516,27 @@ void InstallSignalHandlers()
 -(void)ShowLeaderboard:(NSString*)Category
 {
 	// create the leaderboard display object 
-	GKLeaderboardViewController* LeaderboardDisplay = [[[GKLeaderboardViewController alloc] init] autorelease];
-	LeaderboardDisplay.leaderboardDelegate = self;
+	GKGameCenterViewController* GameCenterDisplay = [[[GKGameCenterViewController alloc] init] autorelease];
+	GameCenterDisplay.viewState = GKGameCenterViewControllerStateLeaderboards;
+	GameCenterDisplay.leaderboardCategory = Category;
+	GameCenterDisplay.gameCenterDelegate = self;
 
-	// set which leaderboard to get 
-	LeaderboardDisplay.category = Category;
+	// show it 
+	[self ShowController : GameCenterDisplay];
+}
+
+/**
+ * Show the achievements interface (call from iOS main thread)
+ */
+-(void)ShowAchievements
+{
+	// create the leaderboard display object 
+	GKGameCenterViewController* GameCenterDisplay = [[[GKGameCenterViewController alloc] init] autorelease];
+	GameCenterDisplay.viewState = GKGameCenterViewControllerStateAchievements;
+	GameCenterDisplay.gameCenterDelegate = self;
 
 	// show it 
-	[self ShowController : LeaderboardDisplay];
+	[self ShowController : GameCenterDisplay];
 }
 
 /**
@@ -538,5 +551,15 @@ CORE_API bool IOSShowLeaderboardUI(const FString& CategoryName)
 	return true;
 }
 
+/**
+* Show the achievements interface (call from game thread)
+*/
+CORE_API bool IOSShowAchievementsUI()
+{
+	// route the function to iOS thread
+	[[IOSAppDelegate GetDelegate] performSelectorOnMainThread:@selector(ShowAchievements) withObject:nil waitUntilDone : NO];
+
+	return true;
+}
 
 @end
diff --git a/Engine/Source/Runtime/Core/Public/IOS/IOSAppDelegate.h b/Engine/Source/Runtime/Core/Public/IOS/IOSAppDelegate.h
--- a/Engine/Source/Runtime/Core/Public/IOS/IOSAppDelegate.h
+++ b/Engine/Source/Runtime/Core/Public/IOS/IOSAppDelegate.h
@@ -4,7 +4,7 @@
 
 #import <UIKit/UIKit.h>
 #import <AVFoundation/AVAudioSession.h>
-#import <GameKit/GKLeaderboardViewController.h>
+#import <GameKit/GKGameCenterViewController.h>
 
 @class EAGLView;
 @class IOSViewController;
@@ -27,7 +27,7 @@ namespace FAppEntry
 #if !UE_BUILD_SHIPPING
 	UIGestureRecognizerDelegate,
 #endif
-	GKLeaderboardViewControllerDelegate,
+	GKGameCenterControllerDelegate,
 UITextFieldDelegate, AVAudioSessionDelegate>
 
 /** Window object */
