#UE4 - Fix TAttribute constructor initializers out of order

diff --git a/Engine/Source/Runtime/Core/Public/Misc/Attribute.h b/Engine/Source/Runtime/Core/Public/Misc/Attribute.h
--- a/Engine/Source/Runtime/Core/Public/Misc/Attribute.h
+++ b/Engine/Source/Runtime/Core/Public/Misc/Attribute.h
@@ -28,8 +28,8 @@ public:
 	 * Default constructor
 	 */
 	TAttribute()
-		: Getter()
-		, Value()         // NOTE: Potentially uninitialized for atomics!!
+		: Value()         // NOTE: Potentially uninitialized for atomics!!
+		, Getter()
 		, bIsSet(false)
 	{
 	}
@@ -42,8 +42,8 @@ public:
 	 */
 	template< typename OtherType >
 	TAttribute( const OtherType& InInitialValue )
-		: Getter()
-		, Value( InInitialValue )
+		: Value( InInitialValue )
+		, Getter()
 		, bIsSet(true)
 	{
 	}
@@ -59,8 +59,8 @@ public:
 	 */
 	template< class SourceType >	
 	TAttribute( TSharedRef< SourceType > InUserObject, typename FGetter::template TSPMethodDelegate_Const< SourceType >::FMethodPtr InMethodPtr )
-		: Getter( FGetter::CreateSP( InUserObject, InMethodPtr ) )
-		, Value()
+		: Value()
+		, Getter( FGetter::CreateSP( InUserObject, InMethodPtr ) )
 		, bIsSet(true)
 	{
 	}
@@ -76,8 +76,8 @@ public:
 	 */
 	template< class SourceType >	
 	TAttribute( SourceType* InUserObject, typename FGetter::template TSPMethodDelegate_Const< SourceType >::FMethodPtr InMethodPtr )
-		: Getter( FGetter::CreateSP( InUserObject, InMethodPtr ) )
-		, Value()
+		: Value()
+		, Getter( FGetter::CreateSP( InUserObject, InMethodPtr ) )	
 		, bIsSet(true)
 	{
 	}
@@ -284,8 +284,8 @@ public:
 
 	/** Special explicit constructor for TAttribute::Create() */
 	TAttribute( typename FGetter::FStaticDelegate::FFuncPtr InFuncPtr, bool bExplicitConstructor )
-		: Getter( FGetter::CreateStatic( InFuncPtr ) )
-		, Value()
+		: Value()
+		, Getter( FGetter::CreateStatic( InFuncPtr ) )
 		, bIsSet( true )
 	{
 	}
@@ -293,8 +293,8 @@ public:
 	
 	/** Special explicit constructor for TAttribute::Create() */
 	TAttribute( const FGetter& InGetter, bool bExplicitConstructor )
-		: Getter( InGetter )
-		, Value()
+		: Value()
+		, Getter( InGetter )
 		, bIsSet( true )
 	{
 	}
@@ -303,7 +303,7 @@ public:
 private:
 
 	// We declare ourselves as a friend (templated using OtherType) so we can access members as needed
-    	template< class OtherType > friend class TAttribute;
+    template< class OtherType > friend class TAttribute;
 
 	/** Current value.  Mutable so that we can cache the value locally when using a bound Getter (allows const ref return value.) */
 	mutable ObjectType Value;
