Mac gamepad code refactor

diff --git a/Engine/Source/Runtime/Core/Private/Mac/HIDInputInterface.cpp b/Engine/Source/Runtime/Core/Private/Mac/HIDInputInterface.cpp
--- a/Engine/Source/Runtime/Core/Private/Mac/HIDInputInterface.cpp
+++ b/Engine/Source/Runtime/Core/Private/Mac/HIDInputInterface.cpp
@@ -3,42 +3,38 @@
 #include "CorePrivate.h"
 #include "HIDInputInterface.h"
 
-TSharedRef< HIDInputInterface > HIDInputInterface::Create( const TSharedRef< FGenericApplicationMessageHandler >& InMessageHandler )
+
+static int32 GetDevicePropertyAsInt32(IOHIDDeviceRef DeviceRef, CFStringRef Property)
+{
+	CFTypeRef Ref = IOHIDDeviceGetProperty(DeviceRef, Property);
+	if (Ref && CFGetTypeID(Ref) == CFNumberGetTypeID())
+	{
+		int32 Value = 0;
+		CFNumberGetValue((CFNumberRef)Ref, kCFNumberSInt32Type, &Value);
+		return Value;
+	}
+	return 0;
+}
+
+TSharedRef<HIDInputInterface> HIDInputInterface::Create(const TSharedRef< FGenericApplicationMessageHandler >& InMessageHandler)
 {
-	return MakeShareable( new HIDInputInterface( InMessageHandler ) );
+	return MakeShareable(new HIDInputInterface(InMessageHandler));
 }
 
-HIDInputInterface::HIDInputInterface( const TSharedRef< FGenericApplicationMessageHandler >& InMessageHandler )
-	: MessageHandler( InMessageHandler )
+HIDInputInterface::HIDInputInterface(const TSharedRef<FGenericApplicationMessageHandler>& InMessageHandler)
+:	MessageHandler(InMessageHandler)
 {
-	for ( int32 ControllerIndex=0; ControllerIndex < MAX_NUM_HIDINPUT_CONTROLLERS; ++ControllerIndex )
+	for (int32 ControllerIndex=0; ControllerIndex < MAX_NUM_HIDINPUT_CONTROLLERS; ++ControllerIndex)
 	{
 		FControllerState& ControllerState = ControllerStates[ControllerIndex];
-		FMemory::Memzero( &ControllerState, sizeof(FControllerState) );
-		
+		FMemory::Memzero(&ControllerState, sizeof(FControllerState));
+
 		ControllerState.ControllerId = ControllerIndex;
 	}
 
 	InitialButtonRepeatDelay = 0.2f;
 	ButtonRepeatDelay = 0.1f;
 
-	HIDToXboxControllerMapping[0] = -1;
-	HIDToXboxControllerMapping[1] = 0;		// A
-	HIDToXboxControllerMapping[2] = 1;		// B
-	HIDToXboxControllerMapping[3] = 2;		// X
-	HIDToXboxControllerMapping[4] = 3;		// Y
-	HIDToXboxControllerMapping[5] = 4;		// L1
-	HIDToXboxControllerMapping[6] = 5;		// R1
-	HIDToXboxControllerMapping[7] = 10;		// L2
-	HIDToXboxControllerMapping[8] = 11;		// R2
-	HIDToXboxControllerMapping[9] = 7;		// Back
-	HIDToXboxControllerMapping[10] = 6;		// Start
-	HIDToXboxControllerMapping[11] = 8;		// Left thumbstick
-	HIDToXboxControllerMapping[12] = 9;		// Right thumbstick
-	HIDToXboxControllerMapping[13] = -1;
-	HIDToXboxControllerMapping[14] = -1;
-	HIDToXboxControllerMapping[15] = -1;
-
 	Buttons[0] = EControllerButtons::FaceButtonBottom;
 	Buttons[1] = EControllerButtons::FaceButtonRight;
 	Buttons[2] = EControllerButtons::FaceButtonLeft;
@@ -55,180 +51,295 @@ HIDInputInterface::HIDInputInterface( const TSharedRef< FGenericApplicationMessa
 	Buttons[13] = EControllerButtons::DPadDown;
 	Buttons[14] = EControllerButtons::DPadLeft;
 	Buttons[15] = EControllerButtons::DPadRight;
+	Buttons[16] = EControllerButtons::LeftStickUp;
+	Buttons[17] = EControllerButtons::LeftStickDown;
+	Buttons[18] = EControllerButtons::LeftStickLeft;
+	Buttons[19] = EControllerButtons::LeftStickRight;
+	Buttons[20] = EControllerButtons::RightStickUp;
+	Buttons[21] = EControllerButtons::RightStickDown;
+	Buttons[22] = EControllerButtons::RightStickLeft;
+	Buttons[23] = EControllerButtons::RightStickRight;
 
 	// Init HID Manager
-	HIDManager = IOHIDManagerCreate( kCFAllocatorDefault, 0L );
-	if( !HIDManager )
+	HIDManager = IOHIDManagerCreate(kCFAllocatorDefault, 0L);
+	if (!HIDManager)
 	{
 		return;	// This will cause all subsequent calls to return information that nothing's connected
 	}
 
-	IOReturn Result = IOHIDManagerOpen( HIDManager, kIOHIDOptionsTypeNone );
-	if( Result != kIOReturnSuccess )
+	IOReturn Result = IOHIDManagerOpen(HIDManager, kIOHIDOptionsTypeNone);
+	if (Result != kIOReturnSuccess)
 	{
-		CFRelease( HIDManager );
+		CFRelease(HIDManager);
 		HIDManager = NULL;
 		return;
 	}
 
 	// Set HID Manager to detect devices of two distinct types: Gamepads and joysticks
-	CFMutableArrayRef MatchingArray = CFArrayCreateMutable( kCFAllocatorDefault, 0, &kCFTypeArrayCallBacks );
-	if( !MatchingArray )
+	CFMutableArrayRef MatchingArray = CFArrayCreateMutable(kCFAllocatorDefault, 0, &kCFTypeArrayCallBacks);
+	if (!MatchingArray)
 	{
-		CFRelease( HIDManager );
+		CFRelease(HIDManager);
 		HIDManager = NULL;
 		return;
 	}
 
-	CFDictionaryRef MatchingJoysticks = CreateDeviceMatchingDictionary( kHIDPage_GenericDesktop, kHIDUsage_GD_Joystick );
-	if( !MatchingJoysticks )
+	CFDictionaryRef MatchingJoysticks = CreateDeviceMatchingDictionary(kHIDPage_GenericDesktop, kHIDUsage_GD_Joystick);
+	if (!MatchingJoysticks)
 	{
-		CFRelease( MatchingArray );
-		CFRelease( HIDManager );
+		CFRelease(MatchingArray);
+		CFRelease(HIDManager);
 		HIDManager = 0;
 		return;
 	}
 
-	CFArrayAppendValue( MatchingArray, MatchingJoysticks );
-	CFRelease( MatchingJoysticks );
+	CFArrayAppendValue(MatchingArray, MatchingJoysticks);
+	CFRelease(MatchingJoysticks);
 
-	CFDictionaryRef MatchingGamepads = CreateDeviceMatchingDictionary( kHIDPage_GenericDesktop, kHIDUsage_GD_GamePad );
-	if( !MatchingGamepads )
+	CFDictionaryRef MatchingGamepads = CreateDeviceMatchingDictionary(kHIDPage_GenericDesktop, kHIDUsage_GD_GamePad);
+	if (!MatchingGamepads)
 	{
-		CFRelease( MatchingArray );
-		CFRelease( HIDManager );
+		CFRelease(MatchingArray);
+		CFRelease(HIDManager);
 		HIDManager = 0;
 		return;
 	}
 
-	CFArrayAppendValue( MatchingArray, MatchingGamepads );
-	CFRelease( MatchingGamepads );
+	CFArrayAppendValue(MatchingArray, MatchingGamepads);
+	CFRelease(MatchingGamepads);
 
-	IOHIDManagerSetDeviceMatchingMultiple( HIDManager, MatchingArray );
-	CFRelease( MatchingArray );
+	IOHIDManagerSetDeviceMatchingMultiple(HIDManager, MatchingArray);
+	CFRelease(MatchingArray);
 
 	// Setup HID Manager's add/remove devices callbacks
-	IOHIDManagerRegisterDeviceMatchingCallback( HIDManager, HIDDeviceMatchingCallback, this );
-	IOHIDManagerRegisterDeviceRemovalCallback( HIDManager, HIDDeviceRemovalCallback, this );
+	IOHIDManagerRegisterDeviceMatchingCallback(HIDManager, HIDDeviceMatchingCallback, this);
+	IOHIDManagerRegisterDeviceRemovalCallback(HIDManager, HIDDeviceRemovalCallback, this);
 
 	// Add HID Manager to run loop
-	IOHIDManagerScheduleWithRunLoop( HIDManager, CFRunLoopGetCurrent(), kCFRunLoopDefaultMode );
+	IOHIDManagerScheduleWithRunLoop(HIDManager, CFRunLoopGetCurrent(), kCFRunLoopDefaultMode);
 }
 
-void HIDInputInterface::Tick( float DeltaTime )
+void HIDInputInterface::FHIDDeviceInfo::SetupMappings()
 {
+	FMemory::Memset(ButtonsMapping, -1, sizeof(ButtonsMapping));
+
+	const int32 VendorID = GetDevicePropertyAsInt32(DeviceRef, CFSTR(kIOHIDVendorIDKey));
+	const int32 ProductID = GetDevicePropertyAsInt32(DeviceRef, CFSTR(kIOHIDProductIDKey));
+
+	if (VendorID == 0x54c && ProductID == 0x268)
+	{
+		// PlayStation 3 Controller
+		ButtonsMapping[1]	= 7;	// Select		->	Back
+		ButtonsMapping[2]	= 8;	// L3			->	Left Thumbstick
+		ButtonsMapping[3]	= 9;	// R3			->	Right Thumbstick
+		ButtonsMapping[4]	= 6;	// Start		->	Start
+		ButtonsMapping[5]	= 12;	// DPad Up		->	DPad Up
+		ButtonsMapping[6]	= 15;	// DPad Right	->	DPad Right
+		ButtonsMapping[7]	= 13;	// DPad Down	->	DPad Down
+		ButtonsMapping[8]	= 14;	// DPad Left	->	DPad Left
+		ButtonsMapping[9]	= 10;	// L2			->	Left Trigger
+		ButtonsMapping[10]	= 11;	// R2			->	Right Trigger
+		ButtonsMapping[11]	= 4;	// L1			->	Left Shoulder
+		ButtonsMapping[12]	= 5;	// R1			->	Right Shoulder
+		ButtonsMapping[13]	= 3;	// Triangle		->	Y
+		ButtonsMapping[14]	= 1;	// Circle		->	B
+		ButtonsMapping[15]	= 0;	// Cross		->	A
+		ButtonsMapping[16]	= 2;	// Square		->	X
+
+		LeftAnalogXMapping = kHIDUsage_GD_X;
+		LeftAnalogYMapping = kHIDUsage_GD_Y;
+		LeftTriggerAnalogMapping = kHIDUsage_GD_Rx;
+		RightAnalogXMapping = kHIDUsage_GD_Z;
+		RightAnalogYMapping = kHIDUsage_GD_Rz;
+		RightTriggerAnalogMapping = kHIDUsage_GD_Ry;
+	}
+	else if (VendorID == 0x45e && ProductID == 0x28e)
+	{
+		// Xbox 360 Controller
+		ButtonsMapping[1]	= 0;	// A
+		ButtonsMapping[2]	= 1;	// B
+		ButtonsMapping[3]	= 2;	// X
+		ButtonsMapping[4]	= 3;	// Y
+		ButtonsMapping[5]	= 4;	// Left Shoulder
+		ButtonsMapping[6]	= 5;	// Right Shoulder
+		ButtonsMapping[7]	= 8;	// Left Thumbstick
+		ButtonsMapping[8]	= 9;	// Right Thumbstick
+		ButtonsMapping[9]	= 6;	// Start
+		ButtonsMapping[10]	= 7;	// Back
+		ButtonsMapping[12]	= 12;	// DPad Up
+		ButtonsMapping[13]	= 13;	// DPad Down
+		ButtonsMapping[14]	= 14;	// DPad Left
+		ButtonsMapping[15]	= 15;	// DPad Right
+
+		LeftAnalogXMapping = kHIDUsage_GD_X;
+		LeftAnalogYMapping = kHIDUsage_GD_Y;
+		LeftTriggerAnalogMapping = kHIDUsage_GD_Z;
+		RightAnalogXMapping = kHIDUsage_GD_Rx;
+		RightAnalogYMapping = kHIDUsage_GD_Ry;
+		RightTriggerAnalogMapping = kHIDUsage_GD_Rz;
+	}
+	else
+	{
+		// Generic (based on Logitech RumblePad 2)
+		ButtonsMapping[1]	= 2;	// X
+		ButtonsMapping[2]	= 0;	// A
+		ButtonsMapping[3]	= 1;	// B
+		ButtonsMapping[4]	= 3;	// Y
+		ButtonsMapping[5]	= 4;	// Left Shoulder
+		ButtonsMapping[6]	= 5;	// Right Shoulder
+		ButtonsMapping[7]	= 10;	// Left Trigger
+		ButtonsMapping[8]	= 11;	// Right Trigger
+		ButtonsMapping[9]	= 7;	// Back
+		ButtonsMapping[10]	= 6;	// Start
+		ButtonsMapping[11]	= 8;	// Left Thumbstick
+		ButtonsMapping[12]	= 9;	// Right Thumbstick
+
+		LeftAnalogXMapping = kHIDUsage_GD_X;
+		LeftAnalogYMapping = kHIDUsage_GD_Y;
+		LeftTriggerAnalogMapping = kHIDUsage_GD_Rx;
+		RightAnalogXMapping = kHIDUsage_GD_Z;
+		RightAnalogYMapping = kHIDUsage_GD_Rz;
+		RightTriggerAnalogMapping = kHIDUsage_GD_Ry;
+	}
 }
 
-void HIDInputInterface::SendControllerEvents()
+void HIDInputInterface::OnNewHIDController(IOReturn Result, IOHIDDeviceRef DeviceRef)
 {
-	for ( int32 ControllerIndex=0; ControllerIndex < MAX_NUM_HIDINPUT_CONTROLLERS; ++ControllerIndex )
-	{
-		FControllerState& ControllerState = ControllerStates[ControllerIndex];
+	int32 ControllerIndex = -1;
 
-		if( ControllerState.Device.DeviceRef )
+	for (int32 Index = 0; Index < MAX_NUM_HIDINPUT_CONTROLLERS; ++Index)
+	{
+		if (ControllerStates[Index].Device.DeviceRef == NULL)
 		{
-			bool CurrentStates[MAX_NUM_CONTROLLER_BUTTONS] = {0};
-
-			for( int ButtonIndex = 0; ButtonIndex < ControllerState.Device.Buttons.Num(); ++ButtonIndex )
-			{
-				FHIDButtonInfo& Button = ControllerState.Device.Buttons[ButtonIndex];
-				IOHIDValueRef Value;
-				IOReturn Result = IOHIDDeviceGetValue( ControllerState.Device.DeviceRef, Button.Element, &Value );
+			ControllerIndex = Index;
+			break;
+		}
+	}
 
-				if( Result == kIOReturnSuccess )
-				{
-					CurrentStates[HIDToXboxControllerMapping[Button.Usage]] = IOHIDValueGetIntegerValue( Value ) > 0;
-				}
-			}
+	if (ControllerIndex != -1 && (Result = IOHIDDeviceOpen(DeviceRef, kIOHIDOptionsTypeNone)) == kIOReturnSuccess)
+	{
+		CFArrayRef ElementsArray = IOHIDDeviceCopyMatchingElements(DeviceRef, NULL, 0);
+		if (ElementsArray)
+		{
+			FHIDDeviceInfo& DeviceInfo = ControllerStates[ControllerIndex].Device;
+			DeviceInfo.DeviceRef = DeviceRef;
+			DeviceInfo.Elements.Empty();
+			DeviceInfo.SetupMappings();
 
-			if( ControllerState.Device.DPad.Element )
+			CFIndex ElementsCount = CFArrayGetCount(ElementsArray);
+			for (CFIndex ElementIndex = 0; ElementIndex < ElementsCount; ++ElementIndex)
 			{
-				IOHIDValueRef Value;
-				IOReturn Result = IOHIDDeviceGetValue( ControllerState.Device.DeviceRef, ControllerState.Device.DPad.Element, &Value );
-					
-				if( Result == kIOReturnSuccess )
+				FHIDElementInfo Element;
+				Element.ElementRef = (IOHIDElementRef)CFArrayGetValueAtIndex(ElementsArray, ElementIndex);
+				Element.Type = IOHIDElementGetType(Element.ElementRef);
+				Element.UsagePage = IOHIDElementGetUsagePage(Element.ElementRef);
+				Element.Usage = IOHIDElementGetUsage(Element.ElementRef);
+				Element.MinValue = IOHIDElementGetLogicalMin(Element.ElementRef);
+				Element.MaxValue = IOHIDElementGetLogicalMax(Element.ElementRef);
+
+				if ((Element.Type == kIOHIDElementTypeInput_Button && Element.UsagePage == kHIDPage_Button && Element.Usage < MAX_NUM_CONTROLLER_BUTTONS)
+					|| ((Element.Type == kIOHIDElementTypeInput_Misc || Element.Type == kIOHIDElementTypeInput_Axis) && Element.UsagePage == kHIDPage_GenericDesktop))
 				{
-					int16 NewValue = IOHIDValueGetIntegerValue( Value );
-						
-					switch( NewValue )
-					{
-						case 0: CurrentStates[12] = true;								break; // Up
-						case 1: CurrentStates[12] = true;	CurrentStates[13] = true;	break; // Up + right
-						case 2: CurrentStates[13] = true;								break; // Right
-						case 3: CurrentStates[13] = true;	CurrentStates[14] = true;	break; // Right + down
-						case 4: CurrentStates[14] = true;								break; // Down
-						case 5: CurrentStates[14] = true;	CurrentStates[15] = true;	break; // Down + left
-						case 6: CurrentStates[15] = true;								break; // Left
-						case 7: CurrentStates[15] = true;	CurrentStates[12] = true;	break; // Left + up
-					}
+					DeviceInfo.Elements.Add(Element);
 				}
 			}
 
-			for( int AxisIndex = 0; AxisIndex < ControllerState.Device.Axes.Num(); ++AxisIndex )
-			{
-				FHIDAxisInfo& Axis = ControllerState.Device.Axes[AxisIndex];
+			CFRelease(ElementsArray);
+		}
+		else
+		{
+			IOHIDDeviceClose(DeviceRef, kIOHIDOptionsTypeNone);
+		}
+	}
+}
 
-				IOHIDValueRef Value;
-				IOReturn Result = IOHIDDeviceGetValue( ControllerState.Device.DeviceRef, Axis.Element, &Value );
+void HIDInputInterface::SendControllerEvents()
+{
+	for (int32 ControllerIndex = 0; ControllerIndex < MAX_NUM_HIDINPUT_CONTROLLERS; ++ControllerIndex)
+	{
+		FControllerState& ControllerState = ControllerStates[ControllerIndex];
 
-				if( Result == kIOReturnSuccess )
-				{
-					int16 NewValue = IOHIDValueGetIntegerValue( Value );
+		if( ControllerState.Device.DeviceRef )
+		{
+			bool CurrentButtonStates[MAX_NUM_CONTROLLER_BUTTONS] = {0};
 
-					int16 HalfRange = ( Axis.MaxValue + Axis.MinValue ) / 2;
-					float FloatValue = FMath::Clamp( ( ( NewValue - HalfRange ) / ( float )Axis.MaxValue ) * 2.0f, -1.0f, 1.0f );
+			for (int Index = 0; Index < ControllerState.Device.Elements.Num(); ++Index)
+			{
+				FHIDElementInfo& Element = ControllerState.Device.Elements[Index];
 
-					switch( Axis.Usage )
+				IOHIDValueRef ValueRef;
+				IOReturn Result = IOHIDDeviceGetValue(ControllerState.Device.DeviceRef, Element.ElementRef, &ValueRef);
+				if (Result == kIOReturnSuccess && IOHIDValueGetLength(ValueRef) <= sizeof(int32))
+				{
+					const int32 NewValue = IOHIDValueGetIntegerValue(ValueRef);
+					if (Element.UsagePage == kHIDPage_Button)
 					{
-						case kHIDUsage_GD_X:
-							if( ControllerState.LeftXAnalog != FloatValue )
-							{
-								MessageHandler->OnControllerAnalog( EControllerButtons::LeftAnalogX, ControllerState.ControllerId, FloatValue );
-								ControllerState.LeftXAnalog = FloatValue;
-							}
-							break;
-
-						case kHIDUsage_GD_Y:
-							if( ControllerState.LeftYAnalog != FloatValue )
-							{
-								MessageHandler->OnControllerAnalog( EControllerButtons::LeftAnalogY, ControllerState.ControllerId, FloatValue * -1 );
-								ControllerState.LeftYAnalog = FloatValue;
-							}
-							break;
-
-						case kHIDUsage_GD_Z:
-							if( ControllerState.RightXAnalog != FloatValue )
-							{
-								MessageHandler->OnControllerAnalog( EControllerButtons::RightAnalogX, ControllerState.ControllerId, FloatValue );
-								ControllerState.RightXAnalog = FloatValue;
-							}
-							CurrentStates[10] = FloatValue > 0.01f;
-							break;
-
-						case kHIDUsage_GD_Rz:
-							if( ControllerState.RightYAnalog != FloatValue )
-							{
-								MessageHandler->OnControllerAnalog( EControllerButtons::RightAnalogY, ControllerState.ControllerId, FloatValue * -1 );
-								ControllerState.RightYAnalog = FloatValue;
-							}
-							CurrentStates[11] = FloatValue > 0.01f;
-							break;
-
-						case kHIDUsage_GD_Rx:
-							if( ControllerState.LeftTriggerAnalog != FloatValue )
-							{
-								MessageHandler->OnControllerAnalog( EControllerButtons::LeftTriggerAnalog, ControllerState.ControllerId, FloatValue );
-								ControllerState.LeftTriggerAnalog = FloatValue;
-							}
-							break;
-
-						case kHIDUsage_GD_Ry:
-							if( ControllerState.RightTriggerAnalog != FloatValue )
+						const int8 MappedIndex = ControllerState.Device.ButtonsMapping[Element.Usage];
+						if (MappedIndex != -1)
+						{
+							CurrentButtonStates[MappedIndex] = NewValue > 0;
+						}
+					}
+					else
+					{
+						const bool bIsTrigger = Element.Usage == ControllerState.Device.LeftTriggerAnalogMapping || Element.Usage == ControllerState.Device.RightTriggerAnalogMapping;
+						const float Percentage = FMath::GetRangePct(Element.MinValue, Element.MaxValue, NewValue);
+						const float FloatValue = bIsTrigger ? Percentage : Percentage * 2.0f - 1.0f;
+
+						if (Element.Usage == ControllerState.Device.LeftAnalogXMapping && ControllerState.LeftAnalogX != NewValue)
+						{
+							MessageHandler->OnControllerAnalog(EControllerButtons::LeftAnalogX, ControllerState.ControllerId, FloatValue);
+							CurrentButtonStates[18] = FloatValue < -0.2f;	// EControllerButtons::LeftStickLeft
+							CurrentButtonStates[19] = FloatValue > 0.2f;	// EControllerButtons::LeftStickRight
+							ControllerState.LeftAnalogX = NewValue;
+						}
+						else if (Element.Usage == ControllerState.Device.LeftAnalogYMapping && ControllerState.LeftAnalogY != NewValue)
+						{
+							MessageHandler->OnControllerAnalog(EControllerButtons::LeftAnalogY, ControllerState.ControllerId, -FloatValue);
+							CurrentButtonStates[16] = FloatValue < -0.2f;	// EControllerButtons::LeftStickUp
+							CurrentButtonStates[17] = FloatValue > 0.2f;	// EControllerButtons::LeftStickDown
+							ControllerState.LeftAnalogY = NewValue;
+						}
+						else if (Element.Usage == ControllerState.Device.RightAnalogXMapping && ControllerState.RightAnalogX != NewValue)
+						{
+							MessageHandler->OnControllerAnalog(EControllerButtons::RightAnalogX, ControllerState.ControllerId, FloatValue);
+							CurrentButtonStates[22] = FloatValue < -0.2f;	// EControllerButtons::RightStickLeft
+							CurrentButtonStates[23] = FloatValue > 0.2f;	// EControllerButtons::RightStickDown
+							ControllerState.RightAnalogX = NewValue;
+						}
+						else if (Element.Usage == ControllerState.Device.RightAnalogYMapping && ControllerState.RightAnalogY != NewValue)
+						{
+							MessageHandler->OnControllerAnalog(EControllerButtons::RightAnalogY, ControllerState.ControllerId, -FloatValue);
+							CurrentButtonStates[20] = FloatValue < -0.2f;	// EControllerButtons::RightStickUp
+							CurrentButtonStates[21] = FloatValue > 0.2f;	// EControllerButtons::RightStickRight
+							ControllerState.RightAnalogY = NewValue;
+						}
+						else if (Element.Usage == ControllerState.Device.LeftTriggerAnalogMapping && ControllerState.LeftTriggerAnalog != NewValue)
+						{
+							MessageHandler->OnControllerAnalog(EControllerButtons::LeftTriggerAnalog, ControllerState.ControllerId, FloatValue);
+							CurrentButtonStates[10] = FloatValue > 0.01f;
+							ControllerState.LeftTriggerAnalog = NewValue;
+						}
+						else if (Element.Usage == ControllerState.Device.RightTriggerAnalogMapping && ControllerState.RightTriggerAnalog != NewValue)
+						{
+							MessageHandler->OnControllerAnalog(EControllerButtons::RightTriggerAnalog, ControllerState.ControllerId, FloatValue);
+							CurrentButtonStates[11] = FloatValue > 0.01f;
+							ControllerState.RightTriggerAnalog = NewValue;
+						}
+						else if (Element.Usage == kHIDUsage_GD_Hatswitch)
+						{
+							switch (NewValue)
 							{
-								MessageHandler->OnControllerAnalog( EControllerButtons::RightTriggerAnalog, ControllerState.ControllerId, FloatValue );
-								ControllerState.RightTriggerAnalog = FloatValue;
+								case 0: CurrentButtonStates[12] = true;									break; // Up
+								case 1: CurrentButtonStates[12] = true;	CurrentButtonStates[13] = true;	break; // Up + right
+								case 2: CurrentButtonStates[13] = true;									break; // Right
+								case 3: CurrentButtonStates[13] = true;	CurrentButtonStates[14] = true;	break; // Right + down
+								case 4: CurrentButtonStates[14] = true;									break; // Down
+								case 5: CurrentButtonStates[14] = true;	CurrentButtonStates[15] = true;	break; // Down + left
+								case 6: CurrentButtonStates[15] = true;									break; // Left
+								case 7: CurrentButtonStates[15] = true;	CurrentButtonStates[12] = true;	break; // Left + up
 							}
-							break;
+						}
 					}
 				}
 			}
@@ -238,177 +349,81 @@ void HIDInputInterface::SendControllerEvents()
 			// For each button check against the previous state and send the correct message if any
 			for (int32 ButtonIndex = 0; ButtonIndex < MAX_NUM_CONTROLLER_BUTTONS; ++ButtonIndex)
 			{
-				if (CurrentStates[ButtonIndex] != ControllerState.ButtonStates[ButtonIndex])
+				if (CurrentButtonStates[ButtonIndex] != ControllerState.ButtonStates[ButtonIndex])
 				{
-					if( CurrentStates[ButtonIndex] )
+					if (CurrentButtonStates[ButtonIndex])
 					{
-						MessageHandler->OnControllerButtonPressed( Buttons[ButtonIndex], ControllerState.ControllerId, false );
+						MessageHandler->OnControllerButtonPressed(Buttons[ButtonIndex], ControllerState.ControllerId, false);
 					}
 					else
 					{
-						MessageHandler->OnControllerButtonReleased( Buttons[ButtonIndex], ControllerState.ControllerId, false );
+						MessageHandler->OnControllerButtonReleased(Buttons[ButtonIndex], ControllerState.ControllerId, false);
 					}
 
-					if ( CurrentStates[ButtonIndex] != 0 )
+					if (CurrentButtonStates[ButtonIndex] != 0)
 					{
 						// this button was pressed - set the button's NextRepeatTime to the InitialButtonRepeatDelay
 						ControllerState.NextRepeatTime[ButtonIndex] = CurrentTime + InitialButtonRepeatDelay;
 					}
 				}
-				else if ( CurrentStates[ButtonIndex] != 0 && ControllerState.NextRepeatTime[ButtonIndex] <= CurrentTime )
+				else if (CurrentButtonStates[ButtonIndex] != 0 && ControllerState.NextRepeatTime[ButtonIndex] <= CurrentTime)
 				{
-					MessageHandler->OnControllerButtonPressed( Buttons[ButtonIndex], ControllerState.ControllerId, true );
+					MessageHandler->OnControllerButtonPressed(Buttons[ButtonIndex], ControllerState.ControllerId, true);
 
 					// set the button's NextRepeatTime to the ButtonRepeatDelay
 					ControllerState.NextRepeatTime[ButtonIndex] = CurrentTime + ButtonRepeatDelay;
 				}
 
 				// Update the state for next time
-				ControllerState.ButtonStates[ButtonIndex] = CurrentStates[ButtonIndex];
+				ControllerState.ButtonStates[ButtonIndex] = CurrentButtonStates[ButtonIndex];
 			}	
 		}
 	}
 }
 
-CFMutableDictionaryRef HIDInputInterface::CreateDeviceMatchingDictionary( UInt32 UsagePage, UInt32 Usage )
+CFMutableDictionaryRef HIDInputInterface::CreateDeviceMatchingDictionary(UInt32 UsagePage, UInt32 Usage)
 {
 	// Create a dictionary to add usage page/usages to
-	CFMutableDictionaryRef Dict = CFDictionaryCreateMutable( kCFAllocatorDefault, 0, &kCFTypeDictionaryKeyCallBacks, &kCFTypeDictionaryValueCallBacks );
-	if( !Dict )
-	{
-		return NULL;
-	}
-
-	// Add key for device type to refine the matching dictionary.
-	CFNumberRef PageCFNumberRef = CFNumberCreate( kCFAllocatorDefault, kCFNumberIntType, &UsagePage );
-
-	if( !PageCFNumberRef )
+	CFMutableDictionaryRef Dict = CFDictionaryCreateMutable(kCFAllocatorDefault, 0, &kCFTypeDictionaryKeyCallBacks, &kCFTypeDictionaryValueCallBacks);
+	if (Dict)
 	{
-		CFRelease( Dict );
-		return NULL;
-	}
-
-	CFDictionarySetValue( Dict, CFSTR( kIOHIDDeviceUsagePageKey ), PageCFNumberRef );
-	CFRelease( PageCFNumberRef );
+		// Add key for device type to refine the matching dictionary.
+		CFNumberRef PageCFNumberRef = CFNumberCreate(kCFAllocatorDefault, kCFNumberIntType, &UsagePage);
+		if (PageCFNumberRef)
+		{
+			CFDictionarySetValue(Dict, CFSTR(kIOHIDDeviceUsagePageKey), PageCFNumberRef);
+			CFRelease(PageCFNumberRef);
 
-	// Note: the usage is only valid if the usage page is also defined
-	CFNumberRef UsageCFNumberRef = CFNumberCreate( kCFAllocatorDefault, kCFNumberIntType, &Usage );
+			// Note: the usage is only valid if the usage page is also defined
+			CFNumberRef UsageCFNumberRef = CFNumberCreate(kCFAllocatorDefault, kCFNumberIntType, &Usage);
+			if (UsageCFNumberRef)
+			{
+				CFDictionarySetValue(Dict, CFSTR(kIOHIDDeviceUsageKey), UsageCFNumberRef);
+				CFRelease(UsageCFNumberRef);
 
-	if( !UsageCFNumberRef )
-	{
-		CFRelease( Dict );
-		return NULL;
+				return Dict;
+			}
+		}
 	}
 
-	CFDictionarySetValue( Dict, CFSTR( kIOHIDDeviceUsageKey ), UsageCFNumberRef );
-	CFRelease( UsageCFNumberRef );
-
-    return Dict;
+	return NULL;
 }
 
-void HIDInputInterface::HIDDeviceMatchingCallback( void* Context, IOReturn Result, void* Sender, IOHIDDeviceRef DeviceRef )
+void HIDInputInterface::HIDDeviceMatchingCallback(void* Context, IOReturn Result, void* Sender, IOHIDDeviceRef DeviceRef)
 {
-	HIDInputInterface* HIDInput = ( HIDInputInterface* )Context;
-	HIDInput->OnNewHIDController( Result, DeviceRef );
+	HIDInputInterface* HIDInput = (HIDInputInterface*)Context;
+	HIDInput->OnNewHIDController(Result, DeviceRef);
 }
 
-void HIDInputInterface::HIDDeviceRemovalCallback( void* Context, IOReturn Result, void* Sender, IOHIDDeviceRef DeviceRef )
+void HIDInputInterface::HIDDeviceRemovalCallback(void* Context, IOReturn Result, void* Sender, IOHIDDeviceRef DeviceRef)
 {
-	HIDInputInterface* HIDInput = ( HIDInputInterface* )Context;
-	for( int32 Index = 0; Index < MAX_NUM_HIDINPUT_CONTROLLERS; ++Index )
+	HIDInputInterface* HIDInput = (HIDInputInterface*)Context;
+	for (int32 Index = 0; Index < MAX_NUM_HIDINPUT_CONTROLLERS; ++Index)
 	{
-		if( HIDInput->ControllerStates[Index].Device.DeviceRef == DeviceRef )
+		if(HIDInput->ControllerStates[Index].Device.DeviceRef == DeviceRef)
 		{
 			HIDInput->ControllerStates[Index].Device.DeviceRef = NULL;
 			break;
 		}
 	}
 }
-
-void HIDInputInterface::OnNewHIDController( IOReturn Result, IOHIDDeviceRef DeviceRef )
-{
-	int32 ControllerIndex = -1;
-
-	for( int32 Index = 0; Index < MAX_NUM_HIDINPUT_CONTROLLERS; ++Index )
-	{
-		if( ControllerStates[Index].Device.DeviceRef == NULL )
-		{
-			ControllerIndex = Index;
-			break;
-		}
-	}
-
-	if( ControllerIndex == -1 )
-	{
-		return; // There's no place for a new device. Ignoring it.
-	}
-
-	Result = IOHIDDeviceOpen( DeviceRef, kIOHIDOptionsTypeNone );
-
-	if( Result != kIOReturnSuccess )
-	{
-		return;	// couldn't open the device
-	}
-
-	FHIDDeviceInfo& DeviceInfo = ControllerStates[ControllerIndex].Device;
-
-	// Get all elements from the device
-	CFArrayRef Elements = IOHIDDeviceCopyMatchingElements( DeviceRef, NULL, 0 );
-	if( !Elements )
-	{
-		IOHIDDeviceClose( DeviceRef, kIOHIDOptionsTypeNone );
-		return;
-	}
-
-	CFIndex ElementsCount = CFArrayGetCount( Elements );
-
-	DeviceInfo.DeviceRef = DeviceRef;
-	DeviceInfo.Buttons.Empty();
-	DeviceInfo.Axes.Empty();
-	DeviceInfo.DPad.Element = NULL;
-
-	for( CFIndex ElementIndex = 0; ElementIndex < ElementsCount; ++ElementIndex )
-	{
-		IOHIDElementRef Element = (IOHIDElementRef)CFArrayGetValueAtIndex( Elements, ElementIndex );
-
-		IOHIDElementType ElementType = IOHIDElementGetType( Element );
-		uint32 UsagePage = IOHIDElementGetUsagePage( Element );
-		uint32 Usage = IOHIDElementGetUsage( Element );
-
-		if( ( ElementType == kIOHIDElementTypeInput_Button ) && ( UsagePage == kHIDPage_Button ) && ( Usage < MAX_NUM_CONTROLLER_BUTTONS ) && ( HIDToXboxControllerMapping[Usage] != -1 ) )
-		{
-			FHIDButtonInfo Button;
-			Button.Element = Element;
-			Button.Usage = Usage;
-			DeviceInfo.Buttons.Push( Button );
-		}
-		else if( ( ( ElementType == kIOHIDElementTypeInput_Axis ) || ( ElementType == kIOHIDElementTypeInput_Misc ) ) && UsagePage == kHIDPage_GenericDesktop
-				&& Usage != kHIDUsage_GD_Hatswitch && Usage != kHIDUsage_GD_Pointer )
-		{
-			IOHIDValueRef Value;
-			IOHIDDeviceGetValue( DeviceRef, Element, &Value );
-
-			FHIDAxisInfo Axis;
-			Axis.Element = Element;
-			Axis.Usage = Usage;
-			Axis.MinValue = IOHIDElementGetPhysicalMin( Element );
-			Axis.MaxValue = IOHIDElementGetPhysicalMax( Element );
-
-			DeviceInfo.Axes.Push( Axis );
-		}
-		else if( Usage == kHIDUsage_GD_Hatswitch && DeviceInfo.DPad.Element == NULL )
-		{
-			DeviceInfo.DPad.Element = Element;
-			DeviceInfo.DPad.Usage = Usage;
-			DeviceInfo.DPad.MinValue = IOHIDElementGetPhysicalMin( Element );
-			DeviceInfo.DPad.MaxValue = IOHIDElementGetPhysicalMax( Element );
-		}
-	}
-
-	CFRelease( Elements );
-}
-
-void HIDInputInterface::SetMessageHandler( const TSharedRef< FGenericApplicationMessageHandler >& InMessageHandler )
-{
-	MessageHandler = InMessageHandler;
-}
diff --git a/Engine/Source/Runtime/Core/Private/Mac/HIDInputInterface.h b/Engine/Source/Runtime/Core/Private/Mac/HIDInputInterface.h
--- a/Engine/Source/Runtime/Core/Private/Mac/HIDInputInterface.h
+++ b/Engine/Source/Runtime/Core/Private/Mac/HIDInputInterface.h
@@ -1,4 +1,5 @@
 // Copyright 1998-2014 Epic Games, Inc. All Rights Reserved.
+
 #pragma once
 
 #include <IOKit/hid/IOHIDLib.h>
@@ -6,8 +7,7 @@
 #define MAX_NUM_HIDINPUT_CONTROLLERS 4
 
 /** Max number of controller buttons.  Must be < 256*/
-#define MAX_NUM_CONTROLLER_BUTTONS 16
-
+#define MAX_NUM_CONTROLLER_BUTTONS 24
 
 /**
  * Interface class for HID Input devices
@@ -16,48 +16,50 @@ class HIDInputInterface
 {
 public:
 
-	static TSharedRef< HIDInputInterface > Create( const TSharedRef< FGenericApplicationMessageHandler >& InMessageHandler );
-
-	void Tick( float DeltaTime );
+	static TSharedRef<HIDInputInterface> Create(const TSharedRef<FGenericApplicationMessageHandler>& InMessageHandler);
 
 	void SendControllerEvents();
 
-	void SetMessageHandler( const TSharedRef< FGenericApplicationMessageHandler >& InMessageHandler );
+	void SetMessageHandler(const TSharedRef<FGenericApplicationMessageHandler>& InMessageHandler)
+	{
+		MessageHandler = InMessageHandler;
+	}
 
 private:
 
-	HIDInputInterface( const TSharedRef< FGenericApplicationMessageHandler >& MessageHandler );
-
-	static CFMutableDictionaryRef CreateDeviceMatchingDictionary( UInt32 UsagePage, UInt32 Usage );
+	HIDInputInterface(const TSharedRef<FGenericApplicationMessageHandler>& MessageHandler);
 
-	static void HIDDeviceMatchingCallback( void* Context, IOReturn Result, void* Sender, IOHIDDeviceRef DeviceRef );
+	void OnNewHIDController(IOReturn Result, IOHIDDeviceRef DeviceRef);
 
-	static void HIDDeviceRemovalCallback( void* Context, IOReturn Result, void* Sender, IOHIDDeviceRef DeviceRef );
-
-	void OnNewHIDController( IOReturn Result, IOHIDDeviceRef DeviceRef );
+	static CFMutableDictionaryRef CreateDeviceMatchingDictionary(UInt32 UsagePage, UInt32 Usage);
+	static void HIDDeviceMatchingCallback(void* Context, IOReturn Result, void* Sender, IOHIDDeviceRef DeviceRef);
+	static void HIDDeviceRemovalCallback(void* Context, IOReturn Result, void* Sender, IOHIDDeviceRef DeviceRef);
 
 private:
 
-	struct FHIDButtonInfo
+	struct FHIDElementInfo
 	{
-		IOHIDElementRef Element;
-		uint32 Usage;
-	};
-
-	struct FHIDAxisInfo
-	{
-		IOHIDElementRef Element;
-		uint32 Usage;
-		int16 MinValue;
-		int16 MaxValue;
+		IOHIDElementRef ElementRef;
+		IOHIDElementType Type;
+		uint16 UsagePage;
+		uint16 Usage;
+		int32 MinValue;
+		int32 MaxValue;
 	};
 
 	struct FHIDDeviceInfo
 	{
 		IOHIDDeviceRef DeviceRef;
-		TArray<FHIDButtonInfo> Buttons;
-		TArray<FHIDAxisInfo> Axes;
-		FHIDAxisInfo DPad;
+		TArray<FHIDElementInfo> Elements;
+		int8 ButtonsMapping[MAX_NUM_CONTROLLER_BUTTONS];
+		uint16 LeftAnalogXMapping;
+		uint16 LeftAnalogYMapping;
+		uint16 LeftTriggerAnalogMapping;
+		uint16 RightAnalogXMapping;
+		uint16 RightAnalogYMapping;
+		uint16 RightTriggerAnalogMapping;
+
+		void SetupMappings();
 	};
 
 	struct FControllerState
@@ -69,22 +71,22 @@ private:
 		double NextRepeatTime[MAX_NUM_CONTROLLER_BUTTONS];
 
 		/** Raw Left thumb x analog value */
-		int16 LeftXAnalog;
+		int32 LeftAnalogX;
 		
 		/** Raw left thumb y analog value */
-		int16 LeftYAnalog;
+		int32 LeftAnalogY;
 
 		/** Raw Right thumb x analog value */
-		int16 RightXAnalog;
+		int32 RightAnalogX;
 
 		/** Raw Right thumb x analog value */
-		int16 RightYAnalog;
+		int32 RightAnalogY;
 
 		/** Left Trigger analog value */
-		int16 LeftTriggerAnalog;
+		int32 LeftTriggerAnalog;
 
 		/** Right trigger analog value */
-		int16 RightTriggerAnalog;
+		int32 RightTriggerAnalog;
 
 		/** Id of the controller */
 		int32 ControllerId;
@@ -94,9 +96,6 @@ private:
 
 private:
 
-	/** In the engine, all controllers map to xbox controllers for consistency */
-	uint8 HIDToXboxControllerMapping[MAX_NUM_CONTROLLER_BUTTONS];
-
 	/** Names of all the buttons */
 	EControllerButtons::Type Buttons[MAX_NUM_CONTROLLER_BUTTONS];
 
@@ -111,5 +110,5 @@ private:
 
 	IOHIDManagerRef HIDManager;
 
-	TSharedRef< FGenericApplicationMessageHandler > MessageHandler;
+	TSharedRef<FGenericApplicationMessageHandler> MessageHandler;
 };
diff --git a/Engine/Source/Runtime/Core/Private/Mac/MacApplication.cpp b/Engine/Source/Runtime/Core/Private/Mac/MacApplication.cpp
--- a/Engine/Source/Runtime/Core/Private/Mac/MacApplication.cpp
+++ b/Engine/Source/Runtime/Core/Private/Mac/MacApplication.cpp
@@ -718,7 +718,6 @@ FSlateCocoaWindow* FMacApplication::FindEventWindow( NSEvent* Event )
 void FMacApplication::PollGameDeviceState( const float TimeDelta )
 {
 	// Poll game device state and send new events
-	HIDInput->Tick( TimeDelta );
 	HIDInput->SendControllerEvents();
 }
 
