#summary Certain strings with character values between 128 and 255 were incorrectly being marked as ANSI strings, leading invalid ANSI characters to be replaced with "?".

#change Fixed TCString<WIDECHAR>::IsPureAnsi to make sure the character is less than 128, not 256.

#ttp 332392 - Blueprints: L10N: Compiling a blueprint causes text within FText members to replace any non-ascii characters with '?' VIDEO ATTACHED

diff --git a/Engine/Source/Runtime/Core/Public/Misc/CString.h b/Engine/Source/Runtime/Core/Public/Misc/CString.h
--- a/Engine/Source/Runtime/Core/Public/Misc/CString.h
+++ b/Engine/Source/Runtime/Core/Public/Misc/CString.h
@@ -613,7 +613,7 @@ bool TCString<WIDECHAR>::IsPureAnsi(const WIDECHAR* Str)
 {
 	for( ; *Str; Str++ )
 	{
-		if( *Str>0xff )
+		if( *Str>0x7f )
 		{
 			return 0;
 		}
