Linux should support quit to the OS (fixes TTP #331986).

diff --git a/Engine/Source/Runtime/Core/Public/Linux/LinuxPlatformProperties.h b/Engine/Source/Runtime/Core/Public/Linux/LinuxPlatformProperties.h
--- a/Engine/Source/Runtime/Core/Public/Linux/LinuxPlatformProperties.h
+++ b/Engine/Source/Runtime/Core/Public/Linux/LinuxPlatformProperties.h
@@ -85,4 +85,9 @@ struct FLinuxPlatformProperties
 	{
 		return !IS_DEDICATED_SERVER;
 	}
+
+	static FORCEINLINE bool SupportsQuit()
+	{
+		return true;
+	}
 };
