Added wireless controller ID to Mac gamepad code

diff --git a/Engine/Source/Runtime/Core/Private/Mac/HIDInputInterface.cpp b/Engine/Source/Runtime/Core/Private/Mac/HIDInputInterface.cpp
--- a/Engine/Source/Runtime/Core/Private/Mac/HIDInputInterface.cpp
+++ b/Engine/Source/Runtime/Core/Private/Mac/HIDInputInterface.cpp
@@ -153,7 +153,7 @@ void HIDInputInterface::FHIDDeviceInfo::SetupMappings()
 		RightAnalogYMapping = kHIDUsage_GD_Rz;
 		RightTriggerAnalogMapping = kHIDUsage_GD_Ry;
 	}
-	else if (VendorID == 0x45e && ProductID == 0x28e)
+	else if (VendorID == 0x45e && (ProductID == 0x28e || ProductID == 0x719))
 	{
 		// Xbox 360 Controller
 		ButtonsMapping[1]	= 0;	// A
