#ue4
#ios
* properly set the thread priorities on iOS

diff --git a/Engine/Source/Runtime/Core/Private/IOS/IOSPlatformProcess.cpp b/Engine/Source/Runtime/Core/Private/IOS/IOSPlatformProcess.cpp
--- a/Engine/Source/Runtime/Core/Private/IOS/IOSPlatformProcess.cpp
+++ b/Engine/Source/Runtime/Core/Private/IOS/IOSPlatformProcess.cpp
@@ -6,6 +6,7 @@
 
 #include "CorePrivate.h"
 #include "ApplePlatformRunnableThread.h"
+#include "IOSAppDelegate.h"
 #include <mach-o/dyld.h>
 
 const TCHAR* FIOSPlatformProcess::ComputerName()
@@ -48,16 +49,16 @@ void FIOSPlatformProcess::LaunchURL( const TCHAR* URL, const TCHAR* Parms, FStri
 void FIOSPlatformProcess::SetRealTimeMode()
 {
     mach_timebase_info_data_t TimeBaseInfo;
-    mach_timebase_info( &TimeBaseInfo );
-    double MsToAbs = ((double)TimeBaseInfo.denom / (double)TimeBaseInfo.numer) * 1000000.0;
-    uint32 NormalProcessingTimeMs = 33;
-    uint32 ConstraintProcessingTimeMs = 66;
-    
-    thread_time_constraint_policy_data_t Policy;
-    Policy.period      = 0;
-    Policy.computation = (uint32_t)(NormalProcessingTimeMs * MsToAbs);
-    Policy.constraint  = (uint32_t)(ConstraintProcessingTimeMs * MsToAbs);
-    Policy.preemptible = false;
-    
- //   thread_policy_set(pthread_mach_thread_np(pthread_self()), THREAD_TIME_CONSTRAINT_POLICY, (thread_policy_t)&Policy, THREAD_TIME_CONSTRAINT_POLICY_COUNT);
+	mach_timebase_info( &TimeBaseInfo );
+	double MsToAbs = ((double)TimeBaseInfo.denom / (double)TimeBaseInfo.numer) * 1000000.0;
+	uint32 NormalProcessingTimeMs = 20;
+	uint32 ConstraintProcessingTimeMs = 60;
+
+	thread_time_constraint_policy_data_t Policy;
+	Policy.period      = 0;
+	Policy.computation = (uint32_t)(NormalProcessingTimeMs * MsToAbs);
+	Policy.constraint  = (uint32_t)(ConstraintProcessingTimeMs * MsToAbs);
+	Policy.preemptible = true;
+
+	thread_policy_set(pthread_mach_thread_np(pthread_self()), THREAD_TIME_CONSTRAINT_POLICY, (thread_policy_t)&Policy, THREAD_TIME_CONSTRAINT_POLICY_COUNT);
 }
diff --git a/Engine/Source/Runtime/Launch/Private/IOS/LaunchIOS.cpp b/Engine/Source/Runtime/Launch/Private/IOS/LaunchIOS.cpp
--- a/Engine/Source/Runtime/Launch/Private/IOS/LaunchIOS.cpp
+++ b/Engine/Source/Runtime/Launch/Private/IOS/LaunchIOS.cpp
@@ -114,6 +114,8 @@ void FAppEntry::PlatformInit()
 
 void FAppEntry::Init()
 {
+	FPlatformProcess::SetRealTimeMode();
+
 	//extern TCHAR GCmdLine[16384];
 	GEngineLoop.PreInit(FCommandLine::Get());
 
diff --git a/Engine/Source/Runtime/RenderCore/Private/RenderingThread.cpp b/Engine/Source/Runtime/RenderCore/Private/RenderingThread.cpp
--- a/Engine/Source/Runtime/RenderCore/Private/RenderingThread.cpp
+++ b/Engine/Source/Runtime/RenderCore/Private/RenderingThread.cpp
@@ -241,6 +241,9 @@ void RenderingThreadMain( FEvent* TaskGraphBoundSyncEvent )
 		TaskGraphBoundSyncEvent->Trigger();
 	}
 
+	// set the thread back to real time mode
+	FPlatformProcess::SetRealTimeMode();
+
 	check(GIsThreadedRendering);
 	FTaskGraphInterface::Get().ProcessThreadUntilRequestReturn(ENamedThreads::RenderThread);
 	FPlatformMisc::MemoryBarrier();
