moved CircularBuffer, CircularQueue and Queue to /Core/Containers/

diff --git a/Engine/Source/Runtime/Core/Public/Containers/CircularBuffer.h b/Engine/Source/Runtime/Core/Public/Containers/CircularBuffer.h
new file mode 100644
--- /dev/null
+++ b/Engine/Source/Runtime/Core/Public/Containers/CircularBuffer.h
@@ -0,0 +1,104 @@
+// Copyright 1998-2014 Epic Games, Inc. All Rights Reserved.
+
+/*=============================================================================
+	CircularBuffer.h: Declares the TCircularBuffer template class.
+=============================================================================*/
+
+#pragma once
+
+
+/**
+ * Template for circular buffers.
+ *
+ * The size of the buffer is rounded up to the next power of two in order speed up indexing
+ * operations using a simple bit mask instead of the commonly used modulus operator that may
+ * be slow on some platforms.
+ */
+template<typename ElementType> class TCircularBuffer
+{
+public:
+
+	/**
+	 * Creates and initializes a new instance of the TCircularBuffer class.
+	 *
+	 * @param Capacity - The number of elements that the buffer can store (will be rounded up to the next power of 2).
+	 */
+	TCircularBuffer( uint32 Capacity )
+	{
+		checkSlow(Capacity > 0);
+		checkSlow(Capacity <= 0xffffffffU);
+
+		Elements.AddZeroed(FMath::RoundUpToPowerOfTwo(Capacity));
+	
+		IndexMask = Elements.Num() - 1;
+	}
+
+	/**
+	 * Creates and initializes a new instance of the TCircularBuffer class.
+	 *
+	 * @param Capacity -The number of elements that the buffer can store (will be rounded up to the next power of 2).
+	 * @param InitialValue - The initial value for the buffer's elements.
+	 */
+	TCircularBuffer( uint32 Capacity, const ElementType& InitialValue )
+	{
+		checkSlow(Capacity <= 0xffffffffU);
+
+		Elements.Init(InitialValue, FMath::RoundUpToPowerOfTwo(Capacity));
+	
+		IndexMask = Elements.Num() - 1;
+	}
+
+public:
+
+	/**
+	 * Returns the mutable element at the specified index.
+	 *
+	 * @param Index - The index of the element to return.
+	 */
+	FORCEINLINE ElementType& operator[]( uint32 Index )
+	{
+		return Elements[Index & IndexMask];
+	}
+
+	/**
+	 * Returns the immutable element at the specified index.
+	 *
+	 * @param Index - The index of the element to return.
+	 */
+	FORCEINLINE const ElementType& operator[]( uint32 Index ) const
+	{
+		return Elements[Index & IndexMask];
+	}
+	
+public:
+
+	/**
+	 * Returns the number of elements that the buffer can hold.
+	 *
+	 * @return Buffer capacity.
+	 */
+	FORCEINLINE uint32 Capacity( ) const
+	{
+		return Elements.Num();
+	}
+
+	/**
+	 * Calculates the index that follows the given index.
+	 *
+	 * @param CurrentIndex - The current index.
+	 *
+	 * @return The next index.
+	 */
+	FORCEINLINE uint32 GetNextIndex( uint32 CurrentIndex ) const
+	{
+		return ((CurrentIndex + 1) & IndexMask);
+	}
+
+private:
+
+	// Holds the mask for indexing the buffer's elements.
+	uint32 IndexMask;
+
+	// Holds the buffer's elements.
+	TArray<ElementType> Elements;
+};
diff --git a/Engine/Source/Runtime/Core/Public/Containers/CircularQueue.h b/Engine/Source/Runtime/Core/Public/Containers/CircularQueue.h
new file mode 100644
--- /dev/null
+++ b/Engine/Source/Runtime/Core/Public/Containers/CircularQueue.h
@@ -0,0 +1,168 @@
+// Copyright 1998-2014 Epic Games, Inc. All Rights Reserved.
+
+/*=============================================================================
+	CircularQueue.h: Declares the TCircularQueue template.
+=============================================================================*/
+
+#pragma once
+
+
+/**
+ * Implements a lock-free first-in first-out queue using a circular array.
+ *
+ * This class is thread safe only in two-thread scenarios, where the first thread
+ * always reads and the second thread always writes. The head and tail indices are
+ * stored in volatile memory to prevent the compiler from reordering the instructions
+ * that are meant to update the indices AFTER items have been inserted or removed.
+ *
+ * The number of items that can be enqueued is one less than the queue's capacity,
+ * because one item will be used for detecting full and empty states.
+ *
+ * @param ElementType - The type of elements held in the queue.
+ */
+template<typename ElementType> class TCircularQueue
+{
+public:
+
+	/**
+	 * Default constructor.
+	 *
+	 * @param Size - The number of elements that the queue can hold (will be rounded up to the next power of 2).
+	 */
+	TCircularQueue( uint32 Size)
+		: Buffer(Size)
+		, Head(0)
+		, Tail(0)
+	{ }
+
+	/**
+	 * Destructor.
+	 */
+	virtual ~TCircularQueue( ) { }
+
+public:
+
+	/**
+	 * Gets the number of elements in the queue.
+	 *
+	 * @return Number of queued elements.
+	 */
+	uint32 Count( ) const
+	{
+		int32 Count = Head - Tail;
+
+		if (Count < 0)
+		{
+			Count += Buffer.GetSize();
+		}
+
+		return Count;
+	}
+
+	/**
+	 * Removes an item from the front of the queue.
+	 *
+	 * @param OutElement - Will contain the element if the queue is not empty.
+	 *
+	 * @return TRUE if an element has been returned, FALSE if the queue was empty.
+	 */
+	bool Dequeue( ElementType& OutElement )
+	{
+		if (Head != Tail)
+		{
+			OutElement = Buffer[Head];
+			Head = Buffer.GetNextIndex(Head);
+
+			return true;
+		}
+
+		return false;
+	}
+
+	/**
+	 * Empties the queue.
+	 *
+	 * @see IsEmpty
+	 */
+	void Empty( )
+	{
+		Head = Tail;
+	}
+
+	/**
+	 * Adds an item to the end of the queue.
+	 *
+	 * @param Element - The element to add.
+	 *
+	 * @return TRUE if the item was added, FALSE if the queue was full.
+	 */
+	bool Enqueue( const ElementType& Element )
+	{
+		uint32 NewTail = Buffer.GetNextIndex(Tail);
+
+		if (NewTail != Head)
+		{
+			Buffer[Tail] = Element;
+			Tail = NewTail;
+
+			return true;
+		}
+
+		return false;
+	}
+
+	/**
+	 * Checks whether the queue is empty.
+	 *
+	 * @return TRUE if the queue is empty, FALSE otherwise.
+	 *
+	 * @see Empty
+	 * @see IsFull
+	 */
+	FORCEINLINE bool IsEmpty( ) const
+	{
+		return (Head == Tail);
+	}
+
+	/**
+	 * Checks whether the queue is full.
+	 *
+	 * @return TRUE if the queue is full, FALSE otherwise.
+	 *
+	 * @see IsEmpty
+	 */
+	bool IsFull( ) const
+	{
+		return (Buffer.GetNextIndex(Tail) == Head);
+	}
+
+	/**
+	 * Returns the oldest item in the queue without removing it.
+	 *
+	 * @param OutItem - Will contain the item if the queue is not empty.
+	 *
+	 * @return TRUE if an item has been returned, FALSE if the queue was empty.
+	 */
+	bool Peek( ElementType& OutItem )
+	{
+		if (Head != Tail)
+		{
+			OutItem = Buffer[Head];
+
+			return true;
+		}
+
+		return false;
+	}
+
+private:
+
+	// Holds the buffer.
+	TCircularBuffer<ElementType> Buffer;
+
+	// Holds the index to the first item in the buffer.
+	MS_ALIGN(CACHE_LINE_SIZE) volatile uint32 Head GCC_ALIGN(CACHE_LINE_SIZE);
+
+	// Holds the index to the last item in the buffer.
+	MS_ALIGN(CACHE_LINE_SIZE) volatile uint32 Tail GCC_ALIGN(CACHE_LINE_SIZE);
+};
diff --git a/Engine/Source/Runtime/Core/Public/Containers/Queue.h b/Engine/Source/Runtime/Core/Public/Containers/Queue.h
new file mode 100644
--- /dev/null
+++ b/Engine/Source/Runtime/Core/Public/Containers/Queue.h
@@ -0,0 +1,195 @@
+// Copyright 1998-2014 Epic Games, Inc. All Rights Reserved.
+
+/*=============================================================================
+	Queue.h: Declares the TQueue class template.
+=============================================================================*/
+
+#pragma once
+
+
+namespace EQueueMode
+{
+	/**
+	 * Enumerates concurrent queue modes.
+	 */
+	enum Type
+	{
+		/**
+		 * Multiple-producers, single-consumer queue.
+		 */
+		Mpsc,
+
+		/**
+		 * Single-producer, single-consumer queue.
+		 */
+		Spsc
+	};
+}
+
+
+/**
+ * Template for queues.
+ *
+ * This template implements an unbounded non-intrusive queue using a lock-free linked
+ * list that stores copies of the queued items. The template can operate in two modes:
+ * Multiple-producers single-consumer (MPSC) and Single-producer single-consumer (SPSC).
+ *
+ * The queue is thread-safe in both modes. The Dequeue() method ensures thread-safety by
+ * writing it in a way that does not depend on possible instruction reordering on the CPU.
+ * The Enqueue() method uses an atomic compare-and-swap in multiple-producers scenarios.
+ *
+ * @param ItemType - The type of items stored in the queue.
+ * @param Mode - The queue mode (single-producer, single-consumer by default).
+ *
+ * @todo gmp: Implement node pooling.
+ */
+template<typename ItemType, EQueueMode::Type Mode = EQueueMode::Spsc>
+class TQueue
+{
+public:
+
+	/**
+	 * Default constructor.
+	 */
+	TQueue( )
+	{
+		Head = Tail = new TNode();
+	}
+
+	/**
+	 * Destructor.
+	 */
+	~TQueue( )
+	{
+		while (Tail != nullptr)
+		{
+			TNode* Node = Tail;
+			Tail = Tail->NextNode;
+
+			delete Node;
+		}
+	}
+
+public:
+
+	/**
+	 * Removes and returns the item from the tail of the queue.
+	 *
+	 * @param OutValue - Will hold the returned value.
+	 *
+	 * @return true if a value was returned, false if the queue was empty.
+	 */
+	bool Dequeue( ItemType& OutItem )
+	{
+		TNode* Popped = Tail->NextNode;
+
+		if (Popped != nullptr)
+		{
+			OutItem = Popped->Item;
+
+			TNode* OldTail = Tail;
+
+			Tail = Popped;
+			Tail->Item = ItemType();
+
+			delete OldTail;
+
+			return true;
+		}
+
+		return false;
+	}
+
+	/**
+	 * Adds an item to the head of the queue.
+	 *
+	 * @param Item - The item to add.
+	 *
+	 * @return true if the item was added, false otherwise.
+	 */
+	bool Enqueue( const ItemType& Item )
+	{
+		TNode* NewNode = new TNode(Item);
+
+		if (NewNode != nullptr)
+		{
+			TNode* OldHead;
+
+			if (Mode == EQueueMode::Mpsc)
+			{
+				OldHead = (TNode*)FPlatformAtomics::InterlockedExchangePtr((void**)&Head, NewNode);
+			}
+			else
+			{
+				OldHead = Head;
+				Head = NewNode;
+			}		
+
+			OldHead->NextNode = NewNode;
+
+			return true;
+		}
+
+		return false;
+	}
+
+	/**
+	 * Checks whether the queue is empty.
+	 *
+	 * @return true if the queue is empty, false otherwise.
+	 */
+	bool IsEmpty( ) const
+	{
+		return (Tail->NextNode == nullptr);
+	}
+
+	/**
+	 * Peeks at the queue's tail item without removing it.
+	 *
+	 * @param OutItem - Will hold the peeked at item.
+	 *
+	 * @return true if an item was returned, false if the queue was empty.
+	 */
+	bool Peek( ItemType& OutItem )
+	{
+		if (Tail->NextNode != nullptr)
+		{
+			OutItem = Tail->NextNode->Item;
+
+			return true;
+		}
+
+		return false;
+	}
+
+private:
+
+	// Structure for the internal linked list.
+	struct TNode
+	{
+		// Holds a pointer to the next node in the list.
+		TNode* volatile NextNode;
+
+		// Holds the node's item.
+		ItemType Item;
+
+
+		// Default constructor.
+		TNode( )
+			: NextNode(nullptr)
+		{ }
+
+		// Creates and initializes a new node.
+		TNode( const ItemType& InItem )
+			: NextNode(nullptr)
+			, Item(InItem)
+		{ }
+	};
+
+
+	// Holds a pointer to the head of the list.
+	MS_ALIGN(16) TNode* volatile Head GCC_ALIGN(16);
+
+	// Holds a pointer to the tail of the list.
+	TNode* Tail;
+};
diff --git a/Engine/Source/Runtime/Core/Public/Misc/CircularBuffer.h b/Engine/Source/Runtime/Core/Public/Misc/CircularBuffer.h
deleted file mode 100644
diff --git a/Engine/Source/Runtime/Core/Public/Misc/CircularQueue.h b/Engine/Source/Runtime/Core/Public/Misc/CircularQueue.h
deleted file mode 100644
diff --git a/Engine/Source/Runtime/Core/Public/Misc/Queue.h b/Engine/Source/Runtime/Core/Public/Misc/Queue.h
deleted file mode 100644
