We now use the IMM context provided by Windows which results in the IME state being remembered correctly for UE4

TTP# 331255 - L10N: IME: The user should be able to toggle IME languages and IME needs to remember the state when switching between textboxes

Previously an IMM context was being setup for each Slate context but, since each context manages its own IME state, this resulted in the IME state being reset whenever focus was changed. Even worse was that it didn't even remember the setting for the same Slate widget, because the IMM context was setup again each time it was given focus.

Windows is able to manage the IMM context for each window/thread itself, so I've changed it to do that instead of us handing it manually.
The only caveat is, since we're using Slate, that we need to close the composition ourselves when changing focus since they're not "real" windows (see the call to ImmNotifyIME(IMMContext, NI_COMPOSITIONSTR, CPS_CANCEL, 0))

diff --git a/Engine/Source/Runtime/Core/Private/Windows/WindowsTextInputMethodSystem.cpp b/Engine/Source/Runtime/Core/Private/Windows/WindowsTextInputMethodSystem.cpp
--- a/Engine/Source/Runtime/Core/Private/Windows/WindowsTextInputMethodSystem.cpp
+++ b/Engine/Source/Runtime/Core/Private/Windows/WindowsTextInputMethodSystem.cpp
@@ -156,20 +156,6 @@ bool FWindowsTextInputMethodSystem::InitializeIMM()
 
 	UpdateIMMProperty(::GetKeyboardLayout(0));
 
-	auto Callback = [](HWND hwnd, LPARAM lparam) -> BOOL
-	{
-		DWORD UEProcessID = lparam;
-		DWORD WindowProcessID;
-		::GetWindowThreadProcessId(hwnd, &(WindowProcessID));
-		if(UEProcessID == WindowProcessID)
-		{
-			::ImmAssociateContext(hwnd, nullptr);
-			::ImmAssociateContextEx(hwnd, nullptr, IACE_CHILDREN);
-		}
-		return TRUE;
-	};
-	EnumWindows(Callback, GetCurrentProcessId());
-
 	UE_LOG(LogWindowsTextInputMethodSystem, Verbose, TEXT("Initialized IMM!"));
 
 	return true;
@@ -193,7 +179,6 @@ void FWindowsTextInputMethodSystem::UpdateWindowPositions(HIMC IMMContext)
 	if(ActiveContext.IsValid())
 	{
 		FInternalContext& InternalContext = ContextToInternalContextMap[ActiveContext];
-		check(IMMContext == InternalContext.IMMContext.IMMContext);
 
 		// Get start of composition area.
 		const uint32 BeginIndex = InternalContext.IMMContext.CompositionBeginIndex;
@@ -416,8 +401,10 @@ TSharedPtr<ITextInputMethodChangeNotifier> FWindowsTextInputMethodSystem::Regist
 
 	HRESULT Result;
 
+	FInternalContext& InternalContext = ContextToInternalContextMap.Add(Context);
+
 	// TSF Implementation
-	FCOMPtr<FTextStoreACP>& TextStore = ContextToInternalContextMap.Add(Context).TSFContext;
+	FCOMPtr<FTextStoreACP>& TextStore = InternalContext.TSFContext;
 	TextStore.Attach(new FTextStoreACP(Context));
 
 	Result = TSFThreadManager->CreateDocumentMgr(&(TextStore->TSFDocumentManager));
@@ -468,8 +455,11 @@ void FWindowsTextInputMethodSystem::UnregisterContext(const TSharedRef<ITextInpu
 
 	HRESULT Result;
 
+	check(ContextToInternalContextMap.Contains(Context));
+	FInternalContext& InternalContext = ContextToInternalContextMap[Context];
+
 	// TSF Implementation
-	FCOMPtr<FTextStoreACP>& TextStore = ContextToInternalContextMap[Context].TSFContext;
+	FCOMPtr<FTextStoreACP>& TextStore = InternalContext.TSFContext;
 
 	Result = TextStore->TSFDocumentManager->Pop(TF_POPF_ALL);
 	if(FAILED(Result))
@@ -493,10 +483,6 @@ void FWindowsTextInputMethodSystem::ActivateContext(const TSharedRef<ITextInputM
 	FInternalContext& InternalContext = ContextToInternalContextMap[Context];
 
 	// IMM Implementation
-	InternalContext.IMMContext.IMMContext = ::ImmCreateContext();
-	const TSharedPtr<FGenericWindow> GenericWindow = Context->GetWindow();
-	InternalContext.IMMContext.AssociatedWindow = GenericWindow.IsValid() ? reinterpret_cast<HWND>(GenericWindow->GetOSWindowHandle()) : nullptr;
-	::ImmAssociateContext(InternalContext.IMMContext.AssociatedWindow, InternalContext.IMMContext.IMMContext);
 	InternalContext.IMMContext.IsComposing = false;
 
 	// TSF Implementation
@@ -540,13 +526,12 @@ void FWindowsTextInputMethodSystem::DeactivateContext(const TSharedRef<ITextInpu
 	}
 
 	// IMM Implementation
-	if(CurrentAPI == EAPI::IMM)
-	{
-		Context->EndComposition();
-	}
-	::ImmAssociateContext(InternalContext.IMMContext.AssociatedWindow, nullptr);
-	::ImmDestroyContext(InternalContext.IMMContext.IMMContext);
-	InternalContext.IMMContext.IMMContext = nullptr;
+	const TSharedPtr<FGenericWindow> GenericWindow = Context->GetWindow();
+	const HWND Hwnd = GenericWindow.IsValid() ? reinterpret_cast<HWND>(GenericWindow->GetOSWindowHandle()) : nullptr;
+	HIMC IMMContext = ::ImmGetContext(Hwnd);
+	// Request the composition is canceled to ensure that the composition input UI is closed, and that a WM_IME_ENDCOMPOSITION message is sent
+	::ImmNotifyIME(IMMContext, NI_COMPOSITIONSTR, CPS_CANCEL, 0);
+	::ImmReleaseContext(Hwnd, IMMContext);
 
 	// General Implementation
 	ActiveContext = nullptr;
@@ -585,27 +570,6 @@ int32 FWindowsTextInputMethodSystem::ProcessMessage(HWND hwnd, uint32 msg, WPARA
 
 				UE_LOG(LogWindowsTextInputMethodSystem, Verbose, TEXT("Setting IMM context."));
 			}
-			else
-			{
-				::ImmAssociateContext(hwnd, nullptr);
-				::ImmAssociateContextEx(hwnd, nullptr, IACE_CHILDREN);
-
-				auto Callback = [](HWND hwnd, LPARAM lparam) -> BOOL
-				{
-					DWORD UEProcessID = lparam;
-					DWORD WindowProcessID;
-					::GetWindowThreadProcessId(hwnd, &(WindowProcessID));
-					if(UEProcessID == WindowProcessID)
-					{
-						::ImmAssociateContext(hwnd, nullptr);
-						::ImmAssociateContextEx(hwnd, nullptr, IACE_CHILDREN);
-					}
-					return TRUE;
-				};
-				EnumWindows(Callback, GetCurrentProcessId());
-
-				UE_LOG(LogWindowsTextInputMethodSystem, Verbose, TEXT("Resetting IMM context."));
-			}
 
 			return DefWindowProc(hwnd, msg, wParam, lParam);
 		}
@@ -645,6 +609,7 @@ int32 FWindowsTextInputMethodSystem::ProcessMessage(HWND hwnd, uint32 msg, WPARA
 			if(ActiveContext.IsValid())
 			{
 				FInternalContext& InternalContext = ContextToInternalContextMap[ActiveContext];
+				check(InternalContext.IMMContext.IsComposing);
 
 				HIMC IMMContext = ::ImmGetContext(hwnd);
 
diff --git a/Engine/Source/Runtime/Core/Public/Windows/WindowsTextInputMethodSystem.h b/Engine/Source/Runtime/Core/Public/Windows/WindowsTextInputMethodSystem.h
--- a/Engine/Source/Runtime/Core/Public/Windows/WindowsTextInputMethodSystem.h
+++ b/Engine/Source/Runtime/Core/Public/Windows/WindowsTextInputMethodSystem.h
@@ -92,8 +92,6 @@ private:
 		FCOMPtr<FTextStoreACP> TSFContext;
 		struct
 		{
-			HWND AssociatedWindow;
-			HIMC IMMContext;
 			bool IsComposing;
 			int32 CompositionBeginIndex;
 			uint32 CompositionLength;
