#UE4 - Prevent FMallocBinned from being used in programs and WITH_EDITORONLY_DATA targets in general.

diff --git a/Engine/Source/Runtime/Core/Private/Windows/WindowsPlatformMemory.cpp b/Engine/Source/Runtime/Core/Private/Windows/WindowsPlatformMemory.cpp
--- a/Engine/Source/Runtime/Core/Private/Windows/WindowsPlatformMemory.cpp
+++ b/Engine/Source/Runtime/Core/Private/Windows/WindowsPlatformMemory.cpp
@@ -62,7 +62,7 @@ FMalloc* FWindowsPlatformMemory::BaseAllocator()
 {
 #if FORCE_ANSI_ALLOCATOR
 	return new FMallocAnsi();
-#elif WITH_EDITOR && TBB_ALLOCATOR_ALLOWED
+#elif (WITH_EDITORONLY_DATA || IS_PROGRAM) && TBB_ALLOCATOR_ALLOWED
 	return new FMallocTBB();
 #else
 	return new FMallocBinned((uint32)(GetConstants().PageSize&MAX_uint32), (uint64)MAX_uint32+1);
