#CrashReportClient: Doesn't display assertion text (TTP#33238,fixed, windows only)

diff --git a/Engine/Source/Runtime/Core/Private/Windows/WindowsPlatformExceptionHandling.cpp b/Engine/Source/Runtime/Core/Private/Windows/WindowsPlatformExceptionHandling.cpp
--- a/Engine/Source/Runtime/Core/Private/Windows/WindowsPlatformExceptionHandling.cpp
+++ b/Engine/Source/Runtime/Core/Private/Windows/WindowsPlatformExceptionHandling.cpp
@@ -120,47 +120,76 @@ static void GetModuleVersion( TCHAR* ModuleName, TCHAR* StringBuffer, DWORD MaxS
 static void SetReportParameters( HREPORT ReportHandle, EXCEPTION_POINTERS* ExceptionInfo )
 {
 	HRESULT Result;
-	TCHAR StringBuffer[1024] = { 0 };
-	TCHAR ModuleName[1024] = { 0 };
+	enum 
+	{
+		BUFFER_SIZE = 1024,
+	};
+	TCHAR StringBuffer[BUFFER_SIZE] = {0};
+	TCHAR LocalBuffer[BUFFER_SIZE] = {0};
 
 	// Set the parameters for the standard problem signature
 	HMODULE ModuleHandle = GetModuleHandle( NULL );
 
-	StringCchPrintf( StringBuffer, 1024, TEXT( "UE4-%s" ), FApp::GetGameName() );
+	StringCchPrintf( StringBuffer, BUFFER_SIZE, TEXT( "UE4-%s" ), FApp::GetGameName() );
 	Result = WerReportSetParameter( ReportHandle, WER_P0, TEXT( "Application Name" ), StringBuffer );
 
-	GetModuleFileName( ModuleHandle, ModuleName, 1024 );
-	PathStripPath( ModuleName );
-	GetModuleVersion( ModuleName, StringBuffer, 1024 );
+	GetModuleFileName( ModuleHandle, LocalBuffer, BUFFER_SIZE );
+	PathStripPath( LocalBuffer );
+	GetModuleVersion( LocalBuffer, StringBuffer, BUFFER_SIZE );
 	Result = WerReportSetParameter( ReportHandle, WER_P1, TEXT( "Application Version" ), StringBuffer );
 
-	StringCchPrintf( StringBuffer, 1024, TEXT( "%08x" ), GetTimestampForLoadedLibrary( ModuleHandle ) );
+	StringCchPrintf( StringBuffer, BUFFER_SIZE, TEXT( "%08x" ), GetTimestampForLoadedLibrary( ModuleHandle ) );
 	Result = WerReportSetParameter( ReportHandle, WER_P2, TEXT( "Application Timestamp" ), StringBuffer );
 
 	HMODULE FaultModuleHandle = NULL;
 	GetModuleHandleEx( GET_MODULE_HANDLE_EX_FLAG_FROM_ADDRESS | GET_MODULE_HANDLE_EX_FLAG_UNCHANGED_REFCOUNT, ( LPCTSTR )ExceptionInfo->ExceptionRecord->ExceptionAddress, &FaultModuleHandle );
 
-	GetModuleFileName( FaultModuleHandle, ModuleName, 1024 );
-	PathStripPath( ModuleName );
-	Result = WerReportSetParameter( ReportHandle, WER_P3, TEXT( "Fault Module Name" ), ModuleName );
+	GetModuleFileName( FaultModuleHandle, LocalBuffer, BUFFER_SIZE );
+	PathStripPath( LocalBuffer );
+	Result = WerReportSetParameter( ReportHandle, WER_P3, TEXT( "Fault Module Name" ), LocalBuffer );
 
-	GetModuleVersion( ModuleName, StringBuffer, 1024 );
+	GetModuleVersion( LocalBuffer, StringBuffer, BUFFER_SIZE );
 	Result = WerReportSetParameter( ReportHandle, WER_P4, TEXT( "Fault Module Version" ), StringBuffer );
 
-	StringCchPrintf( StringBuffer, 1024, TEXT( "%08x" ), GetTimestampForLoadedLibrary( FaultModuleHandle ) );
+	StringCchPrintf( StringBuffer, BUFFER_SIZE, TEXT( "%08x" ), GetTimestampForLoadedLibrary( FaultModuleHandle ) );
 	Result = WerReportSetParameter( ReportHandle, WER_P5, TEXT( "Fault Module Timestamp" ), StringBuffer );
 
-	StringCchPrintf( StringBuffer, 1024, TEXT( "%08x" ), ExceptionInfo->ExceptionRecord->ExceptionCode );
+	StringCchPrintf( StringBuffer, BUFFER_SIZE, TEXT( "%08x" ), ExceptionInfo->ExceptionRecord->ExceptionCode );
 	Result = WerReportSetParameter( ReportHandle, WER_P6, TEXT( "Exception Code" ), StringBuffer );
 
 	INT_PTR ExceptionOffset = ( char* )( ExceptionInfo->ExceptionRecord->ExceptionAddress ) - ( char* )FaultModuleHandle;
-	StringCchPrintf( StringBuffer, 1024, TEXT( "%p" ), ExceptionOffset );
+	StringCchPrintf( StringBuffer, BUFFER_SIZE, TEXT( "%p" ), ExceptionOffset );
 	Result = WerReportSetParameter( ReportHandle, WER_P7, TEXT( "Exception Offset" ), StringBuffer );
 
-	StringCchPrintf( StringBuffer, 1024, TEXT( "!%s!" ), FCommandLine::Get() );
+	// Look for the assertion fail.
+	const TCHAR* AssertPos = FCString::Stristr( GErrorHist, TEXT( "Assertion failed: " ) );
+	const TCHAR* StackPos = FCString::Stristr( GErrorHist, TEXT( "\nStack:" ) );
+	if( AssertPos != nullptr && StackPos != nullptr && StackPos > AssertPos )
+	{
+		// Use LocalBuffer to store the assertion fail.
+		FCString::Strncpy( LocalBuffer, AssertPos, FMath::Min<int32>( StackPos - AssertPos, BUFFER_SIZE ) );
+
+		const int32 BufferLen = FCString::Strlen( LocalBuffer );
+
+		// Replace " with ' and replace \n with #
+		for( int32 ChIndex = 0; ChIndex < BufferLen; ++ChIndex )
+		{
+			if( LocalBuffer[ChIndex] == TEXT( '\"' ) )
+			{
+				LocalBuffer[ChIndex] = TEXT( '\'' );
+			}
+
+			if( LocalBuffer[ChIndex] == TEXT( '\n' ) )
+			{
+				LocalBuffer[ChIndex] = TEXT( '#' );
+			}
+		}
+	}
+
+	StringCchPrintf( StringBuffer, BUFFER_SIZE, TEXT( "!%s!AssertLog=\"%s\"" ), FCommandLine::Get(), LocalBuffer );
 	Result = WerReportSetParameter( ReportHandle, WER_P8, TEXT( "Commandline" ), StringBuffer );
 
-	StringCchPrintf( StringBuffer, 1024, TEXT( "%s!%s!%s!%d" ), TEXT( BRANCH_NAME ), FPlatformProcess::BaseDir(), GetEngineMode(), BUILT_FROM_CHANGELIST );
+	StringCchPrintf( StringBuffer, BUFFER_SIZE, TEXT( "%s!%s!%s!%d" ), TEXT( BRANCH_NAME ), FPlatformProcess::BaseDir(), GetEngineMode(), BUILT_FROM_CHANGELIST );
 	Result = WerReportSetParameter( ReportHandle, WER_P9, TEXT( "BranchBaseDir" ), StringBuffer );
 }
 
