Core - Removed obsolete comments regarding async work

diff --git a/Engine/Source/Runtime/Core/Public/Async/AsyncWork.h b/Engine/Source/Runtime/Core/Public/Async/AsyncWork.h
--- a/Engine/Source/Runtime/Core/Public/Async/AsyncWork.h
+++ b/Engine/Source/Runtime/Core/Public/Async/AsyncWork.h
@@ -33,14 +33,10 @@
 		}
 	};
 
-	start an example job at low priority
+	start an example job
 
 	(new FAutoDeleteAsyncTask<ExampleAutoDeleteAsyncTask>(5)->StartBackgroundTask();
 
-	start an example job at high priority
-
-	(new FAutoDeleteAsyncTask<ExampleAutoDeleteAsyncTask>(5)->StartHiPriorityTask();
-
 	do an example job now, on this thread
 
 	(new FAutoDeleteAsyncTask<ExampleAutoDeleteAsyncTask>(5)->StartSynchronousTask();
@@ -215,7 +211,7 @@ public:
 		}
 	};
 
-	start an example job at low priority
+	start an example job
 
 	FAutoDeleteAsyncTask<ExampleAutoDeleteAsyncTask> MyTask(5);
 	
