Reflecting CL2055153 for Mac builds -  Prevent FMallocBinned from being used in programs and WITH_EDITORONLY_DATA targets in general.

diff --git a/Engine/Source/Runtime/Core/Private/Mac/MacPlatformMemory.cpp b/Engine/Source/Runtime/Core/Private/Mac/MacPlatformMemory.cpp
--- a/Engine/Source/Runtime/Core/Private/Mac/MacPlatformMemory.cpp
+++ b/Engine/Source/Runtime/Core/Private/Mac/MacPlatformMemory.cpp
@@ -23,7 +23,7 @@ FMalloc* FMacPlatformMemory::BaseAllocator()
 {
 #if FORCE_ANSI_ALLOCATOR
 	return new FMallocAnsi();
-#elif WITH_EDITOR && TBB_ALLOCATOR_ALLOWED
+#elif (WITH_EDITORONLY_DATA || IS_PROGRAM) && TBB_ALLOCATOR_ALLOWED
 	return new FMallocTBB();
 #else
 	return new FMallocBinned((uint32)(GetConstants().PageSize&MAX_uint32), 0x100000000);
