Added ability to daemonize (detach from user session/terminal) to platform abstraction layer.

diff --git a/Engine/Source/Runtime/Core/Private/GenericPlatform/GenericPlatformProcess.cpp b/Engine/Source/Runtime/Core/Private/GenericPlatform/GenericPlatformProcess.cpp
--- a/Engine/Source/Runtime/Core/Private/GenericPlatform/GenericPlatformProcess.cpp
+++ b/Engine/Source/Runtime/Core/Private/GenericPlatform/GenericPlatformProcess.cpp
@@ -412,3 +412,9 @@ bool FGenericPlatformProcess::DeleteInterprocessSynchObject(FSemaphore * Object)
 	UE_LOG(LogHAL, Fatal, TEXT("FGenericPlatformProcess::DeleteInterprocessSynchObject not implemented on this platform"));
 	return false;
 }
+
+bool FGenericPlatformProcess::Daemonize()
+{
+	UE_LOG(LogHAL, Fatal, TEXT("FGenericPlatformProcess::Daemonize not implemented on this platform"));
+	return false;
+}
diff --git a/Engine/Source/Runtime/Core/Private/Linux/LinuxPlatformProcess.cpp b/Engine/Source/Runtime/Core/Private/Linux/LinuxPlatformProcess.cpp
--- a/Engine/Source/Runtime/Core/Private/Linux/LinuxPlatformProcess.cpp
+++ b/Engine/Source/Runtime/Core/Private/Linux/LinuxPlatformProcess.cpp
@@ -421,3 +421,16 @@ bool FLinuxPlatformProcess::GetProcReturnCode( FProcHandle & ProcHandle, int32*
 
 	return ProcHandle.GetReturnCode(ReturnCode);
 }
+
+bool FLinuxPlatformProcess::Daemonize()
+{
+	if (daemon(1, 1) == -1)
+	{
+		int ErrNo = errno;
+		UE_LOG(LogHAL, Warning, TEXT("daemon(1, 1) failed with errno = %d (%s)"), ErrNo,
+			StringCast< TCHAR >(strerror(ErrNo)).Get());
+		return false;
+	}
+
+	return true;
+}
diff --git a/Engine/Source/Runtime/Core/Private/Windows/WindowsPlatformProcess.cpp b/Engine/Source/Runtime/Core/Private/Windows/WindowsPlatformProcess.cpp
--- a/Engine/Source/Runtime/Core/Private/Windows/WindowsPlatformProcess.cpp
+++ b/Engine/Source/Runtime/Core/Private/Windows/WindowsPlatformProcess.cpp
@@ -1070,4 +1070,10 @@ bool FWindowsPlatformProcess::DeleteInterprocessSynchObject(FSemaphore * Object)
 	return bSucceeded;
 }
 
+bool FWindowsPlatformProcess::Daemonize()
+{
+	// TODO: implement
+	return true;
+}
+
 #include "HideWindowsPlatformTypes.h"
diff --git a/Engine/Source/Runtime/Core/Public/GenericPlatform/GenericPlatformProcess.h b/Engine/Source/Runtime/Core/Public/GenericPlatform/GenericPlatformProcess.h
--- a/Engine/Source/Runtime/Core/Public/GenericPlatform/GenericPlatformProcess.h
+++ b/Engine/Source/Runtime/Core/Public/GenericPlatform/GenericPlatformProcess.h
@@ -455,6 +455,13 @@ struct CORE_API FGenericPlatformProcess
 	 * @param Object object to destroy
 	 */
 	static bool DeleteInterprocessSynchObject(FSemaphore * Object);
+
+	/**
+	 * Makes process run as a system service (daemon), i.e. detaches it from whatever user session it was initially run from.
+	 *
+	 * @return true if successful
+	 */
+	static bool Daemonize();
 };
 
 
diff --git a/Engine/Source/Runtime/Core/Public/Linux/LinuxPlatformProcess.h b/Engine/Source/Runtime/Core/Public/Linux/LinuxPlatformProcess.h
--- a/Engine/Source/Runtime/Core/Public/Linux/LinuxPlatformProcess.h
+++ b/Engine/Source/Runtime/Core/Public/Linux/LinuxPlatformProcess.h
@@ -143,6 +143,7 @@ struct CORE_API FLinuxPlatformProcess : public FGenericPlatformProcess
 	static void WaitForProc( FProcHandle & ProcessHandle );
 	static void TerminateProc( FProcHandle & ProcessHandle, bool KillTree = false );
 	static bool GetProcReturnCode( FProcHandle & ProcHandle, int32* ReturnCode );
+	static bool Daemonize();
 };
 
 typedef FLinuxPlatformProcess FPlatformProcess;
diff --git a/Engine/Source/Runtime/Core/Public/Windows/WindowsPlatformProcess.h b/Engine/Source/Runtime/Core/Public/Windows/WindowsPlatformProcess.h
--- a/Engine/Source/Runtime/Core/Public/Windows/WindowsPlatformProcess.h
+++ b/Engine/Source/Runtime/Core/Public/Windows/WindowsPlatformProcess.h
@@ -109,6 +109,7 @@ struct CORE_API FWindowsPlatformProcess : public FGenericPlatformProcess
 	static FString ReadPipe( void* ReadPipe );
 	static FSemaphore * NewInterprocessSynchObject(const FString & Name, bool bCreate, uint32 MaxLocks = 1);
 	static bool DeleteInterprocessSynchObject(FSemaphore * Object);
+	static bool Daemonize();
 
 	// End FGenericPlatformProcess interface
 
