Disabled SH specialization functions until I can figure out the cross platform syntax

diff --git a/Engine/Source/Runtime/Core/Public/Math/SHMath.h b/Engine/Source/Runtime/Core/Public/Math/SHMath.h
--- a/Engine/Source/Runtime/Core/Public/Math/SHMath.h
+++ b/Engine/Source/Runtime/Core/Public/Math/SHMath.h
@@ -331,6 +331,9 @@ public:
 	}
 } GCC_ALIGN(16);
 
+// For non-standard specialization syntax
+#if _MSC_VER
+
 /** Specialization for 2nd order to avoid expensive trig functions. */
 template<> TSHVector<2> TSHVector<2>::SHBasisFunction(const FVector& Vector) 
 {
@@ -360,6 +363,8 @@ template<> TSHVector<3> TSHVector<3>::SHBasisFunction(const FVector& Vector)
 	return Result;
 }
 
+#endif
+
 /** A vector of colored spherical harmonic coefficients. */
 template<int32 MaxSHOrder>
 class TSHVectorRGB
