Improvements to the way that IMM IME input handles composition sessions

TTP# 332137 - Localization: IME: Google Japanese IME does not handle spacebar correctly

When you first type a character after having pressed space with an IMM IME, you get a WM_IME_COMPOSITION message with both a GCS_RESULTSTR (for the previous composition) and GCS_COMPSTR (for the start of the new composition). I've updated the handling logic to process the result first (and end the current composition session), and then process the new composition after (which starts a new composition session, updating the composition range correctly); this was previously being done the other way around, and resulted in an incorrect composition range.

This change also removes a lot of code that was managing selection state manually, as this isn't required and makes the IMM API calls more in line with the TSF API calls. Selection is now only set in two places:
1) After getting a GCS_RESULTSTR (the caret is moved to the end of the result string).
2) In response to a GCS_CURSORPOS request (this also now honours the CS_NOMOVECARET flag correctly).

Finally, this change fixes issues replacing selected text. Previously if you had text selected and you typed via an IMM IME, your new text would be placed at the start or end of the selection rather than replacing it. Any selected text is now explicitly cleared the first time we get a GCS_COMPSTR, which also further simplifies the handing of the caret position while a composition is running.

diff --git a/Engine/Source/Runtime/Core/Private/Windows/WindowsTextInputMethodSystem.cpp b/Engine/Source/Runtime/Core/Private/Windows/WindowsTextInputMethodSystem.cpp
--- a/Engine/Source/Runtime/Core/Private/Windows/WindowsTextInputMethodSystem.cpp
+++ b/Engine/Source/Runtime/Core/Private/Windows/WindowsTextInputMethodSystem.cpp
@@ -10,6 +10,16 @@ DEFINE_LOG_CATEGORY_STATIC(LogWindowsTextInputMethodSystem, Log, All);
 
 namespace
 {
+	FString GetIMMStringAsFString(HIMC IMMContext, const DWORD StringType)
+	{
+		const LONG StringNeededSize = ::ImmGetCompositionString(IMMContext, StringType, nullptr, 0);
+		TCHAR* RawString = reinterpret_cast<TCHAR*>(malloc(StringNeededSize));
+		::ImmGetCompositionString(IMMContext, StringType, RawString, StringNeededSize);
+		const FString Str(StringNeededSize / sizeof(TCHAR), RawString);
+		free(RawString);
+		return Str;
+	}
+
 	class FTextInputMethodChangeNotifier : public ITextInputMethodChangeNotifier
 	{
 	public:
@@ -168,13 +178,13 @@ void FWindowsTextInputMethodSystem::UpdateIMMProperty(HKL KeyboardLayoutHandle)
 	IMEProperties = ::ImmGetProperty(KeyboardLayoutHandle, IGP_PROPERTY);
 }
 
-bool FWindowsTextInputMethodSystem::ShouldDrawCompositionString() const
+bool FWindowsTextInputMethodSystem::ShouldDrawIMMCompositionString() const
 {
 	// If the IME doesn't have any kind of special UI and it draws the composition window at the caret, we can draw it ourselves.
 	return !(IMEProperties & IME_PROP_SPECIAL_UI) && (IMEProperties & IME_PROP_AT_CARET);
 }
 
-void FWindowsTextInputMethodSystem::UpdateWindowPositions(HIMC IMMContext)
+void FWindowsTextInputMethodSystem::UpdateIMMWindowPositions(HIMC IMMContext)
 {
 	if(ActiveContext.IsValid())
 	{
@@ -216,6 +226,37 @@ void FWindowsTextInputMethodSystem::UpdateWindowPositions(HIMC IMMContext)
 	}
 }
 
+void FWindowsTextInputMethodSystem::BeginIMMComposition()
+{
+	check(ActiveContext.IsValid());
+
+	FInternalContext& InternalContext = ContextToInternalContextMap[ActiveContext];
+				
+	InternalContext.IMMContext.IsComposing = true;
+	ActiveContext->BeginComposition();
+
+	uint32 SelectionBeginIndex = 0;
+	uint32 SelectionLength = 0;
+	ITextInputMethodContext::ECaretPosition SelectionCaretPosition = ITextInputMethodContext::ECaretPosition::Ending;
+	ActiveContext->GetSelectionRange(SelectionBeginIndex, SelectionLength, SelectionCaretPosition);
+				
+	// Set the initial composition range based on the start of the current selection
+	// We ignore the relative caret position as once you start typing any selected text is 
+	// removed before new text is added, so the caret will effectively always be at the start
+	InternalContext.IMMContext.CompositionBeginIndex = SelectionBeginIndex;
+	InternalContext.IMMContext.CompositionLength = 0;
+	ActiveContext->UpdateCompositionRange(InternalContext.IMMContext.CompositionBeginIndex, InternalContext.IMMContext.CompositionLength);
+}
+
+void FWindowsTextInputMethodSystem::EndIMMComposition()
+{
+	check(ActiveContext.IsValid());
+
+	FInternalContext& InternalContext = ContextToInternalContextMap[ActiveContext];
+	InternalContext.IMMContext.IsComposing = false;
+	ActiveContext->EndComposition();
+}
+
 bool FWindowsTextInputMethodSystem::InitializeTSF()
 {
 	UE_LOG(LogWindowsTextInputMethodSystem, Verbose, TEXT("Initializing TSF..."));
@@ -563,7 +604,7 @@ int32 FWindowsTextInputMethodSystem::ProcessMessage(HWND hwnd, uint32 msg, WPARA
 			if(ActiveContext.IsValid())
 			{
 				// Disable showing an IME-implemented composition window if we're going to draw the composition string ourselves.
-				if(wParam && ShouldDrawCompositionString())
+				if(wParam && ShouldDrawIMMCompositionString())
 				{
 					lParam &= ~ISC_SHOWUICOMPOSITIONWINDOW;
 				}
@@ -588,16 +629,8 @@ int32 FWindowsTextInputMethodSystem::ProcessMessage(HWND hwnd, uint32 msg, WPARA
 		{
 			if(ActiveContext.IsValid())
 			{
-				FInternalContext& InternalContext = ContextToInternalContextMap[ActiveContext];
-				uint32 BeginIndex;
-				uint32 Length;
-				ITextInputMethodContext::ECaretPosition CaretPosition;
-				ActiveContext->GetSelectionRange(BeginIndex, Length, CaretPosition);
-				InternalContext.IMMContext.IsComposing = true;
-				InternalContext.IMMContext.CompositionBeginIndex = BeginIndex + (CaretPosition == ITextInputMethodContext::ECaretPosition::Beginning ? 0 : Length);
-				InternalContext.IMMContext.CompositionLength = 0;
-				ActiveContext->BeginComposition();
-			
+				BeginIMMComposition();
+				
 				UE_LOG(LogWindowsTextInputMethodSystem, Verbose, TEXT("Beginning IMM composition."));
 			}
 
@@ -613,60 +646,69 @@ int32 FWindowsTextInputMethodSystem::ProcessMessage(HWND hwnd, uint32 msg, WPARA
 
 				HIMC IMMContext = ::ImmGetContext(hwnd);
 
-				UpdateWindowPositions(IMMContext);
+				UpdateIMMWindowPositions(IMMContext);
+
+				const bool bHasCompositionStringFlag = !!(lParam & GCS_COMPSTR);
+				const bool bHasResultStringFlag = !!(lParam & GCS_RESULTSTR);
+				const bool bHasNoMoveCaretFlag = !!(lParam & CS_NOMOVECARET);
+				const bool bHasCursorPosFlag = !!(lParam & GCS_CURSORPOS);
+
+				// Check Result
+				if(bHasResultStringFlag)
+				{
+					const FString ResultString = GetIMMStringAsFString(IMMContext, GCS_RESULTSTR);
+					UE_LOG(LogWindowsTextInputMethodSystem, Verbose, TEXT("WM_IME_COMPOSITION Result String: %s"), *ResultString);
+
+					// Update Result
+					ActiveContext->SetTextInRange(InternalContext.IMMContext.CompositionBeginIndex, InternalContext.IMMContext.CompositionLength, ResultString);
+
+					// Once we get a result, we're done; set the caret to the end of the result and end the current composition
+					ActiveContext->SetSelectionRange(InternalContext.IMMContext.CompositionBeginIndex + ResultString.Len(), 0, ITextInputMethodContext::ECaretPosition::Ending);
+					EndIMMComposition();
+				}
 
 				// Check Composition
+				if(bHasCompositionStringFlag)
 				{
-					LONG CompositionStringNeededSize = ::ImmGetCompositionString(IMMContext, GCS_COMPSTR, nullptr, 0);
-					TCHAR* RawString = reinterpret_cast<TCHAR*>(malloc(CompositionStringNeededSize));
-					::ImmGetCompositionString(IMMContext, GCS_COMPSTR, RawString, CompositionStringNeededSize);
-					FString CompositionString(CompositionStringNeededSize / sizeof(TCHAR), RawString);
-					free(RawString);
+					const FString CompositionString = GetIMMStringAsFString(IMMContext, GCS_COMPSTR);
 					UE_LOG(LogWindowsTextInputMethodSystem, Verbose, TEXT("WM_IME_COMPOSITION Composition String: %s"), *CompositionString);
 
-					// Update Composition
-					if(lParam & GCS_COMPSTR)
+					// We've typed a character, so we need to clear out any currently selected text to mimic what happens when you normally type into a text input
+					uint32 SelectionBeginIndex = 0;
+					uint32 SelectionLength = 0;
+					ITextInputMethodContext::ECaretPosition SelectionCaretPosition = ITextInputMethodContext::ECaretPosition::Ending;
+					ActiveContext->GetSelectionRange(SelectionBeginIndex, SelectionLength, SelectionCaretPosition);
+					if(SelectionLength)
 					{
-						ActiveContext->SetSelectionRange(InternalContext.IMMContext.CompositionBeginIndex, InternalContext.IMMContext.CompositionLength, ITextInputMethodContext::ECaretPosition::Ending);
-						ActiveContext->SetTextInRange(InternalContext.IMMContext.CompositionBeginIndex, InternalContext.IMMContext.CompositionLength, CompositionString);
-						InternalContext.IMMContext.CompositionLength = CompositionString.Len();
-						ActiveContext->SetSelectionRange(InternalContext.IMMContext.CompositionBeginIndex + InternalContext.IMMContext.CompositionLength, 0, ITextInputMethodContext::ECaretPosition::Ending);
+						ActiveContext->SetTextInRange(SelectionBeginIndex, SelectionLength, TEXT(""));
 					}
-				}
-
-				// Check Result
-				{
-					LONG ResultStringNeededSize = ::ImmGetCompositionString(IMMContext, GCS_RESULTSTR, nullptr, 0);
-					TCHAR* RawString = reinterpret_cast<TCHAR*>(malloc(ResultStringNeededSize));
-					::ImmGetCompositionString(IMMContext, GCS_RESULTSTR, RawString, ResultStringNeededSize);
-					FString ResultString(ResultStringNeededSize / sizeof(TCHAR), RawString);
-					free(RawString);
-					UE_LOG(LogWindowsTextInputMethodSystem, Verbose, TEXT("WM_IME_COMPOSITION Result String: %s"), *ResultString);
 
-					// Update Result
-					if(lParam & GCS_RESULTSTR)
+					// If we received a result (handled above) then the previous composition will have been ended, so we need to start a new one now
+					// This ensures that each composition ends up as its own distinct undo
+					if(!InternalContext.IMMContext.IsComposing)
 					{
-						ActiveContext->SetSelectionRange(InternalContext.IMMContext.CompositionBeginIndex, InternalContext.IMMContext.CompositionLength, ITextInputMethodContext::ECaretPosition::Ending);
-						ActiveContext->SetTextInRange(InternalContext.IMMContext.CompositionBeginIndex, InternalContext.IMMContext.CompositionLength, ResultString);
-						InternalContext.IMMContext.CompositionLength = ResultString.Len();
-						ActiveContext->SetSelectionRange(InternalContext.IMMContext.CompositionBeginIndex + InternalContext.IMMContext.CompositionLength, 0, ITextInputMethodContext::ECaretPosition::Ending);
+						BeginIMMComposition();
 					}
+
+					// Update Composition
+					ActiveContext->SetTextInRange(InternalContext.IMMContext.CompositionBeginIndex, InternalContext.IMMContext.CompositionLength, CompositionString);
+
+					// Update Composition Range
+					InternalContext.IMMContext.CompositionLength = CompositionString.Len();
+					ActiveContext->UpdateCompositionRange(InternalContext.IMMContext.CompositionBeginIndex, InternalContext.IMMContext.CompositionLength);
 				}
 
 				// Check Cursor
+				if(!bHasNoMoveCaretFlag && bHasCursorPosFlag)
 				{
-					LONG CursorPositionResult = ::ImmGetCompositionString(IMMContext, GCS_CURSORPOS, nullptr, 0);
-					int16 CursorPosition = CursorPositionResult & 0xFFFF;
+					const LONG CursorPositionResult = ::ImmGetCompositionString(IMMContext, GCS_CURSORPOS, nullptr, 0);
+					const int16 CursorPosition = CursorPositionResult & 0xFFFF;
 
 					// Update Cursor
-					if(lParam & GCS_CURSORPOS)
-					{
-						UE_LOG(LogWindowsTextInputMethodSystem, Verbose, TEXT("WM_IME_COMPOSITION Cursor Position: %d"), CursorPosition);
-						ActiveContext->SetSelectionRange(InternalContext.IMMContext.CompositionBeginIndex + CursorPosition, 0, ITextInputMethodContext::ECaretPosition::Ending);
-					}
+					UE_LOG(LogWindowsTextInputMethodSystem, Verbose, TEXT("WM_IME_COMPOSITION Cursor Position: %d"), CursorPosition);
+					ActiveContext->SetSelectionRange(InternalContext.IMMContext.CompositionBeginIndex + CursorPosition, 0, ITextInputMethodContext::ECaretPosition::Ending);
 				}
 
-				ActiveContext->UpdateCompositionRange(InternalContext.IMMContext.CompositionBeginIndex, InternalContext.IMMContext.CompositionLength);
 				::ImmReleaseContext(hwnd, IMMContext);
 
 				UE_LOG(LogWindowsTextInputMethodSystem, Verbose, TEXT("Updating IMM composition."));
@@ -680,9 +722,8 @@ int32 FWindowsTextInputMethodSystem::ProcessMessage(HWND hwnd, uint32 msg, WPARA
 			// On composition end, notify context of the end.
 			if(ActiveContext.IsValid())
 			{
-				FInternalContext& InternalContext = ContextToInternalContextMap[ActiveContext];
-				InternalContext.IMMContext.IsComposing = false;
-				ActiveContext->EndComposition();
+				EndIMMComposition();
+				
 				UE_LOG(LogWindowsTextInputMethodSystem, Verbose, TEXT("Ending IMM composition."));
 			}
 
diff --git a/Engine/Source/Runtime/Core/Public/Windows/WindowsTextInputMethodSystem.h b/Engine/Source/Runtime/Core/Public/Windows/WindowsTextInputMethodSystem.h
--- a/Engine/Source/Runtime/Core/Public/Windows/WindowsTextInputMethodSystem.h
+++ b/Engine/Source/Runtime/Core/Public/Windows/WindowsTextInputMethodSystem.h
@@ -64,8 +64,10 @@ private:
 	// IMM Implementation
 	bool InitializeIMM();
 	void UpdateIMMProperty(HKL KeyboardLatoutHandle);
-	bool ShouldDrawCompositionString() const;
-	void UpdateWindowPositions(HIMC IMMContext);
+	bool ShouldDrawIMMCompositionString() const;
+	void UpdateIMMWindowPositions(HIMC IMMContext);
+	void BeginIMMComposition();
+	void EndIMMComposition();
 
 	// TSF Implementation
 	bool InitializeTSF();
