Warning fix for FCrc::StrCrc32 when instantiated with a type which is 1 byte.

diff --git a/Engine/Source/Runtime/Core/Public/Misc/Crc.h b/Engine/Source/Runtime/Core/Public/Misc/Crc.h
--- a/Engine/Source/Runtime/Core/Public/Misc/Crc.h
+++ b/Engine/Source/Runtime/Core/Public/Misc/Crc.h
@@ -2,6 +2,8 @@
 
 #pragma once
 
+#include "Templates/EnableIf.h"
+
 /** 
  * CRC hash generation for different types of input data
  **/
@@ -19,7 +21,7 @@ struct CORE_API FCrc
 
 	/** String CRC. */
 	template <typename CharType>
-	static uint32 StrCrc32(const CharType* Data, uint32 CRC = 0)
+	static typename TEnableIf<sizeof(CharType) != 1, uint32>::Type StrCrc32(const CharType* Data, uint32 CRC = 0)
 	{
 		// We ensure that we never try to do a StrCrc32 with a CharType of more than 4 bytes.  This is because
 		// we always want to treat every CRC as if it was based on 4 byte chars, even if it's less, because we
@@ -40,6 +42,21 @@ struct CORE_API FCrc
 		return ~CRC;
 	}
 
+	template <typename CharType>
+	static typename TEnableIf<sizeof(CharType) == 1, uint32>::Type StrCrc32(const CharType* Data, uint32 CRC = 0)
+	{
+		/* Overload for when CharType is a byte, which causes warnings when right-shifting by 8 */
+		CRC = ~CRC;
+		while (CharType Ch = *Data++)
+		{
+			CRC = (CRC >> 8) ^ CRCTablesSB8[0][(CRC ^ Ch) & 0xFF];
+			CRC = (CRC >> 8) ^ CRCTablesSB8[0][(CRC     ) & 0xFF];
+			CRC = (CRC >> 8) ^ CRCTablesSB8[0][(CRC     ) & 0xFF];
+			CRC = (CRC >> 8) ^ CRCTablesSB8[0][(CRC     ) & 0xFF];
+		}
+		return ~CRC;
+	}
+
 	/**
 	 * DEPRECATED
 	 * These tables and functions are deprecated because they're using tables and implementations
