Enabled NaN Diagnostic for Orion cooked build to track down TTP 331879

diff --git a/Engine/Source/Runtime/Core/Public/Math/UnrealMathUtility.h b/Engine/Source/Runtime/Core/Public/Math/UnrealMathUtility.h
--- a/Engine/Source/Runtime/Core/Public/Math/UnrealMathUtility.h
+++ b/Engine/Source/Runtime/Core/Public/Math/UnrealMathUtility.h
@@ -9,7 +9,9 @@
 //#define IMPLEMENT_ASSIGNMENT_OPERATOR_MANUALLY
 
 // Assert on non finite numbers. Used to track NaNs.
-#define ENABLE_NAN_DIAGNOSTIC 0
+#ifndef ENABLE_NAN_DIAGNOSTIC
+	#define ENABLE_NAN_DIAGNOSTIC 0
+#endif
 
 /*-----------------------------------------------------------------------------
 	Definitions.
