Corrected issue where a copy+pasted material comment node containing quote marks will be truncated before the first quote mark.
#ttp 332014 - Material editor: copying comments creates issues when the comments are longer than a certain number of characters
#branch UE4
#proj Runtime.Core
#add Reinstated some commented-out lines in FString so that FString::ReplaceCharWithEscapeChar also deals with backslash, and single and double quotes.

diff --git a/Engine/Source/Runtime/Core/Private/Containers/String.cpp b/Engine/Source/Runtime/Core/Private/Containers/String.cpp
--- a/Engine/Source/Runtime/Core/Private/Containers/String.cpp
+++ b/Engine/Source/Runtime/Core/Private/Containers/String.cpp
@@ -954,18 +954,16 @@ FString FString::ReplaceQuotesWithEscapedQuotes() const
 	return *this;
 }
 
-#define MAX_SUPPORTED_ESCAPE_CHARS 3
+#define MAX_SUPPORTED_ESCAPE_CHARS 6
 
 static const TCHAR* CharToEscapeSeqMap[MAX_SUPPORTED_ESCAPE_CHARS][2] =
 {
 	{ TEXT("\n"), TEXT("\\n")  },
 	{ TEXT("\r"), TEXT("\\r")  },
-	{ TEXT("\t"), TEXT("\\t")  }
-
-	// these are currently disabled as the escaped backslash causes problems, the escaped ' and " cause too much churn in the inis
-// 	{ TEXT("\'"), TEXT("\\'")  },
-//	{ TEXT("\\"), TEXT("\\\\") },
-// 	{ TEXT("\""), TEXT("\\\"") }
+	{ TEXT("\t"), TEXT("\\t")  },
+	{ TEXT("\'"), TEXT("\\'")  },
+	{ TEXT("\\"), TEXT("\\\\") },
+	{ TEXT("\""), TEXT("\\\"") }
 };
 
 /**
