Updated FMacPlatformMisc::PumpMessages() to do what Windows version does - lower editor's main thread priority and mute audio when the app is in the background.

diff --git a/Engine/Source/Runtime/Core/Private/Mac/MacPlatformMisc.cpp b/Engine/Source/Runtime/Core/Private/Mac/MacPlatformMisc.cpp
--- a/Engine/Source/Runtime/Core/Private/Mac/MacPlatformMisc.cpp
+++ b/Engine/Source/Runtime/Core/Private/Mac/MacPlatformMisc.cpp
@@ -617,6 +617,39 @@ void FMacPlatformMisc::PumpMessages( bool bFromMainLoop )
 
 			[NSApp sendEvent: Event];
 		}
+
+		const bool HasFocus = [NSApp isActive];
+
+		// If editor thread doesn't have the focus, don't suck up too much CPU time.
+		if( GIsEditor )
+		{
+			static bool HadFocus=1;
+			if( HadFocus && !HasFocus )
+			{
+				// Drop our priority to speed up whatever is in the foreground.
+				struct sched_param Sched;
+				FMemory::Memzero(&Sched, sizeof(struct sched_param));
+				Sched.sched_priority = 5;
+				pthread_setschedparam(pthread_self(), SCHED_RR, &Sched);
+			}
+			else if( HasFocus && !HadFocus )
+			{
+				// Boost our priority back to normal.
+				struct sched_param Sched;
+				FMemory::Memzero(&Sched, sizeof(struct sched_param));
+				Sched.sched_priority = 15;
+				pthread_setschedparam(pthread_self(), SCHED_RR, &Sched);
+			}
+			if( !HasFocus )
+			{
+				// Sleep for a bit to not eat up all CPU time.
+				FPlatformProcess::Sleep(0.005f);
+			}
+			HadFocus = HasFocus;
+		}
+
+		// if app is active, allow sound, otherwise silence it
+		GVolumeMultiplier = HasFocus ? 1.0f : 0.0f;
 	}
 }
 
