Fix for SHBasisFunction specializations from Josh

diff --git a/Engine/Source/Runtime/Core/Public/Math/SHMath.h b/Engine/Source/Runtime/Core/Public/Math/SHMath.h
--- a/Engine/Source/Runtime/Core/Public/Math/SHMath.h
+++ b/Engine/Source/Runtime/Core/Public/Math/SHMath.h
@@ -331,11 +331,9 @@ public:
 	}
 } GCC_ALIGN(16);
 
-// For non-standard specialization syntax
-#if _MSC_VER
-
 /** Specialization for 2nd order to avoid expensive trig functions. */
-template<> TSHVector<2> TSHVector<2>::SHBasisFunction(const FVector& Vector) 
+template<> 
+inline TSHVector<2> TSHVector<2>::SHBasisFunction(const FVector& Vector) 
 {
 	TSHVector<2> Result;
 	Result.V[0] = 0.282095f; 
@@ -346,7 +344,8 @@ template<> TSHVector<2> TSHVector<2>::SHBasisFunction(const FVector& Vector)
 }
 
 /** Specialization for 3rd order to avoid expensive trig functions. */
-template<> TSHVector<3> TSHVector<3>::SHBasisFunction(const FVector& Vector) 
+template<> 
+inline TSHVector<3> TSHVector<3>::SHBasisFunction(const FVector& Vector) 
 {
 	TSHVector<3> Result;
 	Result.V[0] = 0.282095f; 
@@ -363,8 +362,6 @@ template<> TSHVector<3> TSHVector<3>::SHBasisFunction(const FVector& Vector)
 	return Result;
 }
 
-#endif
-
 /** A vector of colored spherical harmonic coefficients. */
 template<int32 MaxSHOrder>
 class TSHVectorRGB
