Set thread priority only on iOS 6 or less
#ue4
#ios

diff --git a/Engine/Source/Runtime/Core/Private/IOS/IOSPlatformProcess.cpp b/Engine/Source/Runtime/Core/Private/IOS/IOSPlatformProcess.cpp
--- a/Engine/Source/Runtime/Core/Private/IOS/IOSPlatformProcess.cpp
+++ b/Engine/Source/Runtime/Core/Private/IOS/IOSPlatformProcess.cpp
@@ -48,17 +48,20 @@ void FIOSPlatformProcess::LaunchURL( const TCHAR* URL, const TCHAR* Parms, FStri
 
 void FIOSPlatformProcess::SetRealTimeMode()
 {
-    mach_timebase_info_data_t TimeBaseInfo;
-	mach_timebase_info( &TimeBaseInfo );
-	double MsToAbs = ((double)TimeBaseInfo.denom / (double)TimeBaseInfo.numer) * 1000000.0;
-	uint32 NormalProcessingTimeMs = 20;
-	uint32 ConstraintProcessingTimeMs = 60;
+	if ([IOSAppDelegate GetDelegate].OSVersion < 7)
+	{
+		mach_timebase_info_data_t TimeBaseInfo;
+		mach_timebase_info( &TimeBaseInfo );
+		double MsToAbs = ((double)TimeBaseInfo.denom / (double)TimeBaseInfo.numer) * 1000000.0;
+		uint32 NormalProcessingTimeMs = 20;
+		uint32 ConstraintProcessingTimeMs = 60;
 
-	thread_time_constraint_policy_data_t Policy;
-	Policy.period      = 0;
-	Policy.computation = (uint32_t)(NormalProcessingTimeMs * MsToAbs);
-	Policy.constraint  = (uint32_t)(ConstraintProcessingTimeMs * MsToAbs);
-	Policy.preemptible = true;
+		thread_time_constraint_policy_data_t Policy;
+		Policy.period      = 0;
+		Policy.computation = (uint32_t)(NormalProcessingTimeMs * MsToAbs);
+		Policy.constraint  = (uint32_t)(ConstraintProcessingTimeMs * MsToAbs);
+		Policy.preemptible = true;
 
-	thread_policy_set(pthread_mach_thread_np(pthread_self()), THREAD_TIME_CONSTRAINT_POLICY, (thread_policy_t)&Policy, THREAD_TIME_CONSTRAINT_POLICY_COUNT);
+		thread_policy_set(pthread_mach_thread_np(pthread_self()), THREAD_TIME_CONSTRAINT_POLICY, (thread_policy_t)&Policy, THREAD_TIME_CONSTRAINT_POLICY_COUNT);
+	}
 }
