Fixed issue with network mapped app data folders in WindowPlatformSurvey

#ttp 293676 - ROCKET: Gather hardware stats can fail because CreateProc runs with unnormalized filepaths

#branch UE4
#proj Runtime.Core

#change FWindowsPlatformSurvey::BeginSurveyHardware() - added calls to FPaths::MakePlatformFilename() for the dxdiag exe path and the output filepath passed to it as an argument. These are usually handled inside platform IO functions but here we're passing the paths straight to a process.

diff --git a/Engine/Source/Runtime/Core/Private/Windows/WindowsPlatformSurvey.cpp b/Engine/Source/Runtime/Core/Private/Windows/WindowsPlatformSurvey.cpp
--- a/Engine/Source/Runtime/Core/Private/Windows/WindowsPlatformSurvey.cpp
+++ b/Engine/Source/Runtime/Core/Private/Windows/WindowsPlatformSurvey.cpp
@@ -118,14 +118,17 @@ void FWindowsPlatformSurvey::BeginSurveyHardware()
 	FString OutputFilepath = GetDxDiagOutputFilepath();
 
 	// Make sure the directory exists before we run dxdiag.  It won't create a directory for us (it will instead just silently do nothing.)
-	IFileManager::Get().MakeDirectory( *FPaths::GetPath(OutputFilepath), true );
-
-	FString ProcessArgs = FString::Printf(TEXT("/t %s"), *OutputFilepath);
+	IFileManager::Get().MakeDirectory(*FPaths::GetPath(OutputFilepath), true);
 
 	// Delete existing output file
 	IFileManager::Get().Delete(*OutputFilepath);
 
+	// Convert paths passed to CreateProc() to Windows format
+	FPaths::MakePlatformFilename(DxDiagFilepath);
+	FPaths::MakePlatformFilename(OutputFilepath);
+
 	// Run dxdiag as a external process, outputting to a text file
+	FString ProcessArgs = FString::Printf(TEXT("/t %s"), *OutputFilepath);
 	if (!FPlatformProcess::CreateProc(*DxDiagFilepath, *ProcessArgs, true, false, false, NULL, 0, NULL, NULL ).IsValid())
 	{
 		UE_LOG(LogWindows, Error, TEXT("FWindowsPlatformSurvey::BeginSurveyHardware() couldn't start up the dxdiag process"));
