Fix for broken localized links to online documentation from Built-in Editor Tutorials

SRombauts
#github https://github.com/EpicGames/UnrealEngine/pull/101

diff --git a/Engine/Source/Runtime/Core/Private/Internationalization/ICUCulture.cpp b/Engine/Source/Runtime/Core/Private/Internationalization/ICUCulture.cpp
--- a/Engine/Source/Runtime/Core/Private/Internationalization/ICUCulture.cpp
+++ b/Engine/Source/Runtime/Core/Private/Internationalization/ICUCulture.cpp
@@ -148,10 +148,16 @@ FString FCulture::FICUCultureImplementation::GetUnrealLegacyThreeLetterISOLangua
 {
 	FString Result( ICULocale.getISO3Language() );
 
-	// Legacy Overrides:
-	     if(Result == TEXT("eng")) { Result = TEXT("INT"); }
-	else if(Result == TEXT("jpn")) { Result = TEXT("JPN"); }
-	else if(Result == TEXT("kor")) { Result = TEXT("KOR"); }
+	// Legacy Overrides (INT, JPN, KOR), also for new web localization (CHN)
+	// and now for any other languages (FRA, DEU...) for correct redirection of documentation web links
+	if (Result == TEXT("eng"))
+	{
+		Result = TEXT("INT");
+	}
+	else
+	{
+		Result = Result.ToUpper();
+	}
 
 	return Result;
 }
