L10N: Built fresh ToolTip and PropertyName localization targets.

diff --git a/Engine/Config/Localization/PropertyNames.ini b/Engine/Config/Localization/PropertyNames.ini
--- a/Engine/Config/Localization/PropertyNames.ini
+++ b/Engine/Config/Localization/PropertyNames.ini
@@ -1,6 +1,6 @@
 [CommonSettings]
-SourcePath=../../../Engine/Content/Localization/Editor/PropertyNames
-DestinationPath=../../../Engine/Content/Localization/Editor/PropertyNames
+SourcePath=../../../Engine/Content/Localization/PropertyNames
+DestinationPath=../../../Engine/Content/Localization/PropertyNames
 ManifestName=PropertyNames.manifest
 ArchiveName=PropertyNames.archive
 SourceCulture=en
diff --git a/Engine/Config/Localization/ToolTips.ini b/Engine/Config/Localization/ToolTips.ini
--- a/Engine/Config/Localization/ToolTips.ini
+++ b/Engine/Config/Localization/ToolTips.ini
@@ -1,6 +1,6 @@
 [CommonSettings]
-SourcePath=../../../Engine/Content/Localization/Editor/ToolTips
-DestinationPath=../../../Engine/Content/Localization/Editor/ToolTips
+SourcePath=../../../Engine/Content/Localization/ToolTips
+DestinationPath=../../../Engine/Content/Localization/ToolTips
 ManifestName=ToolTips.manifest
 ArchiveName=ToolTips.archive
 SourceCulture=en
diff --git a/Engine/Source/Runtime/Core/Private/Misc/Paths.cpp b/Engine/Source/Runtime/Core/Private/Misc/Paths.cpp
--- a/Engine/Source/Runtime/Core/Private/Misc/Paths.cpp
+++ b/Engine/Source/Runtime/Core/Private/Misc/Paths.cpp
@@ -270,7 +270,7 @@ const TArray<FString>& FPaths::GetPropertyNameLocalizationPaths()
 		}
 		else
 		{
-			Results.AddUnique(TEXT("../../../Engine/Content/Localization/Editor/PropertyNames")); // Hardcoded convention.
+			Results.AddUnique(TEXT("../../../Engine/Content/Localization/PropertyNames")); // Hardcoded convention.
 		}
 	}
 
@@ -295,7 +295,7 @@ const TArray<FString>& FPaths::GetToolTipLocalizationPaths()
 		}
 		else
 		{
-			Results.AddUnique(TEXT("../../../Engine/Content/Localization/Editor/ToolTips")); // Hardcoded convention.
+			Results.AddUnique(TEXT("../../../Engine/Content/Localization/ToolTips")); // Hardcoded convention.
 		}
 	}
 
