Support for BOM-less UTF-8 buffers.

#change A file without a BOM is now interpreted as UTF-8 instead of ANSI.

diff --git a/Engine/Source/Runtime/Core/Private/Misc/CoreMisc.cpp b/Engine/Source/Runtime/Core/Private/Misc/CoreMisc.cpp
--- a/Engine/Source/Runtime/Core/Private/Misc/CoreMisc.cpp
+++ b/Engine/Source/Runtime/Core/Private/Misc/CoreMisc.cpp
@@ -134,19 +134,20 @@ void FFileHelper::BufferToString( FString& Result, const uint8* Buffer, int32 Si
 			ResultArray[ i ] = CharCast<TCHAR>( (UCS2CHAR)(( uint16 )Buffer[i * 2 + 3] + ( uint16 )Buffer[i * 2 + 2] * 256) );
 		}
 	}
-	else if ( Size >= 3 && Buffer[0] == 0xef && Buffer[1] == 0xbb && Buffer[2] == 0xbf )
+	else
 	{
-		FUTF8ToTCHAR Conv((const ANSICHAR*)&Buffer[3], Size - 3);
+		if ( Size >= 3 && Buffer[0] == 0xef && Buffer[1] == 0xbb && Buffer[2] == 0xbf )
+		{
+			// Skip over UTF-8 BOM if there is one
+			Buffer += 3;
+			Size   -= 3;
+		}
+
+		FUTF8ToTCHAR Conv((const ANSICHAR*)Buffer, Size);
 		int32 Length = Conv.Length();
 		ResultArray.AddUninitialized(Length + 1); // For the null terminator
 		CopyAssignItems(ResultArray.GetTypedData(), Conv.Get(), Length);
 	}
-	else
-	{
-		// ANSI. Additional 1 for null terminator.
-		ResultArray.AddUninitialized( Size + 1 );
-		FPlatformString::Convert(ResultArray.GetTypedData(), Size, (ANSICHAR*)Buffer, Size);
-	}
 
 	if (ResultArray.Num() == 1)
 	{
