Fixed Mac backspace key not working as delete in Content Browser.

diff --git a/Engine/Source/Runtime/Core/Private/Mac/MacApplication.cpp b/Engine/Source/Runtime/Core/Private/Mac/MacApplication.cpp
--- a/Engine/Source/Runtime/Core/Private/Mac/MacApplication.cpp
+++ b/Engine/Source/Runtime/Core/Private/Mac/MacApplication.cpp
@@ -1129,8 +1129,8 @@ TCHAR FMacApplication::ConvertChar(TCHAR Character)
 
 TCHAR FMacApplication::TranslateCharCode(TCHAR CharCode, uint32 KeyCode)
 {
-	// Keys like F1-F12, enter or backspace do not need translation
-	bool bNeedsTranslation = (CharCode < NSOpenStepUnicodeReservedBase || CharCode > 0xF8FF) && CharCode != 0x7F;
+	// Keys like F1-F12 or Enter do not need translation
+	bool bNeedsTranslation = CharCode < NSOpenStepUnicodeReservedBase || CharCode > 0xF8FF;
 	if( bNeedsTranslation )
 	{
 		// For non-numpad keys, the key code depends on the keyboard layout, so find out what was pressed by converting the key code to a Latin character
@@ -1160,7 +1160,7 @@ TCHAR FMacApplication::TranslateCharCode(TCHAR CharCode, uint32 KeyCode)
 		}
 	}
 	// Private use range should not be returned
-	else if(CharCode >= NSOpenStepUnicodeReservedBase && CharCode <= 0xF8FF)
+	else
 	{
 		CharCode = 0;
 	}
