Fixed mac compile error

diff --git a/Engine/Source/Runtime/Core/Public/Misc/IFilter.h b/Engine/Source/Runtime/Core/Public/Misc/IFilter.h
--- a/Engine/Source/Runtime/Core/Public/Misc/IFilter.h
+++ b/Engine/Source/Runtime/Core/Public/Misc/IFilter.h
@@ -4,18 +4,18 @@
 /**
  *	A generic interface that represents a Filter of ItemType
  */
-template< typename ItemType >
+template< typename TItemType >
 class IFilter
 {
 public:
-	typedef ItemType ItemType;
+	typedef TItemType ItemType;
 
 	virtual ~IFilter(){ }
 
 	/** Returns whether the specified Item passes the Filter's restrictions */
-	virtual bool PassesFilter( ItemType InItem ) const = 0;
+	virtual bool PassesFilter( TItemType InItem ) const = 0;
 
 	/** Broadcasts anytime the restrictions of the Filter changes */
-	DECLARE_EVENT( IFilter<ItemType>, FChangedEvent );
+	DECLARE_EVENT( IFilter<TItemType>, FChangedEvent );
 	virtual FChangedEvent& OnChanged() = 0;
 };
\ No newline at end of file
