Removed obsolete GPrimaryMonitorWidth, GPrimaryMonitorHeight, GPrimaryMonitorWorkRect, GVirtualScreenRect and GGameWindow* global variables and functions that were used to initialize them.

diff --git a/Engine/Source/Developer/CrashTracker/Private/CrashVideoCapture.cpp b/Engine/Source/Developer/CrashTracker/Private/CrashVideoCapture.cpp
--- a/Engine/Source/Developer/CrashTracker/Private/CrashVideoCapture.cpp
+++ b/Engine/Source/Developer/CrashTracker/Private/CrashVideoCapture.cpp
@@ -9,9 +9,6 @@
 
 DEFINE_LOG_CATEGORY(LogCrashTracker);
 
-/** Virtual screen rectangle including all monitors. */
-extern CORE_API RECT GVirtualScreenRect;
-
 /*-----------------------------------------------------------------------------
 	Stat declarations.
 -----------------------------------------------------------------------------*/
diff --git a/Engine/Source/Runtime/Core/Private/Android/AndroidMisc.cpp b/Engine/Source/Runtime/Core/Private/Android/AndroidMisc.cpp
--- a/Engine/Source/Runtime/Core/Private/Android/AndroidMisc.cpp
+++ b/Engine/Source/Runtime/Core/Private/Android/AndroidMisc.cpp
@@ -39,15 +39,6 @@ void FAndroidMisc::LoadPreInitModules()
 
 void FAndroidMisc::PlatformPreInit()
 {
-#if WITH_ENGINE
-
-	// @todo Android Get the total screen size of the primary monitor.v
-	// virtual screens
-	// Get the screen rect of the primary monitor, excluding taskbar etc.
-	// Refer iOS Misc
-	// needed?
-#endif
-
 	FAndroidAppEntry::PlatformInit();
 }
 
@@ -59,15 +50,6 @@ void FAndroidMisc::PlatformInit()
 	// Display memory info
 }
 
-
-void FAndroidMisc::PlatformPostInit()
-{
-	//@todo Android : needed?
-	//Initialize the game window to the main screen.
-	//Should not need to set opaque, since that should be part of the default initialization.
-	// create the main landscape window object
-}
-
 void FAndroidMisc::GetEnvironmentVariable(const TCHAR* VariableName, TCHAR* Result, int32 ResultLength)
 {
 	*Result = 0;
diff --git a/Engine/Source/Runtime/Core/Private/IOS/IOSPlatformMisc.cpp b/Engine/Source/Runtime/Core/Private/IOS/IOSPlatformMisc.cpp
--- a/Engine/Source/Runtime/Core/Private/IOS/IOSPlatformMisc.cpp
+++ b/Engine/Source/Runtime/Core/Private/IOS/IOSPlatformMisc.cpp
@@ -10,49 +10,12 @@
 #include "IOSApplication.h"
 #include "IOSAppDelegate.h"
 
-/** Width of the primary monitor, in pixels. */
-CORE_API int32 GPrimaryMonitorWidth = 0;
-/** Height of the primary monitor, in pixels. */
-CORE_API int32 GPrimaryMonitorHeight = 0;
-/** Rectangle of the work area on the primary monitor (excluding taskbar, etc) in "virtual screen coordinates" (pixels). */
-CORE_API RECT GPrimaryMonitorWorkRect;
-/** Virtual screen rectangle including all monitors. */
-CORE_API RECT GVirtualScreenRect;
-
 /** Amount of free memory in MB reported by the system at starup */
 CORE_API int32 GStartupFreeMemoryMB;
 
-/** Settings for the game Window */
-CORE_API UIWindow *GGameWindow = NULL;
-
 /** Global pointer to memory warning handler */
 void (* GMemoryWarningHandler)(const FGenericMemoryWarningContext& Context) = NULL;
 
-void FIOSPlatformMisc::PlatformPreInit()
-{
-#if WITH_ENGINE
-	CGRect ScreenFrame = [[UIScreen mainScreen] bounds];
-	CGRect VisibleFrame = [[UIScreen mainScreen] applicationFrame];
-
-	// Get the total screen size of the primary monitor.
-	GPrimaryMonitorWidth = ScreenFrame.size.width;
-	GPrimaryMonitorHeight = ScreenFrame.size.height;
-	// @todo IOS: virtual screens
-	GVirtualScreenRect.left = 0;
-	GVirtualScreenRect.top = 0;
-	GVirtualScreenRect.right = ScreenFrame.size.width;
-	GVirtualScreenRect.bottom = ScreenFrame.size.height;
-
-	// Get the screen rect of the primary monitor, excluding taskbar etc.
-	//@todo - zombie - The y might be inverted on iOS. Will need to do verification on it. -astarkey
-	GPrimaryMonitorWorkRect.left = VisibleFrame.origin.x;
-	GPrimaryMonitorWorkRect.right = VisibleFrame.origin.x + VisibleFrame.size.width;
-	GPrimaryMonitorWorkRect.top = ScreenFrame.size.height - (VisibleFrame.origin.y + VisibleFrame.size.height);
-	GPrimaryMonitorWorkRect.bottom = ScreenFrame.size.height - VisibleFrame.origin.y;
-
-#endif		// WITH_ENGINE
-}
-
 static int32 GetFreeMemoryMB()
 {
 	// get free memory
@@ -91,23 +54,6 @@ void FIOSPlatformMisc::PlatformInit()
 	UE_LOG(LogInit, Log, TEXT("Free Memory at startup: %d MB"), GStartupFreeMemoryMB);
 }
 
-
-#if WITH_ENGINE
-void FIOSPlatformMisc::PlatformPostInit()
-{
-	if ( FApp::IsGame() )
-	{
-		//Initialize the game window to the main screen.
-		//Should not need to set opaque, since that should be part of the default initialization.
-		// create the main landscape window object
-		CGRect ApplicationRect = [[UIScreen mainScreen] bounds];
-		GGameWindow = [[UIWindow alloc] initWithFrame: ApplicationRect];
-
-		verify( GGameWindow );
-	}
-}
-#endif		// WITH_ENGINE
-
 GenericApplication* FIOSPlatformMisc::CreateApplication()
 {
 	return FIOSApplication::CreateIOSApplication();
diff --git a/Engine/Source/Runtime/Core/Private/Mac/MacApplication.cpp b/Engine/Source/Runtime/Core/Private/Mac/MacApplication.cpp
--- a/Engine/Source/Runtime/Core/Private/Mac/MacApplication.cpp
+++ b/Engine/Source/Runtime/Core/Private/Mac/MacApplication.cpp
@@ -874,15 +874,10 @@ void FMacApplication::GetDisplayMetrics( FDisplayMetrics& OutDisplayMetrics ) co
 	OutDisplayMetrics.PrimaryDisplayHeight = ScreenFrame.size.height;
 
 	// Virtual desktop area
-	OutDisplayMetrics.VirtualDisplayRect.Left = 0;
-	OutDisplayMetrics.VirtualDisplayRect.Top = 0;
-	OutDisplayMetrics.VirtualDisplayRect.Right = 0;
-	OutDisplayMetrics.VirtualDisplayRect.Bottom = 0;
-
 	NSRect WholeWorkspace = {0};
-	for(NSScreen* Screen in AllScreens)
+	for (NSScreen* Screen in AllScreens)
 	{
-		if(Screen)
+		if (Screen)
 		{
 			WholeWorkspace = NSUnionRect(WholeWorkspace, [Screen frame]);
 		}
diff --git a/Engine/Source/Runtime/Core/Private/Mac/MacPlatformMisc.cpp b/Engine/Source/Runtime/Core/Private/Mac/MacPlatformMisc.cpp
--- a/Engine/Source/Runtime/Core/Private/Mac/MacPlatformMisc.cpp
+++ b/Engine/Source/Runtime/Core/Private/Mac/MacPlatformMisc.cpp
@@ -20,25 +20,6 @@
 #include <mach-o/dyld.h>
 #include <libproc.h>
 
-/** Width of the primary monitor, in pixels. */
-CORE_API int32 GPrimaryMonitorWidth = 0;
-/** Height of the primary monitor, in pixels. */
-CORE_API int32 GPrimaryMonitorHeight = 0;
-/** Rectangle of the work area on the primary monitor (excluding taskbar, etc) in "virtual screen coordinates" (pixels). */
-CORE_API RECT GPrimaryMonitorWorkRect;
-/** Virtual screen rectangle including all monitors. */
-CORE_API RECT GVirtualScreenRect;
-
-
-/** Settings for the game Window */
-CORE_API NSWindow *GGameWindow = NULL;
-CORE_API int32 GGameWindowPosX = 0;
-CORE_API int32 GGameWindowPosY = 0;
-CORE_API int32 GGameWindowWidth = 100;
-CORE_API int32 GGameWindowHeight = 100;
-
-CORE_API NSWindow* GExtendedAlertWindow = NULL;
-CORE_API NSWindow* GEULAWindow = NULL;
 
 /**
  * Information that cannot be obtained during a signal-handler is initialised here.
@@ -185,56 +166,6 @@ struct MacApplicationInfo
 };
 static MacApplicationInfo GMacAppInfo;
 
-void FMacPlatformMisc::PlatformPreInit()
-{
-#if WITH_ENGINE
-	SCOPED_AUTORELEASE_POOL;
-
-	NSArray* AllScreens = [NSScreen screens];
-	NSScreen* PrimaryScreen = (NSScreen*)[AllScreens objectAtIndex: 0];
-
-	NSRect ScreenFrame = [PrimaryScreen frame];
-	NSRect VisibleFrame = [PrimaryScreen visibleFrame];
-
-	// Get the total screen size of the primary monitor.
-	GPrimaryMonitorWidth = ScreenFrame.size.width;
-	GPrimaryMonitorHeight = ScreenFrame.size.height;
-
-	GVirtualScreenRect.left = 0;
-	GVirtualScreenRect.top = 0;
-	GVirtualScreenRect.right = 0;
-	GVirtualScreenRect.bottom = 0;
-
-	for( int32 Index = 0; Index < [AllScreens count]; Index++ )
-	{
-		NSScreen* Screen = (NSScreen*)[AllScreens objectAtIndex: Index];
-		NSRect DisplayFrame = [Screen frame];
-		if(DisplayFrame.origin.x != ScreenFrame.origin.x)
-		{
-			GVirtualScreenRect.right += DisplayFrame.size.width;
-		}
-		else
-		{
-			GVirtualScreenRect.right = FMath::Max(GVirtualScreenRect.right, (int)DisplayFrame.size.width);
-		}
-		if(DisplayFrame.origin.y != ScreenFrame.origin.y)
-		{
-			GVirtualScreenRect.bottom += DisplayFrame.size.height;
-		}
-		else
-		{
-			GVirtualScreenRect.bottom = FMath::Max(GVirtualScreenRect.bottom, (int)DisplayFrame.size.height);
-		}
-	}
-
-	// Get the screen rect of the primary monitor, exclusing taskbar etc.
-	GPrimaryMonitorWorkRect.left = VisibleFrame.origin.x;
-	GPrimaryMonitorWorkRect.right = VisibleFrame.origin.x + VisibleFrame.size.width;
-	GPrimaryMonitorWorkRect.top = ScreenFrame.size.height - (VisibleFrame.origin.y + VisibleFrame.size.height);
-	GPrimaryMonitorWorkRect.bottom = ScreenFrame.size.height - VisibleFrame.origin.y;
-#endif		// WITH_ENGINE
-}
-
 void FMacPlatformMisc::PlatformInit()
 {
 	// Increase the maximum number of simultaneously open files
@@ -278,59 +209,14 @@ void FMacPlatformMisc::PlatformInit()
 	GMacAppInfo.Init();
 }
 
-
-#if WITH_ENGINE
-void FMacPlatformMisc::PlatformPostInit()
-{
-	if ( FApp::IsGame() )
-	{
-		SCOPED_AUTORELEASE_POOL;
-
-		// Obtain width and height of primary monitor.
-		int32 ScreenWidth  = GPrimaryMonitorWidth;
-		int32 ScreenHeight = GPrimaryMonitorHeight;
-		int32 WindowWidth  = ScreenWidth / 2;
-		int32 WindowHeight = ScreenHeight / 2;
-		int32 WindowPosX = (ScreenWidth - WindowWidth ) / 2;
-		int32 WindowPosY = (ScreenHeight - WindowHeight ) / 2;
-
-		const FString Name = FApp::GetGameName();
-
-		NSRect ViewRect = NSMakeRect(WindowPosX, WindowPosY, WindowWidth, WindowHeight);
-
-		GGameWindow = [[NSWindow alloc] initWithContentRect: ViewRect
-									styleMask: (NSTitledWindowMask | NSClosableWindowMask | NSMiniaturizableWindowMask | NSResizableWindowMask)
-									  backing: NSBackingStoreBuffered
-										defer: NO];
-
-		if( GGameWindow )
-		{
-			[GGameWindow setAcceptsMouseMovedEvents: YES];
-			[GGameWindow setCollectionBehavior: NSWindowCollectionBehaviorFullScreenPrimary];
-			[GGameWindow setOpaque: YES];
-
-			CFStringRef CFName = FPlatformString::TCHARToCFString( *Name );
-			[GGameWindow setTitle: ( NSString *)CFName];
-			CFRelease( CFName );
-		}
-		verify( GGameWindow );
-	}
-}
-
-#endif		// WITH_ENGINE
-
-/**
- * Prevents screen-saver from kicking in by moving the mouse by 0 pixels. This works even on
- * Vista in the presence of a group policy for password protected screen saver.
- */
 void FMacPlatformMisc::PreventScreenSaver()
 {
-	CGEventRef LocationEvent = CGEventCreate( NULL );
-	CGPoint MouseLocation = CGEventGetLocation( LocationEvent );
-	CFRelease( LocationEvent );
-	CGEventRef MouseMoveCommand = CGEventCreateMouseEvent( NULL, kCGEventMouseMoved, MouseLocation, 0 );
-	CGEventPost( kCGHIDEventTap, MouseMoveCommand );
-	CFRelease( MouseMoveCommand );
+	CGEventRef LocationEvent = CGEventCreate(NULL);
+	CGPoint MouseLocation = CGEventGetLocation(LocationEvent);
+	CFRelease(LocationEvent);
+	CGEventRef MouseMoveCommand = CGEventCreateMouseEvent(NULL, kCGEventMouseMoved, MouseLocation, 0);
+	CGEventPost(kCGHIDEventTap, MouseMoveCommand);
+	CFRelease(MouseMoveCommand);
 }
 
 GenericApplication* FMacPlatformMisc::CreateApplication()
@@ -801,6 +687,8 @@ void FMacPlatformMisc::CreateGuid(FGuid& Result)
 	Result = GUID;
 }
 
+CORE_API NSWindow* GExtendedAlertWindow = NULL;
+
 static int32 ShowExtendedAlertWindow( EAppMsgType::Type MsgType, NSString* Text, NSString* Caption )
 {
 	SCOPED_AUTORELEASE_POOL;
diff --git a/Engine/Source/Runtime/Core/Private/Windows/WindowsPlatformMisc.cpp b/Engine/Source/Runtime/Core/Private/Windows/WindowsPlatformMisc.cpp
--- a/Engine/Source/Runtime/Core/Private/Windows/WindowsPlatformMisc.cpp
+++ b/Engine/Source/Runtime/Core/Private/Windows/WindowsPlatformMisc.cpp
@@ -44,25 +44,6 @@
  */
 #define WITH_FIREWALL_SUPPORT	0
 
-/** Width of the primary monitor, in pixels. */
-CORE_API int32 GPrimaryMonitorWidth = 0;
-/** Height of the primary monitor, in pixels. */
-CORE_API int32 GPrimaryMonitorHeight = 0;
-/** Rectangle of the work area on the primary monitor (excluding taskbar, etc) in "virtual screen coordinates" (pixels). */
-CORE_API RECT GPrimaryMonitorWorkRect;
-/** Virtual screen rectangle including all monitors. */
-CORE_API RECT GVirtualScreenRect;
-
-
-/** Settings for the game Window */
-CORE_API HWND GGameWindow = NULL;
-CORE_API bool GGameWindowUsingStartupWindowProc = false;	// Whether the game is using the startup window procedure
-CORE_API uint32 GGameWindowStyle = 0;
-CORE_API int32 GGameWindowPosX = 0;
-CORE_API int32 GGameWindowPosY = 0;
-CORE_API int32 GGameWindowWidth = 100;
-CORE_API int32 GGameWindowHeight = 100;
-
 extern "C"
 {
 	CORE_API HINSTANCE hInstance = NULL;
@@ -173,19 +154,6 @@ void FWindowsPlatformMisc::PlatformPreInit()
 		FPlatformMisc::RequestExit( false );
 	}
 
-#if WITH_ENGINE
-	// Get the total screen size of the primary monitor.
-	GPrimaryMonitorWidth = ::GetSystemMetrics( SM_CXSCREEN );
-	GPrimaryMonitorHeight = ::GetSystemMetrics( SM_CYSCREEN );
-	GVirtualScreenRect.left = ::GetSystemMetrics( SM_XVIRTUALSCREEN );
-	GVirtualScreenRect.top = ::GetSystemMetrics( SM_YVIRTUALSCREEN );
-	GVirtualScreenRect.right = ::GetSystemMetrics( SM_CXVIRTUALSCREEN ) + GVirtualScreenRect.left;
-	GVirtualScreenRect.bottom = ::GetSystemMetrics( SM_CYVIRTUALSCREEN ) + GVirtualScreenRect.top;
-
-	// Get the screen rect of the primary monitor, exclusing taskbar etc.
-	SystemParametersInfo( SPI_GETWORKAREA, 0, &GPrimaryMonitorWorkRect, 0 );
-#endif		// WITH_ENGINE
-
 	// initialize the file SHA hash mapping
 	InitSHAHashes();
 }
@@ -217,41 +185,6 @@ void FWindowsPlatformMisc::PlatformInit()
 }
 
 
-#if WITH_ENGINE
-/**
- * Temporary window procedure for the game window during startup.
- * It gets replaced later on with SetWindowLong().
- */
-LRESULT CALLBACK StartupWindowProc(HWND hWnd, uint32 Message, WPARAM wParam, LPARAM lParam)
-{
-	switch (Message)
-	{
-		// Prevent power management
-		case WM_POWERBROADCAST:
-		{
-			switch( wParam )
-			{
-				case PBT_APMQUERYSUSPEND:
-				case PBT_APMQUERYSTANDBY:
-					return BROADCAST_QUERY_DENY;
-			}
-		}
-	}
-
-	return DefWindowProc(hWnd, Message, wParam, lParam);
-}
-
-void FWindowsPlatformMisc::PlatformPostInit()
-{
-
-}
-
-#endif		// WITH_ENGINE
-
-/**
- * Prevents screen-saver from kicking in by moving the mouse by 0 pixels. This works even on
- * Vista in the presence of a group policy for password protected screen saver.
- */
 void FWindowsPlatformMisc::PreventScreenSaver()
 {
 	INPUT Input = { 0 };
@@ -1107,47 +1040,6 @@ EAppReturnType::Type FWindowsPlatformMisc::MessageBoxExt( EAppMsgType::Type MsgT
 	return EAppReturnType::Cancel;
 }
 
-#if WITH_ENGINE
-void FWindowsPlatformMisc::ShowGameWindow( void* StaticWndProc )
-{
-	if ( FPlatformProperties::SupportsWindowedMode() && GGameWindow && GGameWindowUsingStartupWindowProc )
-	{
-		// Convert position from screen coordinates to workspace coordinates.
-		HMONITOR Monitor = MonitorFromWindow(GGameWindow, MONITOR_DEFAULTTOPRIMARY);
-		MONITORINFO MonitorInfo;
-		MonitorInfo.cbSize = sizeof(MONITORINFO);
-		GetMonitorInfo( Monitor, &MonitorInfo );
-		int32 PosX = GGameWindowPosX - MonitorInfo.rcWork.left;
-		int32 PosY = GGameWindowPosY - MonitorInfo.rcWork.top;
-
-		// Clear out old messages using the old StartupWindowProc
-		WinPumpMessages();
-
-		SetWindowLong(GGameWindow, GWL_STYLE, GGameWindowStyle);
-		SetWindowLongPtr(GGameWindow, GWLP_WNDPROC, (LONG_PTR)StaticWndProc);
-		GGameWindowUsingStartupWindowProc = false;
-
-		// Restore the minimized window to the correct position, size and styles.
-		WINDOWPLACEMENT Placement;
-		FMemory::Memzero(&Placement, sizeof(WINDOWPLACEMENT));
-		Placement.length					= sizeof(WINDOWPLACEMENT);
-		GetWindowPlacement(GGameWindow, &Placement);
-		Placement.flags						= 0;
-		Placement.showCmd					= SW_SHOWNORMAL;	// Restores the minimized window (SW_SHOW won't do that)
-		Placement.rcNormalPosition.left		= PosX;
-		Placement.rcNormalPosition.right	= PosX + GGameWindowWidth;
-		Placement.rcNormalPosition.top		= PosY;
-		Placement.rcNormalPosition.bottom	= PosY + GGameWindowHeight;
-		SetWindowPlacement(GGameWindow, &Placement);
-		UpdateWindow(GGameWindow);
-
-		// Pump the messages using the new (correct) WindowProc
-		WinPumpMessages();
-	}
-}
-#endif		// WITH_ENGINE
-
-
 static bool HandleGameExplorerIntegration()
 {
 	if (FPlatformProperties::SupportsWindowedMode())
diff --git a/Engine/Source/Runtime/Core/Public/Android/AndroidMisc.h b/Engine/Source/Runtime/Core/Public/Android/AndroidMisc.h
--- a/Engine/Source/Runtime/Core/Public/Android/AndroidMisc.h
+++ b/Engine/Source/Runtime/Core/Public/Android/AndroidMisc.h
@@ -57,7 +57,6 @@ struct CORE_API FAndroidMisc : public FGenericPlatformMisc
 	static void LocalPrint(const TCHAR *Message);
 	static void PlatformPreInit();
 	static void PlatformInit();
-	static void PlatformPostInit();
 	static void GetEnvironmentVariable(const TCHAR* VariableName, TCHAR* Result, int32 ResultLength);
 	static void* GetHardwareWindow();
 	static void SetHardwareWindow(void* InWindow);
diff --git a/Engine/Source/Runtime/Core/Public/GenericPlatform/GenericPlatformMisc.h b/Engine/Source/Runtime/Core/Public/GenericPlatform/GenericPlatformMisc.h
--- a/Engine/Source/Runtime/Core/Public/GenericPlatform/GenericPlatformMisc.h
+++ b/Engine/Source/Runtime/Core/Public/GenericPlatform/GenericPlatformMisc.h
@@ -484,15 +484,6 @@ public:
 		return false;
 	}
 
-	/*
-	 *	Shows the intial game window in the proper position and size.
-	 *	It also changes the window proc from StartupWindowProc to
-	 *	UWindowsClient::StaticWndProc.
-	 *	This function doesn't have any effect if called a second time.
-	 */
-	static void ShowGameWindow( void* StaticWndProc )
-	{
-	}
 	/**
 	 * Handles Game Explorer, Firewall and FirstInstall commands, typically from the installer
 	 * @returns false if the game cannot continue.
diff --git a/Engine/Source/Runtime/Core/Public/IOS/IOSPlatformMisc.h b/Engine/Source/Runtime/Core/Public/IOS/IOSPlatformMisc.h
--- a/Engine/Source/Runtime/Core/Public/IOS/IOSPlatformMisc.h
+++ b/Engine/Source/Runtime/Core/Public/IOS/IOSPlatformMisc.h
@@ -11,9 +11,7 @@
 **/
 struct CORE_API FIOSPlatformMisc : public FGenericPlatformMisc
 {
-	static void PlatformPreInit();
 	static void PlatformInit();
-	static void PlatformPostInit();
 	static class GenericApplication* CreateApplication();
 	static void GetEnvironmentVariable(const TCHAR* VariableName, TCHAR* Result, int32 ResultLength);
 	static void* GetHardwareWindow();
diff --git a/Engine/Source/Runtime/Core/Public/Mac/MacPlatformMisc.h b/Engine/Source/Runtime/Core/Public/Mac/MacPlatformMisc.h
--- a/Engine/Source/Runtime/Core/Public/Mac/MacPlatformMisc.h
+++ b/Engine/Source/Runtime/Core/Public/Mac/MacPlatformMisc.h
@@ -11,9 +11,7 @@
 **/
 struct CORE_API FMacPlatformMisc : public FGenericPlatformMisc
 {
-	static void PlatformPreInit();
 	static void PlatformInit();
-	static void PlatformPostInit();
 	static class GenericApplication* CreateApplication();
 	static void GetEnvironmentVariable(const TCHAR* VariableName, TCHAR* Result, int32 ResultLength);
 	static TArray<uint8> GetMacAddress();
@@ -55,7 +53,6 @@ struct CORE_API FMacPlatformMisc : public FGenericPlatformMisc
 	static void CreateGuid(class FGuid& Result);
 	static EAppReturnType::Type MessageBoxExt( EAppMsgType::Type MsgType, const TCHAR* Text, const TCHAR* Caption );
 	static void PreventScreenSaver();
-	static void ShowGameWindow( void* StaticWndProc );
 	static bool CommandLineCommands();
 	static int32 NumberOfCores();
 	static int32 NumberOfCoresIncludingHyperthreads();
diff --git a/Engine/Source/Runtime/Core/Public/Windows/WindowsPlatformMisc.h b/Engine/Source/Runtime/Core/Public/Windows/WindowsPlatformMisc.h
--- a/Engine/Source/Runtime/Core/Public/Windows/WindowsPlatformMisc.h
+++ b/Engine/Source/Runtime/Core/Public/Windows/WindowsPlatformMisc.h
@@ -14,7 +14,6 @@ struct CORE_API FWindowsPlatformMisc : public FGenericPlatformMisc
 {
 	static void PlatformPreInit();
 	static void PlatformInit();
-	static void PlatformPostInit();
 	static class GenericApplication* CreateApplication();
 	static void SetGracefulTerminationHandler();
 	static void GetEnvironmentVariable(const TCHAR* VariableName, TCHAR* Result, int32 ResultLength);
@@ -47,7 +46,6 @@ struct CORE_API FWindowsPlatformMisc : public FGenericPlatformMisc
 	static void CreateGuid(class FGuid& Result);
 	static EAppReturnType::Type MessageBoxExt( EAppMsgType::Type MsgType, const TCHAR* Text, const TCHAR* Caption );
 	static void PreventScreenSaver();
-	static void ShowGameWindow( void* StaticWndProc );
 	static bool CommandLineCommands();
 	static bool Is64bitOperatingSystem();
 	static bool IsValidAbsolutePathFormat(const FString& Path);
diff --git a/Engine/Source/Runtime/Launch/Private/Mac/LaunchMac.cpp b/Engine/Source/Runtime/Launch/Private/Mac/LaunchMac.cpp
--- a/Engine/Source/Runtime/Launch/Private/Mac/LaunchMac.cpp
+++ b/Engine/Source/Runtime/Launch/Private/Mac/LaunchMac.cpp
@@ -12,7 +12,6 @@
 
 static FString GSavedCommandLine;
 extern CORE_API NSWindow* GExtendedAlertWindow;
-extern CORE_API NSWindow* GEULAWindow;
 extern int32 GuardedMain( const TCHAR* CmdLine );
 extern void LaunchStaticShutdownAfterError();
 
@@ -45,7 +44,6 @@ void EngineCrashHandler(const FGenericCrashContext & GenericContext)
 @interface UE4AppDelegate : NSObject <NSApplicationDelegate>
 {
 	IBOutlet NSWindow *AlertWindow;
-	IBOutlet NSWindow *EULAWindow;
 #if WITH_EDITOR
 	NSString* Filename;
 	bool bHasFinishedLaunching;
@@ -56,8 +54,6 @@ void EngineCrashHandler(const FGenericCrashContext & GenericContext)
 - (BOOL)application:(NSApplication *)theApplication openFile:(NSString *)filename;
 #endif
 - (IBAction)OnAlertWindowButtonPress:(id)Sender;
-- (IBAction)OnEULAAccepted:(id)Sender;
-- (IBAction)OnEULADeclined:(id)Sender;
 
 @end
 
@@ -70,7 +66,6 @@ void EngineCrashHandler(const FGenericCrashContext & GenericContext)
 	bHasFinishedLaunching = false;
 #endif
 	GExtendedAlertWindow = AlertWindow;
-	GEULAWindow = EULAWindow;
 }
 
 #if WITH_EDITOR
@@ -113,16 +108,6 @@ void EngineCrashHandler(const FGenericCrashContext & GenericContext)
 	}
 }
 
-- (IBAction)OnEULAAccepted:(id)Sender
-{
-	[NSApp stopModalWithCode:1];
-}
-
-- (IBAction)OnEULADeclined:(id)Sender
-{
-	_Exit(0);
-}
-
 - (IBAction)OnQuitRequest:(id)Sender
 {
 	if (GEngine)
diff --git a/Engine/Source/Runtime/SlateRHIRenderer/Private/SlateRHIRenderer.cpp b/Engine/Source/Runtime/SlateRHIRenderer/Private/SlateRHIRenderer.cpp
--- a/Engine/Source/Runtime/SlateRHIRenderer/Private/SlateRHIRenderer.cpp
+++ b/Engine/Source/Runtime/SlateRHIRenderer/Private/SlateRHIRenderer.cpp
@@ -24,11 +24,6 @@ namespace CrashTrackerConstants
 // Defines the maximum size that a slate viewport will create
 #define MAX_VIEWPORT_SIZE 16384
 
-#if PLATFORM_WINDOWS || PLATFORM_MAC
-/** Virtual screen rectangle including all monitors. */
-extern CORE_API RECT GVirtualScreenRect;
-#endif
-
 static FMatrix CreateProjectionMatrix( uint32 Width, uint32 Height )
 {
 	
@@ -196,10 +191,12 @@ void FSlateRHIRenderer::Initialize()
 #if PLATFORM_WINDOWS || PLATFORM_MAC
 	if (GIsEditor)
 	{
-		const FIntPoint VirtualScreenOrigin = FIntPoint(GVirtualScreenRect.left, GVirtualScreenRect.top);
-		const FIntPoint VirtualScreenLowerRight = FIntPoint(GVirtualScreenRect.right, GVirtualScreenRect.bottom);
+		FDisplayMetrics DisplayMetrics;
+		FSlateApplication::Get().GetDisplayMetrics(DisplayMetrics);
+		const FIntPoint VirtualScreenOrigin = FIntPoint(DisplayMetrics.VirtualDisplayRect.Left, DisplayMetrics.VirtualDisplayRect.Top);
+		const FIntPoint VirtualScreenLowerRight = FIntPoint(DisplayMetrics.VirtualDisplayRect.Right, DisplayMetrics.VirtualDisplayRect.Bottom);
 		const FIntRect VirtualScreen = FIntRect(VirtualScreenOrigin, VirtualScreenLowerRight);
-	
+
 		CrashTrackerResource = new FSlateCrashReportResource(VirtualScreen);
 		BeginInitResource(CrashTrackerResource);
 	}
