Fixed a bug in Mac GetWorkArea that was causing menus and tooltips to appear at incorrect position if primary monitor's resolution was lower than secondary's.

diff --git a/Engine/Source/Runtime/Core/Private/Mac/MacApplication.cpp b/Engine/Source/Runtime/Core/Private/Mac/MacApplication.cpp
--- a/Engine/Source/Runtime/Core/Private/Mac/MacApplication.cpp
+++ b/Engine/Source/Runtime/Core/Private/Mac/MacApplication.cpp
@@ -832,7 +832,7 @@ FPlatformRect FMacApplication::GetWorkArea( const FPlatformRect& CurrentWindow )
 
 	FPlatformRect WorkArea;
 	WorkArea.Left = VisibleFrame.origin.x;
-	WorkArea.Top = ScreenHeight - VisibleFrame.size.height - VisibleFrame.origin.y;
+	WorkArea.Top = ScreenHeight - VisibleFrame.size.height - (VisibleFrame.origin.y - [Screen frame].origin.y);
 	WorkArea.Right = WorkArea.Left + VisibleFrame.size.width;
 	WorkArea.Bottom = WorkArea.Top + VisibleFrame.size.height;
 
