Fix incorrect batch file path used to scrape Visual Studio environment variables, causing new project generation and hot-reload to stop working. (TTP 333871)

diff --git a/Engine/Source/Runtime/Core/Private/Modules/ModuleManager.cpp b/Engine/Source/Runtime/Core/Private/Modules/ModuleManager.cpp
--- a/Engine/Source/Runtime/Core/Private/Modules/ModuleManager.cpp
+++ b/Engine/Source/Runtime/Core/Private/Modules/ModuleManager.cpp
@@ -1617,7 +1617,7 @@ bool FModuleManager::BuildUnrealBuildTool(FOutputDevice &Ar)
 	FPlatformMisc::GetVSComnTools(11, VCVarsBat);
 #endif
 
-	VCVarsBat += FPaths::Combine(*VCVarsBat, L"../../VC/bin/x86_amd64/vcvarsx86_amd64.bat");
+	VCVarsBat = FPaths::Combine(*VCVarsBat, L"../../VC/bin/x86_amd64/vcvarsx86_amd64.bat");
 
 	// Check to make sure we found one.
 	if ( VCVarsBat.IsEmpty() || !FPaths::FileExists(VCVarsBat) )
