TSharedPtr build time improvements

#ttp 332486: TSharedPtr compile times seem to be exponential

#add Move semantics for TWeakPtrs, TSharedPtrs and TSharedRefs.  Conversions between TWeakPtrs and TSharedPtrs are not move-aware because moving would be worse than copying.  Also, TSharedRefs are move assignable but not move constructible, because we can't leave a TSharedRef as NULL.
#fix Removed extraneous copies of F(Shared/Weak)Referencers in assignment operators which were hurting build times.
#change Changed FReferenceController's definition to reduce template instantiations and remove a Clang-specific hack.
#change Split FReferenceControllerImpl into FReferenceControllerOps and FReferenceControllerState to reduce template instantiations.
#remove Deleted DynamicCastShared*, because we don't support RTTI.
#fix Replaced the only visible use of DynamicCastShared*, in PropertyEditorModule.cpp, as it seemed unnecessary.

diff --git a/Engine/Source/Editor/PropertyEditor/Private/PropertyEditorModule.cpp b/Engine/Source/Editor/PropertyEditor/Private/PropertyEditorModule.cpp
--- a/Engine/Source/Editor/PropertyEditor/Private/PropertyEditorModule.cpp
+++ b/Engine/Source/Editor/PropertyEditor/Private/PropertyEditorModule.cpp
@@ -259,7 +259,7 @@ TSharedRef< IPropertyTableWidgetHandle > FPropertyEditorModule::CreatePropertyTa
 {
 	TSharedRef< FPropertyTableWidgetHandle > FWidgetHandle = MakeShareable( new FPropertyTableWidgetHandle( SNew( SPropertyTable, PropertyTable ) ) );
 
-	TSharedRef< IPropertyTableWidgetHandle > IWidgetHandle = DynamicCastSharedRef<IPropertyTableWidgetHandle>(FWidgetHandle);
+	TSharedRef< IPropertyTableWidgetHandle > IWidgetHandle = StaticCastSharedRef<IPropertyTableWidgetHandle>(FWidgetHandle);
 
 	 return IWidgetHandle;
 }
@@ -269,7 +269,7 @@ TSharedRef< IPropertyTableWidgetHandle > FPropertyEditorModule::CreatePropertyTa
 	TSharedRef< FPropertyTableWidgetHandle > FWidgetHandle = MakeShareable( new FPropertyTableWidgetHandle( SNew( SPropertyTable, PropertyTable )
 		.ColumnCustomizations( Customizations )));
 
-	TSharedRef< IPropertyTableWidgetHandle > IWidgetHandle = DynamicCastSharedRef<IPropertyTableWidgetHandle>(FWidgetHandle);
+	TSharedRef< IPropertyTableWidgetHandle > IWidgetHandle = StaticCastSharedRef<IPropertyTableWidgetHandle>(FWidgetHandle);
 
 	 return IWidgetHandle;
 }
diff --git a/Engine/Source/Runtime/Core/Public/Templates/SharedPointer.h b/Engine/Source/Runtime/Core/Public/Templates/SharedPointer.h
--- a/Engine/Source/Runtime/Core/Public/Templates/SharedPointer.h
+++ b/Engine/Source/Runtime/Core/Public/Templates/SharedPointer.h
@@ -42,10 +42,8 @@
  *		TSharedFromThis - You can derive your own class from this to acquire a TSharedRef from "this"
  *		StaticCastSharedRef() - Static cast utility function, typically used to downcast to a derived type. 
  *		ConstCastSharedRef() - Converts a 'const' reference to 'mutable' smart reference
- *		DynamicCastSharedRef() - Static cast utility function, typically used to downcast to a derived type. 
  *		StaticCastSharedPtr() - Dynamic cast utility function, typically used to downcast to a derived type. 
  *		ConstCastSharedPtr() - Converts a 'const' smart pointer to 'mutable' smart pointer
- *		DynamicCastSharedPtr() - Dynamic cast utility function, typically used to downcast to a derived type. 
  *
  *
  *	Examples:
@@ -124,18 +122,6 @@ FORCEINLINE TSharedRef< CastToType, Mode > StaticCastSharedRef( TSharedRef< Cast
 }
 
 
-/**
- * Casts a shared reference of one type to another type. (dynamic_cast)  Useful for downcasting.
- *
- * @param  InSharedRef  The shared reference to cast
- */
-template< class CastToType, class CastFromType, ESPMode::Type Mode >
-FORCEINLINE TSharedRef< CastToType, Mode > DynamicCastSharedRef( TSharedRef< CastFromType, Mode > const& InSharedRef )
-{
-	return TSharedRef< CastToType, Mode >( InSharedRef, SharedPointerInternals::FDynamicCastTag() );
-}
-
-
 /**
  * This tricky code is used to prevent the use of TSharedPtrs and TSharedRefs with UObjects
  */
@@ -246,22 +232,6 @@ public:
 	}
 	
   
-	/**
-	 * Special constructor used internally to dynamically cast one shared reference type to another.  You
-	 * should never call this constructor directly.  Instead, use the DynamicCastSharedRef() function.
-	 * This constructor creates a shared reference as a shared reference to an existing shared reference after
-	 * dynamically casting that reference's object.  This constructor is needed for dynamic casts.
-	 *
-	 * @param  InSharedRef  The shared reference whose object we should create an additional reference to
-	 */
-	template< class OtherType >
-	FORCEINLINE TSharedRef( TSharedRef< OtherType, Mode > const& InSharedRef, SharedPointerInternals::FDynamicCastTag )
-		: Object( dynamic_cast< ObjectType* >( InSharedRef.Object ) ),
-		  SharedReferenceCount( InSharedRef.SharedReferenceCount )
-	{
-	}
-
-
 	/**
 	 * Special constructor used internally to create a shared reference from an existing shared reference,
 	 * while using the specified object reference instead of the incoming shared reference's object
@@ -276,7 +246,19 @@ public:
 		  SharedReferenceCount( OtherSharedRef.SharedReferenceCount )
 	{
 	}
-
+	FORCEINLINE TSharedRef( TSharedRef const& InSharedRef )
+		: Object( InSharedRef.Object ),
+		  SharedReferenceCount( InSharedRef.SharedReferenceCount )
+	{
+	}
+	FORCEINLINE TSharedRef( TSharedRef&& InSharedRef )
+		: Object( InSharedRef.Object ),
+		  SharedReferenceCount( InSharedRef.SharedReferenceCount )
+	{
+		// We're intentionally not moving here, because we don't want to leave InSharedRef in a
+		// null state, because that breaks the class invariant.  But we provide a move constructor
+		// anyway in case the compiler complains that we have a move assign but no move construct.
+	}
 
 	/**
 	 * Assignment operator replaces this shared reference with the specified shared reference.  The object
@@ -287,12 +269,15 @@ public:
 	 */
 	FORCEINLINE TSharedRef& operator=( TSharedRef const& InSharedRef )
 	{
-		SharedPointerInternals::FSharedReferencer< Mode > OldSharedReferenceCount = SharedReferenceCount;
 		SharedReferenceCount = InSharedRef.SharedReferenceCount;
 		Object = InSharedRef.Object;
 		return *this;
 	}
-
+	FORCEINLINE TSharedRef& operator=( TSharedRef&& InSharedRef )
+	{
+		FMemory::Memswap(this, &InSharedRef, sizeof(TSharedRef));
+		return *this;
+	}
 
 	/**
 	 * Assignment operator replaces this shared reference with the specified shared reference.  The object
@@ -397,6 +382,17 @@ private:
 		check( IsValid() );
 	}
 
+	template< class OtherType >
+	FORCEINLINE explicit TSharedRef( TSharedPtr< OtherType, Mode >&& InSharedPtr )
+		: Object( InSharedPtr.Object ),
+		  SharedReferenceCount( MoveTemp(InSharedPtr.SharedReferenceCount) )
+	{
+		InSharedPtr.Object = NULL;
+
+		// If this assert goes off, it means a shared reference was created from a shared pointer that was NULL.
+		// Shared references are never allowed to be null.  Consider using TSharedPtr instead.
+		check( IsValid() );
+	}
 
 	/**
 	 * Checks to see if this shared reference is actually pointing to an object. 
@@ -510,6 +506,17 @@ public:
 		  SharedReferenceCount( InSharedPtr.SharedReferenceCount )
 	{
 	}
+	FORCEINLINE TSharedPtr( TSharedPtr const& InSharedPtr )
+		: Object( InSharedPtr.Object ),
+		  SharedReferenceCount( InSharedPtr.SharedReferenceCount )
+	{
+	}
+	FORCEINLINE TSharedPtr( TSharedPtr&& InSharedPtr )
+		: Object( InSharedPtr.Object ),
+		  SharedReferenceCount( MoveTemp(InSharedPtr.SharedReferenceCount) )
+	{
+		InSharedPtr.Object = NULL;
+	}
 
 
 	/**
@@ -524,6 +531,8 @@ public:
 		: Object( InSharedRef.Object ),
 		  SharedReferenceCount( InSharedRef.SharedReferenceCount )
 	{
+		// There is no rvalue overload of this constructor, because 'stealing' the pointer from a
+		// TSharedRef would leave it as null, which would invalidate its invariant.
 	}
 
 
@@ -559,22 +568,6 @@ public:
 	}
 	
   
-	/**
-	 * Special constructor used internally to dynamically cast one shared pointer type to another.  You
-	 * should never call this constructor directly.  Instead, use the DynamicCastSharedPtr() function.
-	 * This constructor creates a shared pointer as a shared reference to an existing shared pointer after
-	 * dynamically casting that pointer's object.  This constructor is needed for dynamic casts.
-	 *
-	 * @param  InSharedPtr  The shared pointer whose object we should create an additional reference to
-	 */
-	template< class OtherType >
-	FORCEINLINE TSharedPtr( TSharedPtr< OtherType, Mode > const& InSharedPtr, SharedPointerInternals::FDynamicCastTag )
-		: Object( dynamic_cast< ObjectType* >( InSharedPtr.Object ) ),
-		  SharedReferenceCount( InSharedPtr.SharedReferenceCount )
-	{
-	}
-	
-  
 	/**
 	 * Special constructor used internally to create a shared pointer from an existing shared pointer,
 	 * while using the specified object pointer instead of the incoming shared pointer's object
@@ -612,12 +605,20 @@ public:
 	 */
 	FORCEINLINE TSharedPtr& operator=( TSharedPtr const& InSharedPtr )
 	{
-		SharedPointerInternals::FSharedReferencer< Mode > OldSharedReferenceCount = SharedReferenceCount;
 		SharedReferenceCount = InSharedPtr.SharedReferenceCount;
 		Object = InSharedPtr.Object;
 		return *this;
 	}
-
+	FORCEINLINE TSharedPtr& operator=( TSharedPtr&& InSharedPtr )
+	{
+		if (this != &InSharedPtr)
+		{
+			Object = InSharedPtr.Object;
+			InSharedPtr.Object = NULL;
+			SharedReferenceCount = MoveTemp(InSharedPtr.SharedReferenceCount);
+		}
+		return *this;
+	}
 
 	/**
 	 * Assignment operator replaces this shared pointer with the specified shared pointer.  The object
@@ -851,6 +852,24 @@ public:
 		  WeakReferenceCount( InWeakPtr.WeakReferenceCount )
 	{
 	}
+	template< class OtherType >
+	FORCEINLINE TWeakPtr( TWeakPtr< OtherType, Mode >&& InWeakPtr )
+		: Object            ( InWeakPtr.Object ),
+		  WeakReferenceCount( MoveTemp(InWeakPtr.WeakReferenceCount) )
+	{
+		InWeakPtr.Object = NULL;
+	}
+	FORCEINLINE TWeakPtr( TWeakPtr const& InWeakPtr )
+		: Object            ( InWeakPtr.Object ),
+		  WeakReferenceCount( InWeakPtr.WeakReferenceCount )
+	{
+	}
+	FORCEINLINE TWeakPtr( TWeakPtr&& InWeakPtr )
+		: Object            ( InWeakPtr.Object ),
+		  WeakReferenceCount( MoveTemp(InWeakPtr.WeakReferenceCount) )
+	{
+		InWeakPtr.Object = NULL;
+	}
 
 
 	/**
@@ -875,6 +894,16 @@ public:
 		WeakReferenceCount = InWeakPtr.WeakReferenceCount;
 		return *this;
 	}
+	FORCEINLINE TWeakPtr& operator=( TWeakPtr&& InWeakPtr )
+	{
+		if (this != &InWeakPtr)
+		{
+			Object             = InWeakPtr.Object;
+			InWeakPtr.Object   = NULL;
+			WeakReferenceCount = MoveTemp(InWeakPtr.WeakReferenceCount);
+		}
+		return *this;
+	}
 
 
 	/**
@@ -890,6 +919,14 @@ public:
 		WeakReferenceCount = InWeakPtr.WeakReferenceCount;
 		return *this;
 	}
+	template <typename OtherType>
+	FORCEINLINE TWeakPtr& operator=( TWeakPtr<OtherType, Mode>&& InWeakPtr )
+	{
+		Object             = InWeakPtr.Object;
+		InWeakPtr.Object   = NULL;
+		WeakReferenceCount = MoveTemp(InWeakPtr.WeakReferenceCount);
+		return *this;
+	}
 
 
 	/**
@@ -1369,19 +1406,6 @@ FORCEINLINE TSharedRef< CastToType, Mode > ConstCastSharedRef( TSharedRef< CastF
 
 
 
-/**
- * Casts a shared pointer of one type to another type. (dynamic_cast)  Useful for downcasting.
- *
- * @param  InSharedPtr  The shared pointer to cast
- */
-template< class CastToType, class CastFromType, ESPMode::Type Mode >
-FORCEINLINE TSharedPtr< CastToType, Mode > DynamicCastSharedPtr( TSharedPtr< CastFromType, Mode > const& InSharedPtr )
-{
-	return TSharedPtr< CastToType, Mode >( InSharedPtr, SharedPointerInternals::FDynamicCastTag() );
-}
-
-
-
 /**
  * Casts a 'const' shared pointer to 'mutable' shared pointer. (const_cast)
  *
diff --git a/Engine/Source/Runtime/Core/Public/Templates/SharedPointerInternals.h b/Engine/Source/Runtime/Core/Public/Templates/SharedPointerInternals.h
--- a/Engine/Source/Runtime/Core/Public/Templates/SharedPointerInternals.h
+++ b/Engine/Source/Runtime/Core/Public/Templates/SharedPointerInternals.h
@@ -57,7 +57,6 @@ namespace SharedPointerInternals
 	/** Dummy structs used internally as template arguments for typecasts */
 	struct FStaticCastTag {};
 	struct FConstCastTag {};
-	struct FDynamicCastTag {};
 
 	// NOTE: The following is an Unreal extension to standard shared_ptr behavior
 	struct FNullTag {};
@@ -79,266 +78,207 @@ namespace SharedPointerInternals
 		}
 	};
 
+	struct FReferenceControllerState
+	{
+		/** Constructor */
+		FORCEINLINE FReferenceControllerState(void* InObject, void (*InDestroyObject)(void*))
+			: SharedReferenceCount(1),
+			  WeakReferenceCount  (1),
+			  Object              (InObject),
+			  DestroyObject       (InDestroyObject)
+		{
+		}
+
+		// NOTE: The primary reason these reference counters are 32-bit values (and not 16-bit to save
+		//       memory), is that atomic operations require at least 32-bit objects.
 
+		/** Number of shared references to this object.  When this count reaches zero, the associated object
+		    will be destroyed (even if there are still weak references!) */
+		int32 SharedReferenceCount;
+
+		/** Number of weak references to this object.  If there are any shared references, that counts as one
+		   weak reference too. */
+		int32 WeakReferenceCount;
+
+		/** The object associated with this reference counter.  */
+		void* Object;
+
+		/** Destroys the object associated with this reference counter.  */
+		void (*DestroyObject)(void*);
+
+	private:
+		FReferenceControllerState( FReferenceControllerState const& );
+		FReferenceControllerState& operator=( FReferenceControllerState const& );
+	};
 
 	/**
 	 * FReferenceController is a standalone heap-allocated object that tracks the number of references
 	 * to an object referenced by TSharedRef, TSharedPtr or TWeakPtr objects.
+	 *
+	 * It is specialized for different threading modes.
 	 */
 	template< ESPMode::Type Mode >
-	class FReferenceController
-	{
-
-	public:
-		// Note: Intentionally not including a virtual destructor here because it causes problems using weak pointers across dlls
-		// If a module where the object and shared pointer controller was created are unloaded before the module with weak references to it is unloaded then the module with weak references
-		// will crash deleting this controller because its vtable has been destroyed so it cant determine the proper destructor to call if its virtual.  The virtual destructor does not provide any benifits
-		// here because there is nothing they need to delete.
-
-		/** Default constructor */
-		FORCEINLINE FReferenceController()
-			: SharedReferenceCount( 1 ),
-			  WeakReferenceCount( 1 )
-		{
-		}
-
+	struct FReferenceControllerOps;
 
+	template<>
+	struct FReferenceControllerOps<ESPMode::ThreadSafe>
+	{
 		/** Returns the shared reference count */
-		FORCEINLINE const int32 GetSharedReferenceCount() const
+		static FORCEINLINE const int32 GetSharedReferenceCount(const FReferenceControllerState* State)
 		{
-			if( Mode == ESPMode::ThreadSafe )
-			{
-				// This reference count may be accessed by multiple threads
-				return static_cast< int32 const volatile& >( SharedReferenceCount );
-			}
-			else
-			{
-				return SharedReferenceCount;
-			}
+			// This reference count may be accessed by multiple threads
+			return static_cast< int32 const volatile& >( State->SharedReferenceCount );
 		}
 
-
 		/** Adds a shared reference to this counter */
-		FORCEINLINE void AddSharedReference()
+		static FORCEINLINE void AddSharedReference(FReferenceControllerState* State)
 		{
-			if( Mode == ESPMode::ThreadSafe )
-			{
-				FPlatformAtomics::InterlockedIncrement( &SharedReferenceCount );
-			}
-			else
-			{
-				++SharedReferenceCount;
-			}
+			FPlatformAtomics::InterlockedIncrement( &State->SharedReferenceCount );
 		}
 
-
 		/**
 		 * Adds a shared reference to this counter ONLY if there is already at least one reference
 		 *
 		 * @return  True if the shared reference was added successfully
 		 */
-		inline bool ConditionallyAddSharedReference()
+		static bool ConditionallyAddSharedReference(FReferenceControllerState* State)
 		{
-			if( Mode == ESPMode::ThreadSafe )
+			for( ; ; )
 			{
-				for( ; ; )
-				{
-					// Peek at the current shared reference count.  Remember, this value may be updated by
-					// multiple threads.
-					const int32 OriginalCount = static_cast< int32 const volatile& >( SharedReferenceCount );
-					if( OriginalCount == 0 )
-					{
-						// Never add a shared reference if the pointer has already expired
-						return false;
-					}
-
-					// Attempt to increment the reference count.
-					const int32 ActualOriginalCount = FPlatformAtomics::InterlockedCompareExchange( &SharedReferenceCount, OriginalCount + 1, OriginalCount );
-
-					// We need to make sure that we never revive a counter that has already expired, so if the
-					// actual value what we expected (because it was touched by another thread), then we'll try
-					// again.  Note that only in very unusual cases will this actually have to loop.
-					if( ActualOriginalCount == OriginalCount )
-					{
-						return true;
-					}
-				}
-			}
-			else
-			{
-				if( SharedReferenceCount == 0 )
+				// Peek at the current shared reference count.  Remember, this value may be updated by
+				// multiple threads.
+				const int32 OriginalCount = static_cast< int32 const volatile& >( State->SharedReferenceCount );
+				if( OriginalCount == 0 )
 				{
 					// Never add a shared reference if the pointer has already expired
 					return false;
 				}
 
-				++SharedReferenceCount;
-				return true;
+				// Attempt to increment the reference count.
+				const int32 ActualOriginalCount = FPlatformAtomics::InterlockedCompareExchange( &State->SharedReferenceCount, OriginalCount + 1, OriginalCount );
+
+				// We need to make sure that we never revive a counter that has already expired, so if the
+				// actual value what we expected (because it was touched by another thread), then we'll try
+				// again.  Note that only in very unusual cases will this actually have to loop.
+				if( ActualOriginalCount == OriginalCount )
+				{
+					return true;
+				}
 			}
 		}
 
 
 		/** Releases a shared reference to this counter */
-		FORCEINLINE void ReleaseSharedReference()
+		static FORCEINLINE void ReleaseSharedReference(FReferenceControllerState* State)
 		{
-			checkSlow( SharedReferenceCount > 0 );
-			if( Mode == ESPMode::ThreadSafe )
-			{
-				if( FPlatformAtomics::InterlockedDecrement( &SharedReferenceCount ) == 0 )
-				{
-					// Last shared reference was released!  Destroy the referenced object.
-					DestroyObject();
+			checkSlow( State->SharedReferenceCount > 0 );
 
-					// No more shared referencers, so decrement the weak reference count by one.  When the weak
-					// reference count reaches zero, this object will be deleted.
-					ReleaseWeakReference();
-				}
-			}
-			else
+			if( FPlatformAtomics::InterlockedDecrement( &State->SharedReferenceCount ) == 0 )
 			{
-				if( --SharedReferenceCount == 0 )
-				{
-					// Last shared reference was released!  Destroy the referenced object.
-					DestroyObject();
+				// Last shared reference was released!  Destroy the referenced object.
+				State->DestroyObject(State->Object);
 
-					// No more shared referencers, so decrement the weak reference count by one.  When the weak
-					// reference count reaches zero, this object will be deleted.
-					ReleaseWeakReference();
-				}
+				// No more shared referencers, so decrement the weak reference count by one.  When the weak
+				// reference count reaches zero, this object will be deleted.
+				ReleaseWeakReference(State);
 			}
 		}
 
 
 		/** Adds a weak reference to this counter */
-		FORCEINLINE void AddWeakReference()
+		static FORCEINLINE void AddWeakReference(FReferenceControllerState* State)
 		{
-			if( Mode == ESPMode::ThreadSafe )
-			{
-				FPlatformAtomics::InterlockedIncrement( &WeakReferenceCount );
-			}
-			else
-			{
-				++WeakReferenceCount;
-			}
+			FPlatformAtomics::InterlockedIncrement( &State->WeakReferenceCount );
 		}
 
-// We intentionally don't have a virtual destructor to avoid problems with the destructor being executed when
-// pointers are cleaned up at shutdown. Some pointers being cleaned up are to objects that were instantiated
-// in other modules, even modules that may have been unloaded already.
-#ifdef __clang__
-#pragma clang diagnostic push
-#pragma clang diagnostic ignored "-Wdelete-non-virtual-dtor"
-#endif
-
 		/** Releases a weak reference to this counter */
-		inline void ReleaseWeakReference()
+		static void ReleaseWeakReference(FReferenceControllerState* State)
 		{
-			checkSlow( WeakReferenceCount > 0 );
-			if( Mode == ESPMode::ThreadSafe )
+			checkSlow( State->WeakReferenceCount > 0 );
+
+			if( FPlatformAtomics::InterlockedDecrement( &State->WeakReferenceCount ) == 0 )
 			{
-				if( FPlatformAtomics::InterlockedDecrement( &WeakReferenceCount ) == 0 )
-				{
-					// No more references to this reference count.  Destroy ourselves!
-					delete this;
-				}
-			}
-			else
-			{
-				if( --WeakReferenceCount == 0 )
-				{
-					// No more references to this reference count.  Destroy ourselves!
-					delete this;
-				}
+				// No more references to this reference count.  Destroy it!
+				delete State;
 			}
 		}
-
-#ifdef __clang__
-#pragma clang diagnostic pop
-#endif
-
-	protected:
-
-		/** Destroys the object associated with this reference counter.  Implemented in derived classes. */
-		virtual void DestroyObject() = 0;
-
-
-	private:
-
-		/** Hide generated copy constructor */
-		FReferenceController( FReferenceController const& );
-
-		/** Hide generated assignment operator */
-		FReferenceController& operator=( FReferenceController const& );
-
-
-	private:
-
-		// NOTE: The primary reason these reference counters are 32-bit values (and not 16-bit to save
-		//       memory), is that atomic operations require at least 32-bit objects.
-
-		/** Number of shared references to this object.  When this count reaches zero, the associated object
-		    will be destroyed (even if there are still weak references!) */
-		int32 SharedReferenceCount;
-
-		/** Number of weak references to this object.  If there are any shared references, that counts as one
-		   weak reference too. */
-		int32 WeakReferenceCount;
 	};
 
-
-
-	/**
-	 * Polymorphic implementation of FReferenceController that contains the actual referenced object.
-	 * The interesting bit here is that this object is instantiated with a templated ObjectType that
-	 * matches the object's type as it was when the first reference to it was created.  This allows
-	 * us to delete the object as the same type it was originally instantiated with.
-	 */
-	template< class ObjectType, ESPMode::Type Mode >
-	class FReferenceControllerImpl
-		: public FReferenceController< Mode >
+	template<>
+	struct FReferenceControllerOps<ESPMode::NotThreadSafe>
 	{
-
-	public:
-
-		/** Construct with a single reference to the specified object (reference count = 1) */
-		FORCEINLINE explicit FReferenceControllerImpl( ObjectType* InObject )
-			: Object( InObject )
+		/** Returns the shared reference count */
+		static FORCEINLINE const int32 GetSharedReferenceCount(const FReferenceControllerState* State)
 		{
+			return State->SharedReferenceCount;
 		}
 
+		/** Adds a shared reference to this counter */
+		static FORCEINLINE void AddSharedReference(FReferenceControllerState* State)
+		{
+			++State->SharedReferenceCount;
+		}
 
-	protected:
-
-		/** FReferenceController: Destroys the object associated with this reference counter */
-		virtual void DestroyObject()
+		/**
+		 * Adds a shared reference to this counter ONLY if there is already at least one reference
+		 *
+		 * @return  True if the shared reference was added successfully
+		 */
+		static bool ConditionallyAddSharedReference(FReferenceControllerState* State)
 		{
-			// Destroy the object.  This function is called with the number of shared references reaches zero.
-			delete Object;
+			if( State->SharedReferenceCount == 0 )
+			{
+				// Never add a shared reference if the pointer has already expired
+				return false;
+			}
+
+			++State->SharedReferenceCount;
+			return true;
 		}
 
 
-	private:
+		/** Releases a shared reference to this counter */
+		static FORCEINLINE void ReleaseSharedReference(FReferenceControllerState* State)
+		{
+			checkSlow( State->SharedReferenceCount > 0 );
 
-		/** Hide generated copy constructor */
-		FReferenceControllerImpl( FReferenceControllerImpl const& );
+			if( --State->SharedReferenceCount == 0 )
+			{
+				// Last shared reference was released!  Destroy the referenced object.
+				State->DestroyObject(State->Object);
 
-		/** Hide generated assignment operator */
-		FReferenceControllerImpl& operator=( FReferenceControllerImpl const& );
+				// No more shared referencers, so decrement the weak reference count by one.  When the weak
+				// reference count reaches zero, this object will be deleted.
+				ReleaseWeakReference(State);
+			}
+		}
 
 
-	private:
+		/** Adds a weak reference to this counter */
+		static FORCEINLINE void AddWeakReference(FReferenceControllerState* State)
+		{
+			++State->WeakReferenceCount;
+		}
 
-		// NOTE: It's important to understand why the reference controller has it's own copy of the
-		//       the object.  This is because we want to be able to delete the object as the type
-		//       that it was originally instantiated as, and even if we retain only the type
-		//       information (and no object pointer), it would be difficult to pass down the object
-		//       pointer in all possible cases where the object could be destroyed.
+		/** Releases a weak reference to this counter */
+		static void ReleaseWeakReference(FReferenceControllerState* State)
+		{
+			checkSlow( State->WeakReferenceCount > 0 );
 
-		/** The object we're holding a reference to */
-		ObjectType* Object;
+			if( --State->WeakReferenceCount == 0 )
+			{
+				// No more references to this reference count.  Destroy it!
+				delete State;
+			}
+		}
 	};
 
-
+	/** Type-erasing function deleting objects of a given type via a common pointer type (void*) */
+	template <typename Type>
+	void DestroyObject(void* Object)
+	{
+		delete (Type*)Object;
+	}
 
 	/**
 	 * FSharedReferencer is a wrapper around a pointer to a reference controller that is used by either a
@@ -347,6 +287,7 @@ namespace SharedPointerInternals
 	template< ESPMode::Type Mode >
 	class FSharedReferencer
 	{
+		typedef FReferenceControllerOps<Mode> TOps;
 
 	public:
 
@@ -360,7 +301,7 @@ namespace SharedPointerInternals
 		/** Constructor that counts a single reference to the specified object */
 		template< class OtherType > 
 		inline explicit FSharedReferencer( OtherType* InObject )
-			: ReferenceController( new FReferenceControllerImpl< OtherType, Mode >( InObject ) )
+			: ReferenceController( new FReferenceControllerState( InObject, &DestroyObject<OtherType> ) )
 		{
 		}
 
@@ -373,14 +314,36 @@ namespace SharedPointerInternals
 			// shared reference count
 			if( ReferenceController != NULL )
 			{
-				ReferenceController->AddSharedReference();
+				TOps::AddSharedReference(ReferenceController);
 			}
 		}
 
+		/** Move constructor creates no new references */
+		FORCEINLINE FSharedReferencer( FSharedReferencer&& InSharedReference )
+			: ReferenceController( InSharedReference.ReferenceController )
+		{
+			InSharedReference.ReferenceController = NULL;
+		}
+
 
 		/** Creates a shared referencer object from a weak referencer object.  This will only result
 		    in a valid object reference if the object already has at least one other shared referencer. */
-		inline FSharedReferencer( FWeakReferencer< Mode > const& InWeakReference );
+		FSharedReferencer( FWeakReferencer< Mode > const& InWeakReference )
+			: ReferenceController( InWeakReference.ReferenceController )
+		{
+			// If the incoming reference had an object associated with it, then go ahead and increment the
+			// shared reference count
+			if( ReferenceController != NULL )
+			{
+				// Attempt to elevate a weak reference to a shared one.  For this to work, the object this
+				// weak counter is associated with must already have at least one shared reference.  We'll
+				// never revive a pointer that has already expired!
+				if( !TOps::ConditionallyAddSharedReference(ReferenceController) )
+				{
+					ReferenceController = NULL;
+				}
+			}
+		}
 
 
 		/** Destructor */
@@ -390,29 +353,28 @@ namespace SharedPointerInternals
 			{
 				// Tell the reference counter object that we're no longer referencing the object with
 				// this shared pointer
-				ReferenceController->ReleaseSharedReference();
+				TOps::ReleaseSharedReference(ReferenceController);
 			}
 		}
 
-		
 		/** Assignment operator adds a reference to the assigned object.  If this counter was previously
 		    referencing an object, that reference will be released. */
 		inline FSharedReferencer& operator=( FSharedReferencer const& InSharedReference )
 		{
 			// Make sure we're not be reassigned to ourself!
-			FReferenceController< Mode >* NewReferenceController = InSharedReference.ReferenceController;
+			auto NewReferenceController = InSharedReference.ReferenceController;
 			if( NewReferenceController != ReferenceController )
 			{
 				// First, add a shared reference to the new object
 				if( NewReferenceController != NULL )
 				{
-					NewReferenceController->AddSharedReference();
+					TOps::AddSharedReference(NewReferenceController);
 				}
 
 				// Release shared reference to the old object
 				if( ReferenceController != NULL )
 				{
-					ReferenceController->ReleaseSharedReference();
+					TOps::ReleaseSharedReference(ReferenceController);
 				}
 
 				// Assume ownership of the assigned reference counter
@@ -422,6 +384,29 @@ namespace SharedPointerInternals
 			return *this;
 		}
 
+		/** Move assignment operator adds no references to the assigned object.  If this counter was previously
+		    referencing an object, that reference will be released. */
+		inline FSharedReferencer& operator=( FSharedReferencer&& InSharedReference )
+		{
+			// Make sure we're not be reassigned to ourself!
+			auto NewReferenceController = InSharedReference.ReferenceController;
+			auto OldReferenceController = ReferenceController;
+			if( NewReferenceController != OldReferenceController )
+			{
+				// Assume ownership of the assigned reference counter
+				InSharedReference.ReferenceController = NULL;
+				ReferenceController                   = NewReferenceController;
+
+				// Release shared reference to the old object
+				if( OldReferenceController != NULL )
+				{
+					TOps::ReleaseSharedReference(OldReferenceController);
+				}
+			}
+
+			return *this;
+		}
+
 
 		/**
 		 * Tests to see whether or not this shared counter contains a valid reference
@@ -441,7 +426,7 @@ namespace SharedPointerInternals
 		 */
 		FORCEINLINE const int32 GetSharedReferenceCount() const
 		{
-			return ReferenceController != NULL ? ReferenceController->GetSharedReferenceCount() : 0;
+			return ReferenceController != NULL ? TOps::GetSharedReferenceCount(ReferenceController) : 0;
 		}
 
 
@@ -466,7 +451,7 @@ namespace SharedPointerInternals
 	private:
 
 		/** Pointer to the reference controller for the object a shared reference/pointer is referencing */
-		FReferenceController< Mode >* ReferenceController;
+		FReferenceControllerState* ReferenceController;
 	};
 
 
@@ -478,6 +463,7 @@ namespace SharedPointerInternals
 	template< ESPMode::Type Mode >
 	class FWeakReferencer
 	{
+		typedef FReferenceControllerOps<Mode> TOps;
 
 	public:
 
@@ -495,11 +481,19 @@ namespace SharedPointerInternals
 			// If the weak referencer has a valid controller, then go ahead and add a weak reference to it!
 			if( ReferenceController != NULL )
 			{
-				ReferenceController->AddWeakReference();
+				TOps::AddWeakReference(ReferenceController);
 			}
 		}
 
 
+		/** Construct a weak referencer object from an rvalue weak referencer */
+		FORCEINLINE FWeakReferencer( FWeakReferencer&& InWeakRefCountPointer )
+			: ReferenceController( InWeakRefCountPointer.ReferenceController )
+		{
+			InWeakRefCountPointer.ReferenceController = NULL;
+		}
+
+
 		/** Construct a weak referencer object from a shared referencer object */
 		FORCEINLINE FWeakReferencer( FSharedReferencer< Mode > const& InSharedRefCountPointer )
 			: ReferenceController( InSharedRefCountPointer.ReferenceController )
@@ -507,7 +501,7 @@ namespace SharedPointerInternals
 			// If the shared referencer had a valid controller, then go ahead and add a weak reference to it!
 			if( ReferenceController != NULL )
 			{
-				ReferenceController->AddWeakReference();
+				TOps::AddWeakReference(ReferenceController);
 			}
 		}
 
@@ -519,7 +513,7 @@ namespace SharedPointerInternals
 			{
 				// Tell the reference counter object that we're no longer referencing the object with
 				// this weak pointer
-				ReferenceController->ReleaseWeakReference();
+				TOps::ReleaseWeakReference(ReferenceController);
 			}
 		}
 
@@ -534,6 +528,22 @@ namespace SharedPointerInternals
 		}
 
 
+		/** Assignment operator from an rvalue weak referencer object.  If this counter was previously referencing an
+		    object, that reference will be released. */
+		FORCEINLINE FWeakReferencer& operator=( FWeakReferencer&& InWeakReference )
+		{
+			auto OldReferenceController         = ReferenceController;
+			ReferenceController                 = InWeakReference.ReferenceController;
+			InWeakReference.ReferenceController = NULL;
+			if( OldReferenceController != NULL )
+			{
+				TOps::ReleaseWeakReference(OldReferenceController);
+			}
+
+			return *this;
+		}
+
+
 		/** Assignment operator from a shared reference counter.  If this counter was previously referencing an
 		   object, that reference will be released. */
 		FORCEINLINE FWeakReferencer& operator=( FSharedReferencer< Mode > const& InSharedReference )
@@ -551,7 +561,7 @@ namespace SharedPointerInternals
 		 */
 		FORCEINLINE const bool IsValid() const
 		{
-			return ReferenceController != NULL && ReferenceController->GetSharedReferenceCount() > 0;
+			return ReferenceController != NULL && TOps::GetSharedReferenceCount(ReferenceController) > 0;
 		}
 
 
@@ -559,7 +569,7 @@ namespace SharedPointerInternals
 
 		/** Assigns a new reference controller to this counter object, first adding a reference to it, then
 		    releasing the previous object. */
-		inline void AssignReferenceController( FReferenceController< Mode >* NewReferenceController )
+		inline void AssignReferenceController( FReferenceControllerState* NewReferenceController )
 		{
 			// Only proceed if the new reference counter is different than our current
 			if( NewReferenceController != ReferenceController )
@@ -567,13 +577,13 @@ namespace SharedPointerInternals
 				// First, add a weak reference to the new object
 				if( NewReferenceController != NULL )
 				{
-					NewReferenceController->AddWeakReference();
+					TOps::AddWeakReference(NewReferenceController);
 				}
 
 				// Release weak reference to the old object
 				if( ReferenceController != NULL )
 				{
-					ReferenceController->ReleaseWeakReference();
+					TOps::ReleaseWeakReference(ReferenceController);
 				}
 
 				// Assume ownership of the assigned reference counter
@@ -591,31 +601,10 @@ namespace SharedPointerInternals
 	private:
 
 		/** Pointer to the reference controller for the object a TWeakPtr is referencing */
-		FReferenceController< Mode >* ReferenceController;
+		FReferenceControllerState* ReferenceController;
 	};
 
 
-	
-	/** Construct from a weak referencer object */
-	template< ESPMode::Type Mode >
-	inline FSharedReferencer< Mode >::FSharedReferencer( FWeakReferencer< Mode > const& InWeakReference )
-		: ReferenceController( InWeakReference.ReferenceController )
-	{
-		// If the incoming reference had an object associated with it, then go ahead and increment the
-		// shared reference count
-		if( ReferenceController != NULL )
-		{
-			// Attempt to elevate a weak reference to a shared one.  For this to work, the object this
-			// weak counter is associated with must already have at least one shared reference.  We'll
-			// never revive a pointer that has already expired!
-			if( !ReferenceController->ConditionallyAddSharedReference() )
-			{
-				ReferenceController = NULL;
-			}
-		}
-	}
-
-
 	/** Templated helper function (const) that creates a shared pointer from an object instance */
 	template< class SharedPtrType, class ObjectType, class OtherType, ESPMode::Type Mode >
 	FORCEINLINE void EnableSharedFromThis( TSharedPtr< SharedPtrType, Mode > const* InSharedPtr, ObjectType const* InObject, TSharedFromThis< OtherType, Mode > const* InShareable )
