Fix an occasional crash when a modal dialog appears and some input is allowed to re-enter the editor code.

diff --git a/Engine/Source/Runtime/Core/Private/Windows/WindowsApplication.cpp b/Engine/Source/Runtime/Core/Private/Windows/WindowsApplication.cpp
--- a/Engine/Source/Runtime/Core/Private/Windows/WindowsApplication.cpp
+++ b/Engine/Source/Runtime/Core/Private/Windows/WindowsApplication.cpp
@@ -709,10 +709,18 @@ int32 FWindowsApplication::ProcessMessage( HWND hwnd, uint32 msg, WPARAM wParam,
 					if ( CurrentNativeEventWindow->IsRegularWindow() )
 					{
 						EWindowZone::Type Zone;
-						// Assumes this is not allowed to leave Slate or touch rendering
+					
+						if( MessageHandler->ShouldProcessUserInputMessages( CurrentNativeEventWindowPtr ) )
 						{
+							// Assumes this is not allowed to leave Slate or touch rendering
 							Zone = MessageHandler->GetWindowZoneForPoint( CurrentNativeEventWindow, LocalMouseX, LocalMouseY );
 						}
+						else
+						{
+							// Default to client area so that we are able to see the feedback effect when attempting to click on a non-modal window when a modal window is active
+							// Any other window zones could have side effects and NotInWindow prevents the feedback effect.
+							Zone = EWindowZone::ClientArea;
+						}
 
 						static const LRESULT Results[] = {HTNOWHERE, HTTOPLEFT, HTTOP, HTTOPRIGHT, HTLEFT, HTCLIENT,
 							HTRIGHT, HTBOTTOMLEFT, HTBOTTOM, HTBOTTOMRIGHT,
