UE4: Added some more temp logging to track down a crash when cooking.

diff --git a/Engine/Source/Runtime/Core/Private/Misc/CoreMisc.cpp b/Engine/Source/Runtime/Core/Private/Misc/CoreMisc.cpp
--- a/Engine/Source/Runtime/Core/Private/Misc/CoreMisc.cpp
+++ b/Engine/Source/Runtime/Core/Private/Misc/CoreMisc.cpp
@@ -718,9 +718,12 @@ class ITargetPlatformManagerModule* GetTargetPlatformManager()
 		static bool bIntialized = false;
 		if (!bIntialized)
 		{
+			UE_LOG(LogInit, Log, TEXT("Begin loading target platform manager module. Pre-check"));
 			check(IsInGameThread());
 			bIntialized = true;
+			UE_LOG(LogInit, Log, TEXT("Begin loading target platform manager module. Post-check"));
 			SingletonInterface = FModuleManager::LoadModulePtr<ITargetPlatformManagerModule>("TargetPlatform");
+			UE_LOG(LogInit, Log, TEXT("Target platform manager module loaded."));
 		}
 	}
 	return SingletonInterface;
