UE4: Fixed a bug that caused non-assert crashes to fail to produce a callstack in the log.

diff --git a/Engine/Source/Runtime/Core/Private/Windows/WindowsPlatformExceptionHandling.cpp b/Engine/Source/Runtime/Core/Private/Windows/WindowsPlatformExceptionHandling.cpp
--- a/Engine/Source/Runtime/Core/Private/Windows/WindowsPlatformExceptionHandling.cpp
+++ b/Engine/Source/Runtime/Core/Private/Windows/WindowsPlatformExceptionHandling.cpp
@@ -449,27 +449,25 @@ int32 ReportCrash( LPEXCEPTION_POINTERS ExceptionInfo )
 #if WINVER > 0x502	// Windows Error Reporting is not supported on Windows XP
 		// Loop in the new system while keeping the old until it's fully implemented
 		NewReportCrash( ExceptionInfo, NULL, EErrorReportUI::ShowDialog );
-
-		if (GUseCrashReportClient)
-		{
-			return EXCEPTION_EXECUTE_HANDLER;
-		}
 #endif		// WINVER
 
-		// Try to create file for minidump.
-		HANDLE FileHandle	= CreateFileW( MiniDumpFilenameW, GENERIC_WRITE, 0, NULL, CREATE_ALWAYS, FILE_ATTRIBUTE_NORMAL, NULL );
-			
-		// Write a minidump.
-		if( FileHandle != INVALID_HANDLE_VALUE )
+		if (!GUseCrashReportClient)
 		{
-			MINIDUMP_EXCEPTION_INFORMATION DumpExceptionInfo;
+			// Try to create file for minidump.
+			HANDLE FileHandle	= CreateFileW( MiniDumpFilenameW, GENERIC_WRITE, 0, NULL, CREATE_ALWAYS, FILE_ATTRIBUTE_NORMAL, NULL );
+			
+			// Write a minidump.
+			if( FileHandle != INVALID_HANDLE_VALUE )
+			{
+				MINIDUMP_EXCEPTION_INFORMATION DumpExceptionInfo;
 
-			DumpExceptionInfo.ThreadId			= GetCurrentThreadId();
-			DumpExceptionInfo.ExceptionPointers	= ExceptionInfo;
-			DumpExceptionInfo.ClientPointers	= false;
+				DumpExceptionInfo.ThreadId			= GetCurrentThreadId();
+				DumpExceptionInfo.ExceptionPointers	= ExceptionInfo;
+				DumpExceptionInfo.ClientPointers	= false;
 
-			MiniDumpWriteDump( GetCurrentProcess(), GetCurrentProcessId(), FileHandle, MiniDumpNormal, &DumpExceptionInfo, NULL, NULL );
-			CloseHandle( FileHandle );
+				MiniDumpWriteDump( GetCurrentProcess(), GetCurrentProcessId(), FileHandle, MiniDumpNormal, &DumpExceptionInfo, NULL, NULL );
+				CloseHandle( FileHandle );
+			}
 		}
 
 		const SIZE_T StackTraceSize = 65535;
@@ -482,8 +480,11 @@ int32 ReportCrash( LPEXCEPTION_POINTERS ExceptionInfo )
 		FMemory::SystemFree( StackTrace );
 
 #if UE_BUILD_SHIPPING && WITH_EDITOR
+		if (!GUseCrashReportClient)
+		{
 			uint32 dwOpt = 0;
 			EFaultRepRetVal repret = ReportFault( ExceptionInfo, dwOpt);
+		}
 #endif
 	}
 
