Back out changelist 2060753. These temp logging changes are no longer needed.

diff --git a/Engine/Source/Runtime/Core/Private/Misc/CoreMisc.cpp b/Engine/Source/Runtime/Core/Private/Misc/CoreMisc.cpp
--- a/Engine/Source/Runtime/Core/Private/Misc/CoreMisc.cpp
+++ b/Engine/Source/Runtime/Core/Private/Misc/CoreMisc.cpp
@@ -718,12 +718,9 @@ class ITargetPlatformManagerModule* GetTargetPlatformManager()
 		static bool bIntialized = false;
 		if (!bIntialized)
 		{
-			UE_LOG(LogInit, Log, TEXT("Begin loading target platform manager module. Pre-check"));
 			check(IsInGameThread());
 			bIntialized = true;
-			UE_LOG(LogInit, Log, TEXT("Begin loading target platform manager module. Post-check"));
 			SingletonInterface = FModuleManager::LoadModulePtr<ITargetPlatformManagerModule>("TargetPlatform");
-			UE_LOG(LogInit, Log, TEXT("Target platform manager module loaded."));
 		}
 	}
 	return SingletonInterface;
