Delegates: consistent line spacing, fixed spelling errors

diff --git a/Engine/Source/Runtime/Core/Public/Templates/DelegateSignatureImpl.inl b/Engine/Source/Runtime/Core/Public/Templates/DelegateSignatureImpl.inl
--- a/Engine/Source/Runtime/Core/Public/Templates/DelegateSignatureImpl.inl
+++ b/Engine/Source/Runtime/Core/Public/Templates/DelegateSignatureImpl.inl
@@ -6,11 +6,10 @@
 #endif
 
 
-// NOTE: This file in re-included for EVERY delegate signature that we support.  That is, every combination
+// NOTE: This file is re-included for EVERY delegate signature that we support.  That is, every combination
 //		 of parameter count, return value presence or other function modifier will include this file to
 //		 generate a delegate interface type and implementation type for that signature.
 
-
 #define DELEGATE_INSTANCE_INTERFACE_CLASS FUNC_COMBINE( IBaseDelegateInstance_, FUNC_SUFFIX )
 
 #define SP_METHOD_DELEGATE_INSTANCE_CLASS_NoVars FUNC_COMBINE( TBaseSPMethodDelegateInstance_, FUNC_SUFFIX )
@@ -85,22 +84,21 @@
 #define PAYLOAD_TEMPLATE_PASSIN_FourVars PAYLOAD_TEMPLATE_PASSIN_ThreeVars, Var4
 
 
-
 /**
- * Single-cast delegate base object.   You'll use the various DECLARE_DELEGATE macros to create the
- * actual delegate type, templated to the function signature the delegate is compatible with.  Then,
- * you can create an instance of that class when you want to bind a function to the delegate.
+ * Single-cast delegate base object.
+ *
+ * Use the various DECLARE_DELEGATE macros to create the actual delegate type, templated to
+ * the function signature the delegate is compatible with. Then, you can create an instance
+ * of that class when you want to bind a function to the delegate.
  */
 template< FUNC_TEMPLATE_DECL_TYPENAME >
 class DELEGATE_CLASS
 {
-
 public:
 
 	/** Declare the user's delegate interface (e.g. IMyDelegate) */
 	typedef DELEGATE_INSTANCE_INTERFACE_CLASS< FUNC_TEMPLATE_ARGS > IDelegate;
 
-
 	/** Declare the user's "fast" shared pointer-based delegate object (e.g. TSPMyDelegate) */
 	template< class UserClass                                  > class TSPMethodDelegate                 : public SP_METHOD_DELEGATE_INSTANCE_CLASS_NoVars         < UserClass, FUNC_TEMPLATE_ARGS,                                  ESPMode::Fast > { public: TSPMethodDelegate                ( const TSharedRef< UserClass, ESPMode::Fast >& InUserObject, typename SP_METHOD_DELEGATE_INSTANCE_CLASS_NoVars         < UserClass, FUNC_TEMPLATE_ARGS,                                  ESPMode::Fast >::FMethodPtr InMethodPtr                                  ) : SP_METHOD_DELEGATE_INSTANCE_CLASS_NoVars         < UserClass, FUNC_TEMPLATE_ARGS,                                  ESPMode::Fast >( InUserObject, InMethodPtr                                    ) {} };
 	template< class UserClass                                  > class TSPMethodDelegate_Const           : public SP_METHOD_DELEGATE_INSTANCE_CLASS_NoVars_Const   < UserClass, FUNC_TEMPLATE_ARGS,                                  ESPMode::Fast > { public: TSPMethodDelegate_Const          ( const TSharedRef< UserClass, ESPMode::Fast >& InUserObject, typename SP_METHOD_DELEGATE_INSTANCE_CLASS_NoVars_Const   < UserClass, FUNC_TEMPLATE_ARGS,                                  ESPMode::Fast >::FMethodPtr InMethodPtr                                  ) : SP_METHOD_DELEGATE_INSTANCE_CLASS_NoVars_Const   < UserClass, FUNC_TEMPLATE_ARGS,                                  ESPMode::Fast >( InUserObject, InMethodPtr                                    ) {} };
@@ -156,6 +154,7 @@ public:
 	template< PAYLOAD_TEMPLATE_DECL_ThreeVars > class TStaticDelegate_ThreeVars : public STATIC_DELEGATE_INSTANCE_CLASS_ThreeVars< FUNC_TEMPLATE_ARGS, PAYLOAD_TEMPLATE_LIST_ThreeVars > { public: TStaticDelegate_ThreeVars( typename STATIC_DELEGATE_INSTANCE_CLASS_ThreeVars< FUNC_TEMPLATE_ARGS, PAYLOAD_TEMPLATE_LIST_ThreeVars >::FFuncPtr InFuncPtr, PAYLOAD_TEMPLATE_ARGS_ThreeVars ) : STATIC_DELEGATE_INSTANCE_CLASS_ThreeVars< FUNC_TEMPLATE_ARGS, PAYLOAD_TEMPLATE_LIST_ThreeVars >( InFuncPtr, PAYLOAD_TEMPLATE_PASSIN_ThreeVars ) {} };
 	template< PAYLOAD_TEMPLATE_DECL_FourVars  > class TStaticDelegate_FourVars  : public STATIC_DELEGATE_INSTANCE_CLASS_FourVars < FUNC_TEMPLATE_ARGS, PAYLOAD_TEMPLATE_LIST_FourVars  > { public: TStaticDelegate_FourVars ( typename STATIC_DELEGATE_INSTANCE_CLASS_FourVars < FUNC_TEMPLATE_ARGS, PAYLOAD_TEMPLATE_LIST_FourVars  >::FFuncPtr InFuncPtr, PAYLOAD_TEMPLATE_ARGS_FourVars  ) : STATIC_DELEGATE_INSTANCE_CLASS_FourVars < FUNC_TEMPLATE_ARGS, PAYLOAD_TEMPLATE_LIST_FourVars  >( InFuncPtr, PAYLOAD_TEMPLATE_PASSIN_FourVars  ) {} };
 
+public:
 
 	/**
 	 * Static: Creates a raw C++ pointer global function delegate
@@ -185,10 +184,11 @@ public:
 		return DELEGATE_CLASS< FUNC_TEMPLATE_ARGS >( TStaticDelegate_FourVars< PAYLOAD_TEMPLATE_LIST_FourVars >::Create( InFunc, PAYLOAD_TEMPLATE_PASSIN_FourVars ) );
 	}
 
-	
 	/**
-	 * Static: Creates a raw C++ pointer member function delegate.  Raw pointer doesn't use any sort of reference, so may be unsafe to call if the object was
-	 * deleted out from underneath your delegate.  Be careful when calling Execute()!
+	 * Static: Creates a raw C++ pointer member function delegate.
+	 *
+	 * Raw pointer doesn't use any sort of reference, so may be unsafe to call if the object was
+	 * deleted out from underneath your delegate. Be careful when calling Execute()!
 	 */
 	template< class UserClass >
 	inline static DELEGATE_CLASS< FUNC_TEMPLATE_ARGS > CreateRaw( UserClass* InUserObject, typename TRawMethodDelegate< UserClass >::FMethodPtr InFunc )
@@ -240,10 +240,12 @@ public:
 	{
 		return DELEGATE_CLASS< FUNC_TEMPLATE_ARGS >( TRawMethodDelegate_FourVars_Const< UserClass, PAYLOAD_TEMPLATE_LIST_FourVars >::Create( InUserObject, InFunc, PAYLOAD_TEMPLATE_PASSIN_FourVars ) );
 	}
-
 	
 	/**
-	 * Static: Creates a shared pointer-based (fast, not thread-safe) member function delegate.  Shared pointer delegates keep a weak reference to your object.  You can use ExecuteIfBound() to call them.
+	 * Static: Creates a shared pointer-based (fast, not thread-safe) member function delegate.
+	 *
+	 * Shared pointer delegates keep a weak reference to your object.
+	 * You can use ExecuteIfBound() to call them.
 	 */
 	template< class UserClass >
 	inline static DELEGATE_CLASS< FUNC_TEMPLATE_ARGS > CreateSP( const TSharedRef< UserClass, ESPMode::Fast >& InUserObjectRef, typename TSPMethodDelegate< UserClass >::FMethodPtr InFunc )
@@ -296,9 +298,11 @@ public:
 		return DELEGATE_CLASS< FUNC_TEMPLATE_ARGS >( TSPMethodDelegate_FourVars_Const< UserClass, PAYLOAD_TEMPLATE_LIST_FourVars >::Create( InUserObjectRef, InFunc, PAYLOAD_TEMPLATE_PASSIN_FourVars ) );
 	}
 
-
 	/**
-	 * Static: Creates a shared pointer-based (fast, not thread-safe) member function delegate.  Shared pointer delegates keep a weak reference to your object.  You can use ExecuteIfBound() to call them.
+	 * Static: Creates a shared pointer-based (fast, not thread-safe) member function delegate.
+	 *
+	 * Shared pointer delegates keep a weak reference to your object.
+	 * You can use ExecuteIfBound() to call them.
 	 */
 	template< class UserClass >
 	inline static DELEGATE_CLASS< FUNC_TEMPLATE_ARGS > CreateSP( UserClass* InUserObject, typename TSPMethodDelegate< UserClass >::FMethodPtr InFunc )
@@ -350,10 +354,12 @@ public:
 	{
 		return CreateSP( StaticCastSharedRef< UserClass >( InUserObject->AsShared() ), InFunc, PAYLOAD_TEMPLATE_PASSIN_FourVars );
 	}
-
 	
 	/**
-	 * Static: Creates a shared pointer-based (slower, conditionally thread-safe) member function delegate.  Shared pointer delegates keep a weak reference to your object.  You can use ExecuteIfBound() to call them.
+	 * Static: Creates a shared pointer-based (slower, conditionally thread-safe) member function delegate.
+	 *
+	 * Shared pointer delegates keep a weak reference to your object.
+	 * You can use ExecuteIfBound() to call them.
 	 */
 	template< class UserClass >
 	inline static DELEGATE_CLASS< FUNC_TEMPLATE_ARGS > CreateThreadSafeSP( const TSharedRef< UserClass, ESPMode::ThreadSafe >& InUserObjectRef, typename TThreadSafeSPMethodDelegate< UserClass >::FMethodPtr InFunc )
@@ -406,9 +412,11 @@ public:
 		return DELEGATE_CLASS< FUNC_TEMPLATE_ARGS >( TThreadSafeSPMethodDelegate_FourVars_Const< UserClass, PAYLOAD_TEMPLATE_LIST_FourVars >::Create( InUserObjectRef, InFunc, PAYLOAD_TEMPLATE_PASSIN_FourVars ) );
 	}
 
-
 	/**
-	 * Static: Creates a shared pointer-based (slower, conditionally thread-safe) member function delegate.  Shared pointer delegates keep a weak reference to your object.  You can use ExecuteIfBound() to call them.
+	 * Static: Creates a shared pointer-based (slower, conditionally thread-safe) member function delegate.
+	 *
+	 * Shared pointer delegates keep a weak reference to your object.
+	 * You can use ExecuteIfBound() to call them.
 	 */
 	template< class UserClass >
 	inline static DELEGATE_CLASS< FUNC_TEMPLATE_ARGS > CreateThreadSafeSP( UserClass* InUserObject, typename TThreadSafeSPMethodDelegate< UserClass >::FMethodPtr InFunc )
@@ -461,9 +469,11 @@ public:
 		return CreateThreadSafeSP( StaticCastSharedRef< UserClass >( InUserObject->AsShared() ), InFunc, PAYLOAD_TEMPLATE_PASSIN_FourVars );
 	}
 
-
 	/**
-	 * Static: Creates a UObject-based member function delegate.  UObject delegates keep a weak reference to your object.  You can use ExecuteIfBound() to call them.
+	 * Static: Creates a UObject-based member function delegate.
+	 *
+	 * UObject delegates keep a weak reference to your object.
+	 * You can use ExecuteIfBound() to call them.
 	 */
 	template< class UserClass >
 	inline static DELEGATE_CLASS< FUNC_TEMPLATE_ARGS > CreateUObject( UserClass* InUserObject, typename TUObjectMethodDelegate< UserClass >::FMethodPtr InFunc )
@@ -516,15 +526,14 @@ public:
 		return DELEGATE_CLASS< FUNC_TEMPLATE_ARGS >( TUObjectMethodDelegate_FourVars_Const< UserClass, PAYLOAD_TEMPLATE_LIST_FourVars >::Create( InUserObject, InFunc, PAYLOAD_TEMPLATE_PASSIN_FourVars ) );
 	}
 
+public:
 
 	/**
 	 * Default constructor
 	 */
 	inline DELEGATE_CLASS()
 		: Delegate( NULL )
-	{
-	}
-
+	{ }
 
 	/**
 	 * Delegate instance destructor
@@ -534,16 +543,13 @@ public:
 		Unbind();
 	}
 
-
 	/**
 	 * Construct from a delegate instance pointer.  The delegate will assume ownership of the incoming delegate instance!
 	 * IMPORTANT: This is a system-internal function and you should never be using this in regular C++ code
 	 */
 	inline DELEGATE_CLASS( IDelegate* InitDelegate )
 		: Delegate( InitDelegate )
-	{
-	}
-
+	{ }
 
 	/**
 	 * Construct from existing delegate object
@@ -556,7 +562,6 @@ public:
 		*this = OtherDelegate;
 	}
 
-
 	/**
 	 * Assignment operator
 	 *
@@ -575,13 +580,12 @@ public:
 			Unbind();
 			Delegate = NewDelegate;
 		}
+
 		return *this;
 	}
 
-
 public:
 
-
 	/**
 	 * Creates a raw C++ pointer global function delegate
 	 */
@@ -610,10 +614,11 @@ public:
 		*this = CreateStatic( InFunc, PAYLOAD_TEMPLATE_PASSIN_FourVars );
 	}
 
-
 	/**
-	 * Binds a raw C++ pointer delegate.  Raw pointer doesn't use any sort of reference, so may be unsafe to call if the object was
-	 * deleted out from underneath your delegate.  Be careful when calling Execute()!
+	 * Binds a raw C++ pointer delegate.
+	 *
+	 * Raw pointer doesn't use any sort of reference, so may be unsafe to call if the object was
+	 * deleted out from underneath your delegate. Be careful when calling Execute()!
 	 */
 	template< class UserClass >
 	inline void BindRaw( UserClass* InUserObject, typename TRawMethodDelegate< UserClass >::FMethodPtr InFunc )
@@ -666,7 +671,6 @@ public:
 		*this = CreateRaw( InUserObject, InFunc, PAYLOAD_TEMPLATE_PASSIN_FourVars );
 	}
 
-
 	/**
 	 * Binds a shared pointer-based (fast, not thread-safe) member function delegate.  Shared pointer delegates keep a weak reference to your object.  You can use ExecuteIfBound() to call them.
 	 */
@@ -721,9 +725,11 @@ public:
 		*this = CreateSP( InUserObjectRef, InFunc, PAYLOAD_TEMPLATE_PASSIN_FourVars );
 	}
 
-
 	/**
-	 * Creates a shared pointer-based (fast, not thread-safe) member function delegate.  Shared pointer delegates keep a weak reference to your object.  You can use ExecuteIfBound() to call them.
+	 * Creates a shared pointer-based (fast, not thread-safe) member function delegate.
+	 *
+	 * Shared pointer delegates keep a weak reference to your object.
+	 * You can use ExecuteIfBound() to call them.
 	 */
 	template< class UserClass >
 	inline void BindSP( UserClass* InUserObject, typename TSPMethodDelegate< UserClass >::FMethodPtr InFunc )
@@ -776,9 +782,11 @@ public:
 		*this = CreateSP( InUserObject, InFunc, PAYLOAD_TEMPLATE_PASSIN_FourVars );
 	}
 
-
 	/**
-	 * Binds a shared pointer-based (slower, conditionally thread-safe) member function delegate.  Shared pointer delegates keep a weak reference to your object.  You can use ExecuteIfBound() to call them.
+	 * Binds a shared pointer-based (slower, conditionally thread-safe) member function delegate.
+	 *
+	 * Shared pointer delegates keep a weak reference to your object.
+	 * You can use ExecuteIfBound() to call them.
 	 */
 	template< class UserClass >
 	inline void BindThreadSafeSP( const TSharedRef< UserClass, ESPMode::ThreadSafe >& InUserObjectRef, typename TThreadSafeSPMethodDelegate< UserClass >::FMethodPtr InFunc )
@@ -831,9 +839,11 @@ public:
 		*this = CreateThreadSafeSP( InUserObjectRef, InFunc, PAYLOAD_TEMPLATE_PASSIN_FourVars );
 	}
 
-
 	/**
-	 * Creates a shared pointer-based (slower, conditionally thread-safe) member function delegate.  Shared pointer delegates keep a weak reference to your object.  You can use ExecuteIfBound() to call them.
+	 * Creates a shared pointer-based (slower, conditionally thread-safe) member function delegate.
+	 *
+	 * Shared pointer delegates keep a weak reference to your object.
+	 * You can use ExecuteIfBound() to call them.
 	 */
 	template< class UserClass >
 	inline void BindThreadSafeSP( UserClass* InUserObject, typename TThreadSafeSPMethodDelegate< UserClass >::FMethodPtr InFunc )
@@ -886,9 +896,11 @@ public:
 		*this = CreateThreadSafeSP( InUserObject, InFunc, PAYLOAD_TEMPLATE_PASSIN_FourVars );
 	}
 
-
 	/**
-	 * Creates a UObject-based member function delegate.  UObject delegates keep a weak reference to your object.  You can use ExecuteIfBound() to call them.
+	 * Creates a UObject-based member function delegate.
+	 *
+	 * UObject delegates keep a weak reference to your object.
+	 * You can use ExecuteIfBound() to call them.
 	 */
 	template< class UserClass >
 	inline void BindUObject( UserClass* InUserObject, typename TUObjectMethodDelegate< UserClass >::FMethodPtr InFunc )
@@ -941,8 +953,6 @@ public:
 		*this = CreateUObject( InUserObject, InFunc, PAYLOAD_TEMPLATE_PASSIN_FourVars );
 	}
 
-
-
 	/**
 	 * Unbinds this delegate
 	 */
@@ -955,7 +965,6 @@ public:
 		}
 	}
 
-
 	/**
 	 * Checks to see if the user object bound to this delegate is still valid
 	 *
@@ -1007,6 +1016,7 @@ public:
 		{
 			return (class UObject*)Delegate->GetRawUserObject();
 		}
+
 		return NULL;
 	}
 
@@ -1023,11 +1033,11 @@ public:
 		{
 			return Delegate->ExecuteIfSafe( FUNC_PARAM_PASSTHRU );
 		}
+
 		return false;
 	}
 #endif
 
-
 	/**
 	 * Equality operator
 	 */
@@ -1046,13 +1056,11 @@ public:
 				return true;
 			}
 
-
 			// No match!
 			return false;
 		}
 	}
 
-
 private:
 
 	// Declare ourselves as a friend to our multi-cast counterpart, so that it can access our interface directly
@@ -1061,7 +1069,6 @@ private:
 	/** Pointer to the actual delegate object instance.  We need to use a pointer because the actual type of delegate class
 	    that's created will depend on the type of binding (raw, shared pointer, etc) as well as the user's object type */
 	IDelegate* Delegate;
-
 };
 
 
@@ -1077,7 +1084,6 @@ private:
 template< FUNC_TEMPLATE_DECL_TYPENAME >
 class BASE_MULTICAST_DELEGATE_CLASS
 {
-
 public:
 
 	/** The actual single-cast delegate class for this multi-cast delegate */
@@ -1097,7 +1103,6 @@ public:
 		CompactInvocationList();
 	}
 
-
 	/**
 	 * Adds a raw C++ pointer global function delegate
 	 *
@@ -1128,10 +1133,11 @@ public:
 		Add( FDelegate::CreateStatic( InFunc, PAYLOAD_TEMPLATE_PASSIN_FourVars ) );
 	}
 
-
 	/**
-	 * Adds a raw C++ pointer delegate.  Raw pointer doesn't use any sort of reference, so may be unsafe to call if the object was
-	 * deleted out from underneath your delegate.  Be careful when calling Execute()!
+	 * Adds a raw C++ pointer delegate.
+	 *
+	 * Raw pointer doesn't use any sort of reference, so may be unsafe to call if the object was
+	 * deleted out from underneath your delegate. Be careful when calling Execute()!
 	 *
 	 * @param	InUserObject	User object to bind to
 	 * @param	InFunc			Class method function address
@@ -1187,9 +1193,10 @@ public:
 		Add( FDelegate::CreateRaw( InUserObject, InFunc, PAYLOAD_TEMPLATE_PASSIN_FourVars ) );
 	}
 
-
 	/**
-	 * Adds a shared pointer-based (fast, not thread-safe) member function delegate.  Shared pointer delegates keep a weak reference to your object.
+	 * Adds a shared pointer-based (fast, not thread-safe) member function delegate.
+	 *
+	 * Shared pointer delegates keep a weak reference to your object.
 	 *
 	 * @param	InUserObjectRef	User object to bind to
 	 * @param	InFunc			Class method function address
@@ -1245,9 +1252,10 @@ public:
 		Add( FDelegate::CreateSP( InUserObjectRef, InFunc, PAYLOAD_TEMPLATE_PASSIN_FourVars ) );
 	}
 
-
 	/**
-	 * Adds a shared pointer-based (fast, not thread-safe) member function delegate.  Shared pointer delegates keep a weak reference to your object.
+	 * Adds a shared pointer-based (fast, not thread-safe) member function delegate.
+	 *
+	 * Shared pointer delegates keep a weak reference to your object.
 	 *
 	 * @param	InUserObject	User object to bind to
 	 * @param	InFunc			Class method function address
@@ -1303,7 +1311,6 @@ public:
 		Add( FDelegate::CreateSP( InUserObject, InFunc, PAYLOAD_TEMPLATE_PASSIN_FourVars ) );
 	}
 
-
 	/**
 	 * Adds a shared pointer-based (slower, conditionally thread-safe) member function delegate.  Shared pointer delegates keep a weak reference to your object.
 	 *
@@ -1361,9 +1368,10 @@ public:
 		Add( FDelegate::CreateThreadSafeSP( InUserObjectRef, InFunc, PAYLOAD_TEMPLATE_PASSIN_FourVars ) );
 	}
 
-
 	/**
-	 * Adds a shared pointer-based (slower, conditionally thread-safe) member function delegate.  Shared pointer delegates keep a weak reference to your object.
+	 * Adds a shared pointer-based (slower, conditionally thread-safe) member function delegate.
+	 *
+	 * Shared pointer delegates keep a weak reference to your object.
 	 *
 	 * @param	InUserObject	User object to bind to
 	 * @param	InFunc			Class method function address
@@ -1419,9 +1427,10 @@ public:
 		Add( FDelegate::CreateThreadSafeSP( InUserObject, InFunc, PAYLOAD_TEMPLATE_PASSIN_FourVars ) );
 	}
 
-
 	/**
-	 * Adds a UObject-based member function delegate.  UObject delegates keep a weak reference to your object.
+	 * Adds a UObject-based member function delegate.
+	 *
+	 * UObject delegates keep a weak reference to your object.
 	 *
 	 * @param	InUserObject	User object to bind to
 	 * @param	InFunc			Class method function address
@@ -1477,7 +1486,6 @@ public:
 		Add( FDelegate::CreateUObject( InUserObject, InFunc, PAYLOAD_TEMPLATE_PASSIN_FourVars ) );
 	}
 
-
 	/**
 	 *	Removes all functions from this multi-cast delegate's invocation list that are bound to the specified UserObject.
 	 *	Note that the order of the delegates may not be preserved!
@@ -1491,7 +1499,6 @@ public:
 		CompactInvocationList();
 	}
 
-
 	/**
 	 * Removes a function from this multi-cast delegate's invocation list (performance is O(N)).  Note that the
 	 * order of the delegates may not be preserved!
@@ -1512,7 +1519,6 @@ public:
 		CompactInvocationList();
 	}
 
-
 	// NOTE: These direct Remove methods are only supported for multi-cast delegates with no payload attached.
 	// See the comment in the multi-cast delegate Remove() method above for more details.
 
@@ -1527,7 +1533,6 @@ public:
 		Remove( FDelegate::CreateStatic( InFunc ) );
 	}
 
-
 	/**
 	 * Removes a raw C++ pointer delegate (performance is O(N)).  Note that the order of the delegates may not
 	 * be preserved!
@@ -1546,7 +1551,6 @@ public:
 		Remove( FDelegate::CreateRaw( InUserObject, InFunc ) );
 	}
 
-
 	/**
 	 * Removes a shared pointer-based member function delegate (performance is O(N)).  Note that the order of
 	 * the delegates may not be preserved!
@@ -1565,7 +1569,6 @@ public:
 		Remove( FDelegate::CreateSP( InUserObjectRef, InFunc ) );
 	}
 
-
 	/**
 	 * Removes a shared pointer-based member function delegate (performance is O(N)).  Note that the order of
 	 * the delegates may not be preserved!
@@ -1584,7 +1587,6 @@ public:
 		Remove( FDelegate::CreateSP( InUserObject, InFunc ) );
 	}
 
-
 	/**
 	 * Removes a shared pointer-based member function delegate (performance is O(N)).  Note that the order of
 	 * the delegates may not be preserved!
@@ -1603,7 +1605,6 @@ public:
 		Remove( FDelegate::CreateThreadSafeSP( InUserObjectRef, InFunc ) );
 	}
 
-
 	/**
 	 * Removes a shared pointer-based member function delegate (performance is O(N)).  Note that the order of
 	 * the delegates may not be preserved!
@@ -1622,7 +1623,6 @@ public:
 		Remove( FDelegate::CreateThreadSafeSP( InUserObject, InFunc ) );
 	}
 
-
 	/**
 	 * Removes a UObject-based member function delegate (performance is O(N)).  Note that the order of the
 	 * delegates may not be preserved!
@@ -1667,12 +1667,12 @@ protected:
 	 */
 	inline BASE_MULTICAST_DELEGATE_CLASS()
 		: ExpirableObjectCount( 0 )
-	{
-	}
-
+	{ }
 
 	/**
-	 * Checks to see if any functions are bound to this multi-cast delegate. Private to avoid assumptions about only one delegate existing
+	 * Checks to see if any functions are bound to this multi-cast delegate.
+	 *
+	 * This method is protected to avoid assumptions about only one delegate existing.
 	 *
 	 * @return	True if any functions are bound
 	 */
@@ -1681,7 +1681,6 @@ protected:
 		return InvocationList.Num() > 0;
 	}
 
-
 	/**
 	 * Removes all functions from this delegate's invocation list
 	 */
@@ -1691,7 +1690,6 @@ protected:
 		ExpirableObjectCount = 0;
 	}
 
-
 	/**
 	 * Broadcasts this delegate to all bound objects, except to those that may have expired
 	 */
@@ -1721,7 +1719,6 @@ protected:
 		}
 	}
 
-
 private:
 
 	/**
@@ -1753,7 +1750,6 @@ private:
 		}
 	}
 
-
 	/**
 	 * Removes a function from this multi-cast delegate's invocation list (performance is O(N)).  Note that the
 	 * order of the delegates may not be preserved!
@@ -1815,7 +1811,6 @@ private:
 		}
 	}
 
-
 	/** Cleans up any delegates in our invocation list that have expired (performance is O(N)) */
 	void CompactInvocationList() const
 	{
@@ -1836,7 +1831,6 @@ private:
 		}
 	}
 
-
 private:
 
 	/** Ordered list functions to invoke when the Broadcast function is called */
@@ -1852,9 +1846,6 @@ private:
 template< FUNC_TEMPLATE_DECL_TYPENAME >
 class MULTICAST_DELEGATE_CLASS : public BASE_MULTICAST_DELEGATE_CLASS< FUNC_TEMPLATE_DECL >
 {
-private:  // Prevents erroneous use by other classes.
-	typedef BASE_MULTICAST_DELEGATE_CLASS< FUNC_TEMPLATE_DECL > Super;
-
 public:
 
 	/**
@@ -1867,7 +1858,6 @@ public:
 		return Super::IsBound();
 	}
 
-
 	/**
 	 * Removes all functions from this delegate's invocation list
 	 */
@@ -1876,7 +1866,6 @@ public:
 		return Super::Clear();
 	}
 
-
 	/**
 	 * Broadcasts this delegate to all bound objects, except to those that may have expired
 	 */
@@ -1884,11 +1873,15 @@ public:
 	{
 		return Super::Broadcast( FUNC_PARAM_PASSTHRU );
 	}
+
+private: 
+	
+	// Prevents erroneous use by other classes.
+	typedef BASE_MULTICAST_DELEGATE_CLASS< FUNC_TEMPLATE_DECL > Super;
 };
 #endif		// FUNC_IS_VOID
 
 
-
 /**
  * Dynamic delegate base object (UObject-based, serializable).  You'll use the various DECLARE_DYNAMIC_DELEGATE
  * macros to create the actual delegate type, templated to the function signature the delegate is compatible with.
@@ -1898,16 +1891,12 @@ template< FUNC_TEMPLATE_DECL_TYPENAME, typename TWeakPtr = FWeakObjectPtr >
 class DYNAMIC_DELEGATE_CLASS
 	: public TScriptDelegate<TWeakPtr>
 {
-
 public:
 
 	/**
 	 * Default constructor
 	 */
-	DYNAMIC_DELEGATE_CLASS()
-	{
-	}
-
+	DYNAMIC_DELEGATE_CLASS() { }
 
 	/**
 	 * Construction from an FScriptDelegate must be explicit.  This is really only used by UObject system internals.
@@ -1916,9 +1905,7 @@ public:
 	 */
 	explicit DYNAMIC_DELEGATE_CLASS( const TScriptDelegate<TWeakPtr>& InScriptDelegate )
 		: TScriptDelegate<TWeakPtr>( InScriptDelegate )
-	{
-	}
-
+	{ }
 
 	/**
 	 * Templated helper class to define a typedef for user's method pointer, then used below
@@ -1926,13 +1913,11 @@ public:
 	template< class UserClass >
 	class TMethodPtrResolver
 	{
-
 	public:
 
 		typedef RetValType ( UserClass::*FMethodPtr )( FUNC_PARAM_LIST );
 	};
 
-
 	/**
 	 * Binds a UObject instance and a UObject method address to this delegate.
 	 *
@@ -1969,11 +1954,9 @@ public:
 	// NOTE:  Execute() method must be defined in derived classes
 
 	// NOTE:  ExecuteIfBound() method must be defined in derived classes
-
 };
 
 
-
 /**
  * Dynamic multi-cast delegate base object (UObject-based, serializable).  You'll use the various
  * DECLARE_DYNAMIC_MULTICAST_DELEGATE macros to create the actual delegate type, templated to the function
@@ -1984,20 +1967,15 @@ template< FUNC_TEMPLATE_DECL_TYPENAME, typename TWeakPtr = FWeakObjectPtr >
 class DYNAMIC_MULTICAST_DELEGATE_CLASS
 	: public TMulticastScriptDelegate<TWeakPtr>
 {
-
 public:
 
 	/** The actual single-cast delegate class for this multi-cast delegate */
 	typedef DYNAMIC_DELEGATE_CLASS< FUNC_TEMPLATE_ARGS > FDelegate;
 
-
 	/**
 	 * Default constructor
 	 */
-	DYNAMIC_MULTICAST_DELEGATE_CLASS()
-	{
-	}
-
+	DYNAMIC_MULTICAST_DELEGATE_CLASS() { }
 
 	/**
 	 * Construction from an FMulticastScriptDelegate must be explicit.  This is really only used by UObject system internals.
@@ -2006,9 +1984,7 @@ public:
 	 */
 	explicit DYNAMIC_MULTICAST_DELEGATE_CLASS( const TMulticastScriptDelegate<TWeakPtr>& InMulticastScriptDelegate )
 		: TMulticastScriptDelegate<TWeakPtr>( InMulticastScriptDelegate )
-	{
-	}
-
+	{ }
 
 	/**
 	 * Tests if a UObject instance and a UObject method address pair are already bound to this multi-cast delegate.
@@ -2034,7 +2010,6 @@ public:
 		return this->Contains( NewDelegate );
 	}
 
-
 	/**
 	 * Binds a UObject instance and a UObject method address to this multi-cast delegate.
 	 *
@@ -2058,7 +2033,6 @@ public:
 		this->Add( NewDelegate );
 	}
 
-
 	/**
 	 * Binds a UObject instance and a UObject method address to this multi-cast delegate, but only if it hasn't been bound before.
 	 *
@@ -2082,7 +2056,6 @@ public:
 		this->AddUnique( NewDelegate );
 	}
 
-
 	/**
 	 * Unbinds a UObject instance and a UObject method address from this multi-cast delegate.
 	 *
@@ -2106,13 +2079,10 @@ public:
 		this->Remove( NewDelegate );
 	}
 
-
 	// NOTE:  Broadcast() method must be defined in derived classes
-
 };
 
 
-
 #undef PAYLOAD_TEMPLATE_DECL_OneVar
 #undef PAYLOAD_TEMPLATE_LIST_OneVar
 #undef PAYLOAD_TEMPLATE_ARGS_OneVar
@@ -2164,5 +2134,3 @@ public:
 #undef STATIC_DELEGATE_INSTANCE_CLASS_FourVars
 
 #undef DELEGATE_INSTANCE_INTERFACE_CLASS
-
-
diff --git a/Engine/Source/Runtime/Core/Public/UObject/ScriptDelegates.h b/Engine/Source/Runtime/Core/Public/UObject/ScriptDelegates.h
--- a/Engine/Source/Runtime/Core/Public/UObject/ScriptDelegates.h
+++ b/Engine/Source/Runtime/Core/Public/UObject/ScriptDelegates.h
@@ -6,26 +6,25 @@
 
 #pragma once
 
-
 #include "WeakObjectPtrTemplates.h"
 
+
 /**
  * Script delegate base class
  */
 template <typename TWeakPtr = FWeakObjectPtr>
 class TScriptDelegate
 {
-
 public:
 
 	/** Constructor */
 	TScriptDelegate() 
 		: Object( NULL ),
 		  FunctionName( NAME_None )
-	{
-	}
+	{ }
 
 private:
+
 	template <class UObjectTemplate>
 	inline bool IsBound_Internal() const
 	{
@@ -41,6 +40,7 @@ private:
 	}
 
 public:
+
 	/** 
 	 * Checks to see if the user object bound to this delegate is still valid
 	 *
@@ -103,7 +103,6 @@ public:
 		return TEXT( "<Unbound>" );
 	}
 
-
 	/** Delegate serialization */
 	friend FArchive& operator<<( FArchive& Ar, TScriptDelegate& D )
 	{
@@ -111,7 +110,6 @@ public:
 		return Ar;
 	}
 
-
 	/** Comparison operators */
 	FORCEINLINE bool operator==( const TScriptDelegate& Other ) const
 	{
@@ -139,7 +137,6 @@ public:
 		Object = InObject;
 	}
 
-
 	/**
 	 * Sets the function name.  Usually, you should never call this yourself.  Use BindDynamic() instead!
 	 *
@@ -150,7 +147,6 @@ public:
 		FunctionName = InFunctionName;
 	}
 
-
 	/** 
 	 * Gets the object bound to this delegate
 	 *
@@ -162,7 +158,6 @@ public:
 		return static_cast< UObject* >( Object.Get() );
 	}
 
-
 	/**
 	 * Gets the object bound to this delegate (const)
 	 *
@@ -174,7 +169,6 @@ public:
 		return static_cast< const UObject* >( Object.Get() );
 	}
 
-
 	/**
 	 * Gets the name of the function to call on the bound object
 	 *
@@ -185,7 +179,6 @@ public:
 		return FunctionName;
 	}
 
-
 	/**
 	 * Executes a delegate by calling the named function on the object bound to the delegate.  You should
 	 * always first verify that the delegate is safe to execute by calling IsBound() before calling this function.
@@ -226,6 +219,7 @@ protected:
 	friend class FCallDelegateHelper;
 };
 
+
 template<typename TWeakPtr> struct TIsZeroConstructType<TScriptDelegate<TWeakPtr> > { enum { Value = TIsZeroConstructType<TWeakPtr>::Value }; };
 
 
@@ -235,16 +229,14 @@ template<typename TWeakPtr> struct TIsZeroConstructType<TScriptDelegate<TWeakPtr
 template <typename TWeakPtr = FWeakObjectPtr>
 class TMulticastScriptDelegate
 {
-
 public:
 
 	/**
 	 * Default constructor
 	 */
-	inline TMulticastScriptDelegate()
-	{
-	}
+	inline TMulticastScriptDelegate() { }
 
+public:
 
 	/**
 	 * Checks to see if any functions are bound to this multi-cast delegate
@@ -256,7 +248,6 @@ public:
 		return InvocationList.Num() > 0;
 	}
 
-
 	/**
 	 * Checks whether a function delegate is already a member of this multi-cast delegate's invocation list
 	 *
@@ -313,7 +304,6 @@ public:
 		CompactInvocationList();
 	}
 
-
 	/**
 	 * Removes all functions from this delegate's invocation list
 	 */
@@ -322,7 +312,6 @@ public:
 		InvocationList.Empty();
 	}
 
-
 	/**
 	 * Converts this delegate to a string representation
 	 *
@@ -348,7 +337,6 @@ public:
 		return TEXT( "<Unbound>" );
 	}
 
-
 	/** Multi-cast delegate serialization */
 	friend FArchive& operator<<( FArchive& Ar, TMulticastScriptDelegate<TWeakPtr>& D )
 	{
@@ -369,7 +357,6 @@ public:
 		return Ar;
 	}
 
-
 	/**
 	 * Executes a multi-cast delegate by calling all functions on objects bound to the delegate.  Always
 	 * safe to call, even if when no objects are bound, or if objects have expired.  In general, you should
@@ -422,7 +409,6 @@ public:
 		return OutputList;
 	}
 
-
 protected:
 
 	/**
@@ -455,7 +441,6 @@ protected:
 		InvocationList.AddUnique( InDelegate );
 	}
 
-
 	/**
 	 * Removes a function from this multi-cast delegate's invocation list (performance is O(N)).  Note that the
 	 * order of the delegates may not be preserved!
@@ -477,7 +462,6 @@ protected:
 		}
 	}
 
-
 	/** Cleans up any delegates in our invocation list that have expired (performance is O(N)) */
 	void CompactInvocationList() const
 	{
@@ -491,7 +475,6 @@ protected:
 		}
 	}
 
-
 protected:
 
 	/** Ordered list functions to invoke when the Broadcast function is called */
@@ -505,4 +488,5 @@ protected:
 	friend class FCallDelegateHelper;
 };
 
+
 template<typename TWeakPtr> struct TIsZeroConstructType<TMulticastScriptDelegate<TWeakPtr> > { enum { Value = TIsZeroConstructType<TWeakPtr>::Value }; };
