Fixed incorrect Right and Bottom values for PrimaryDisplayWorkAreaRect and VirtualDisplayRect set in FWindowsApplication::GetDisplayMetrics()

diff --git a/Engine/Source/Runtime/Core/Private/Windows/WindowsApplication.cpp b/Engine/Source/Runtime/Core/Private/Windows/WindowsApplication.cpp
--- a/Engine/Source/Runtime/Core/Private/Windows/WindowsApplication.cpp
+++ b/Engine/Source/Runtime/Core/Private/Windows/WindowsApplication.cpp
@@ -414,14 +414,14 @@ void FWindowsApplication::GetDisplayMetrics( FDisplayMetrics& OutDisplayMetrics
 	SystemParametersInfo( SPI_GETWORKAREA, 0, &WorkAreaRect, 0 );
 	OutDisplayMetrics.PrimaryDisplayWorkAreaRect.Left = WorkAreaRect.left;
 	OutDisplayMetrics.PrimaryDisplayWorkAreaRect.Top = WorkAreaRect.top;
-	OutDisplayMetrics.PrimaryDisplayWorkAreaRect.Right = 1.0f + WorkAreaRect.right;
-	OutDisplayMetrics.PrimaryDisplayWorkAreaRect.Bottom = 1.0f + WorkAreaRect.bottom;
+	OutDisplayMetrics.PrimaryDisplayWorkAreaRect.Right = WorkAreaRect.right;
+	OutDisplayMetrics.PrimaryDisplayWorkAreaRect.Bottom = WorkAreaRect.bottom;
 
 	// Virtual desktop area
 	OutDisplayMetrics.VirtualDisplayRect.Left = ::GetSystemMetrics( SM_XVIRTUALSCREEN );
 	OutDisplayMetrics.VirtualDisplayRect.Top = ::GetSystemMetrics( SM_YVIRTUALSCREEN );
-	OutDisplayMetrics.VirtualDisplayRect.Right = 1.0f + ::GetSystemMetrics( SM_CXVIRTUALSCREEN );
-	OutDisplayMetrics.VirtualDisplayRect.Bottom = 1.0f + ::GetSystemMetrics( SM_CYVIRTUALSCREEN );
+	OutDisplayMetrics.VirtualDisplayRect.Right = OutDisplayMetrics.VirtualDisplayRect.Left + ::GetSystemMetrics( SM_CXVIRTUALSCREEN );
+	OutDisplayMetrics.VirtualDisplayRect.Bottom = OutDisplayMetrics.VirtualDisplayRect.Top + ::GetSystemMetrics( SM_CYVIRTUALSCREEN );
 
 	// Get connected monitor information
 	GetMonitorInfo(OutDisplayMetrics.MonitorInfo);
diff --git a/Engine/Source/Runtime/SlateCore/Private/Widgets/SWindow.cpp b/Engine/Source/Runtime/SlateCore/Private/Widgets/SWindow.cpp
--- a/Engine/Source/Runtime/SlateCore/Private/Widgets/SWindow.cpp
+++ b/Engine/Source/Runtime/SlateCore/Private/Widgets/SWindow.cpp
@@ -221,12 +221,12 @@ void SWindow::Construct(const FArguments& InArgs)
 		}
 
 		float PrimaryWidthPadding = DisplayMetrics.PrimaryDisplayWidth - 
-			(PrimaryDisplayRect.Right - 1 - PrimaryDisplayRect.Left);
+			(PrimaryDisplayRect.Right - PrimaryDisplayRect.Left);
 		float PrimaryHeightPadding = DisplayMetrics.PrimaryDisplayHeight - 
-			(PrimaryDisplayRect.Bottom - 1 - PrimaryDisplayRect.Top);
+			(PrimaryDisplayRect.Bottom - PrimaryDisplayRect.Top);
 
-		float VirtualWidth = (VirtualDisplayRect.Right - 1 - VirtualDisplayRect.Left);
-		float VirtualHeight = (VirtualDisplayRect.Bottom - 1 - VirtualDisplayRect.Top);
+		float VirtualWidth = (VirtualDisplayRect.Right - VirtualDisplayRect.Left);
+		float VirtualHeight = (VirtualDisplayRect.Bottom - VirtualDisplayRect.Top);
 
 		// Make sure that the window size is no larger than the virtual display area.
 		WindowSize.X = FMath::Clamp(WindowSize.X, 0.0f, VirtualWidth - PrimaryWidthPadding);
