I18N: Fixed the teardown process for the i18n system so it won't do anything if it was never initialized.

diff --git a/Engine/Source/Runtime/Core/Private/Internationalization/Internationalization.cpp b/Engine/Source/Runtime/Core/Private/Internationalization/Internationalization.cpp
--- a/Engine/Source/Runtime/Core/Private/Internationalization/Internationalization.cpp
+++ b/Engine/Source/Runtime/Core/Private/Internationalization/Internationalization.cpp
@@ -38,6 +38,14 @@ FInternationalization& FInternationalization::Get()
 	return *Instance;
 }
 
+void FInternationalization::TearDown()
+{
+	if (Instance && Instance->IsInitialized())
+	{
+		Instance->Terminate();
+	}
+}
+
 void FInternationalization::GetTimeZonesIDs(TArray<FString>& TimeZonesIDs)
 {
 	TimeZonesIDs.Empty();
diff --git a/Engine/Source/Runtime/Core/Public/Internationalization/Internationalization.h b/Engine/Source/Runtime/Core/Public/Internationalization/Internationalization.h
--- a/Engine/Source/Runtime/Core/Public/Internationalization/Internationalization.h
+++ b/Engine/Source/Runtime/Core/Public/Internationalization/Internationalization.h
@@ -29,6 +29,8 @@ class FInternationalization
 public:
 	static CORE_API FInternationalization& Get();
 
+	static CORE_API void TearDown();
+
 	static CORE_API FText ForUseOnlyByLocMacroAndGraphNodeTextLiterals_CreateText( const TCHAR* InTextLiteral, const TCHAR* Namespace, const TCHAR* Key )
 	{
 		return FText( InTextLiteral, Namespace, Key );
@@ -57,8 +59,6 @@ public:
 		return InvariantCulture.ToSharedRef();
 	}
 
-	CORE_API void Initialize();
-	CORE_API void Terminate();
 
 	CORE_API bool IsInitialized() {return bIsInitialized;}
 
@@ -77,6 +77,9 @@ private:
 	//@return index of a culture by given name, or -1 if not found
 	CORE_API int32 GetCultureIndex(const FString& Name);
 
+	CORE_API void Initialize();
+	CORE_API void Terminate();
+
 private:
 	static FInternationalization* Instance;
 	bool bIsInitialized;
diff --git a/Engine/Source/Runtime/Launch/Private/LaunchEngineLoop.cpp b/Engine/Source/Runtime/Launch/Private/LaunchEngineLoop.cpp
--- a/Engine/Source/Runtime/Launch/Private/LaunchEngineLoop.cpp
+++ b/Engine/Source/Runtime/Launch/Private/LaunchEngineLoop.cpp
@@ -2514,7 +2514,7 @@ void FEngineLoop::AppExit( )
 		GLog->TearDown();
 	}
 
-	FInternationalization::Get().Terminate();
+	FInternationalization::TearDown();
 }
 
 #undef LOCTEXT_NAMESPACE
