Updated FMacPlatformMisc::MessageBoxExt() to use NSAlert instead of deprecated NSRunAlertPanel.

diff --git a/Engine/Source/Runtime/Core/Private/Mac/MacPlatformMisc.cpp b/Engine/Source/Runtime/Core/Private/Mac/MacPlatformMisc.cpp
--- a/Engine/Source/Runtime/Core/Private/Mac/MacPlatformMisc.cpp
+++ b/Engine/Source/Runtime/Core/Private/Mac/MacPlatformMisc.cpp
@@ -687,181 +687,179 @@ void FMacPlatformMisc::CreateGuid(FGuid& Result)
 	Result = GUID;
 }
 
-CORE_API NSWindow* GExtendedAlertWindow = NULL;
-
-static int32 ShowExtendedAlertWindow( EAppMsgType::Type MsgType, NSString* Text, NSString* Caption )
+EAppReturnType::Type FMacPlatformMisc::MessageBoxExt(EAppMsgType::Type MsgType, const TCHAR* Text, const TCHAR* Caption)
 {
 	SCOPED_AUTORELEASE_POOL;
 
-	checkf(GExtendedAlertWindow);
+	EAppReturnType::Type RetValue = EAppReturnType::Cancel;
+	NSInteger Result;
 
-	// Update window contents
-	NSArray* SubViews = [[GExtendedAlertWindow contentView] subviews];
-	for (NSView* SubView in SubViews)
+	if (MacApplication)
 	{
-		if ([SubView isMemberOfClass:[NSTextField class]])
-		{
-			// Set the caption and message text
-			NSTextField* TextField = (NSTextField*)SubView;
-			if ([[TextField stringValue] isEqualToString: @"Caption"])
+		MacApplication->UseMouseCaptureWindow(false);
+	}
+
+	NSAlert* AlertPanel = [NSAlert new];
+	[AlertPanel setInformativeText:FString(Text).GetNSString()];
+	[AlertPanel setMessageText:FString(Caption).GetNSString()];
+
+	switch (MsgType)
+	{
+		case EAppMsgType::Ok:
+			[AlertPanel addButtonWithTitle:@"OK"];
+			[AlertPanel runModal];
+			RetValue = EAppReturnType::Ok;
+			break;
+
+		case EAppMsgType::YesNo:
+			[AlertPanel addButtonWithTitle:@"Yes"];
+			[AlertPanel addButtonWithTitle:@"No"];
+			Result = [AlertPanel runModal];
+			if (Result == NSAlertFirstButtonReturn)
 			{
-				[TextField setStringValue: Caption];
+				RetValue = EAppReturnType::Yes;
 			}
-			else if ([[TextField stringValue] isEqualToString: @"Message"])
+			else if (Result == NSAlertSecondButtonReturn)
 			{
-				[TextField setStringValue: Text];
+				RetValue = EAppReturnType::No;
 			}
-		}
-		else if (MsgType == EAppMsgType::YesNoYesAllNoAll && [SubView isMemberOfClass:[NSButton class]])
-		{
-			// Remove Cancel button if it's not needed
-			NSButton* Button = (NSButton*)SubView;
-			if ([[Button title] isEqualToString: @"Cancel"])
+			break;
+
+		case EAppMsgType::OkCancel:
+			[AlertPanel addButtonWithTitle:@"OK"];
+			[AlertPanel addButtonWithTitle:@"Cancel"];
+			Result = [AlertPanel runModal];
+			if (Result == NSAlertFirstButtonReturn)
 			{
-				[SubView setHidden: true];
+				RetValue = EAppReturnType::Ok;
 			}
-		}
-	}
-
-	// UE4AppDelegate's OnAlertWindowButtonPress will close the window and set the correct return code
-	return [NSApp runModalForWindow: GExtendedAlertWindow];
-}
-
-EAppReturnType::Type FMacPlatformMisc::MessageBoxExt( EAppMsgType::Type MsgType, const TCHAR* Text, const TCHAR* Caption )
-{
-	SCOPED_AUTORELEASE_POOL;
+			else if (Result == NSAlertSecondButtonReturn)
+			{
+				RetValue = EAppReturnType::Cancel;
+			}
+			break;
+
+		case EAppMsgType::YesNoCancel:
+			[AlertPanel addButtonWithTitle:@"Yes"];
+			[AlertPanel addButtonWithTitle:@"No"];
+			[AlertPanel addButtonWithTitle:@"Cancel"];
+			Result = [AlertPanel runModal];
+			if (Result == NSAlertFirstButtonReturn)
+			{
+				RetValue = EAppReturnType::Yes;
+			}
+			else if (Result == NSAlertSecondButtonReturn)
+			{
+				RetValue = EAppReturnType::No;
+			}
+			else
+			{
+				RetValue = EAppReturnType::Cancel;
+			}
+			break;
+
+		case EAppMsgType::CancelRetryContinue:
+			[AlertPanel addButtonWithTitle:@"Continue"];
+			[AlertPanel addButtonWithTitle:@"Retry"];
+			[AlertPanel addButtonWithTitle:@"Cancel"];
+			Result = [AlertPanel runModal];
+			if (Result == NSAlertFirstButtonReturn)
+			{
+				RetValue = EAppReturnType::Continue;
+			}
+			else if (Result == NSAlertSecondButtonReturn)
+			{
+				RetValue = EAppReturnType::Retry;
+			}
+			else
+			{
+				RetValue = EAppReturnType::Cancel;
+			}
+			break;
+
+		case EAppMsgType::YesNoYesAllNoAll:
+			[AlertPanel addButtonWithTitle:@"Yes"];
+			[AlertPanel addButtonWithTitle:@"No"];
+			[AlertPanel addButtonWithTitle:@"Yes to all"];
+			[AlertPanel addButtonWithTitle:@"No to all"];
+			Result = [AlertPanel runModal];
+			if (Result == NSAlertFirstButtonReturn)
+			{
+				RetValue = EAppReturnType::Yes;
+			}
+			else if (Result == NSAlertSecondButtonReturn)
+			{
+				RetValue = EAppReturnType::No;
+			}
+			else if (Result == NSAlertThirdButtonReturn)
+			{
+				RetValue = EAppReturnType::YesAll;
+			}
+			else
+			{
+				RetValue = EAppReturnType::NoAll;
+			}
+			break;
+
+		case EAppMsgType::YesNoYesAllNoAllCancel:
+			[AlertPanel addButtonWithTitle:@"Yes"];
+			[AlertPanel addButtonWithTitle:@"No"];
+			[AlertPanel addButtonWithTitle:@"Yes to all"];
+			[AlertPanel addButtonWithTitle:@"No to all"];
+			[AlertPanel addButtonWithTitle:@"Cancel"];
+			Result = [AlertPanel runModal];
+			if (Result == NSAlertFirstButtonReturn)
+			{
+				RetValue = EAppReturnType::Yes;
+			}
+			else if (Result == NSAlertSecondButtonReturn)
+			{
+				RetValue = EAppReturnType::No;
+			}
+			else if (Result == NSAlertThirdButtonReturn)
+			{
+				RetValue = EAppReturnType::YesAll;
+			}
+			else if (Result == NSAlertThirdButtonReturn + 1)
+			{
+				RetValue = EAppReturnType::NoAll;
+			}
+			else
+			{
+				RetValue = EAppReturnType::Cancel;
+			}
+			break;
+
+		case EAppMsgType::YesNoYesAll:
+			[AlertPanel addButtonWithTitle:@"Yes"];
+			[AlertPanel addButtonWithTitle:@"No"];
+			[AlertPanel addButtonWithTitle:@"Yes to all"];
+			Result = [AlertPanel runModal];
+			if (Result == NSAlertFirstButtonReturn)
+			{
+				RetValue = EAppReturnType::Yes;
+			}
+			else if (Result == NSAlertSecondButtonReturn)
+			{
+				RetValue = EAppReturnType::No;
+			}
+			else
+			{
+				RetValue = EAppReturnType::YesAll;
+			}
+			break;
 
-	NSString* CocoaText = Text ? (NSString*)FPlatformString::TCHARToCFString(Text) : NULL;
-	NSString* CocoaCaption = Caption ? (NSString*)FPlatformString::TCHARToCFString(Caption) : NULL;
+		default:
+			break;
+	}
 
-	EAppReturnType::Type RetValue = EAppReturnType::Cancel;
-	NSInteger Result;
+	[AlertPanel release];
 
-	if(MacApplication)
-	{
-		MacApplication->UseMouseCaptureWindow(false);
-	}
-	
-	switch (MsgType)
-	{
-	case EAppMsgType::YesNo:
-		Result = NSRunAlertPanel(CocoaCaption, CocoaText, @"Yes", @"No", NULL);
-		if (Result == NSAlertDefaultReturn)
-		{
-			RetValue = EAppReturnType::Yes;
-		}
-		else if (Result == NSAlertAlternateReturn)
-		{
-			RetValue = EAppReturnType::No;
-		}
-		break;
-	case EAppMsgType::OkCancel:
-		// return 1 for Ok, 0 for Cancel
-		Result = NSRunAlertPanel(CocoaCaption, CocoaText, @"OK", @"Cancel", NULL);
-		if (Result == NSAlertDefaultReturn)
-		{
-			RetValue = EAppReturnType::Ok;
-		}
-		else if (Result == NSAlertAlternateReturn)
-		{
-			RetValue = EAppReturnType::Cancel;
-		}
-		break;
-	case EAppMsgType::YesNoCancel:
-		// return 0 for Yes, 1 for No, 2 for Cancel
-		Result = NSRunAlertPanel(CocoaCaption, CocoaText, @"Yes", @"No", @"Cancel");
-		if (Result == NSAlertDefaultReturn)
-		{
-			RetValue = EAppReturnType::Yes;
-		}
-		else if (Result == NSAlertAlternateReturn)
-		{
-			RetValue = EAppReturnType::No;
-		}
-		else
-		{
-			RetValue = EAppReturnType::Cancel;
-		}
-		break;
-	case EAppMsgType::CancelRetryContinue:
-		// return 0 for Cancel, 1 for Retry, 2 for Continue
-		Result = NSRunAlertPanel(CocoaCaption, CocoaText, @"Continue", @"Retry", @"Cancel");
-		if (Result == NSAlertDefaultReturn)
-		{
-			RetValue = EAppReturnType::Continue;
-		}
-		else if (Result == NSAlertAlternateReturn)
-		{
-			RetValue = EAppReturnType::Retry;
-		}
-		else
-		{
-			RetValue = EAppReturnType::Cancel;
-		}
-		break;
-	case EAppMsgType::YesNoYesAllNoAll:
-		// return 0 for No, 1 for Yes, 2 for YesToAll, 3 for NoToAll
-		Result = ShowExtendedAlertWindow(MsgType, CocoaText, CocoaCaption);
-		if (Result == 0)
-		{
-			RetValue = EAppReturnType::No;
-		}
-		else if (Result == 1)
-		{
-			RetValue = EAppReturnType::Yes;
-		}
-		else if (Result == 2)
-		{
-			RetValue = EAppReturnType::YesAll;
-		}
-		else
-		{
-			RetValue = EAppReturnType::NoAll;
-		}
-		break;
-	case EAppMsgType::YesNoYesAllNoAllCancel:
-		// return 0 for No, 1 for Yes, 2 for YesToAll, 3 for NoToAll, 4 for Cancel
-		Result = ShowExtendedAlertWindow(MsgType, CocoaText, CocoaCaption);
-		if (Result == 0)
-		{
-			RetValue = EAppReturnType::No;
-		}
-		else if (Result == 1)
-		{
-			RetValue = EAppReturnType::Yes;
-		}
-		else if (Result == 2)
-		{
-			RetValue = EAppReturnType::YesAll;
-		}
-		else if (Result == 3)
-		{
-			RetValue = EAppReturnType::NoAll;
-		}
-		else
-		{
-			RetValue = EAppReturnType::Cancel;
-		}
-		break;
-	default:
-		NSRunAlertPanel(CocoaCaption, CocoaText, @"OK", NULL, NULL);
-		break;
-	}
-	
-	if(MacApplication)
+	if (MacApplication)
 	{
 		MacApplication->UseMouseCaptureWindow(true);
 	}
 
-	if (CocoaCaption)
-	{
-		CFRelease((CFStringRef)CocoaCaption);
-	}
-	if (CocoaText)
-	{
-		CFRelease((CFStringRef)CocoaText);
-	}
-
 	return RetValue;
 }
 
diff --git a/Engine/Source/Runtime/Launch/Private/Mac/LaunchMac.cpp b/Engine/Source/Runtime/Launch/Private/Mac/LaunchMac.cpp
--- a/Engine/Source/Runtime/Launch/Private/Mac/LaunchMac.cpp
+++ b/Engine/Source/Runtime/Launch/Private/Mac/LaunchMac.cpp
@@ -11,7 +11,6 @@
 
 
 static FString GSavedCommandLine;
-extern CORE_API NSWindow* GExtendedAlertWindow;
 extern int32 GuardedMain( const TCHAR* CmdLine );
 extern void LaunchStaticShutdownAfterError();
 
@@ -43,7 +42,6 @@ void EngineCrashHandler(const FGenericCrashContext & GenericContext)
 
 @interface UE4AppDelegate : NSObject <NSApplicationDelegate>
 {
-	IBOutlet NSWindow *AlertWindow;
 #if WITH_EDITOR
 	NSString* Filename;
 	bool bHasFinishedLaunching;
@@ -53,7 +51,6 @@ void EngineCrashHandler(const FGenericCrashContext & GenericContext)
 #if WITH_EDITOR
 - (BOOL)application:(NSApplication *)theApplication openFile:(NSString *)filename;
 #endif
-- (IBAction)OnAlertWindowButtonPress:(id)Sender;
 
 @end
 
@@ -65,7 +62,6 @@ void EngineCrashHandler(const FGenericCrashContext & GenericContext)
 	Filename = nil;
 	bHasFinishedLaunching = false;
 #endif
-	GExtendedAlertWindow = AlertWindow;
 }
 
 #if WITH_EDITOR
@@ -93,21 +89,6 @@ void EngineCrashHandler(const FGenericCrashContext & GenericContext)
 }
 #endif
 
-- (IBAction)OnAlertWindowButtonPress:(id)Sender
-{
-	static NSString* AllNames[] = { @"No", @"Yes", @"Yes to all", @"No to all", @"Cancel" };
-
-	NSString* ButtonName = [(NSButton*)Sender title];
-
-	for (int32 Index = 0; Index < sizeof(AllNames) / sizeof(AllNames[0]); Index++)
-	{
-		if ([ButtonName isEqualToString: AllNames[Index]])
-		{
-			[NSApp stopModalWithCode: Index];
-		}
-	}
-}
-
 - (IBAction)OnQuitRequest:(id)Sender
 {
 	if (GEngine)
