Fixed sticky modifier key issue on Mac

diff --git a/Engine/Source/Runtime/Core/Private/Mac/MacApplication.cpp b/Engine/Source/Runtime/Core/Private/Mac/MacApplication.cpp
--- a/Engine/Source/Runtime/Core/Private/Mac/MacApplication.cpp
+++ b/Engine/Source/Runtime/Core/Private/Mac/MacApplication.cpp
@@ -187,7 +187,7 @@ void FMacApplication::ProcessDeferredEvents( const float TimeDelta )
 	}
 }
 
-void FMacApplication::HandleExternallyChangedModifier(TSharedPtr< FMacWindow > CurrentEventWindow, NSUInteger NewModifierFlags, NSUInteger FlagsShift, NSUInteger UE4Shift, EMacModifierKeys TranslatedCode)
+void FMacApplication::HandleModifierChange(TSharedPtr< FMacWindow > CurrentEventWindow, NSUInteger NewModifierFlags, NSUInteger FlagsShift, NSUInteger UE4Shift, EMacModifierKeys TranslatedCode)
 {
 	bool CurrentPressed = (CurrentModifierFlags & FlagsShift) != 0;
 	bool NewPressed = (NewModifierFlags & FlagsShift) != 0;
@@ -235,24 +235,20 @@ void FMacApplication::ProcessEvent( NSEvent* Event )
 	{
 		return;
 	}
-	
-	// Mission Control (aka Expose, Spaces) can silently steal focus and modify the modifier flags,
-	// without posting any events, leaving our input caching very confused as the key states won't match.
-	// We should only need to capture the modifiers, because they must go down first, we shouldn't get messages for the
-	// other keys.
-	if( ( EventType != NSFlagsChanged || [Event keyCode] == 0 ) && CurrentModifierFlags != [Event modifierFlags] )
+
+	if( CurrentModifierFlags != [Event modifierFlags] )
 	{
 		NSUInteger ModifierFlags = [Event modifierFlags];
 		
-		HandleExternallyChangedModifier(CurrentEventWindow, ModifierFlags, (1<<4), 7, MMK_RightCommand);
-		HandleExternallyChangedModifier(CurrentEventWindow, ModifierFlags, (1<<3), 6, MMK_LeftCommand);
-		HandleExternallyChangedModifier(CurrentEventWindow, ModifierFlags, (1<<1), 0, MMK_LeftShift);
-		HandleExternallyChangedModifier(CurrentEventWindow, ModifierFlags, (1<<16), 8, MMK_CapsLock);
-		HandleExternallyChangedModifier(CurrentEventWindow, ModifierFlags, (1<<5), 4, MMK_LeftAlt);
-		HandleExternallyChangedModifier(CurrentEventWindow, ModifierFlags, (1<<0), 2, MMK_LeftControl);
-		HandleExternallyChangedModifier(CurrentEventWindow, ModifierFlags, (1<<2), 1, MMK_RightShift);
-		HandleExternallyChangedModifier(CurrentEventWindow, ModifierFlags, (1<<6), 5, MMK_RightAlt);
-		HandleExternallyChangedModifier(CurrentEventWindow, ModifierFlags, (1<<13), 3, MMK_RightControl);
+		HandleModifierChange(CurrentEventWindow, ModifierFlags, (1<<4), 7, MMK_RightCommand);
+		HandleModifierChange(CurrentEventWindow, ModifierFlags, (1<<3), 6, MMK_LeftCommand);
+		HandleModifierChange(CurrentEventWindow, ModifierFlags, (1<<1), 0, MMK_LeftShift);
+		HandleModifierChange(CurrentEventWindow, ModifierFlags, (1<<16), 8, MMK_CapsLock);
+		HandleModifierChange(CurrentEventWindow, ModifierFlags, (1<<5), 4, MMK_LeftAlt);
+		HandleModifierChange(CurrentEventWindow, ModifierFlags, (1<<0), 2, MMK_LeftControl);
+		HandleModifierChange(CurrentEventWindow, ModifierFlags, (1<<2), 1, MMK_RightShift);
+		HandleModifierChange(CurrentEventWindow, ModifierFlags, (1<<6), 5, MMK_RightAlt);
+		HandleModifierChange(CurrentEventWindow, ModifierFlags, (1<<13), 3, MMK_RightControl);
 		
 		CurrentModifierFlags = ModifierFlags;
 	}
@@ -562,54 +558,6 @@ void FMacApplication::ProcessEvent( NSEvent* Event )
 			}
 			break;
 		}
-
-		case NSFlagsChanged:
-		{
-			const uint32 KeyCode = [Event keyCode];
-			const uint32 Flags = [Event modifierFlags];
-			CurrentModifierFlags = Flags;
-			uint32 TranslatedCode = 0;
-
-			if( KeyCode >= 54 && KeyCode <= 62 )
-			{
-				bool Pressed = false;
-				uint32 Shift = 0;
-
-				switch( KeyCode )
-				{
-					case 54:	Pressed = ( ( Flags & (1<<4) ) != 0 );	Shift = 7;	TranslatedCode = MMK_RightCommand;	break; // right cmd
-					case 55:	Pressed = ( ( Flags & (1<<3) ) != 0 );	Shift = 6;	TranslatedCode = MMK_LeftCommand;	break; // left cmd
-					case 56:	Pressed = ( ( Flags & (1<<1) ) != 0 );	Shift = 0;	TranslatedCode = MMK_LeftShift;		break; // left shift
-					case 57:	Pressed = ( ( Flags & (1<<16)) != 0 );	Shift = 8;	TranslatedCode = MMK_CapsLock;		break; // caps lock
-					case 58:	Pressed = ( ( Flags & (1<<5) ) != 0 );	Shift = 4;	TranslatedCode = MMK_LeftAlt;		break; // left alt
-					case 59:	Pressed = ( ( Flags & (1<<0) ) != 0 );	Shift = 2;	TranslatedCode = MMK_LeftControl;	break; // left ctrl
-					case 60:	Pressed = ( ( Flags & (1<<2) ) != 0 );	Shift = 1;	TranslatedCode = MMK_RightShift;	break; // right shift
-					case 61:	Pressed = ( ( Flags & (1<<6) ) != 0 );	Shift = 5;	TranslatedCode = MMK_RightAlt;		break; // right alt
-					case 62:	Pressed = ( ( Flags & (1<<13) ) != 0 );	Shift = 3;	TranslatedCode = MMK_RightControl;	break; // right ctrl
-					default:
-						check(0);
-						break;
-				}
-
-				if( Pressed )
-				{
-					ModifierKeysFlags |= 1 << Shift;
-					if( CurrentEventWindow.IsValid() )
-					{
-						MessageHandler->OnKeyDown( TranslatedCode, 0, false );
-					}
-				}
-				else
-				{
-					ModifierKeysFlags &= ~(1 << Shift);
-					if( CurrentEventWindow.IsValid() )
-					{
-						MessageHandler->OnKeyUp( TranslatedCode, 0, false );
-					}
-				}
-			}
-			break;
-		}
 	}
 }
 
diff --git a/Engine/Source/Runtime/Core/Public/Mac/MacApplication.h b/Engine/Source/Runtime/Core/Public/Mac/MacApplication.h
--- a/Engine/Source/Runtime/Core/Public/Mac/MacApplication.h
+++ b/Engine/Source/Runtime/Core/Public/Mac/MacApplication.h
@@ -112,7 +112,7 @@ private:
 
 	void UpdateMouseCaptureWindow( FSlateCocoaWindow* TargetWindow );
 
-	void HandleExternallyChangedModifier(TSharedPtr< FMacWindow > CurrentEventWindow, NSUInteger NewModifierFlags, NSUInteger FlagsShift, NSUInteger UE4Shift, EMacModifierKeys TranslatedCode);
+	void HandleModifierChange(TSharedPtr< FMacWindow > CurrentEventWindow, NSUInteger NewModifierFlags, NSUInteger FlagsShift, NSUInteger UE4Shift, EMacModifierKeys TranslatedCode);
 
 #if WITH_EDITOR
 	void RecordUsage(EGestureEvent::Type Gesture);
