Stats - Restored indexed rendering thread, but now it's based on NAME_RenderThread, removed all references to "RenderingThread", leaving only these required for backward compatibility
Stats - Added missing GetGroupCategory to the thread node in the GetRawStackStats

diff --git a/Engine/Source/Developer/ProfilerClient/Private/ProfilerClientManager.cpp b/Engine/Source/Developer/ProfilerClient/Private/ProfilerClientManager.cpp
--- a/Engine/Source/Developer/ProfilerClient/Private/ProfilerClientManager.cpp
+++ b/Engine/Source/Developer/ProfilerClient/Private/ProfilerClientManager.cpp
@@ -926,24 +926,20 @@ int32 FServiceConnection::FindOrAddThread(const FStatNameAndInfo& Thread)
 {
 	SCOPE_CYCLE_COUNTER(STAT_PC_FindOrAddThread);
 
-	CurrentThreadState.Threads;
-
 	// The description of a thread group contains the thread id
 	const FString Desc = Thread.GetDescription();
 	const uint32 ThreadID = FStatsUtils::ParseThreadID( Desc );
 
-	{
-		const FName ShortName = Thread.GetShortName();
+	const FName ShortName = Thread.GetShortName();
 
-		// add to the meta data
-		FScopeLock ScopeLock(&CriticalSection);
-		const int32 OldNum = MetaData.ThreadDescriptions.Num();
-		MetaData.ThreadDescriptions.Add( ThreadID, ShortName.ToString() );
-		const int32 NewNum = MetaData.ThreadDescriptions.Num();
+	// add to the meta data
+	FScopeLock ScopeLock( &CriticalSection );
+	const int32 OldNum = MetaData.ThreadDescriptions.Num();
+	MetaData.ThreadDescriptions.Add( ThreadID, ShortName.ToString() );
+	const int32 NewNum = MetaData.ThreadDescriptions.Num();
 
-		// meta data has been updated
-		CurrentData.MetaDataUpdated = CurrentData.MetaDataUpdated || OldNum != NewNum;
-	}
+	// meta data has been updated
+	CurrentData.MetaDataUpdated = CurrentData.MetaDataUpdated || OldNum != NewNum;
 
 	return ThreadID;
 }
diff --git a/Engine/Source/Runtime/Core/Private/Stats/StatsData.cpp b/Engine/Source/Runtime/Core/Private/Stats/StatsData.cpp
--- a/Engine/Source/Runtime/Core/Private/Stats/StatsData.cpp
+++ b/Engine/Source/Runtime/Core/Private/Stats/StatsData.cpp
@@ -1019,7 +1019,7 @@ void FStatsThreadState::GetRawStackStats(int64 TargetFrame, FRawStatStackNode& R
 		if (!ThreadRoot)
 		{
 			FString ThreadIdName = FStatsUtils::BuildUniqueThreadName( Packet.ThreadId );
-			ThreadRoot = Root.Children.Add( ThreadName, new FRawStatStackNode( FStatMessage( ThreadName, EStatDataType::ST_int64, STAT_GROUP_TO_FStatGroup( STATGROUP_Threads )::GetGroupName(), NULL, *ThreadIdName, true, true ) ) );
+			ThreadRoot = Root.Children.Add( ThreadName, new FRawStatStackNode( FStatMessage( ThreadName, EStatDataType::ST_int64, STAT_GROUP_TO_FStatGroup( STATGROUP_Threads )::GetGroupName(), STAT_GROUP_TO_FStatGroup( STATGROUP_Threads )::GetGroupCategory(), *ThreadIdName, true, true ) ) );
 			ThreadRoot->Meta.NameAndInfo.SetFlag(EStatMetaFlags::IsPackedCCAndDuration, true);
 			ThreadRoot->Meta.Clear();
 		}
@@ -1251,8 +1251,8 @@ void FStatsThreadState::Condense(int64 TargetFrame, TArray<FStatMessage>& OutSta
 
 void FStatsThreadState::FindOrAddMetaData(FStatMessage const& Item)
 {
-	FName LongName = Item.NameAndInfo.GetRawName();
-	FName ShortName = Item.NameAndInfo.GetShortName();
+	const FName LongName = Item.NameAndInfo.GetRawName();
+	const FName ShortName = Item.NameAndInfo.GetShortName();
 
 	FStatMessage* Result = ShortNameToLongName.Find(ShortName);
 	if (!Result)
@@ -1260,7 +1260,7 @@ void FStatsThreadState::FindOrAddMetaData(FStatMessage const& Item)
 		check(ShortName != LongName);
 		FStatMessage AsSet(Item);
 		AsSet.Clear();
-		
+
 		const FName GroupName = Item.NameAndInfo.GetGroupName();
 
 		// Whether to add to the threads group.
diff --git a/Engine/Source/Runtime/Core/Public/Stats/StatsData.h b/Engine/Source/Runtime/Core/Public/Stats/StatsData.h
--- a/Engine/Source/Runtime/Core/Public/Stats/StatsData.h
+++ b/Engine/Source/Runtime/Core/Public/Stats/StatsData.h
@@ -605,6 +605,11 @@ struct CORE_API FStatsUtils
 	static FString ToEscapedFString(const TCHAR* Source);
 	static FString FromEscapedFString(const TCHAR* Escaped);
 
+	static FString BuildRenderThreadName( uint32 ThreadIndex )
+	{
+		return FString::Printf( TEXT( "%s %u" ), *FName(NAME_RenderThread).GetPlainNameString(), ThreadIndex );
+	}
+	
 	static FString BuildUniqueThreadName( uint32 InThreadID )
 	{
 		// Make unique name.
diff --git a/Engine/Source/Runtime/RenderCore/Private/RenderingThread.cpp b/Engine/Source/Runtime/RenderCore/Private/RenderingThread.cpp
--- a/Engine/Source/Runtime/RenderCore/Private/RenderingThread.cpp
+++ b/Engine/Source/Runtime/RenderCore/Private/RenderingThread.cpp
@@ -10,6 +10,7 @@
 #include "TickableObjectRenderThread.h"
 #include "ExceptionHandling.h"
 #include "TaskGraphInterfaces.h"
+#include "StatsData.h"
 
 //
 // Globals
@@ -453,7 +454,7 @@ void StartRenderingThread()
 
 	EThreadPriority RenderingThreadPrio = TPri_Normal;
 
-	GRenderingThread = FRunnableThread::Create(GRenderingThreadRunnable, *FName(NAME_RenderThread).ToString(), 0, 0, 0, RenderingThreadPrio);
+	GRenderingThread = FRunnableThread::Create( GRenderingThreadRunnable, *FStatsUtils::BuildRenderThreadName( ThreadCount ), 0, 0, 0, RenderingThreadPrio );
 
 	// Wait for render thread to have taskgraph bound before we dispatch any tasks for it.
 	((FRenderingThread*)GRenderingThreadRunnable)->TaskGraphBoundSyncEvent->Wait();
