[GitHub] Fixed FIntPoint::ToString() and FIntRect::ToString() not printing correct values

From EverNewJoy on GitHub (ttps://github.com/EpicGames/UnrealEngine/pull/120)

diff --git a/Engine/Source/Runtime/Core/Public/Math/IntPoint.h b/Engine/Source/Runtime/Core/Public/Math/IntPoint.h
--- a/Engine/Source/Runtime/Core/Public/Math/IntPoint.h
+++ b/Engine/Source/Runtime/Core/Public/Math/IntPoint.h
@@ -414,5 +414,5 @@ FORCEINLINE int32 FIntPoint::SizeSquared() const
 
 FORCEINLINE FString FIntPoint::ToString( ) const
 {
-	return FString::Printf(TEXT("X=%3.3f Y=%3.3f"), X, Y);
+	return FString::Printf(TEXT("X=%d Y=%d"), X, Y);
 }
\ No newline at end of file
