Fixed a problem with the ipad mini 2nd gen using the ipad mini 1st gen settings

diff --git a/Engine/Source/Runtime/Core/Private/IOS/IOSPlatformMisc.cpp b/Engine/Source/Runtime/Core/Private/IOS/IOSPlatformMisc.cpp
--- a/Engine/Source/Runtime/Core/Private/IOS/IOSPlatformMisc.cpp
+++ b/Engine/Source/Runtime/Core/Private/IOS/IOSPlatformMisc.cpp
@@ -352,7 +352,7 @@ FIOSPlatformMisc::EIOSDevice FIOSPlatformMisc::GetIOSDeviceType()
 		{
 			if (Minor >= 4)
 			{
-				DeviceType = IOS_IPadMini;
+				DeviceType = IOS_IPadMini2;
 			}
 			else
 			{
