Move functionality for handling .uprojectdirs into a class which can be reused for any engine installation. Required for UnrealVersionSelector to handle generating project files for multiple installations.

diff --git a/Engine/Source/Runtime/Core/Private/GenericPlatform/GenericPlatformMisc.cpp b/Engine/Source/Runtime/Core/Private/GenericPlatform/GenericPlatformMisc.cpp
--- a/Engine/Source/Runtime/Core/Private/GenericPlatform/GenericPlatformMisc.cpp
+++ b/Engine/Source/Runtime/Core/Private/GenericPlatform/GenericPlatformMisc.cpp
@@ -466,8 +466,7 @@ const TCHAR* FGenericPlatformMisc::GameDir()
 						FPlatformMisc::LowLevelOutputDebugStringf(TEXT("Failed to find game directory: %s\n"), *GameDir);
 
 						// Use the uprojectdirs
-						FUProjectInfoHelper::FillProjectInfo();
-						FString GameProjectFile = FUProjectInfoHelper::GetProjectForGame(GGameName);
+						FString GameProjectFile = FUProjectDictionary::GetDefault().GetRelativeProjectPathForGame(GGameName, FPlatformProcess::BaseDir());
 						if (GameProjectFile.IsEmpty() == false)
 						{
 							// We found a project folder for the game
diff --git a/Engine/Source/Runtime/Core/Private/Misc/UProjectInfo.cpp b/Engine/Source/Runtime/Core/Private/Misc/UProjectInfo.cpp
--- a/Engine/Source/Runtime/Core/Private/Misc/UProjectInfo.cpp
+++ b/Engine/Source/Runtime/Core/Private/Misc/UProjectInfo.cpp
@@ -7,30 +7,23 @@
 DECLARE_LOG_CATEGORY_EXTERN(LogUProjectInfo, Verbose, All);
 DEFINE_LOG_CATEGORY(LogUProjectInfo);
 
-bool FUProjectInfoHelper::bProjectInfoFilled = false;
-TMap<FString, FUProjectInfo> FUProjectInfoHelper::ProjectInfoDictionary;
-TMap<FString, FString> FUProjectInfoHelper::ShortProjectNameDictionary;
-
-void FUProjectInfoHelper::FillProjectInfo()
+FUProjectDictionary::FUProjectDictionary(const FString& InRootDir)
 {
-	if (bProjectInfoFilled == true)
-	{
-		return;
-	}
-
 	double StartTime = FPlatformTime::Seconds();
 
+	FString RootDir = InRootDir;
+	FPaths::NormalizeDirectoryName(RootDir);
+
 	TArray<FString> DirectoriesToSearch;
 	
 	// Find all the .uprojectdirs files contained in the root folder and add their entries to the search array
 	TArray<FString> UProjectDirsFiles;
-
-	FString SearchPath = FPaths::RootDir() / FString(TEXT("*.uprojectdirs"));
+	FString SearchPath = RootDir / FString(TEXT("*.uprojectdirs"));
 	IFileManager::Get().FindFiles(UProjectDirsFiles, *SearchPath, true, false);
 
 	for (int32 FileIndex = 0; FileIndex < UProjectDirsFiles.Num(); ++FileIndex)
 	{
-		FString ProjDirsFile = FPaths::RootDir() / UProjectDirsFiles[FileIndex];
+		FString ProjDirsFile = RootDir / UProjectDirsFiles[FileIndex];
 		UE_LOG(LogUProjectInfo, Log, TEXT("Found uprojectdirs file %s"), *ProjDirsFile);
 
 		TArray<FString> FileStrings;
@@ -49,7 +42,7 @@ void FUProjectInfoHelper::FillProjectInfo()
 					}
 					else
 					{
-						DirectoriesToSearch.Add(FPaths::RootDir() / ProjDirEntry);
+						DirectoriesToSearch.Add(RootDir / ProjDirEntry);
 					}
 				}
 			}
@@ -84,41 +77,49 @@ void FUProjectInfoHelper::FillProjectInfo()
 
 			for (int32 UProjIdx = 0; UProjIdx < SubDirFiles.Num(); UProjIdx++)
 			{
-				FString UProject = SubDir / SubDirFiles[UProjIdx];
-				FString RelativePath = UProject;
-				FPaths::MakePathRelativeTo(RelativePath, FPlatformProcess::BaseDir());
-				RelativePath = RelativePath.Replace(TEXT("\\"), TEXT("/"));
-				UE_LOG(LogUProjectInfo, Log, TEXT("\t\t\t\t%s"), *RelativePath);
+				FString ProjectFileName = SubDir / SubDirFiles[UProjIdx];
+				FPaths::NormalizeFilename(ProjectFileName);
 
-				if (ProjectInfoDictionary.Find(RelativePath) == NULL)
-				{
-					FUProjectInfo NewProjectInfo;
-					NewProjectInfo.ProjectName = FPaths::GetBaseFilename(RelativePath);
-					NewProjectInfo.ProjectFileFolder = FPaths::GetPath(RelativePath);
-					if (NewProjectInfo.ProjectFileFolder.EndsWith("/"))
-					{
-						NewProjectInfo.ProjectFileFolder = NewProjectInfo.ProjectFileFolder.Left(NewProjectInfo.ProjectFileFolder.Len() - 1);
-					}
-					FString ProjectSourceDirectory = FPaths::GetPath(RelativePath) / TEXT("Source");
-					ProjectInfoDictionary.Add(RelativePath, NewProjectInfo);
-					ShortProjectNameDictionary.Add(NewProjectInfo.ProjectName, RelativePath);
-				}
+				FString ShortName = FPaths::GetBaseFilename(ProjectFileName);
+				ShortProjectNameDictionary.Add(ShortName, ProjectFileName);
 			}
 		}
 	}
 
 	double TotalProjectInfoTime = FPlatformTime::Seconds() - StartTime;
 	UE_LOG(LogUProjectInfo, Log, TEXT("FillProjectInfo took %5.4f seconds"), TotalProjectInfoTime);
+}
 
-	bProjectInfoFilled = true;
+bool FUProjectDictionary::IsForeignProject(const FString& InProjectFileName)
+{
+	FString ProjectFileName = InProjectFileName;
+	FPaths::NormalizeFilename(ProjectFileName);
+
+	for(TMap<FString, FString>::TConstIterator Iter(ShortProjectNameDictionary); Iter; ++Iter)
+	{
+		if(Iter.Value() == ProjectFileName)
+		{
+			return false;
+		}
+	}
+
+	return true;
 }
 
-FString FUProjectInfoHelper::GetProjectForGame(const TCHAR* InGameName)
+FString FUProjectDictionary::GetRelativeProjectPathForGame(const TCHAR* InGameName, const FString& BaseDir)
 {
 	FString* ProjectFile = ShortProjectNameDictionary.Find(*(FString(InGameName).ToLower()));
 	if (ProjectFile != NULL)
 	{
-		return *ProjectFile;
+		FString RelativePath = *ProjectFile;
+		FPaths::MakePathRelativeTo(RelativePath, *BaseDir);
+		return RelativePath;
 	}
 	return TEXT("");
 }
+
+FUProjectDictionary& FUProjectDictionary::GetDefault()
+{
+	static FUProjectDictionary DefaultDictionary(FPaths::RootDir());
+	return DefaultDictionary;
+}
diff --git a/Engine/Source/Runtime/Core/Public/Misc/UProjectInfo.h b/Engine/Source/Runtime/Core/Public/Misc/UProjectInfo.h
--- a/Engine/Source/Runtime/Core/Public/Misc/UProjectInfo.h
+++ b/Engine/Source/Runtime/Core/Public/Misc/UProjectInfo.h
@@ -2,38 +2,23 @@
 
 #pragma once
 
-//@NOTE: This is more or less exactly mirrored in UnrealBuildTool.UProjectInfo
-// The UBT one contains more functionality as it needs to match build Targets to UProjects.
-
-/** Helper struct for tracking UProject information */
-struct FUProjectInfo
-{
-	/** Name of the project file */
-	FString ProjectName;
-	/** Folder the project file is located in */
-	FString ProjectFileFolder;
-};
-
-class FUProjectInfoHelper
+/** Dictionary of all the non-foreign projects for an engine installation, found by parsing .uprojectdirs files for source directories. */
+class CORE_API FUProjectDictionary
 {
 public:
-	/** Fill in the project info */
-	static CORE_API void FillProjectInfo();
+	/** Scans the engine root directory for all the known projects */
+	FUProjectDictionary(const FString& InRootDir);
+	
+	/** Determines whether a project is a foreign project or not */
+	bool IsForeignProject(const FString& ProjectFileName);
+
+	/** Gets the project filename for the given game. Empty if not found. */
+	FString GetRelativeProjectPathForGame(const TCHAR* GameName, const FString& BaseDir);
 
-	/** 
-	 *	Get the project path for the given game name.
-	 *
-	 *	@param	InGameName		The game of interest
-	 *
-	 *	@return	FString			The path to the UProject file, empty if not found
-	 */
-	static CORE_API FString GetProjectForGame(const TCHAR* InGameName);
+	/** Gets the project dictionary for the active engine installation */
+	static FUProjectDictionary& GetDefault();
 
-protected:
-	/** True if project info has already been filled... */
-	static bool bProjectInfoFilled;
-	/** Map of relative or complete project file names to the project info */
-	static TMap<FString, FUProjectInfo> ProjectInfoDictionary;
-	/** Map of short project file names to the relative or complete project file name */
-	static TMap<FString, FString> ShortProjectNameDictionary;
+private:
+	/** Map of short game names to full project paths */
+	TMap<FString, FString> ShortProjectNameDictionary;
 };
diff --git a/Engine/Source/Runtime/Launch/Private/LaunchEngineLoop.cpp b/Engine/Source/Runtime/Launch/Private/LaunchEngineLoop.cpp
--- a/Engine/Source/Runtime/Launch/Private/LaunchEngineLoop.cpp
+++ b/Engine/Source/Runtime/Launch/Private/LaunchEngineLoop.cpp
@@ -607,8 +607,7 @@ int32 FEngineLoop::PreInit( const TCHAR* CmdLine )
 				UE_LOG(LogInit, Display, TEXT("Project file not found: %s"), *ProjPath);
 				UE_LOG(LogInit, Display, TEXT("\tAttempting to find via project info helper."));
 				// Use the uprojectdirs
-				FUProjectInfoHelper::FillProjectInfo();
-				FString GameProjectFile = FUProjectInfoHelper::GetProjectForGame(GGameName);
+				FString GameProjectFile = FUProjectDictionary::GetDefault().GetRelativeProjectPathForGame(GGameName, FPlatformProcess::BaseDir());
 				if (GameProjectFile.IsEmpty() == false)
 				{
 					UE_LOG(LogInit, Display, TEXT("\tFound project file %s."), *GameProjectFile);
