Fixed a problem in Mac editor with the app unhiding immediately after hiding if interactive notification window was acitve.

diff --git a/Engine/Source/Runtime/Core/Private/Mac/MacWindow.cpp b/Engine/Source/Runtime/Core/Private/Mac/MacWindow.cpp
--- a/Engine/Source/Runtime/Core/Private/Mac/MacWindow.cpp
+++ b/Engine/Source/Runtime/Core/Private/Mac/MacWindow.cpp
@@ -134,22 +134,25 @@
 
 - (void)orderFrontEvenIfChildAndMakeMain:(bool)bMain andKey:(bool)bKey
 {
-	NSWindow* ParentWindow = [self getParent];
-	if( ParentWindow && [ParentWindow shouldAddChildWindow:self] )
+	if ([NSApp isHidden] == NO)
 	{
-		[ParentWindow removeChildWindow: self];
-		[ParentWindow addChildWindow: self ordered: NSWindowAbove];
-	}
-	
-	[self orderFront:nil];
-	
-	if(bMain && [self canBecomeMainWindow])
-	{
-		[self makeMainWindow];
-	}
-	if(bKey && [self canBecomeKeyWindow])
-	{
-		[self makeKeyWindow];
+		NSWindow* ParentWindow = [self getParent];
+		if( ParentWindow && [ParentWindow shouldAddChildWindow:self] )
+		{
+			[ParentWindow removeChildWindow: self];
+			[ParentWindow addChildWindow: self ordered: NSWindowAbove];
+		}
+		
+		[self orderFront:nil];
+		
+		if(bMain && [self canBecomeMainWindow])
+		{
+			[self makeMainWindow];
+		}
+		if(bKey && [self canBecomeKeyWindow])
+		{
+			[self makeKeyWindow];
+		}
 	}
 }
 
