Fixed a crash on renaming bluepring when using Kotoeri/Romaji keyboard on Mac

diff --git a/Engine/Source/Runtime/Core/Private/Mac/MacTextInputMethodSystem.cpp b/Engine/Source/Runtime/Core/Private/Mac/MacTextInputMethodSystem.cpp
--- a/Engine/Source/Runtime/Core/Private/Mac/MacTextInputMethodSystem.cpp
+++ b/Engine/Source/Runtime/Core/Private/Mac/MacTextInputMethodSystem.cpp
@@ -12,7 +12,8 @@ DEFINE_LOG_CATEGORY_STATIC(LogMacTextInputMethodSystem, Log, All);
 - (id)initWithFrame:(NSRect)frame
 {
     self = [super initWithFrame:frame];
-    if (self) {
+    if (self)
+	{
         // Initialization code here.
 		IMMContext = nil;
 		markedRange = {NSNotFound, 0};
@@ -30,7 +31,7 @@ DEFINE_LOG_CATEGORY_STATIC(LogMacTextInputMethodSystem, Log, All);
 
 - (bool)imkKeyDown:(NSEvent *)theEvent
 {
-	if(IMMContext)
+	if (IMMContext.IsValid())
 	{
 		reallyHandledEvent = true;
 		return [[self inputContext] handleEvent:theEvent] && reallyHandledEvent;
@@ -43,7 +44,7 @@ DEFINE_LOG_CATEGORY_STATIC(LogMacTextInputMethodSystem, Log, All);
 
 - (void)mouseDown:(NSEvent *)theEvent
 {
-	if(IMMContext)
+	if (IMMContext.IsValid())
 	{
 		[[self inputContext] handleEvent:theEvent];
 	}
@@ -51,7 +52,7 @@ DEFINE_LOG_CATEGORY_STATIC(LogMacTextInputMethodSystem, Log, All);
 
 - (void)mouseDragged:(NSEvent *)theEvent
 {
-	if(IMMContext)
+	if (IMMContext.IsValid())
 	{
 		[[self inputContext] handleEvent:theEvent];
 	}
@@ -59,19 +60,18 @@ DEFINE_LOG_CATEGORY_STATIC(LogMacTextInputMethodSystem, Log, All);
 
 - (void)mouseUp:(NSEvent *)theEvent
 {
-	if(IMMContext)
+	if (IMMContext.IsValid())
 	{
 		[[self inputContext] handleEvent:theEvent];
 	}
 }
 
-- (void)activateInputMethod:(ITextInputMethodContext*)InContext
+- (void)activateInputMethod:(const TSharedRef<ITextInputMethodContext>&)InContext
 {
-	check(InContext);
-	if(IMMContext)
+	if (IMMContext.IsValid())
 	{
+		[self unmarkText];
 		[[self inputContext] deactivate];
-		IMMContext = NULL;
 	}
 	IMMContext = InContext;
 	[[self inputContext] activate];
@@ -90,7 +90,7 @@ DEFINE_LOG_CATEGORY_STATIC(LogMacTextInputMethodSystem, Log, All);
  */
 - (void)insertText:(id)aString replacementRange:(NSRange)replacementRange
 {
-	if(IMMContext && [self hasMarkedText])
+	if (IMMContext.IsValid() && [self hasMarkedText])
 	{
 		uint32 SelectionLocation, SelectionLength;
 		if (replacementRange.location == NSNotFound)
@@ -146,7 +146,7 @@ DEFINE_LOG_CATEGORY_STATIC(LogMacTextInputMethodSystem, Log, All);
  */
 - (void)setMarkedText:(id)aString selectedRange:(NSRange)selectedRange replacementRange:(NSRange)replacementRange
 {
-	if(IMMContext)
+	if (IMMContext.IsValid())
 	{
 		uint32 SelectionLocation, SelectionLength;
 		if (replacementRange.location == NSNotFound)
@@ -209,7 +209,7 @@ DEFINE_LOG_CATEGORY_STATIC(LogMacTextInputMethodSystem, Log, All);
 - (void)unmarkText
 {
 	markedRange = {NSNotFound, 0};
-	if(IMMContext)
+	if (IMMContext.IsValid())
 	{
 		IMMContext->EndComposition();
 	}
@@ -219,7 +219,7 @@ DEFINE_LOG_CATEGORY_STATIC(LogMacTextInputMethodSystem, Log, All);
  */
 - (NSRange)selectedRange
 {
-	if(IMMContext)
+	if (IMMContext.IsValid())
 	{
 		uint32 SelectionLocation, SelectionLength;
 		ITextInputMethodContext::ECaretPosition CaretPosition;
@@ -237,7 +237,7 @@ DEFINE_LOG_CATEGORY_STATIC(LogMacTextInputMethodSystem, Log, All);
  */
 - (NSRange)markedRange
 {
-	if(IMMContext)
+	if (IMMContext.IsValid())
 	{
 		return markedRange;
 	}
@@ -251,7 +251,7 @@ DEFINE_LOG_CATEGORY_STATIC(LogMacTextInputMethodSystem, Log, All);
  */
 - (BOOL)hasMarkedText
 {
-	return IMMContext && (markedRange.location != NSNotFound);
+	return IMMContext.IsValid() && (markedRange.location != NSNotFound);
 }
 
 /* Returns attributed string specified by aRange. It may return nil. If non-nil return value and actualRange is non-NULL, it contains the actual range for the return value. The range can be adjusted from various reasons (i.e. adjust to grapheme cluster boundary, performance optimization, etc).
@@ -259,7 +259,7 @@ DEFINE_LOG_CATEGORY_STATIC(LogMacTextInputMethodSystem, Log, All);
 - (NSAttributedString *)attributedSubstringForProposedRange:(NSRange)aRange actualRange:(NSRangePointer)actualRange
 {
 	NSAttributedString* AttributedString = nil;
-	if(IMMContext)
+	if (IMMContext.IsValid())
 	{
 		FString String;
 		IMMContext->GetTextInRange(aRange.location, aRange.length, String);
@@ -291,7 +291,7 @@ DEFINE_LOG_CATEGORY_STATIC(LogMacTextInputMethodSystem, Log, All);
  */
 - (NSRect)firstRectForCharacterRange:(NSRange)aRange actualRange:(NSRangePointer)actualRange
 {
-	if(IMMContext)
+	if (IMMContext.IsValid())
 	{
 		FVector2D Position, Size;
 		IMMContext->GetTextBounds(aRange.location, aRange.length, Position, Size);
@@ -319,7 +319,7 @@ DEFINE_LOG_CATEGORY_STATIC(LogMacTextInputMethodSystem, Log, All);
  */
 - (NSUInteger)characterIndexForPoint:(NSPoint)aPoint
 {
-	if(IMMContext)
+	if (IMMContext.IsValid())
 	{
 		FVector2D Pos(aPoint.x, aPoint.y);
 		int32 Index = IMMContext->GetCharacterIndexFromPoint(Pos);
@@ -484,7 +484,7 @@ void FMacTextInputMethodSystem::ActivateContext(const TSharedRef<ITextInputMetho
 		if(CocoaWindow && [CocoaWindow openGLView])
 		{
 			FSlateTextView* TextView = (FSlateTextView*)[CocoaWindow openGLView];
-			[TextView activateInputMethod:&Context.Get()];
+			[TextView activateInputMethod:Context];
 			return;
 		}
 	}
diff --git a/Engine/Source/Runtime/Core/Public/Mac/MacTextInputMethodSystem.h b/Engine/Source/Runtime/Core/Public/Mac/MacTextInputMethodSystem.h
--- a/Engine/Source/Runtime/Core/Public/Mac/MacTextInputMethodSystem.h
+++ b/Engine/Source/Runtime/Core/Public/Mac/MacTextInputMethodSystem.h
@@ -8,12 +8,12 @@
 
 @interface FSlateTextView : NSView <NSTextInputClient>
 {
-	ITextInputMethodContext* IMMContext;
+	TSharedPtr<ITextInputMethodContext> IMMContext;
 	NSRange markedRange;
 	bool reallyHandledEvent;
 }
 
-- (void)activateInputMethod:(ITextInputMethodContext*)InContext;
+- (void)activateInputMethod:(const TSharedRef<ITextInputMethodContext>&)InContext;
 - (void)deactivateInputMethod;
 
 - (bool)imkKeyDown:(NSEvent*)Event;
