Core - Fixed a memory leak in the TIndirectArray.Reset

diff --git a/Engine/Source/Runtime/Core/Public/Containers/Array.h b/Engine/Source/Runtime/Core/Public/Containers/Array.h
--- a/Engine/Source/Runtime/Core/Public/Containers/Array.h
+++ b/Engine/Source/Runtime/Core/Public/Containers/Array.h
@@ -1015,6 +1015,26 @@ public:
 		}
 	}
 
+	/**
+	* Same as empty, but doesn't change memory allocations, unless the new size is larger than
+	* the current array. It calls the destructors on held items if needed and then zeros the ArrayNum.
+	*
+	* @param NewSize the expected usage size
+	*/
+	void Reset( int32 NewSize = 0 )
+	{
+		// If we have space to hold the excepted size, then don't reallocate
+		if( NewSize <= ArrayMax )
+		{
+			DestructItems( GetTypedData(), ArrayNum );
+			ArrayNum = 0;
+		}
+		else
+		{
+			Empty( NewSize );
+		}
+	}
+
 	void Empty( int32 Slack=0 )
 	{
 		DestructItems(GetTypedData(), ArrayNum);
@@ -1395,26 +1415,6 @@ public:
 		}
 	}
 
-	/**
-	 * Same as empty, but doesn't change memory allocations, unless the new size is larger than
-	 * the current array. It calls the destructors on held items if needed and then zeros the ArrayNum.
-	 *
-	 * @param NewSize the expected usage size
-	 */
-	void Reset(int32 NewSize = 0)
-	{
-		// If we have space to hold the excepted size, then don't reallocate
-		if (NewSize <= ArrayMax)
-		{
-			DestructItems(GetTypedData(), ArrayNum);
-			ArrayNum = 0;
-		}
-		else
-		{
-			Empty(NewSize);
-		}
-	}
-
 	/**
 	 * Searches for the first entry of the specified type, will only work
 	 * with TArray<UObject*>.  Optionally return the item's index, and can
@@ -2119,6 +2119,7 @@ public:
 	}
 	void Reset(int32 NewSize = 0)
 	{
+		DestructAndFreeItems();
 		Array.Reset(NewSize);
 	}
 
@@ -2224,16 +2225,7 @@ public:
 	}
 	void Empty( int32 Slack=0 )
 	{
-		T** Element = GetTypedData();
-		for (int32 i = Array.Num(); i; --i)
-		{
-			// We need a typedef here because VC won't compile the destructor call below if T itself has a member called T
-			typedef T IndirectArrayDestructElementType;
-
-			(*Element)->IndirectArrayDestructElementType::~IndirectArrayDestructElementType();
-			FMemory::Free(*Element);
-			++Element;
-		}
+		DestructAndFreeItems();
 		Array.Empty( Slack );
 	}
 	FORCEINLINE int32 Add(T* Item)
@@ -2287,6 +2279,20 @@ public:
 	}
 
 private:
+	void DestructAndFreeItems()
+	{
+		T** Element = GetTypedData();
+		for( int32 i = Array.Num(); i; --i )
+		{
+			// We need a typedef here because VC won't compile the destructor call below if T itself has a member called T
+			typedef T IndirectArrayDestructElementType;
+
+			(*Element)->IndirectArrayDestructElementType::~IndirectArrayDestructElementType();
+			FMemory::Free( *Element );
+			++Element;
+		}
+	}
+
 	/**
 	 * DO NOT USE DIRECTLY
 	 * STL-like iterators to enable range-based for loop support.
