Fix multi-monitor handling on OS X again to deal with even more unusual display arrangements.
- GetWorkArea needs to calculate top & bottom relative to the primary display.

diff --git a/Engine/Source/Runtime/Core/Private/Mac/MacApplication.cpp b/Engine/Source/Runtime/Core/Private/Mac/MacApplication.cpp
--- a/Engine/Source/Runtime/Core/Private/Mac/MacApplication.cpp
+++ b/Engine/Source/Runtime/Core/Private/Mac/MacApplication.cpp
@@ -777,10 +777,14 @@ FPlatformRect FMacApplication::GetWorkArea( const FPlatformRect& CurrentWindow )
 
 	const int32 ScreenHeight = FMath::TruncToInt([Screen frame].size.height);
 	const NSRect VisibleFrame = [Screen visibleFrame];
+	
+	NSArray* AllScreens = [NSScreen screens];
+	NSScreen* PrimaryScreen = (NSScreen*)[AllScreens objectAtIndex: 0];
+	NSRect PrimaryFrame = [PrimaryScreen frame];
 
 	FPlatformRect WorkArea;
 	WorkArea.Left = VisibleFrame.origin.x;
-	WorkArea.Top = ScreenHeight - VisibleFrame.size.height - (VisibleFrame.origin.y - [Screen frame].origin.y);
+	WorkArea.Top = (PrimaryFrame.origin.y + PrimaryFrame.size.height) - (VisibleFrame.origin.y + VisibleFrame.size.height);
 	WorkArea.Right = WorkArea.Left + VisibleFrame.size.width;
 	WorkArea.Bottom = WorkArea.Top + VisibleFrame.size.height;
 
@@ -839,7 +843,7 @@ void FMacApplication::GetDisplayMetrics( FDisplayMetrics& OutDisplayMetrics ) co
 	OutDisplayMetrics.PrimaryDisplayWorkAreaRect.Left = VisibleFrame.origin.x;
 	OutDisplayMetrics.PrimaryDisplayWorkAreaRect.Top = ScreenFrame.size.height - (VisibleFrame.origin.y + VisibleFrame.size.height);
 	OutDisplayMetrics.PrimaryDisplayWorkAreaRect.Right = VisibleFrame.origin.x + VisibleFrame.size.width;
-	OutDisplayMetrics.PrimaryDisplayWorkAreaRect.Bottom = ScreenFrame.size.height - VisibleFrame.origin.y;
+	OutDisplayMetrics.PrimaryDisplayWorkAreaRect.Bottom = OutDisplayMetrics.PrimaryDisplayWorkAreaRect.Top + VisibleFrame.size.height;
 }
 
 void FMacApplication::OnDragEnter( FSlateCocoaWindow* Window, void *InPasteboard )
