Added handling of mouse move events in Mac editor when the editor is in the background.

diff --git a/Engine/Source/Runtime/Core/Private/Mac/MacApplication.cpp b/Engine/Source/Runtime/Core/Private/Mac/MacApplication.cpp
--- a/Engine/Source/Runtime/Core/Private/Mac/MacApplication.cpp
+++ b/Engine/Source/Runtime/Core/Private/Mac/MacApplication.cpp
@@ -43,7 +43,9 @@ FMacApplication::FMacApplication()
 	, CurrentModifierFlags( 0 )
 {
 	CGDisplayRegisterReconfigurationCallback(FMacApplication::OnDisplayReconfiguration, this);
-	
+
+	[NSEvent addGlobalMonitorForEventsMatchingMask:NSMouseMovedMask handler:^(NSEvent* Event){AddPendingEvent(Event);}];
+
 	TextInputMethodSystem = MakeShareable( new FMacTextInputMethodSystem );
 	if(!TextInputMethodSystem->Initialize())
 	{
