Fix Linux compile.

diff --git a/Engine/Source/Runtime/Core/Public/Linux/LinuxPlatformProcess.h b/Engine/Source/Runtime/Core/Public/Linux/LinuxPlatformProcess.h
--- a/Engine/Source/Runtime/Core/Public/Linux/LinuxPlatformProcess.h
+++ b/Engine/Source/Runtime/Core/Public/Linux/LinuxPlatformProcess.h
@@ -131,6 +131,7 @@ struct CORE_API FLinuxPlatformProcess : public FGenericPlatformProcess
 	static void* GetDllExport( void* DllHandle, const TCHAR* ProcName );
 	static const TCHAR* ComputerName();
 	static const TCHAR* BaseDir();
+	static FString GetApplicationName( uint32 ProcessId );
 	static bool SetProcessLimits(EProcessResource::Type Resource, uint64 Limit);
 	static const TCHAR* ExecutableName(bool bRemoveExtension = true);
 	static void ClosePipe( void* ReadPipe, void* WritePipe );
