Fixed setting of initial window position on Mac

diff --git a/Engine/Source/Runtime/Core/Private/Mac/MacWindow.cpp b/Engine/Source/Runtime/Core/Private/Mac/MacWindow.cpp
--- a/Engine/Source/Runtime/Core/Private/Mac/MacWindow.cpp
+++ b/Engine/Source/Runtime/Core/Private/Mac/MacWindow.cpp
@@ -666,7 +666,8 @@ void FMacWindow::Initialize( FMacApplication* const Application, const TSharedRe
 	PositionX = X;
 	PositionY = Y;
 
-	const NSRect ViewRect = NSMakeRect(X, Y, SizeX, SizeY);
+	const int32 InvertedY = FPlatformMisc::ConvertSlateYPositionToCocoa(Y) - SizeY + 1;
+	const NSRect ViewRect = NSMakeRect(X, InvertedY, SizeX, SizeY);
 
 	uint32 WindowStyle = 0;
 	if( Definition->IsRegularWindow )
