Significantly improve Mac OS X editor responsiveness by deferring any requests to move or reshape a window that has been realised but which has then been ordered or alpha'd offscreen.
After a window has been created & rendered to if it becomes invisible (orderOut or alphaValue == 0) then cache the new frame origin/size & update it only once it becomes visible (ordered onscreen & alphaValue > 0),
this prevents updating invisible tooltips from forcing Cocoa to synchronise Main & Render threads. Since the render thread might be responsible for triggering performDeferredSetFrame it will internally dispatch to the main thread when necessary & performDeferredOrderFront won't show the window until this has ocurred.

diff --git a/Engine/Source/Runtime/Core/Private/Mac/MacWindow.cpp b/Engine/Source/Runtime/Core/Private/Mac/MacWindow.cpp
--- a/Engine/Source/Runtime/Core/Private/Mac/MacWindow.cpp
+++ b/Engine/Source/Runtime/Core/Private/Mac/MacWindow.cpp
@@ -27,11 +27,14 @@
 		bRoundedCorners = false;
 		bDisplayReconfiguring = false;
 		bDeferOrderFront = false;
-		bDeferOpacity = 0.0f;
+		DeferOpacity = 0.0f;
 		bRenderInitialised = false;
 		bZoomed = [super isZoomed];
 		self.bForwardEvents = true;
-		[super setAlphaValue:bDeferOpacity];
+		[super setAlphaValue:DeferOpacity];
+		DeferFrame = [super frame];
+		bDeferSetFrame = false;
+		bDeferSetOrigin = false;
 	}
 	return NewSelf;
 }
@@ -67,14 +70,56 @@
 	if(!bRenderInitialised)
 	{
 		bRenderInitialised = true;
-		
-		if(bDeferOrderFront)
+	}
+	
+	if(bDeferOrderFront)
+	{
+		if(!(bDeferSetFrame || bDeferSetOrigin))
 		{
 			bDeferOrderFront = false;
+			[super setAlphaValue:DeferOpacity];
+		}
+		else
+		{
+			[self performDeferredSetFrame];
+		}
+	}
+}
+
+- (void)performDeferredSetFrame
+{
+	if(bRenderInitialised && (bDeferSetFrame || bDeferSetOrigin))
+	{
+		dispatch_block_t Block = ^{
+			if(bDeferSetOrigin)
+			{
+				DeferFrame.size = [self frame].size;
+			}
 			
-			[self setAlphaValue:bDeferOpacity];
+			[super setFrame:DeferFrame display:NO];
+		};
+		
+		if([NSThread isMainThread])
+		{
+			Block();
 		}
+		else
+		{
+			dispatch_async(dispatch_get_main_queue(), Block);
+		}
+		
+		bDeferSetFrame = false;
+		bDeferSetOrigin = false;
+	}
+}
+
+- (void)orderWindow:(NSWindowOrderingMode)OrderingMode relativeTo:(NSInteger)OtherWindowNumber
+{
+	if([self alphaValue] > 0.0f)
+	{
+		[self performDeferredSetFrame];
 	}
+	[super orderWindow:OrderingMode relativeTo:OtherWindowNumber];
 }
 
 - (bool)roundedCorners
@@ -234,11 +279,15 @@
 {
 	if(!bRenderInitialised)
 	{
-		bDeferOpacity = WindowAlpha;
+		DeferOpacity = WindowAlpha;
 		bDeferOrderFront = true;
 	}
 	else
 	{
+		if([self isVisible] && WindowAlpha > 0.0f)
+		{
+			[self performDeferredSetFrame];
+		}
 		[super setAlphaValue:WindowAlpha];
 	}
 }
@@ -268,6 +317,32 @@
 	[self zoom: self];
 }
 
+- (void)setFrame:(NSRect)FrameRect display:(BOOL)Flag
+{
+	if(!bRenderInitialised || ([self isVisible] && [super alphaValue] > 0.0f))
+	{
+		[super setFrame:FrameRect display:Flag];
+	}
+	else
+	{
+		bDeferSetFrame = true;
+		DeferFrame = FrameRect;
+	}
+}
+
+- (void)setFrameOrigin:(NSPoint)Point
+{
+	if(!bRenderInitialised || ([self isVisible] && [super alphaValue] > 0.0f))
+	{
+		[super setFrameOrigin:Point];
+	}
+	else
+	{
+		bDeferSetOrigin = true;
+		DeferFrame.origin = Point;
+	}
+}
+
 // keyDown and keyUp are empty, but having them lets Cocoa know we handle the keys ourselves
 - (void)keyDown:(NSEvent *)Event
 {
diff --git a/Engine/Source/Runtime/Core/Public/Mac/MacWindow.h b/Engine/Source/Runtime/Core/Public/Mac/MacWindow.h
--- a/Engine/Source/Runtime/Core/Public/Mac/MacWindow.h
+++ b/Engine/Source/Runtime/Core/Public/Mac/MacWindow.h
@@ -15,11 +15,14 @@
 	NSWindow* Parent;
 	EWindowMode::Type WindowMode;
 	NSMutableArray* ChildWindows;
+	NSRect DeferFrame;
+	CGFloat DeferOpacity;
 	bool bAcceptsInput;
 	bool bRoundedCorners;
 	bool bDisplayReconfiguring;
 	bool bDeferOrderFront;
-	CGFloat bDeferOpacity;
+	bool bDeferSetFrame;
+	bool bDeferSetOrigin;
 	bool bRenderInitialised;
 @public
 	bool bZoomed;
@@ -34,8 +37,12 @@
 /** Get the view used for OpenGL rendering. @return The OpenGL view for rendering. */
 - (NSView*)openGLView;
 
+/** Perform render thread deferred order front operation. */
 - (void)performDeferredOrderFront;
 
+/** Perform deferred set frame operation */
+- (void)performDeferredSetFrame;
+
 /** Set whether the window should display with rounded corners. */
 - (void)setRoundedCorners:(bool)bUseRoundedCorners;
 
