Updated MacApplication so it ignores mouse move events for windows owned by other applications

diff --git a/Engine/Source/Runtime/Core/Private/Mac/MacApplication.cpp b/Engine/Source/Runtime/Core/Private/Mac/MacApplication.cpp
--- a/Engine/Source/Runtime/Core/Private/Mac/MacApplication.cpp
+++ b/Engine/Source/Runtime/Core/Private/Mac/MacApplication.cpp
@@ -222,7 +222,7 @@ void FMacApplication::ProcessEvent( NSEvent* Event )
 
 	FSlateCocoaWindow* NativeWindow = FindEventWindow( Event );
 	TSharedPtr< FMacWindow > CurrentEventWindow = FindWindowByNSWindow( Windows, NativeWindow );
-	if( !CurrentEventWindow.IsValid() && LastEventWindow.IsValid() )
+	if( NativeWindow && !CurrentEventWindow.IsValid() && LastEventWindow.IsValid() )
 	{
 		CurrentEventWindow = LastEventWindow;
 		NativeWindow = CurrentEventWindow->GetWindowHandle();
@@ -587,6 +587,16 @@ FSlateCocoaWindow* FMacApplication::FindEventWindow( NSEvent* Event )
 
 	FSlateCocoaWindow* EventWindow = (FSlateCocoaWindow*)[Event window];
 
+	if ([Event type] == NSMouseMoved)
+	{
+		// Ignore windows owned by other applications
+		NSInteger WindowNumber = [NSWindow windowNumberAtPoint:[NSEvent mouseLocation] belowWindowWithWindowNumber:0];
+		if ([NSApp windowWithWindowNumber:WindowNumber] == NULL)
+		{
+			return NULL;
+		}
+	}
+
 	if( IsMouseEvent )
 	{
 		if( DraggedWindow )
