Better fix for a crash in Mac editor when root window was resizing when destroying menu.

diff --git a/Engine/Source/Runtime/Core/Private/Mac/MacWindow.cpp b/Engine/Source/Runtime/Core/Private/Mac/MacWindow.cpp
--- a/Engine/Source/Runtime/Core/Private/Mac/MacWindow.cpp
+++ b/Engine/Source/Runtime/Core/Private/Mac/MacWindow.cpp
@@ -187,14 +187,17 @@
 			[ParentWindow removeChildWindow: self];
 			[ParentWindow addChildWindow: self ordered: NSWindowAbove];
 		}
+
+		if (![self isVisible])
+		{
+			[self orderFront:nil];
+		}
 		
-		[self orderFront:nil];
-		
-		if(bMain && [self canBecomeMainWindow])
+		if (bMain && [self canBecomeMainWindow] && self != [NSApp mainWindow])
 		{
 			[self makeMainWindow];
 		}
-		if(bKey && [self canBecomeKeyWindow])
+		if (bKey && [self canBecomeKeyWindow] && self != [NSApp keyWindow])
 		{
 			[self makeKeyWindow];
 		}
