Linux mouse support improvements from Darkside.

- Result of pair debugging of NickA and MichaelL.
- Introduces some really scary platform-specific #ifdef's in SWindow I am not happy about...  Reportedly going to be accounted for in Slate refactor and will cease to be platform-specific.
- Old junk nuked.

diff --git a/Engine/Source/Runtime/Core/Private/Linux/LinuxApplication.cpp b/Engine/Source/Runtime/Core/Private/Linux/LinuxApplication.cpp
--- a/Engine/Source/Runtime/Core/Private/Linux/LinuxApplication.cpp
+++ b/Engine/Source/Runtime/Core/Private/Linux/LinuxApplication.cpp
@@ -8,16 +8,10 @@
 	#include "SteamControllerInterface.h"
 #endif // STEAM_CONTROLLER_SUPPORT
 
-//	todo:
-//	need to change to a linux one
-//	#include "HIDInputInterface.h"
+#include "ds_extensions.h"
 
 #if WITH_EDITOR
 #include "ModuleManager.h"
-
-//	whats dis???
-//	#include "Developer/Windows/VSAccessor/Public/VSAccessorModule.h"
-
 #endif
 
 //
@@ -218,7 +212,6 @@ void FLinuxApplication::ProcessDeferredMessage( SDL_Event Event )
 		{
 			SDL_MouseMotionEvent motionEvent = Event.motion;
 			FLinuxCursor *LinuxCursor = (FLinuxCursor*)Cursor.Get();
-			LinuxCursor->SetRelative( NativeWindow );
 
 			if(LinuxCursor->IsHidden())
 			{
@@ -242,6 +235,16 @@ void FLinuxApplication::ProcessDeferredMessage( SDL_Event Event )
 				{
 					LinuxCursor->SetPosition( CurrentPosition.X, CurrentPosition.Y );
 				}
+				if( !CurrentEventWindow->GetDefinition().HasOSWindowBorder )
+				{
+					if ( CurrentEventWindow->IsRegularWindow() )
+					{
+						int xOffset, yOffset;
+						SDL_GetWindowPosition( NativeWindow, &xOffset, &yOffset );
+						MessageHandler->GetWindowZoneForPoint( CurrentEventWindow.ToSharedRef(), CurrentPosition.X - xOffset, CurrentPosition.Y - yOffset );
+						MessageHandler->OnCursorSet();
+					}
+				}
 			}
 
 			if(bUsingHighPrecisionMouseInput)
@@ -278,6 +281,9 @@ void FLinuxApplication::ProcessDeferredMessage( SDL_Event Event )
 			case SDL_BUTTON_X2:
 				button = EMouseButtons::Thumb02;
 				break;
+			default:
+				button = EMouseButtons::Invalid;
+				break;
 			}
 			if(buttonEvent.type == SDL_MOUSEBUTTONUP)
 			{
@@ -627,7 +633,6 @@ void FLinuxApplication::ProcessDeferredMessage( SDL_Event Event )
 
 				case SDL_WINDOWEVENT_ENTER:
 					{
-					//	printf( "Ariel - SDL_WINDOWEVENT_ENTER has been send.\n" );
 						if ( CurrentEventWindow.IsValid() )
 						{
 							MessageHandler->OnCursorSet();
@@ -638,13 +643,15 @@ void FLinuxApplication::ProcessDeferredMessage( SDL_Event Event )
 
 				case SDL_WINDOWEVENT_LEAVE:
 					{
-					//	printf( "Ariel - SDL_WINDOWEVENT_LEAVE has been send.\n" );
+						if( CurrentEventWindow.IsValid() && GetCapture() != NULL)
+						{
+							UpdateMouseCaptureWindow((SDL_HWindow)GetCapture());
+						}
 					}
 					break;
 
 				case SDL_WINDOWEVENT_FOCUS_GAINED:
 					{
-					//	printf( "Ariel - SDL_WINDOWEVENT_FOCUS_GAINED has been send.\n" );
 						if ( CurrentEventWindow.IsValid() )
 						{
 							MessageHandler->OnWindowActivationChanged( CurrentEventWindow.ToSharedRef(), EWindowActivation::Activate );
@@ -800,14 +807,16 @@ void FLinuxApplication::UpdateMouseCaptureWindow( SDL_HWindow TargetWindow )
 		}
 		if( MouseCaptureWindow )
 		{
-			// TODO: SDL doesn't allow seem to provide a way to keep mouse focus
+			//SDL_SetWindowGrab(TargetWindow, SDL_TRUE);
+			DSEXT_SetMouseGrab(TargetWindow, SDL_TRUE);
 		}
 	}
 	else
 	{
 		if( MouseCaptureWindow )
 		{
-			// TODO: SDL doesn't allow seem to provide a way to keep mouse focus
+			//SDL_SetWindowGrab(TargetWindow, SDL_FALSE);
+			DSEXT_SetMouseGrab(TargetWindow, SDL_FALSE);
 			MouseCaptureWindow = NULL;
 		}
 	}
diff --git a/Engine/Source/Runtime/Core/Private/Linux/LinuxCursor.cpp b/Engine/Source/Runtime/Core/Private/Linux/LinuxCursor.cpp
--- a/Engine/Source/Runtime/Core/Private/Linux/LinuxCursor.cpp
+++ b/Engine/Source/Runtime/Core/Private/Linux/LinuxCursor.cpp
@@ -3,8 +3,25 @@
 #include "LinuxCursor.h"
 #include "LinuxWindow.h"
 
-FLinuxCursor::FLinuxCursor() : bHidden(false), XOffset(0), YOffset(0)
+#include "ds_extensions.h"
+
+FLinuxCursor::FLinuxCursor() : bHidden(false)
 {
+	//	init the sdl here
+	if	( SDL_WasInit( 0 ) == 0 )
+	{
+		SDL_Init( SDL_INIT_VIDEO );
+	}
+	else
+	{
+		Uint32 subsystem_init = SDL_WasInit( SDL_INIT_EVERYTHING );
+
+		if	( !(subsystem_init & SDL_INIT_VIDEO) )
+		{
+			SDL_InitSubSystem( SDL_INIT_VIDEO );
+		}
+	}
+
 	// Load up cursors that we'll be using
 	for( int32 CurCursorIndex = 0; CurCursorIndex < EMouseCursor::TotalCursorCount; ++CurCursorIndex )
 	{
@@ -131,15 +148,21 @@ FLinuxCursor::~FLinuxCursor()
 
 FVector2D FLinuxCursor::GetPosition() const
 {
-	int CursorPosX, CursorPosY;
-	SDL_GetMouseState(&CursorPosX, &CursorPosY);
+	int CursorX, CursorY;
 
-	return FVector2D( CursorPosX + XOffset, CursorPosY + YOffset );
+	DSEXT_GetMousePosition( &CursorX, &CursorY );
+
+	return FVector2D( CursorX, CursorY );
 }
 
 void FLinuxCursor::SetPosition( const int32 X, const int32 Y )
 {
-	SDL_WarpMouseInWindow( NULL, X - XOffset, Y - YOffset );
+	int WndX, WndY;
+
+	SDL_HWindow WndFocus = SDL_GetMouseFocus();
+
+	SDL_GetWindowPosition( WndFocus, &WndX, &WndY );	//	get top left
+	SDL_WarpMouseInWindow( NULL, X - WndX, Y - WndY );
 }
 
 void FLinuxCursor::SetType( const EMouseCursor::Type InNewCursor )
@@ -240,8 +263,3 @@ bool FLinuxCursor::IsHidden()
 {
 	return bHidden;
 }
-
-void FLinuxCursor::SetRelative( SDL_HWindow Window )
-{
-	SDL_GetWindowPosition( Window, &XOffset, &YOffset );	//	get top left
-}
\ No newline at end of file
diff --git a/Engine/Source/Runtime/Core/Private/Linux/LinuxWindow.cpp b/Engine/Source/Runtime/Core/Private/Linux/LinuxWindow.cpp
--- a/Engine/Source/Runtime/Core/Private/Linux/LinuxWindow.cpp
+++ b/Engine/Source/Runtime/Core/Private/Linux/LinuxWindow.cpp
@@ -30,7 +30,6 @@ FLinuxWindow::FLinuxWindow()
 	, bWasFullscreen( false )
 {
 	PreFullscreenWindowRect.left = PreFullscreenWindowRect.top = PreFullscreenWindowRect.right = PreFullscreenWindowRect.bottom = 0;
-//	bWindowGrab = false;
 }
 
 SDL_HWindow FLinuxWindow::GetHWnd() const
@@ -86,7 +85,12 @@ void FLinuxWindow::Initialize( FLinuxApplication* const Application, const TShar
 		WindowStyle |= SDL_WINDOW_BORDERLESS;
 	}
 
-	//	by testing, the SDL window doesn't need to be reshaped.
+	if ( Definition->HasSizingFrame )
+	{
+		WindowStyle |= SDL_WINDOW_RESIZABLE;
+	}
+
+	//	The SDL window doesn't need to be reshaped.
 	//	the size of the window you input is the sizeof the client.
 	HWnd = SDL_CreateWindow( TCHAR_TO_ANSI( *Definition->Title ), X, Y, ClientWidth, ClientHeight, WindowStyle  );
 
@@ -224,39 +228,15 @@ static void _SetBestFullscreenDisplayMode( SDL_HWindow hWnd, int32 *pWidth, int3
 	*pHeight = dsp_mode.h;
 }
 
-
-#define USE_REAL_FULLSCREEN		1
-#define USE_BOARDER_FUNCTION	0
-
-
-static int gResizeIdx = 0;
-
 void FLinuxWindow::ReshapeWindow( int32 NewX, int32 NewY, int32 NewWidth, int32 NewHeight )
 {
-//	printf( "--------------------------------------------------------\n" );
-//	printf( "ArielLW - SetWindowShape - Start - %03d\n", gResizeIdx );
-//	printf( "          Req  X = %d, Y = %d, W = %d, H = %d\n", NewX, NewY, NewWidth, NewHeight );
-
 	int32 closest_w = NewWidth;
 	int32 closest_h = NewHeight;
 
 	switch( WindowMode )
 	{
-#if USE_REAL_FULLSCREEN
 		case EWindowMode::Fullscreen:
-#else
-		case EWindowMode::WindowedFullscreen:
-#endif
 		{
-			if ( bWasFullscreen == false )
-			{
-#if USE_BOARDER_FUNCTION
-				int i = 100;
-				i = i * i;
-				SDL_SetWindowBordered( HWnd, SDL_FALSE );
-#endif
-			}
-
 			SDL_SetWindowFullscreen( HWnd, 0 );
 			_GetBestFullscreenResolution( HWnd, &closest_w, &closest_h );
 			SDL_SetWindowSize( HWnd, closest_w, closest_h );
@@ -264,39 +244,18 @@ void FLinuxWindow::ReshapeWindow( int32 NewX, int32 NewY, int32 NewWidth, int32
 			_SetBestFullscreenDisplayMode( HWnd, &closest_w, &closest_h );
 			SDL_SetWindowFullscreen( HWnd, SDL_WINDOW_FULLSCREEN );
 
-		//	printf( "          Real X = %d, Y = %d, W = %d, H = %d\n", 0, 0, closest_w, closest_h );
-		//	printf( "        - Fullscreen\n" );
-
 			bWasFullscreen = true;
-			//bIsFullscreenSwitchingRes = true;
 
 		}	break;
 
-#if USE_REAL_FULLSCREEN
 		case EWindowMode::WindowedFullscreen:
-#else
-		case EWindowMode::Fullscreen:
-#endif
 		{
-			if ( bWasFullscreen == false )
-			{
-#if USE_BOARDER_FUNCTION
-				SDL_SetWindowBordered( HWnd, SDL_FALSE );
-#endif
-			}
-
 			SDL_SetWindowFullscreen( HWnd, 0 );
 			SDL_SetWindowPosition( HWnd, 0, 0 );
-		//	_GetBestFullscreenResolution( HWnd, &closest_w, &closest_h );
-		//	SDL_SetWindowSize( HWnd, NewWidth, NewHeight );
 			_SetBestFullscreenDisplayMode( HWnd, &closest_w, &closest_h );
 			SDL_SetWindowFullscreen( HWnd, SDL_WINDOW_FULLSCREEN_DESKTOP );
 
-		//	printf( "          Real X = %d, Y = %d, W = %d, H = %d\n", 0, 0, closest_w, closest_h );
-		//	printf( "        - WindowedFullscreen\n" );
-
 			bWasFullscreen = true;
-			//bIsFullscreenSwitchingRes = true;
 
 		}	break;
 
@@ -305,10 +264,7 @@ void FLinuxWindow::ReshapeWindow( int32 NewX, int32 NewY, int32 NewWidth, int32
 			SDL_SetWindowPosition( HWnd, NewX, NewY );
 			SDL_SetWindowSize( HWnd, NewWidth, NewHeight );
 
-		//	printf( "        - Windowed\n" );
-
 			bWasFullscreen = false;
-			//bIsFullscreenSwitchingRes = false;
 
 		}	break;
 	}
@@ -317,53 +273,29 @@ void FLinuxWindow::ReshapeWindow( int32 NewX, int32 NewY, int32 NewWidth, int32
 	RegionHeight  = NewHeight;
 	VirtualWidth  = NewWidth;
 	VirtualHeight = NewHeight;
-
-//	printf( "ArielLW - SetWindowShape - End   - %03d\n", gResizeIdx );
-	gResizeIdx++;
 }
 
-
-
-
-static int gSetWindowModeidx = 0;
-
 /** Toggle native window between fullscreen and normal mode */
 void FLinuxWindow::SetWindowMode( EWindowMode::Type NewWindowMode )
 {
-//	printf( "--------------------------------------------------------\n" );
-//	printf( "ArielLW - SetWindowMode - Start - %03d\n", gSetWindowModeidx );
-
-//	bWindowGrab = SDL_GetWindowGrab( HWnd );
-
 	if( NewWindowMode != WindowMode )
 	{
 		SDL_DisplayMode dsp_mode;
 		int32 closest_w;
 		int32 closest_h;
 
-		//	SDL_GetWindowSize( HWnd, &closest_w, &closest_h );
 		closest_w = VirtualWidth;
 		closest_h = VirtualHeight;
 
-	//	printf( "          r_w = %03d, r_h = %03d\n", closest_w, closest_h );
-
 		switch( NewWindowMode )
 		{
-#if USE_REAL_FULLSCREEN
 			case EWindowMode::Fullscreen:
-#else
-			case EWindowMode::WindowedFullscreen:
-#endif
 			{
 				if ( bWasFullscreen != true )
 				{
 					SDL_SetWindowPosition( HWnd, 0, 0 );
 					_SetBestFullscreenDisplayMode( HWnd, &closest_w, &closest_h );
-					printf( "          c_w = %03d, c_h = %03d\n", closest_w, closest_h );
 
-#if USE_BOARDER_FUNCTION
-					SDL_SetWindowBordered( HWnd, SDL_FALSE );
-#endif
 					SDL_SetWindowFullscreen( HWnd, SDL_WINDOW_FULLSCREEN );
 					SDL_SetWindowGrab( HWnd, SDL_TRUE );
 
@@ -372,24 +304,15 @@ void FLinuxWindow::SetWindowMode( EWindowMode::Type NewWindowMode )
 
 			}	break;
 
-#if USE_REAL_FULLSCREEN
 			case EWindowMode::WindowedFullscreen:
-#else
-			case EWindowMode::Fullscreen:
-#endif
 			{
 				if ( bWasFullscreen != true )
 				{
-#if USE_BOARDER_FUNCTION
-					SDL_SetWindowBordered( HWnd, SDL_FALSE );
-#endif
 					SDL_SetWindowPosition( HWnd, 0, 0 );
 
 					_SetBestFullscreenDisplayMode( HWnd, &closest_w, &closest_h );
 					SDL_SetWindowFullscreen( HWnd, SDL_WINDOW_FULLSCREEN_DESKTOP );
 
-				//	printf( "          c_w = %03d, c_h = %03d\n", closest_w, closest_h );
-
 					bWasFullscreen = true;
 				}
 
@@ -404,32 +327,22 @@ void FLinuxWindow::SetWindowMode( EWindowMode::Type NewWindowMode )
 				SDL_SetWindowGrab( HWnd, SDL_FALSE );
 
 				bWasFullscreen = false;
-				//bIsFullscreenSwitchingRes = false;
 			}	break;
 		}
 
 
 		WindowMode = NewWindowMode;
 	}
-
-//	printf( "ArielLW - SetWindowMode - End   - %03d\n", gSetWindowModeidx );
-	gSetWindowModeidx++;
 }
 
 
-static int gAdjustCachedSizeIdx = 0;
-
 /** @return	Gives the native window a chance to adjust our stored window size before we cache it off */
 
 void FLinuxWindow::AdjustCachedSize( FVector2D& Size ) const
 {
-//	printf( "--------------------------------------------------------\n" );
-//	printf( "ArielLW - AdjustCachedSize - Start - %03d\n", gAdjustCachedSizeIdx );
-
 	if	( Definition.IsValid() && Definition->SizeWillChangeOften )
 	{
 		Size = FVector2D( VirtualWidth, VirtualHeight );
-	//	printf( "        - W = %d, H = %d\n", VirtualWidth, VirtualHeight );
 	}
 	else
 	if	( HWnd )
@@ -438,28 +351,18 @@ void FLinuxWindow::AdjustCachedSize( FVector2D& Size ) const
 
 		if ( WindowMode == EWindowMode::Windowed )
 		{
-		//	printf( "        - Windowed\n" );
-
 			SDL_GetWindowSize( HWnd, &SizeW, &SizeH );
-		//	printf( "        - Req  Size - RW = %d, RH = %d\n", SizeW, SizeH );
 		}
 		else // windowed fullscreen or fullscreen
 		{
-		//	printf( "        - Fullscreen\n" );
-
 			SizeW = VirtualWidth ;
 			SizeH = VirtualHeight;
-		//	printf( "        - Req  Size - RW = %d, RH = %d\n", SizeW, SizeH );
 
 			_GetBestFullscreenResolution( HWnd, &SizeW, &SizeH );
-		//	printf( "        - Real Size - RW = %d, RH = %d\n", SizeW, SizeH );
 		}
 
 		Size = FVector2D( SizeW, SizeH );
 	}
-
-//	printf( "ArielLW - AdjustCachedSize - End   - %03d\n", gAdjustCachedSizeIdx );
-	gAdjustCachedSizeIdx++;
 }
 
 
@@ -469,8 +372,6 @@ bool FLinuxWindow::GetFullScreenInfo( int32& X, int32& Y, int32& Width, int32& H
 	return true;
 }
 
-
-
 /** @return true if the native window is maximized, false otherwise */
 bool FLinuxWindow::IsMaximized() const
 {
@@ -513,7 +414,8 @@ void FLinuxWindow::SetWindowFocus()
  */
 void FLinuxWindow::SetOpacity( const float InOpacity )
 {
-	//	by searching from internet, ppl said that the SDL2 didn't offer such functionality...
+	//	SDL2 doesn't offer such functionality...
+	//  Could be added to ds_extenstion
 }
 
 /**
@@ -523,43 +425,30 @@ void FLinuxWindow::SetOpacity( const float InOpacity )
  */
 void FLinuxWindow::Enable( bool bEnable )
 {
-	//	by searching from internet, ppl said that the SDL2 didn't offer such functionality...
+	//	SDL2 doesn't offer such functionality...
+	//  Could be added to ds_extenstion
 }
 
 /** @return true if native window exists underneath the coordinates */
 bool FLinuxWindow::IsPointInWindow( int32 X, int32 Y ) const
 {
-	int32 x0 = 0, y0 = 0, x1 = 0, y1 = 0;
+	int32 width = 0, height = 0;
 
-	SDL_GetWindowPosition( HWnd, &x0, &y0 );	//	get top left
-	SDL_GetWindowSize( HWnd, &x1, &y1 );		//	get size : the size of the window is the client size, not the window boder size
+	SDL_GetWindowSize( HWnd, &width, &height );
 	
-	//	todo:
-	//	SDL_GetWindowPosition returns the top left of the client region of the window when the top of the window is not stick the top edge of the screen.
-	//	so ......
-	//	when the top of the whole window is connect to the screen, the real client's top should > 0, but still may return 0...
-	x1 += x0;	//	get the bottom right
-	y1 += y0;
-
-	if ( (X >= x0) && (X < x1) )
-	{
-		if ( (Y >= y0) && (Y < y1) )
-		{
-			return true;
-		}
-	}
-
-	return false;
+	return X > 0 && Y > 0 && X < width && Y < height;
 }
 
 int32 FLinuxWindow::GetWindowBorderSize() const
 {
-	return 0;
+	// todo
+	return 5;
 }
 
 
 bool FLinuxWindow::IsForegroundWindow() const
 {
+	// todo
 	return 1;
 }
 
diff --git a/Engine/Source/Runtime/Core/Public/Linux/LinuxCursor.h b/Engine/Source/Runtime/Core/Public/Linux/LinuxCursor.h
--- a/Engine/Source/Runtime/Core/Public/Linux/LinuxCursor.h
+++ b/Engine/Source/Runtime/Core/Public/Linux/LinuxCursor.h
@@ -33,8 +33,6 @@ public:
 
 	bool IsHidden();
 
-	void SetRelative( SDL_HWindow Window );
-
 	uint32 GetCursorEvent();
 	
 private:
@@ -46,8 +44,6 @@ private:
 
 	FIntRect CursorClipRect;
 
-	int XOffset, YOffset;
-
 	uint32 CursorEvent;
 };
 
diff --git a/Engine/Source/Runtime/SlateCore/Private/Widgets/SWindow.cpp b/Engine/Source/Runtime/SlateCore/Private/Widgets/SWindow.cpp
--- a/Engine/Source/Runtime/SlateCore/Private/Widgets/SWindow.cpp
+++ b/Engine/Source/Runtime/SlateCore/Private/Widgets/SWindow.cpp
@@ -609,7 +609,7 @@ FSlateRect SWindow::GetClippingRectangleInWindow() const
 FMargin SWindow::GetWindowBorderSize() const
 {
 // Mac didn't want a window border, and consoles don't either, so only do this in Windows
-#if PLATFORM_WINDOWS
+#if PLATFORM_WINDOWS || PLATFORM_LINUX
 	if (NativeWindow.IsValid() && NativeWindow->IsMaximized())
 	{
 		int32 OSWindowBorderSize = NativeWindow->GetWindowBorderSize();
@@ -1241,6 +1241,23 @@ FReply SWindow::OnKeyboardFocusReceived( const FGeometry& MyGeometry, const FKey
 
 FReply SWindow::OnMouseButtonDown( const FGeometry& MyGeometry, const FPointerEvent& MouseEvent )
 {
+#if PLATFORM_LINUX
+	if (bHasSizingFrame && MouseEvent.GetEffectingButton() == EKeys::LeftMouseButton)
+	{
+		if ((WindowZone == EWindowZone::TopLeftBorder || WindowZone == EWindowZone::BottomRightBorder ||
+			WindowZone == EWindowZone::BottomLeftBorder || WindowZone == EWindowZone::TopRightBorder ||
+			WindowZone == EWindowZone::TopBorder || WindowZone == EWindowZone::BottomBorder ||
+			WindowZone == EWindowZone::LeftBorder || WindowZone == EWindowZone::RightBorder ||
+			WindowZone == EWindowZone::TitleBar)
+			&& MouseEvent.GetEffectingButton() == EKeys::LeftMouseButton)
+		{
+			MoveResizeZone = WindowZone;
+			MoveResizeStart = MouseEvent.GetScreenSpacePosition();
+			MoveResizeRect = GetRectInScreen();
+			return FReply::Handled().CaptureMouse(SharedThis(this));
+		}
+	}
+#endif
 	if (bDragAnywhere && MouseEvent.GetEffectingButton() == EKeys::LeftMouseButton)
 	{
 		return FReply::Handled().CaptureMouse(SharedThis(this));
@@ -1253,6 +1270,13 @@ FReply SWindow::OnMouseButtonDown( const FGeometry& MyGeometry, const FPointerEv
 
 FReply SWindow::OnMouseButtonUp( const FGeometry& MyGeometry, const FPointerEvent& MouseEvent )
 {
+#if PLATFORM_LINUX
+	if (MoveResizeZone != EWindowZone::Unspecified && MouseEvent.GetEffectingButton() == EKeys::LeftMouseButton)
+	{
+		MoveResizeZone =  EWindowZone::Unspecified;
+		return FReply::Handled().ReleaseMouseCapture();
+	}
+#endif
 	if (bDragAnywhere &&  this->HasMouseCapture() && MouseEvent.GetEffectingButton() == EKeys::LeftMouseButton)
 	{
 		return FReply::Handled().ReleaseMouseCapture();
@@ -1265,6 +1289,110 @@ FReply SWindow::OnMouseButtonUp( const FGeometry& MyGeometry, const FPointerEven
 
 FReply SWindow::OnMouseMove( const FGeometry& MyGeometry, const FPointerEvent& MouseEvent )
 {
+#if PLATFORM_LINUX
+	if (MoveResizeZone == EWindowZone::TopLeftBorder)
+	{
+		FVector2D MoveResizeOffset = MouseEvent.GetScreenSpacePosition() - MoveResizeStart;
+		if (NativeWindow.IsValid())
+		{
+			NativeWindow->ReshapeWindow(
+				FMath::Trunc(MoveResizeRect.Left + MoveResizeOffset.X), FMath::Trunc(MoveResizeRect.Top + MoveResizeOffset.Y),
+				FMath::Trunc(MoveResizeRect.Right - MoveResizeRect.Left - MoveResizeOffset.X), FMath::Trunc(MoveResizeRect.Bottom - MoveResizeRect.Top - MoveResizeOffset.Y)
+				);
+		}
+		return FReply::Handled();
+	}
+	if (MoveResizeZone == EWindowZone::BottomRightBorder)
+	{
+		FVector2D MoveResizeOffset = MouseEvent.GetScreenSpacePosition() - MoveResizeStart;
+		if (NativeWindow.IsValid())
+		{
+			NativeWindow->ReshapeWindow(
+				FMath::Trunc(MoveResizeRect.Left), FMath::Trunc(MoveResizeRect.Top),
+				FMath::Trunc(MoveResizeRect.Right - MoveResizeRect.Left + MoveResizeOffset.X), FMath::Trunc(MoveResizeRect.Bottom - MoveResizeRect.Top + MoveResizeOffset.Y)
+				);
+		}
+		return FReply::Handled();
+	}
+	if (MoveResizeZone == EWindowZone::BottomLeftBorder)
+	{
+		FVector2D MoveResizeOffset = MouseEvent.GetScreenSpacePosition() - MoveResizeStart;
+		if (NativeWindow.IsValid())
+		{
+			NativeWindow->ReshapeWindow(
+				FMath::Trunc(MoveResizeRect.Left + MoveResizeOffset.X), FMath::Trunc(MoveResizeRect.Top),
+				FMath::Trunc(MoveResizeRect.Right - MoveResizeRect.Left - MoveResizeOffset.X), FMath::Trunc(MoveResizeRect.Bottom - MoveResizeRect.Top + MoveResizeOffset.Y)
+				);
+		}
+		return FReply::Handled();
+	}
+	if (MoveResizeZone == EWindowZone::TopRightBorder)
+	{
+		FVector2D MoveResizeOffset = MouseEvent.GetScreenSpacePosition() - MoveResizeStart;
+		if (NativeWindow.IsValid())
+		{
+			NativeWindow->ReshapeWindow(
+				FMath::Trunc(MoveResizeRect.Left), FMath::Trunc(MoveResizeRect.Top + MoveResizeOffset.Y),
+				FMath::Trunc(MoveResizeRect.Right - MoveResizeRect.Left + MoveResizeOffset.X), FMath::Trunc(MoveResizeRect.Bottom - MoveResizeRect.Top - MoveResizeOffset.Y)
+				);
+		}
+		return FReply::Handled();
+	}
+	if (MoveResizeZone == EWindowZone::TopBorder)
+	{
+		FVector2D MoveResizeOffset = MouseEvent.GetScreenSpacePosition() - MoveResizeStart;
+		if (NativeWindow.IsValid())
+		{
+			NativeWindow->ReshapeWindow(
+				FMath::Trunc(MoveResizeRect.Left), FMath::Trunc(MoveResizeRect.Top + MoveResizeOffset.Y),
+				FMath::Trunc(MoveResizeRect.Right - MoveResizeRect.Left), FMath::Trunc(MoveResizeRect.Bottom - MoveResizeRect.Top - MoveResizeOffset.Y)
+				);
+		}
+		return FReply::Handled();
+	}
+	if (MoveResizeZone == EWindowZone::BottomBorder)
+	{
+		FVector2D MoveResizeOffset = MouseEvent.GetScreenSpacePosition() - MoveResizeStart;
+		if (NativeWindow.IsValid())
+		{
+			NativeWindow->ReshapeWindow(
+				FMath::Trunc(MoveResizeRect.Left), FMath::Trunc(MoveResizeRect.Top),
+				FMath::Trunc(MoveResizeRect.Right - MoveResizeRect.Left), FMath::Trunc(MoveResizeRect.Bottom - MoveResizeRect.Top + MoveResizeOffset.Y)
+				);
+		}
+		return FReply::Handled();
+	}
+	if (MoveResizeZone == EWindowZone::LeftBorder)
+	{
+		FVector2D MoveResizeOffset = MouseEvent.GetScreenSpacePosition() - MoveResizeStart;
+		if (NativeWindow.IsValid())
+		{
+			NativeWindow->ReshapeWindow(
+				FMath::Trunc(MoveResizeRect.Left + MoveResizeOffset.X), FMath::Trunc(MoveResizeRect.Top),
+				FMath::Trunc(MoveResizeRect.Right - MoveResizeRect.Left - MoveResizeOffset.X), FMath::Trunc(MoveResizeRect.Bottom - MoveResizeRect.Top)
+				);
+		}
+		return FReply::Handled();
+	}
+	if (MoveResizeZone == EWindowZone::RightBorder)
+	{
+		FVector2D MoveResizeOffset = MouseEvent.GetScreenSpacePosition() - MoveResizeStart;
+		if (NativeWindow.IsValid())
+		{
+			NativeWindow->ReshapeWindow(
+				FMath::Trunc(MoveResizeRect.Left), FMath::Trunc(MoveResizeRect.Top),
+				FMath::Trunc(MoveResizeRect.Right - MoveResizeRect.Left + MoveResizeOffset.X), FMath::Trunc(MoveResizeRect.Bottom - MoveResizeRect.Top)
+				);
+		}
+		return FReply::Handled();
+	}
+	if (MoveResizeZone == EWindowZone::TitleBar)
+	{
+		FVector2D MoveResizeOffset = MouseEvent.GetScreenSpacePosition() - MoveResizeStart;
+		this->MoveWindowTo( FVector2D(MoveResizeRect.Left, MoveResizeRect.Top) + MoveResizeOffset );
+		return FReply::Handled();
+	}
+#endif
 	if ( bDragAnywhere && this->HasMouseCapture() && MouseEvent.IsMouseButtonDown(EKeys::LeftMouseButton) )
 	{
 		this->MoveWindowTo( ScreenPosition + MouseEvent.GetCursorDelta() );
diff --git a/Engine/Source/Runtime/SlateCore/Public/Widgets/SWindow.h b/Engine/Source/Runtime/SlateCore/Public/Widgets/SWindow.h
--- a/Engine/Source/Runtime/SlateCore/Public/Widgets/SWindow.h
+++ b/Engine/Source/Runtime/SlateCore/Public/Widgets/SWindow.h
@@ -544,6 +544,11 @@ public:
 	/** Gets the current Window Zone that mouse position is over. */
 	EWindowZone::Type GetCurrentWindowZone(FVector2D LocalMousePosition);
 
+	/** Used to store the zone where the mouse down event occurred during move / drag */
+	EWindowZone::Type MoveResizeZone;
+	FVector2D MoveResizeStart;
+	FSlateRect MoveResizeRect;
+
 	/** @return Gets the radius of the corner rounding of the window. */
 	int32 GetCornerRadius();
 
