Corrected date time formatting to use the right formatter for times.

diff --git a/Engine/Source/Runtime/Core/Private/Internationalization/ICUCulture.cpp b/Engine/Source/Runtime/Core/Private/Internationalization/ICUCulture.cpp
--- a/Engine/Source/Runtime/Core/Private/Internationalization/ICUCulture.cpp
+++ b/Engine/Source/Runtime/Core/Private/Internationalization/ICUCulture.cpp
@@ -294,7 +294,7 @@ TSharedRef<const icu::DateFormat> FCulture::FICUCultureImplementation::GetTimeFo
 	icu::UnicodeString InputTimeZoneID;
 	ICUUtilities::Convert(TimeZone, InputTimeZoneID, false);
 
-	const TSharedRef<const icu::DateFormat> DefaultFormatter( ICUDateTimeFormat );
+	const TSharedRef<const icu::DateFormat> DefaultFormatter( ICUTimeFormat );
 
 	bool bIsDefaultTimeZone = TimeZone.IsEmpty();
 	if( !bIsDefaultTimeZone )
