Delegates: Added support for UFunctions; partially refactored multicast delegates

Changes:
-  Added UFunction bindings to unicast and multicast delegates
-  Pulled non-templated class methods into shared base classes to reduce the code that needs to be parsed and compiled
-  Multicast delegates are now collections of delegate bindings instead of unicast delegates
-  Fixed one bug where removing a binding in RemoveByDelegateInternal would skip over the current binding in RemoveAndSwap
-  Updated/added some documentation
-  Changed NULL to nullptr

diff --git a/Engine/Source/Programs/DelegateHeaderTool/Windows/Program.cs b/Engine/Source/Programs/DelegateHeaderTool/Windows/Program.cs
--- a/Engine/Source/Programs/DelegateHeaderTool/Windows/Program.cs
+++ b/Engine/Source/Programs/DelegateHeaderTool/Windows/Program.cs
@@ -123,7 +123,10 @@ namespace DelegateHeaderTool
 			var DefineTemplateArgs = new StringBuilder( "#define FUNC_TEMPLATE_ARGS " );
 			var DefineHasParams = String.Format( "#define FUNC_HAS_PARAMS {0}", CallParameterCount > 0 ? 1 : 0 );
 			var DefineParamList = new StringBuilder( "#define FUNC_PARAM_LIST" );
+			var DefineParamMembers = new StringBuilder( "#define FUNC_PARAM_MEMBERS" );
 			var DefineParamPassThru = new StringBuilder( "#define FUNC_PARAM_PASSTHRU" );
+			var DefineParamParmsPassIn = new StringBuilder( "#define FUNC_PARAM_PARMS_PASSIN" );
+			var DefineParamInitializerList = new StringBuilder( "#define FUNC_PARAM_INITIALIZER_LIST" );
 			var DefineIsVoid = String.Format( "#define FUNC_IS_VOID {0}", bSupportsRetVal ? 0 : 1 );
 			var DefineDeclareDelegate = new StringBuilder( "#define DECLARE_DELEGATE" );
             var DefineDeclareMulticastDelegate = new StringBuilder("#define DECLARE_MULTICAST_DELEGATE");
@@ -209,9 +212,14 @@ namespace DelegateHeaderTool
 				{
 					DefineParamList.Append( "," );
 					DefineParamPassThru.Append( "," );
+					DefineParamParmsPassIn.Append( "," );
+					DefineParamInitializerList.Append(",");
 				}
 				DefineParamList.AppendFormat( " Param{0}Type InParam{0}", CallParameterIndex );
+				DefineParamMembers.AppendFormat(" Param{0}Type Param{0};", CallParameterIndex);
 				DefineParamPassThru.AppendFormat( " InParam{0}", CallParameterIndex );
+				DefineParamParmsPassIn.AppendFormat(" Parms.Param{0}", CallParameterIndex);
+				DefineParamInitializerList.AppendFormat(" Param{0}( InParam{0} )", CallParameterIndex);
 
 				DefineDeclareDelegate.AppendFormat( ", Param{0}Type", CallParameterIndex );
                 DefineDeclareMulticastDelegate.AppendFormat(", Param{0}Type", CallParameterIndex);
@@ -320,7 +328,10 @@ namespace DelegateHeaderTool
 			Output.WriteLine( DefineTemplateArgs );
 			Output.WriteLine( DefineHasParams );
 			Output.WriteLine( DefineParamList );
+			Output.WriteLine( DefineParamMembers );
 			Output.WriteLine( DefineParamPassThru );
+			Output.WriteLine( DefineParamParmsPassIn );
+			Output.WriteLine( DefineParamInitializerList );
 			Output.WriteLine( DefineIsVoid );
 			Output.WriteLine( "#include \"DelegateInstanceInterfaceImpl.inl\"" );
 
@@ -459,7 +470,10 @@ namespace DelegateHeaderTool
             Output.WriteLine( "#undef FUNC_TEMPLATE_ARGS" );
             Output.WriteLine( "#undef FUNC_HAS_PARAMS" );
             Output.WriteLine( "#undef FUNC_PARAM_LIST" );
+			Output.WriteLine( "#undef FUNC_PARAM_MEMBERS" );
             Output.WriteLine( "#undef FUNC_PARAM_PASSTHRU" );
+			Output.WriteLine( "#undef FUNC_PARAM_PARMS_PASSIN" );
+			Output.WriteLine( "#undef FUNC_PARAM_INITIALIZER_LIST" );
             Output.WriteLine( "#undef FUNC_IS_VOID" );
         }
 
diff --git a/Engine/Source/Runtime/Core/Public/Templates/Delegate.h b/Engine/Source/Runtime/Core/Public/Templates/Delegate.h
--- a/Engine/Source/Runtime/Core/Public/Templates/Delegate.h
+++ b/Engine/Source/Runtime/Core/Public/Templates/Delegate.h
@@ -1,6 +1,6 @@
 // Copyright 1998-2014 Epic Games, Inc. All Rights Reserved.
 
-
+// Do not remove this header guard! It ensures that DelegateCombinations.h can only be included from here.
 #ifndef __Delegate_h__
 #define __Delegate_h__
 
@@ -8,6 +8,9 @@
 
 #include "SharedPointer.h"
 #include "WeakObjectPtrTemplates.h"
+#include "IDelegateInstance.h"
+#include "DelegateBase.h"
+#include "MulticastDelegateBase.h"
 
 
 /**
@@ -159,30 +162,6 @@
  */
 
 
-/**
- * Types of delegate instances
- */
-namespace EDelegateInstanceType
-{
-	enum Type
-	{
-		/** Member function pointer to method in (fast, not thread-safe) shared pointer-based class */
-		SharedPointerMethod,
-
-		/** Member function pointer to method in (conditionally thread-safe) shared pointer-based class */
-		ThreadSafeSharedPointerMethod,
-
-		/** Raw C++ member function pointer (pointer to class method) */
-		RawMethod,
-
-		/** Member function pointer to method in UObject-based class */
-		UObjectMethod,
-
-		/** Raw C++ static function pointer */
-		Raw
-	};
-}
-
 // This suffix is appended to all header exported delegates
 #define HEADER_GENERATED_DELEGATE_SIGNATURE_SUFFIX TEXT("__DelegateSignature")
 
@@ -351,4 +330,5 @@ class DynamicMulticastDelegateName : public TBaseDynamicMulticastDelegate_##Suff
 DECLARE_DELEGATE( FSimpleDelegate );
 DECLARE_MULTICAST_DELEGATE( FSimpleMulticastDelegate );
 
-#endif		// __Delegate_h__
\ No newline at end of file
+
+#endif		// __Delegate_h__
diff --git a/Engine/Source/Runtime/Core/Public/Templates/DelegateBase.h b/Engine/Source/Runtime/Core/Public/Templates/DelegateBase.h
new file mode 100644
--- /dev/null
+++ b/Engine/Source/Runtime/Core/Public/Templates/DelegateBase.h
@@ -0,0 +1,121 @@
+// Copyright 1998-2014 Epic Games, Inc. All Rights Reserved.
+
+/*=============================================================================
+	DelegateBase.h: Declares the FDelegateBase class.
+=============================================================================*/
+
+#pragma once
+
+
+/**
+ * Base class for unicast delegates.
+ */
+template<typename ObjectPtrType = FWeakObjectPtr>
+class FDelegateBase
+{
+public:
+
+	/**
+	 * Creates and initializes a new instance.
+	 *
+	 * @param InDelegateInstance The delegate instance to assign.
+	 */
+	FDelegateBase( IDelegateInstance* InDelegateInstance )
+		: DelegateInstance(InDelegateInstance)
+	{ }
+
+public:
+
+	/**
+	 * If this is a UFunction or UObject delegate, return the UObject
+	 *
+	 * @return  the object associated with this delegate if there is one
+	 */
+	inline class UObject* GetUObject( ) const
+	{
+		if (DelegateInstance != nullptr)
+		{
+			if ((DelegateInstance->GetType() == EDelegateInstanceType::UFunction) ||
+				(DelegateInstance->GetType() == EDelegateInstanceType::UObjectMethod))
+			{
+				return (class UObject*)DelegateInstance->GetRawUserObject();
+			}
+		}
+
+		return nullptr;
+	}
+
+	/**
+	 * Checks to see if the user object bound to this delegate is still valid
+	 *
+	 * @return  True if the user object is still valid and it's safe to execute the function call
+	 */
+	inline bool IsBound( ) const
+	{
+		return (DelegateInstance != nullptr) && (DelegateInstance->IsSafeToExecute());
+	}
+
+	/** 
+	 * Checks to see if this delegate is bound to the given user object.
+	 *
+	 * @return	True if this delegate is bound to InUserObject, false otherwise.
+	 */
+	inline bool IsBoundToObject( void const* InUserObject ) const
+	{
+		return (InUserObject != nullptr) && (DelegateInstance != nullptr) && DelegateInstance->HasSameObject(InUserObject);
+	}
+
+	/**
+	 * Checks to see if the user object bound to this delegate will ever be valid again
+	 *
+	 * @return  True if the user object could be valid in the future or is valid now.
+	 */
+	inline bool IsCompactable( ) const
+	{
+		return (DelegateInstance != nullptr) && DelegateInstance->IsCompactable();
+	}
+
+	/**
+	 * Unbinds this delegate
+	 */
+	inline void Unbind( )
+	{
+		if (DelegateInstance != nullptr)
+		{
+			delete DelegateInstance;
+			DelegateInstance = nullptr;
+		}
+	}
+
+protected:
+
+	/**
+	 * Gets the delegate instance.
+	 *
+	 * @return The delegate instance.
+	 *
+	 * @see SetDelegateInstance
+	 */
+	IDelegateInstance* GetDelegateInstance( ) const
+	{
+		return DelegateInstance;
+	}
+
+	/**
+	 * Sets the delegate instance.
+	 *
+	 * @param InDelegateInstance The delegate instance to set.
+	 *
+	 * @see GetDelegateInstance.
+	 */
+	inline void SetDelegateInstance( IDelegateInstance* InDelegateInstance )
+	{
+		Unbind();
+		DelegateInstance = InDelegateInstance;
+	}
+
+private:
+
+	// Holds the delegate instance
+	IDelegateInstance* DelegateInstance;
+};
diff --git a/Engine/Source/Runtime/Core/Public/Templates/DelegateCombinations.h b/Engine/Source/Runtime/Core/Public/Templates/DelegateCombinations.h
--- a/Engine/Source/Runtime/Core/Public/Templates/DelegateCombinations.h
+++ b/Engine/Source/Runtime/Core/Public/Templates/DelegateCombinations.h
@@ -16,7 +16,10 @@
 #define FUNC_TEMPLATE_ARGS RetValType
 #define FUNC_HAS_PARAMS 0
 #define FUNC_PARAM_LIST
+#define FUNC_PARAM_MEMBERS
 #define FUNC_PARAM_PASSTHRU
+#define FUNC_PARAM_PARMS_PASSIN
+#define FUNC_PARAM_INITIALIZER_LIST
 #define FUNC_IS_VOID 1
 #include "DelegateInstanceInterfaceImpl.inl"
 
@@ -204,7 +207,10 @@
 #undef FUNC_TEMPLATE_ARGS
 #undef FUNC_HAS_PARAMS
 #undef FUNC_PARAM_LIST
+#undef FUNC_PARAM_MEMBERS
 #undef FUNC_PARAM_PASSTHRU
+#undef FUNC_PARAM_PARMS_PASSIN
+#undef FUNC_PARAM_INITIALIZER_LIST
 #undef FUNC_IS_VOID
 
 #define FUNC_SUFFIX RetVal_NoParams
@@ -214,7 +220,10 @@
 #define FUNC_TEMPLATE_ARGS RetValType
 #define FUNC_HAS_PARAMS 0
 #define FUNC_PARAM_LIST
+#define FUNC_PARAM_MEMBERS
 #define FUNC_PARAM_PASSTHRU
+#define FUNC_PARAM_PARMS_PASSIN
+#define FUNC_PARAM_INITIALIZER_LIST
 #define FUNC_IS_VOID 0
 #include "DelegateInstanceInterfaceImpl.inl"
 
@@ -399,7 +408,10 @@
 #undef FUNC_TEMPLATE_ARGS
 #undef FUNC_HAS_PARAMS
 #undef FUNC_PARAM_LIST
+#undef FUNC_PARAM_MEMBERS
 #undef FUNC_PARAM_PASSTHRU
+#undef FUNC_PARAM_PARMS_PASSIN
+#undef FUNC_PARAM_INITIALIZER_LIST
 #undef FUNC_IS_VOID
 
 #define FUNC_SUFFIX OneParam
@@ -409,7 +421,10 @@
 #define FUNC_TEMPLATE_ARGS RetValType, Param1Type
 #define FUNC_HAS_PARAMS 1
 #define FUNC_PARAM_LIST Param1Type InParam1
+#define FUNC_PARAM_MEMBERS Param1Type Param1;
 #define FUNC_PARAM_PASSTHRU InParam1
+#define FUNC_PARAM_PARMS_PASSIN Parms.Param1
+#define FUNC_PARAM_INITIALIZER_LIST Param1( InParam1 )
 #define FUNC_IS_VOID 1
 #include "DelegateInstanceInterfaceImpl.inl"
 
@@ -597,7 +612,10 @@
 #undef FUNC_TEMPLATE_ARGS
 #undef FUNC_HAS_PARAMS
 #undef FUNC_PARAM_LIST
+#undef FUNC_PARAM_MEMBERS
 #undef FUNC_PARAM_PASSTHRU
+#undef FUNC_PARAM_PARMS_PASSIN
+#undef FUNC_PARAM_INITIALIZER_LIST
 #undef FUNC_IS_VOID
 
 #define FUNC_SUFFIX RetVal_OneParam
@@ -607,7 +625,10 @@
 #define FUNC_TEMPLATE_ARGS RetValType, Param1Type
 #define FUNC_HAS_PARAMS 1
 #define FUNC_PARAM_LIST Param1Type InParam1
+#define FUNC_PARAM_MEMBERS Param1Type Param1;
 #define FUNC_PARAM_PASSTHRU InParam1
+#define FUNC_PARAM_PARMS_PASSIN Parms.Param1
+#define FUNC_PARAM_INITIALIZER_LIST Param1( InParam1 )
 #define FUNC_IS_VOID 0
 #include "DelegateInstanceInterfaceImpl.inl"
 
@@ -792,7 +813,10 @@
 #undef FUNC_TEMPLATE_ARGS
 #undef FUNC_HAS_PARAMS
 #undef FUNC_PARAM_LIST
+#undef FUNC_PARAM_MEMBERS
 #undef FUNC_PARAM_PASSTHRU
+#undef FUNC_PARAM_PARMS_PASSIN
+#undef FUNC_PARAM_INITIALIZER_LIST
 #undef FUNC_IS_VOID
 
 #define FUNC_SUFFIX TwoParams
@@ -802,7 +826,10 @@
 #define FUNC_TEMPLATE_ARGS RetValType, Param1Type, Param2Type
 #define FUNC_HAS_PARAMS 1
 #define FUNC_PARAM_LIST Param1Type InParam1, Param2Type InParam2
+#define FUNC_PARAM_MEMBERS Param1Type Param1; Param2Type Param2;
 #define FUNC_PARAM_PASSTHRU InParam1, InParam2
+#define FUNC_PARAM_PARMS_PASSIN Parms.Param1, Parms.Param2
+#define FUNC_PARAM_INITIALIZER_LIST Param1( InParam1 ), Param2( InParam2 )
 #define FUNC_IS_VOID 1
 #include "DelegateInstanceInterfaceImpl.inl"
 
@@ -990,7 +1017,10 @@
 #undef FUNC_TEMPLATE_ARGS
 #undef FUNC_HAS_PARAMS
 #undef FUNC_PARAM_LIST
+#undef FUNC_PARAM_MEMBERS
 #undef FUNC_PARAM_PASSTHRU
+#undef FUNC_PARAM_PARMS_PASSIN
+#undef FUNC_PARAM_INITIALIZER_LIST
 #undef FUNC_IS_VOID
 
 #define FUNC_SUFFIX RetVal_TwoParams
@@ -1000,7 +1030,10 @@
 #define FUNC_TEMPLATE_ARGS RetValType, Param1Type, Param2Type
 #define FUNC_HAS_PARAMS 1
 #define FUNC_PARAM_LIST Param1Type InParam1, Param2Type InParam2
+#define FUNC_PARAM_MEMBERS Param1Type Param1; Param2Type Param2;
 #define FUNC_PARAM_PASSTHRU InParam1, InParam2
+#define FUNC_PARAM_PARMS_PASSIN Parms.Param1, Parms.Param2
+#define FUNC_PARAM_INITIALIZER_LIST Param1( InParam1 ), Param2( InParam2 )
 #define FUNC_IS_VOID 0
 #include "DelegateInstanceInterfaceImpl.inl"
 
@@ -1185,7 +1218,10 @@
 #undef FUNC_TEMPLATE_ARGS
 #undef FUNC_HAS_PARAMS
 #undef FUNC_PARAM_LIST
+#undef FUNC_PARAM_MEMBERS
 #undef FUNC_PARAM_PASSTHRU
+#undef FUNC_PARAM_PARMS_PASSIN
+#undef FUNC_PARAM_INITIALIZER_LIST
 #undef FUNC_IS_VOID
 
 #define FUNC_SUFFIX ThreeParams
@@ -1195,7 +1231,10 @@
 #define FUNC_TEMPLATE_ARGS RetValType, Param1Type, Param2Type, Param3Type
 #define FUNC_HAS_PARAMS 1
 #define FUNC_PARAM_LIST Param1Type InParam1, Param2Type InParam2, Param3Type InParam3
+#define FUNC_PARAM_MEMBERS Param1Type Param1; Param2Type Param2; Param3Type Param3;
 #define FUNC_PARAM_PASSTHRU InParam1, InParam2, InParam3
+#define FUNC_PARAM_PARMS_PASSIN Parms.Param1, Parms.Param2, Parms.Param3
+#define FUNC_PARAM_INITIALIZER_LIST Param1( InParam1 ), Param2( InParam2 ), Param3( InParam3 )
 #define FUNC_IS_VOID 1
 #include "DelegateInstanceInterfaceImpl.inl"
 
@@ -1383,7 +1422,10 @@
 #undef FUNC_TEMPLATE_ARGS
 #undef FUNC_HAS_PARAMS
 #undef FUNC_PARAM_LIST
+#undef FUNC_PARAM_MEMBERS
 #undef FUNC_PARAM_PASSTHRU
+#undef FUNC_PARAM_PARMS_PASSIN
+#undef FUNC_PARAM_INITIALIZER_LIST
 #undef FUNC_IS_VOID
 
 #define FUNC_SUFFIX RetVal_ThreeParams
@@ -1393,7 +1435,10 @@
 #define FUNC_TEMPLATE_ARGS RetValType, Param1Type, Param2Type, Param3Type
 #define FUNC_HAS_PARAMS 1
 #define FUNC_PARAM_LIST Param1Type InParam1, Param2Type InParam2, Param3Type InParam3
+#define FUNC_PARAM_MEMBERS Param1Type Param1; Param2Type Param2; Param3Type Param3;
 #define FUNC_PARAM_PASSTHRU InParam1, InParam2, InParam3
+#define FUNC_PARAM_PARMS_PASSIN Parms.Param1, Parms.Param2, Parms.Param3
+#define FUNC_PARAM_INITIALIZER_LIST Param1( InParam1 ), Param2( InParam2 ), Param3( InParam3 )
 #define FUNC_IS_VOID 0
 #include "DelegateInstanceInterfaceImpl.inl"
 
@@ -1578,7 +1623,10 @@
 #undef FUNC_TEMPLATE_ARGS
 #undef FUNC_HAS_PARAMS
 #undef FUNC_PARAM_LIST
+#undef FUNC_PARAM_MEMBERS
 #undef FUNC_PARAM_PASSTHRU
+#undef FUNC_PARAM_PARMS_PASSIN
+#undef FUNC_PARAM_INITIALIZER_LIST
 #undef FUNC_IS_VOID
 
 #define FUNC_SUFFIX FourParams
@@ -1588,7 +1636,10 @@
 #define FUNC_TEMPLATE_ARGS RetValType, Param1Type, Param2Type, Param3Type, Param4Type
 #define FUNC_HAS_PARAMS 1
 #define FUNC_PARAM_LIST Param1Type InParam1, Param2Type InParam2, Param3Type InParam3, Param4Type InParam4
+#define FUNC_PARAM_MEMBERS Param1Type Param1; Param2Type Param2; Param3Type Param3; Param4Type Param4;
 #define FUNC_PARAM_PASSTHRU InParam1, InParam2, InParam3, InParam4
+#define FUNC_PARAM_PARMS_PASSIN Parms.Param1, Parms.Param2, Parms.Param3, Parms.Param4
+#define FUNC_PARAM_INITIALIZER_LIST Param1( InParam1 ), Param2( InParam2 ), Param3( InParam3 ), Param4( InParam4 )
 #define FUNC_IS_VOID 1
 #include "DelegateInstanceInterfaceImpl.inl"
 
@@ -1776,7 +1827,10 @@
 #undef FUNC_TEMPLATE_ARGS
 #undef FUNC_HAS_PARAMS
 #undef FUNC_PARAM_LIST
+#undef FUNC_PARAM_MEMBERS
 #undef FUNC_PARAM_PASSTHRU
+#undef FUNC_PARAM_PARMS_PASSIN
+#undef FUNC_PARAM_INITIALIZER_LIST
 #undef FUNC_IS_VOID
 
 #define FUNC_SUFFIX RetVal_FourParams
@@ -1786,7 +1840,10 @@
 #define FUNC_TEMPLATE_ARGS RetValType, Param1Type, Param2Type, Param3Type, Param4Type
 #define FUNC_HAS_PARAMS 1
 #define FUNC_PARAM_LIST Param1Type InParam1, Param2Type InParam2, Param3Type InParam3, Param4Type InParam4
+#define FUNC_PARAM_MEMBERS Param1Type Param1; Param2Type Param2; Param3Type Param3; Param4Type Param4;
 #define FUNC_PARAM_PASSTHRU InParam1, InParam2, InParam3, InParam4
+#define FUNC_PARAM_PARMS_PASSIN Parms.Param1, Parms.Param2, Parms.Param3, Parms.Param4
+#define FUNC_PARAM_INITIALIZER_LIST Param1( InParam1 ), Param2( InParam2 ), Param3( InParam3 ), Param4( InParam4 )
 #define FUNC_IS_VOID 0
 #include "DelegateInstanceInterfaceImpl.inl"
 
@@ -1971,7 +2028,10 @@
 #undef FUNC_TEMPLATE_ARGS
 #undef FUNC_HAS_PARAMS
 #undef FUNC_PARAM_LIST
+#undef FUNC_PARAM_MEMBERS
 #undef FUNC_PARAM_PASSTHRU
+#undef FUNC_PARAM_PARMS_PASSIN
+#undef FUNC_PARAM_INITIALIZER_LIST
 #undef FUNC_IS_VOID
 
 #define FUNC_SUFFIX FiveParams
@@ -1981,7 +2041,10 @@
 #define FUNC_TEMPLATE_ARGS RetValType, Param1Type, Param2Type, Param3Type, Param4Type, Param5Type
 #define FUNC_HAS_PARAMS 1
 #define FUNC_PARAM_LIST Param1Type InParam1, Param2Type InParam2, Param3Type InParam3, Param4Type InParam4, Param5Type InParam5
+#define FUNC_PARAM_MEMBERS Param1Type Param1; Param2Type Param2; Param3Type Param3; Param4Type Param4; Param5Type Param5;
 #define FUNC_PARAM_PASSTHRU InParam1, InParam2, InParam3, InParam4, InParam5
+#define FUNC_PARAM_PARMS_PASSIN Parms.Param1, Parms.Param2, Parms.Param3, Parms.Param4, Parms.Param5
+#define FUNC_PARAM_INITIALIZER_LIST Param1( InParam1 ), Param2( InParam2 ), Param3( InParam3 ), Param4( InParam4 ), Param5( InParam5 )
 #define FUNC_IS_VOID 1
 #include "DelegateInstanceInterfaceImpl.inl"
 
@@ -2169,7 +2232,10 @@
 #undef FUNC_TEMPLATE_ARGS
 #undef FUNC_HAS_PARAMS
 #undef FUNC_PARAM_LIST
+#undef FUNC_PARAM_MEMBERS
 #undef FUNC_PARAM_PASSTHRU
+#undef FUNC_PARAM_PARMS_PASSIN
+#undef FUNC_PARAM_INITIALIZER_LIST
 #undef FUNC_IS_VOID
 
 #define FUNC_SUFFIX RetVal_FiveParams
@@ -2179,7 +2245,10 @@
 #define FUNC_TEMPLATE_ARGS RetValType, Param1Type, Param2Type, Param3Type, Param4Type, Param5Type
 #define FUNC_HAS_PARAMS 1
 #define FUNC_PARAM_LIST Param1Type InParam1, Param2Type InParam2, Param3Type InParam3, Param4Type InParam4, Param5Type InParam5
+#define FUNC_PARAM_MEMBERS Param1Type Param1; Param2Type Param2; Param3Type Param3; Param4Type Param4; Param5Type Param5;
 #define FUNC_PARAM_PASSTHRU InParam1, InParam2, InParam3, InParam4, InParam5
+#define FUNC_PARAM_PARMS_PASSIN Parms.Param1, Parms.Param2, Parms.Param3, Parms.Param4, Parms.Param5
+#define FUNC_PARAM_INITIALIZER_LIST Param1( InParam1 ), Param2( InParam2 ), Param3( InParam3 ), Param4( InParam4 ), Param5( InParam5 )
 #define FUNC_IS_VOID 0
 #include "DelegateInstanceInterfaceImpl.inl"
 
@@ -2364,7 +2433,10 @@
 #undef FUNC_TEMPLATE_ARGS
 #undef FUNC_HAS_PARAMS
 #undef FUNC_PARAM_LIST
+#undef FUNC_PARAM_MEMBERS
 #undef FUNC_PARAM_PASSTHRU
+#undef FUNC_PARAM_PARMS_PASSIN
+#undef FUNC_PARAM_INITIALIZER_LIST
 #undef FUNC_IS_VOID
 
 #define FUNC_SUFFIX SixParams
@@ -2374,7 +2446,10 @@
 #define FUNC_TEMPLATE_ARGS RetValType, Param1Type, Param2Type, Param3Type, Param4Type, Param5Type, Param6Type
 #define FUNC_HAS_PARAMS 1
 #define FUNC_PARAM_LIST Param1Type InParam1, Param2Type InParam2, Param3Type InParam3, Param4Type InParam4, Param5Type InParam5, Param6Type InParam6
+#define FUNC_PARAM_MEMBERS Param1Type Param1; Param2Type Param2; Param3Type Param3; Param4Type Param4; Param5Type Param5; Param6Type Param6;
 #define FUNC_PARAM_PASSTHRU InParam1, InParam2, InParam3, InParam4, InParam5, InParam6
+#define FUNC_PARAM_PARMS_PASSIN Parms.Param1, Parms.Param2, Parms.Param3, Parms.Param4, Parms.Param5, Parms.Param6
+#define FUNC_PARAM_INITIALIZER_LIST Param1( InParam1 ), Param2( InParam2 ), Param3( InParam3 ), Param4( InParam4 ), Param5( InParam5 ), Param6( InParam6 )
 #define FUNC_IS_VOID 1
 #include "DelegateInstanceInterfaceImpl.inl"
 
@@ -2562,7 +2637,10 @@
 #undef FUNC_TEMPLATE_ARGS
 #undef FUNC_HAS_PARAMS
 #undef FUNC_PARAM_LIST
+#undef FUNC_PARAM_MEMBERS
 #undef FUNC_PARAM_PASSTHRU
+#undef FUNC_PARAM_PARMS_PASSIN
+#undef FUNC_PARAM_INITIALIZER_LIST
 #undef FUNC_IS_VOID
 
 #define FUNC_SUFFIX RetVal_SixParams
@@ -2572,7 +2650,10 @@
 #define FUNC_TEMPLATE_ARGS RetValType, Param1Type, Param2Type, Param3Type, Param4Type, Param5Type, Param6Type
 #define FUNC_HAS_PARAMS 1
 #define FUNC_PARAM_LIST Param1Type InParam1, Param2Type InParam2, Param3Type InParam3, Param4Type InParam4, Param5Type InParam5, Param6Type InParam6
+#define FUNC_PARAM_MEMBERS Param1Type Param1; Param2Type Param2; Param3Type Param3; Param4Type Param4; Param5Type Param5; Param6Type Param6;
 #define FUNC_PARAM_PASSTHRU InParam1, InParam2, InParam3, InParam4, InParam5, InParam6
+#define FUNC_PARAM_PARMS_PASSIN Parms.Param1, Parms.Param2, Parms.Param3, Parms.Param4, Parms.Param5, Parms.Param6
+#define FUNC_PARAM_INITIALIZER_LIST Param1( InParam1 ), Param2( InParam2 ), Param3( InParam3 ), Param4( InParam4 ), Param5( InParam5 ), Param6( InParam6 )
 #define FUNC_IS_VOID 0
 #include "DelegateInstanceInterfaceImpl.inl"
 
@@ -2757,7 +2838,10 @@
 #undef FUNC_TEMPLATE_ARGS
 #undef FUNC_HAS_PARAMS
 #undef FUNC_PARAM_LIST
+#undef FUNC_PARAM_MEMBERS
 #undef FUNC_PARAM_PASSTHRU
+#undef FUNC_PARAM_PARMS_PASSIN
+#undef FUNC_PARAM_INITIALIZER_LIST
 #undef FUNC_IS_VOID
 
 #define FUNC_SUFFIX SevenParams
@@ -2767,7 +2851,10 @@
 #define FUNC_TEMPLATE_ARGS RetValType, Param1Type, Param2Type, Param3Type, Param4Type, Param5Type, Param6Type, Param7Type
 #define FUNC_HAS_PARAMS 1
 #define FUNC_PARAM_LIST Param1Type InParam1, Param2Type InParam2, Param3Type InParam3, Param4Type InParam4, Param5Type InParam5, Param6Type InParam6, Param7Type InParam7
+#define FUNC_PARAM_MEMBERS Param1Type Param1; Param2Type Param2; Param3Type Param3; Param4Type Param4; Param5Type Param5; Param6Type Param6; Param7Type Param7;
 #define FUNC_PARAM_PASSTHRU InParam1, InParam2, InParam3, InParam4, InParam5, InParam6, InParam7
+#define FUNC_PARAM_PARMS_PASSIN Parms.Param1, Parms.Param2, Parms.Param3, Parms.Param4, Parms.Param5, Parms.Param6, Parms.Param7
+#define FUNC_PARAM_INITIALIZER_LIST Param1( InParam1 ), Param2( InParam2 ), Param3( InParam3 ), Param4( InParam4 ), Param5( InParam5 ), Param6( InParam6 ), Param7( InParam7 )
 #define FUNC_IS_VOID 1
 #include "DelegateInstanceInterfaceImpl.inl"
 
@@ -2955,7 +3042,10 @@
 #undef FUNC_TEMPLATE_ARGS
 #undef FUNC_HAS_PARAMS
 #undef FUNC_PARAM_LIST
+#undef FUNC_PARAM_MEMBERS
 #undef FUNC_PARAM_PASSTHRU
+#undef FUNC_PARAM_PARMS_PASSIN
+#undef FUNC_PARAM_INITIALIZER_LIST
 #undef FUNC_IS_VOID
 
 #define FUNC_SUFFIX RetVal_SevenParams
@@ -2965,7 +3055,10 @@
 #define FUNC_TEMPLATE_ARGS RetValType, Param1Type, Param2Type, Param3Type, Param4Type, Param5Type, Param6Type, Param7Type
 #define FUNC_HAS_PARAMS 1
 #define FUNC_PARAM_LIST Param1Type InParam1, Param2Type InParam2, Param3Type InParam3, Param4Type InParam4, Param5Type InParam5, Param6Type InParam6, Param7Type InParam7
+#define FUNC_PARAM_MEMBERS Param1Type Param1; Param2Type Param2; Param3Type Param3; Param4Type Param4; Param5Type Param5; Param6Type Param6; Param7Type Param7;
 #define FUNC_PARAM_PASSTHRU InParam1, InParam2, InParam3, InParam4, InParam5, InParam6, InParam7
+#define FUNC_PARAM_PARMS_PASSIN Parms.Param1, Parms.Param2, Parms.Param3, Parms.Param4, Parms.Param5, Parms.Param6, Parms.Param7
+#define FUNC_PARAM_INITIALIZER_LIST Param1( InParam1 ), Param2( InParam2 ), Param3( InParam3 ), Param4( InParam4 ), Param5( InParam5 ), Param6( InParam6 ), Param7( InParam7 )
 #define FUNC_IS_VOID 0
 #include "DelegateInstanceInterfaceImpl.inl"
 
@@ -3150,7 +3243,10 @@
 #undef FUNC_TEMPLATE_ARGS
 #undef FUNC_HAS_PARAMS
 #undef FUNC_PARAM_LIST
+#undef FUNC_PARAM_MEMBERS
 #undef FUNC_PARAM_PASSTHRU
+#undef FUNC_PARAM_PARMS_PASSIN
+#undef FUNC_PARAM_INITIALIZER_LIST
 #undef FUNC_IS_VOID
 
 #define FUNC_SUFFIX EightParams
@@ -3160,7 +3256,10 @@
 #define FUNC_TEMPLATE_ARGS RetValType, Param1Type, Param2Type, Param3Type, Param4Type, Param5Type, Param6Type, Param7Type, Param8Type
 #define FUNC_HAS_PARAMS 1
 #define FUNC_PARAM_LIST Param1Type InParam1, Param2Type InParam2, Param3Type InParam3, Param4Type InParam4, Param5Type InParam5, Param6Type InParam6, Param7Type InParam7, Param8Type InParam8
+#define FUNC_PARAM_MEMBERS Param1Type Param1; Param2Type Param2; Param3Type Param3; Param4Type Param4; Param5Type Param5; Param6Type Param6; Param7Type Param7; Param8Type Param8;
 #define FUNC_PARAM_PASSTHRU InParam1, InParam2, InParam3, InParam4, InParam5, InParam6, InParam7, InParam8
+#define FUNC_PARAM_PARMS_PASSIN Parms.Param1, Parms.Param2, Parms.Param3, Parms.Param4, Parms.Param5, Parms.Param6, Parms.Param7, Parms.Param8
+#define FUNC_PARAM_INITIALIZER_LIST Param1( InParam1 ), Param2( InParam2 ), Param3( InParam3 ), Param4( InParam4 ), Param5( InParam5 ), Param6( InParam6 ), Param7( InParam7 ), Param8( InParam8 )
 #define FUNC_IS_VOID 1
 #include "DelegateInstanceInterfaceImpl.inl"
 
@@ -3348,7 +3447,10 @@
 #undef FUNC_TEMPLATE_ARGS
 #undef FUNC_HAS_PARAMS
 #undef FUNC_PARAM_LIST
+#undef FUNC_PARAM_MEMBERS
 #undef FUNC_PARAM_PASSTHRU
+#undef FUNC_PARAM_PARMS_PASSIN
+#undef FUNC_PARAM_INITIALIZER_LIST
 #undef FUNC_IS_VOID
 
 #define FUNC_SUFFIX RetVal_EightParams
@@ -3358,7 +3460,10 @@
 #define FUNC_TEMPLATE_ARGS RetValType, Param1Type, Param2Type, Param3Type, Param4Type, Param5Type, Param6Type, Param7Type, Param8Type
 #define FUNC_HAS_PARAMS 1
 #define FUNC_PARAM_LIST Param1Type InParam1, Param2Type InParam2, Param3Type InParam3, Param4Type InParam4, Param5Type InParam5, Param6Type InParam6, Param7Type InParam7, Param8Type InParam8
+#define FUNC_PARAM_MEMBERS Param1Type Param1; Param2Type Param2; Param3Type Param3; Param4Type Param4; Param5Type Param5; Param6Type Param6; Param7Type Param7; Param8Type Param8;
 #define FUNC_PARAM_PASSTHRU InParam1, InParam2, InParam3, InParam4, InParam5, InParam6, InParam7, InParam8
+#define FUNC_PARAM_PARMS_PASSIN Parms.Param1, Parms.Param2, Parms.Param3, Parms.Param4, Parms.Param5, Parms.Param6, Parms.Param7, Parms.Param8
+#define FUNC_PARAM_INITIALIZER_LIST Param1( InParam1 ), Param2( InParam2 ), Param3( InParam3 ), Param4( InParam4 ), Param5( InParam5 ), Param6( InParam6 ), Param7( InParam7 ), Param8( InParam8 )
 #define FUNC_IS_VOID 0
 #include "DelegateInstanceInterfaceImpl.inl"
 
@@ -3543,5 +3648,8 @@
 #undef FUNC_TEMPLATE_ARGS
 #undef FUNC_HAS_PARAMS
 #undef FUNC_PARAM_LIST
+#undef FUNC_PARAM_MEMBERS
 #undef FUNC_PARAM_PASSTHRU
+#undef FUNC_PARAM_PARMS_PASSIN
+#undef FUNC_PARAM_INITIALIZER_LIST
 #undef FUNC_IS_VOID
diff --git a/Engine/Source/Runtime/Core/Public/Templates/DelegateInstanceInterfaceImpl.inl b/Engine/Source/Runtime/Core/Public/Templates/DelegateInstanceInterfaceImpl.inl
--- a/Engine/Source/Runtime/Core/Public/Templates/DelegateInstanceInterfaceImpl.inl
+++ b/Engine/Source/Runtime/Core/Public/Templates/DelegateInstanceInterfaceImpl.inl
@@ -12,6 +12,12 @@
 
 #define DELEGATE_INSTANCE_INTERFACE_CLASS FUNC_COMBINE( IBaseDelegateInstance_, FUNC_SUFFIX )
 
+#if FUNC_HAS_PARAMS
+	#define DELEGATE_INSTANCE_COLON_PARAM_INITIALIZER_LIST : FUNC_PARAM_INITIALIZER_LIST
+#else
+	#define DELEGATE_INSTANCE_COLON_PARAM_INITIALIZER_LIST
+#endif
+
 
 /**
  * Delegate instance base interface.  For internal use only.  This abstract class is used to interact with
@@ -19,24 +25,23 @@
  */
 template< FUNC_TEMPLATE_DECL_TYPENAME >
 class DELEGATE_INSTANCE_INTERFACE_CLASS
+	: public IDelegateInstance
 {
-public:
-
-	/**
-	 * Destructor
-	 */
-	virtual ~DELEGATE_INSTANCE_INTERFACE_CLASS()
+	// Structure for UFunction call parameter marshaling.
+	struct FParms
 	{
-		// Empty virtual destructor.  Ensures that user-defined payload member variables are properly
-		// destructed when the delegate instance is destroyed.
-	}
+		FUNC_PARAM_MEMBERS
 
-	/**
-	 * Returns the type of delegate instance
-	 *
-	 * @return  Delegate instance type
-	 */
-	virtual EDelegateInstanceType::Type GetType() const = 0;
+#if !FUNC_IS_VOID
+		RetValType Result;
+#endif
+
+		FParms( FUNC_PARAM_LIST )
+			DELEGATE_INSTANCE_COLON_PARAM_INITIALIZER_LIST
+		{ }
+	};
+
+public:
 
 	/**
 	 * Creates a copy of the delegate instance
@@ -56,61 +61,39 @@ public:
 	virtual bool IsSameFunction( const DELEGATE_INSTANCE_INTERFACE_CLASS< FUNC_TEMPLATE_ARGS >& InOtherDelegate ) const = 0;
 
 	/**
-	 * Returns true if this delegate is bound to the specified UserObject,
-	 *
-	 * @param  InUserObject
-	 *
-	 * @return  True if delegate is bound to the specified UserObject
+	 * Execute the delegate.  If the function pointer is not valid, an error will occur.
 	 */
-	virtual bool HasSameObject( const void* InUserObject ) const = 0;
+	virtual RetValType Execute( FUNC_PARAM_LIST ) const = 0;
 
-	/**
-	 * Returns raw pointer to UserObject,
-	 *
-	 * @return  Raw pointer to UserObject.
-	 */
-	virtual const void* GetRawUserObject() const = 0;
+#if FUNC_IS_VOID
 
 	/**
-	 * Returns raw pointer to the delegate method.
+	 * Execute the delegate, but only if the function pointer is still valid
 	 *
-	 * @return  Raw pointer to the delegate method.
+	 * @return  Returns true if the function was executed
 	 */
-	virtual const void* GetRawMethodPtr() const = 0;
+	// NOTE: Currently only delegates with no return value support ExecuteIfSafe() 
+	virtual bool ExecuteIfSafe( FUNC_PARAM_LIST ) const = 0;
 
-	/**
-	 * Execute the delegate.  If the function pointer is not valid, an error will occur.
-	 */
-	virtual RetValType Execute( FUNC_PARAM_LIST ) const = 0;
+#endif // FUNC_IS_VOID
 
-	/**
-	 * Checks to see if the user object bound to this delegate is still valid
-	 *
-	 * @return  True if the user object is still valid and it's safe to execute the function call
-	 */
-	virtual bool IsSafeToExecute() const = 0;
+public:
 
 	/**
-	 * Checks to see if the user object bound to this delegate can ever be valid again.
-	 * used to compact multicast delegate arrays so they don't expand without limit.
+	 * Execute the delegate with parameters from the virtual machine.
 	 *
-	 * @return  True if the user object can never be used again
+	 * This non-virtual interface (NVI) method marshals a set of call parameters
+	 * from the UObject virtual machine to the bound delegate handler function.
+	 *
+	 * @param InParms An UFunction parameter structure.
 	 */
-	virtual bool IsCompactable() const
+	RetValType ExecuteParms( void* InParms ) const
 	{
-		return !IsSafeToExecute();
+		FParms& Parms = *((FParms*)InParms);
+		return Execute(FUNC_PARAM_PARMS_PASSIN);
 	}
-
-#if FUNC_IS_VOID
-	/**
-	 * Execute the delegate, but only if the function pointer is still valid
-	 *
-	 * @return  Returns true if the function was executed
-	 */
-	// NOTE: Currently only delegates with no return value support ExecuteIfSafe() 
-	virtual bool ExecuteIfSafe( FUNC_PARAM_LIST ) const = 0;
-#endif
 };
 
 
+#undef DELEGATE_INSTANCE_COLON_PARAM_INITIALIZER_LIST
 #undef DELEGATE_INSTANCE_INTERFACE_CLASS
diff --git a/Engine/Source/Runtime/Core/Public/Templates/DelegateInstancesImpl.inl b/Engine/Source/Runtime/Core/Public/Templates/DelegateInstancesImpl.inl
--- a/Engine/Source/Runtime/Core/Public/Templates/DelegateInstancesImpl.inl
+++ b/Engine/Source/Runtime/Core/Public/Templates/DelegateInstancesImpl.inl
@@ -20,6 +20,7 @@
 #define RAW_METHOD_DELEGATE_INSTANCE_CLASS FUNC_COMBINE( TBaseRawMethodDelegateInstance_, FUNC_COMBINE( FUNC_PAYLOAD_SUFFIX, FUNC_CONST_SUFFIX ) )
 #define UOBJECT_METHOD_DELEGATE_INSTANCE_CLASS FUNC_COMBINE( TBaseUObjectMethodDelegateInstance_, FUNC_COMBINE( FUNC_PAYLOAD_SUFFIX, FUNC_CONST_SUFFIX ) )
 #define STATIC_DELEGATE_INSTANCE_CLASS FUNC_COMBINE( TBaseStaticDelegateInstance_, FUNC_COMBINE( FUNC_PAYLOAD_SUFFIX, FUNC_CONST_SUFFIX ) )
+#define UFUNCTION_DELEGATE_INSTANCE_CLASS FUNC_COMBINE( TBaseUFunctionDelegateInstance_, FUNC_COMBINE( FUNC_PAYLOAD_SUFFIX, FUNC_CONST_SUFFIX ) )
 
 #if FUNC_IS_CONST
 	#define DELEGATE_QUALIFIER const
@@ -35,9 +36,11 @@
 	#if FUNC_HAS_PARAMS
 		#define DELEGATE_PARAM_PASSTHRU_COMMA_PAYLOAD_PASSIN FUNC_PARAM_PASSTHRU, FUNC_PAYLOAD_PASSIN
 		#define DELEGATE_PARAM_LIST_COMMA_PAYLOAD_LIST       FUNC_PARAM_LIST, FUNC_PAYLOAD_LIST
+		#define DELEGATE_PARMS_COLON_INITIALIZER_LIST	     : FUNC_PARAM_INITIALIZER_LIST, FUNC_PAYLOAD_INITIALIZER_LIST
 	#else
 		#define DELEGATE_PARAM_PASSTHRU_COMMA_PAYLOAD_PASSIN FUNC_PAYLOAD_PASSIN
 		#define DELEGATE_PARAM_LIST_COMMA_PAYLOAD_LIST       FUNC_PAYLOAD_LIST
+		#define DELEGATE_PARMS_COLON_INITIALIZER_LIST	     : FUNC_PAYLOAD_INITIALIZER_LIST
 	#endif
 #else
 	#define DELEGATE_COMMA_PAYLOAD_LIST
@@ -46,6 +49,11 @@
 	#define DELEGATE_COMMA_PAYLOAD_PASSIN
 	#define DELEGATE_PARAM_PASSTHRU_COMMA_PAYLOAD_PASSIN FUNC_PARAM_PASSTHRU
 	#define DELEGATE_PARAM_LIST_COMMA_PAYLOAD_LIST       FUNC_PARAM_LIST
+	#if FUNC_HAS_PARAMS
+		#define DELEGATE_PARMS_COLON_INITIALIZER_LIST	 : FUNC_PARAM_INITIALIZER_LIST
+	#else
+		#define DELEGATE_PARMS_COLON_INITIALIZER_LIST
+	#endif
 #endif
 
 
@@ -183,6 +191,11 @@ public:
 		return UserObject.HasSameObject( InUserObject );
 	}
 
+	virtual FName GetFunctionName() const OVERRIDE
+	{
+		return NAME_None;
+	}
+
 	virtual const void* GetRawUserObject() const OVERRIDE
 	{
 		return GetRawUserObjectInternal();
@@ -380,6 +393,11 @@ public:
 		return UserObject == InUserObject;
 	}
 
+	virtual FName GetFunctionName() const OVERRIDE
+	{
+		return NAME_None;
+	}
+
 	virtual const void* GetRawUserObject() const OVERRIDE
 	{
 		return GetRawUserObjectInternal();
@@ -569,6 +587,11 @@ public:
 		return UserObject.Get() == InUserObject;
 	}
 
+	virtual FName GetFunctionName() const OVERRIDE
+	{
+		return NAME_None;
+	}
+
 	virtual const void* GetRawUserObject() const OVERRIDE
 	{
 		return GetRawUserObjectInternal();
@@ -750,6 +773,11 @@ public:
 		return false;
 	}
 
+	virtual FName GetFunctionName() const OVERRIDE
+	{
+		return NAME_None;
+	}
+
 	virtual const void* GetRawUserObject() const OVERRIDE
 	{
 		return NULL;
@@ -806,10 +834,180 @@ private:
 };
 
 
+#if !FUNC_IS_CONST
+
+/**
+ * Delegate type implementation class for UFunctions.  For internal use only.
+ *
+ * @params UserClass Must be an UObject derived class.
+ */
+template<class UserClass, FUNC_PAYLOAD_TEMPLATE_DECL_TYPENAME>
+class UFUNCTION_DELEGATE_INSTANCE_CLASS
+	: public DELEGATE_INSTANCE_INTERFACE_CLASS<FUNC_TEMPLATE_ARGS>
+{
+	checkAtCompileTime((CanConvertPointerFromTo<UserClass, UObjectBase>::Result), You_Cannot_Use_UFunction_Delegates_With_Non_UObject_Classes);
+
+	// Structure for UFunction call parameter marshaling.
+	struct FParmsWithPayload
+	{
+		FUNC_PARAM_MEMBERS		// (Param1Type Param1; Param2Type Param2; ...)
+		FUNC_PAYLOAD_MEMBERS	// (Var1Type Var1; Var2Type Var2; ...)
+
+#if !FUNC_IS_VOID
+		RetValType Result;		// UFunction return value.
+#endif
+
+		// Constructor.
+		FParmsWithPayload( DELEGATE_PARAM_LIST_COMMA_PAYLOAD_LIST )
+			DELEGATE_PARMS_COLON_INITIALIZER_LIST
+		{ }
+	};
+
+public:
+
+	/**
+	 * Creates and initializes a new instance.
+	 *
+	 * @param InUserObject The UObject to call the function on.
+	 * @param InFunctionName The name of the function call.
+	 */
+	UFUNCTION_DELEGATE_INSTANCE_CLASS( UserClass* InUserObject, const FName& InFunctionName DELEGATE_COMMA_PAYLOAD_LIST )
+		: FunctionName(InFunctionName)
+		, UserObjectPtr(InUserObject)
+		DELEGATE_COMMA_PAYLOAD_INITIALIZER_LIST
+	{
+		check(InFunctionName != NAME_None);
+		
+		if (InUserObject != nullptr)
+		{
+			CachedFunction = UserObjectPtr->FindFunctionChecked(InFunctionName);
+		}		
+	}
+	
+public:
+
+	// Begin IDelegateInstance interface
+
+	virtual FName GetFunctionName( ) const OVERRIDE
+	{
+		return FunctionName;
+	}
+
+	virtual const void* GetRawMethodPtr( ) const OVERRIDE
+	{
+		return nullptr;
+	}
+
+	virtual const void* GetRawUserObject( ) const OVERRIDE
+	{
+		return UserObjectPtr.Get();
+	}
+
+	virtual EDelegateInstanceType::Type GetType() const OVERRIDE
+	{
+		return EDelegateInstanceType::UFunction;
+	}
+
+	virtual bool HasSameObject( const void* InUserObject ) const OVERRIDE
+	{
+		return (UserObjectPtr.Get() == InUserObject);
+	}
+
+	virtual bool IsCompactable( ) const OVERRIDE
+	{
+		return !UserObjectPtr.Get(true);
+	}
+
+	virtual bool IsSafeToExecute( ) const OVERRIDE
+	{
+		return UserObjectPtr.IsValid();
+	}
+
+	// End IDelegateInstance interface
+
+public:
+
+	// Begin DELEGATE_INSTANCE_INTERFACE_CLASS interface
+
+	virtual DELEGATE_INSTANCE_INTERFACE_CLASS<FUNC_TEMPLATE_ARGS>* CreateCopy() OVERRIDE
+	{
+		return Create(UserObjectPtr.Get(), GetFunctionName() DELEGATE_COMMA_PAYLOAD_PASSIN);
+	}
+
+	virtual RetValType Execute( FUNC_PARAM_LIST ) const OVERRIDE
+	{
+		checkSlow(IsSafeToExecute());
+
+		FParmsWithPayload Parms = FParmsWithPayload(DELEGATE_PARAM_PASSTHRU_COMMA_PAYLOAD_PASSIN);
+		UserObjectPtr->ProcessEvent(CachedFunction, &Parms);
+
+#if !FUNC_IS_VOID
+		return Parms.Result;
+#endif
+	}
+
+#if FUNC_IS_VOID
+	virtual bool ExecuteIfSafe( FUNC_PARAM_LIST ) const OVERRIDE
+	{
+		if (IsSafeToExecute())
+		{
+			Execute(FUNC_PARAM_PASSTHRU);
+
+			return true;
+		}
+
+		return false;
+	}
+#endif
+
+	virtual bool IsSameFunction( const DELEGATE_INSTANCE_INTERFACE_CLASS<FUNC_TEMPLATE_ARGS>& Other ) const OVERRIDE
+	{
+		return ((Other.GetType() == EDelegateInstanceType::UFunction) &&
+				(Other.GetRawUserObject() == GetRawUserObject()) &&
+				(Other.GetFunctionName() == GetFunctionName()));
+	}
+
+	// End DELEGATE_INSTANCE_INTERFACE_CLASS interface
+
+public:
+
+	/**
+	 * Static helper function that returns a delegate object for the specified user class and class method.
+	 *
+	 * @param InObject The user object to call the function on.
+	 * @param InFunctionName The name of the function call.
+	 *
+	 * @return  Returns an interface to the new delegate
+	 */
+	FORCEINLINE static DELEGATE_INSTANCE_INTERFACE_CLASS<FUNC_TEMPLATE_ARGS>* Create( UserClass* InUserObject, const FName& InFunctionName DELEGATE_COMMA_PAYLOAD_LIST )
+	{
+		return new UFUNCTION_DELEGATE_INSTANCE_CLASS<UserClass, FUNC_PAYLOAD_TEMPLATE_ARGS>(InUserObject, InFunctionName DELEGATE_COMMA_PAYLOAD_PASSTHRU);
+	}
+
+public:
+
+	// Holds the cached UFunction to call.
+	UFunction* CachedFunction;
+
+	// Holds the name of the function to call.
+	FName FunctionName;
+
+	// The user object to call the function on.
+	TWeakObjectPtr<UserClass> UserObjectPtr;
+
+	// Payload member variables, if any.
+	FUNC_PAYLOAD_MEMBERS
+};
+
+#endif // FUNC_IS_CONST
+
+
 #undef SP_METHOD_DELEGATE_INSTANCE_CLASS
 #undef RAW_METHOD_DELEGATE_INSTANCE_CLASS
 #undef UOBJECT_METHOD_DELEGATE_INSTANCE_CLASS
 #undef STATIC_DELEGATE_INSTANCE_CLASS
+#undef UFUNCTION_DELEGATE_INSTANCE_CLASS
+
 #undef DELEGATE_INSTANCE_INTERFACE_CLASS
 #undef DELEGATE_QUALIFIER
 #undef DELEGATE_COMMA_PAYLOAD_LIST
@@ -818,3 +1016,4 @@ private:
 #undef DELEGATE_COMMA_PAYLOAD_PASSIN
 #undef DELEGATE_PARAM_PASSTHRU_COMMA_PAYLOAD_PASSIN
 #undef DELEGATE_PARAM_LIST_COMMA_PAYLOAD_LIST
+#undef DELEGATE_PARMS_COLON_INITIALIZER_LIST
diff --git a/Engine/Source/Runtime/Core/Public/Templates/DelegateSignatureImpl.inl b/Engine/Source/Runtime/Core/Public/Templates/DelegateSignatureImpl.inl
--- a/Engine/Source/Runtime/Core/Public/Templates/DelegateSignatureImpl.inl
+++ b/Engine/Source/Runtime/Core/Public/Templates/DelegateSignatureImpl.inl
@@ -34,6 +34,12 @@
 #define RAW_METHOD_DELEGATE_INSTANCE_CLASS_FourVars FUNC_COMBINE( RAW_METHOD_DELEGATE_INSTANCE_CLASS_NoVars, _FourVars )
 #define RAW_METHOD_DELEGATE_INSTANCE_CLASS_FourVars_Const FUNC_COMBINE( RAW_METHOD_DELEGATE_INSTANCE_CLASS_FourVars, _Const )
 
+#define UFUNCTION_DELEGATE_INSTANCE_CLASS_NoVars FUNC_COMBINE( TBaseUFunctionDelegateInstance_, FUNC_SUFFIX )
+#define UFUNCTION_DELEGATE_INSTANCE_CLASS_OneVar FUNC_COMBINE( UFUNCTION_DELEGATE_INSTANCE_CLASS_NoVars, _OneVar )
+#define UFUNCTION_DELEGATE_INSTANCE_CLASS_TwoVars FUNC_COMBINE( UFUNCTION_DELEGATE_INSTANCE_CLASS_NoVars, _TwoVars )
+#define UFUNCTION_DELEGATE_INSTANCE_CLASS_ThreeVars FUNC_COMBINE( UFUNCTION_DELEGATE_INSTANCE_CLASS_NoVars, _ThreeVars )
+#define UFUNCTION_DELEGATE_INSTANCE_CLASS_FourVars FUNC_COMBINE( UFUNCTION_DELEGATE_INSTANCE_CLASS_NoVars, _FourVars )
+
 #define UOBJECT_METHOD_DELEGATE_INSTANCE_CLASS_NoVars FUNC_COMBINE( TBaseUObjectMethodDelegateInstance_, FUNC_SUFFIX )
 #define UOBJECT_METHOD_DELEGATE_INSTANCE_CLASS_NoVars_Const FUNC_COMBINE( UOBJECT_METHOD_DELEGATE_INSTANCE_CLASS_NoVars, _Const )
 #define UOBJECT_METHOD_DELEGATE_INSTANCE_CLASS_OneVar FUNC_COMBINE( UOBJECT_METHOD_DELEGATE_INSTANCE_CLASS_NoVars, _OneVar )
@@ -85,7 +91,7 @@
 
 
 /**
- * Single-cast delegate base object.
+ * Unicast delegate base object.
  *
  * Use the various DECLARE_DELEGATE macros to create the actual delegate type, templated to
  * the function signature the delegate is compatible with. Then, you can create an instance
@@ -93,13 +99,14 @@
  */
 template< FUNC_TEMPLATE_DECL_TYPENAME >
 class DELEGATE_CLASS
+	: public FDelegateBase<>
 {
 public:
 
-	/** Declare the user's delegate interface (e.g. IMyDelegate) */
-	typedef DELEGATE_INSTANCE_INTERFACE_CLASS< FUNC_TEMPLATE_ARGS > IDelegate;
+	/** Type definition for the shared interface of delegate instance types compatible with this delegate class. */
+	typedef DELEGATE_INSTANCE_INTERFACE_CLASS<FUNC_TEMPLATE_ARGS> TDelegateInstanceInterface;
 
-	/** Declare the user's "fast" shared pointer-based delegate object (e.g. TSPMyDelegate) */
+	/** Declare the user's "fast" shared pointer-based delegate instance types. */
 	template< class UserClass                                  > class TSPMethodDelegate                 : public SP_METHOD_DELEGATE_INSTANCE_CLASS_NoVars         < UserClass, FUNC_TEMPLATE_ARGS,                                  ESPMode::Fast > { public: TSPMethodDelegate                ( const TSharedRef< UserClass, ESPMode::Fast >& InUserObject, typename SP_METHOD_DELEGATE_INSTANCE_CLASS_NoVars         < UserClass, FUNC_TEMPLATE_ARGS,                                  ESPMode::Fast >::FMethodPtr InMethodPtr                                  ) : SP_METHOD_DELEGATE_INSTANCE_CLASS_NoVars         < UserClass, FUNC_TEMPLATE_ARGS,                                  ESPMode::Fast >( InUserObject, InMethodPtr                                    ) {} };
 	template< class UserClass                                  > class TSPMethodDelegate_Const           : public SP_METHOD_DELEGATE_INSTANCE_CLASS_NoVars_Const   < UserClass, FUNC_TEMPLATE_ARGS,                                  ESPMode::Fast > { public: TSPMethodDelegate_Const          ( const TSharedRef< UserClass, ESPMode::Fast >& InUserObject, typename SP_METHOD_DELEGATE_INSTANCE_CLASS_NoVars_Const   < UserClass, FUNC_TEMPLATE_ARGS,                                  ESPMode::Fast >::FMethodPtr InMethodPtr                                  ) : SP_METHOD_DELEGATE_INSTANCE_CLASS_NoVars_Const   < UserClass, FUNC_TEMPLATE_ARGS,                                  ESPMode::Fast >( InUserObject, InMethodPtr                                    ) {} };
 	template< class UserClass, PAYLOAD_TEMPLATE_DECL_OneVar    > class TSPMethodDelegate_OneVar          : public SP_METHOD_DELEGATE_INSTANCE_CLASS_OneVar         < UserClass, FUNC_TEMPLATE_ARGS, PAYLOAD_TEMPLATE_LIST_OneVar,    ESPMode::Fast > { public: TSPMethodDelegate_OneVar         ( const TSharedRef< UserClass, ESPMode::Fast >& InUserObject, typename SP_METHOD_DELEGATE_INSTANCE_CLASS_OneVar         < UserClass, FUNC_TEMPLATE_ARGS, PAYLOAD_TEMPLATE_LIST_OneVar,    ESPMode::Fast >::FMethodPtr InMethodPtr, PAYLOAD_TEMPLATE_ARGS_OneVar    ) : SP_METHOD_DELEGATE_INSTANCE_CLASS_OneVar         < UserClass, FUNC_TEMPLATE_ARGS, PAYLOAD_TEMPLATE_LIST_OneVar,    ESPMode::Fast >( InUserObject, InMethodPtr, PAYLOAD_TEMPLATE_PASSIN_OneVar    ) {} };
@@ -111,7 +118,7 @@ public:
 	template< class UserClass, PAYLOAD_TEMPLATE_DECL_FourVars  > class TSPMethodDelegate_FourVars        : public SP_METHOD_DELEGATE_INSTANCE_CLASS_FourVars       < UserClass, FUNC_TEMPLATE_ARGS, PAYLOAD_TEMPLATE_LIST_FourVars,  ESPMode::Fast > { public: TSPMethodDelegate_FourVars       ( const TSharedRef< UserClass, ESPMode::Fast >& InUserObject, typename SP_METHOD_DELEGATE_INSTANCE_CLASS_FourVars       < UserClass, FUNC_TEMPLATE_ARGS, PAYLOAD_TEMPLATE_LIST_FourVars,  ESPMode::Fast >::FMethodPtr InMethodPtr, PAYLOAD_TEMPLATE_ARGS_FourVars  ) : SP_METHOD_DELEGATE_INSTANCE_CLASS_FourVars       < UserClass, FUNC_TEMPLATE_ARGS, PAYLOAD_TEMPLATE_LIST_FourVars,  ESPMode::Fast >( InUserObject, InMethodPtr, PAYLOAD_TEMPLATE_PASSIN_FourVars  ) {} };
 	template< class UserClass, PAYLOAD_TEMPLATE_DECL_FourVars  > class TSPMethodDelegate_FourVars_Const  : public SP_METHOD_DELEGATE_INSTANCE_CLASS_FourVars_Const < UserClass, FUNC_TEMPLATE_ARGS, PAYLOAD_TEMPLATE_LIST_FourVars,  ESPMode::Fast > { public: TSPMethodDelegate_FourVars_Const ( const TSharedRef< UserClass, ESPMode::Fast >& InUserObject, typename SP_METHOD_DELEGATE_INSTANCE_CLASS_FourVars_Const < UserClass, FUNC_TEMPLATE_ARGS, PAYLOAD_TEMPLATE_LIST_FourVars,  ESPMode::Fast >::FMethodPtr InMethodPtr, PAYLOAD_TEMPLATE_ARGS_FourVars  ) : SP_METHOD_DELEGATE_INSTANCE_CLASS_FourVars_Const < UserClass, FUNC_TEMPLATE_ARGS, PAYLOAD_TEMPLATE_LIST_FourVars,  ESPMode::Fast >( InUserObject, InMethodPtr, PAYLOAD_TEMPLATE_PASSIN_FourVars  ) {} };
 
-	/** Declare the user's "thread-safe" shared pointer-based delegate object (e.g. TThreadSafeSPMyDelegate) */
+	/** Declare the user's "thread-safe" shared pointer-based delegate instance types. */
 	template< class UserClass                                  > class TThreadSafeSPMethodDelegate                 : public SP_METHOD_DELEGATE_INSTANCE_CLASS_NoVars         < UserClass, FUNC_TEMPLATE_ARGS,                                  ESPMode::ThreadSafe > { public: TThreadSafeSPMethodDelegate                ( const TSharedRef< UserClass, ESPMode::ThreadSafe >& InUserObject, typename SP_METHOD_DELEGATE_INSTANCE_CLASS_NoVars         < UserClass, FUNC_TEMPLATE_ARGS,                                  ESPMode::ThreadSafe >::FMethodPtr InMethodPtr                                  ) : SP_METHOD_DELEGATE_INSTANCE_CLASS_NoVars         < UserClass, FUNC_TEMPLATE_ARGS,                                  ESPMode::ThreadSafe >( InUserObject, InMethodPtr                                    ) {} };
 	template< class UserClass                                  > class TThreadSafeSPMethodDelegate_Const           : public SP_METHOD_DELEGATE_INSTANCE_CLASS_NoVars_Const   < UserClass, FUNC_TEMPLATE_ARGS,                                  ESPMode::ThreadSafe > { public: TThreadSafeSPMethodDelegate_Const          ( const TSharedRef< UserClass, ESPMode::ThreadSafe >& InUserObject, typename SP_METHOD_DELEGATE_INSTANCE_CLASS_NoVars_Const   < UserClass, FUNC_TEMPLATE_ARGS,                                  ESPMode::ThreadSafe >::FMethodPtr InMethodPtr                                  ) : SP_METHOD_DELEGATE_INSTANCE_CLASS_NoVars_Const   < UserClass, FUNC_TEMPLATE_ARGS,                                  ESPMode::ThreadSafe >( InUserObject, InMethodPtr                                    ) {} };
 	template< class UserClass, PAYLOAD_TEMPLATE_DECL_OneVar    > class TThreadSafeSPMethodDelegate_OneVar          : public SP_METHOD_DELEGATE_INSTANCE_CLASS_OneVar         < UserClass, FUNC_TEMPLATE_ARGS, PAYLOAD_TEMPLATE_LIST_OneVar,    ESPMode::ThreadSafe > { public: TThreadSafeSPMethodDelegate_OneVar         ( const TSharedRef< UserClass, ESPMode::ThreadSafe >& InUserObject, typename SP_METHOD_DELEGATE_INSTANCE_CLASS_OneVar         < UserClass, FUNC_TEMPLATE_ARGS, PAYLOAD_TEMPLATE_LIST_OneVar,    ESPMode::ThreadSafe >::FMethodPtr InMethodPtr, PAYLOAD_TEMPLATE_ARGS_OneVar    ) : SP_METHOD_DELEGATE_INSTANCE_CLASS_OneVar         < UserClass, FUNC_TEMPLATE_ARGS, PAYLOAD_TEMPLATE_LIST_OneVar,    ESPMode::ThreadSafe >( InUserObject, InMethodPtr, PAYLOAD_TEMPLATE_PASSIN_OneVar    ) {} };
@@ -123,7 +130,7 @@ public:
 	template< class UserClass, PAYLOAD_TEMPLATE_DECL_FourVars  > class TThreadSafeSPMethodDelegate_FourVars        : public SP_METHOD_DELEGATE_INSTANCE_CLASS_FourVars       < UserClass, FUNC_TEMPLATE_ARGS, PAYLOAD_TEMPLATE_LIST_FourVars,  ESPMode::ThreadSafe > { public: TThreadSafeSPMethodDelegate_FourVars       ( const TSharedRef< UserClass, ESPMode::ThreadSafe >& InUserObject, typename SP_METHOD_DELEGATE_INSTANCE_CLASS_FourVars       < UserClass, FUNC_TEMPLATE_ARGS, PAYLOAD_TEMPLATE_LIST_FourVars,  ESPMode::ThreadSafe >::FMethodPtr InMethodPtr, PAYLOAD_TEMPLATE_ARGS_FourVars  ) : SP_METHOD_DELEGATE_INSTANCE_CLASS_FourVars       < UserClass, FUNC_TEMPLATE_ARGS, PAYLOAD_TEMPLATE_LIST_FourVars,  ESPMode::ThreadSafe >( InUserObject, InMethodPtr, PAYLOAD_TEMPLATE_PASSIN_FourVars  ) {} };
 	template< class UserClass, PAYLOAD_TEMPLATE_DECL_FourVars  > class TThreadSafeSPMethodDelegate_FourVars_Const  : public SP_METHOD_DELEGATE_INSTANCE_CLASS_FourVars_Const < UserClass, FUNC_TEMPLATE_ARGS, PAYLOAD_TEMPLATE_LIST_FourVars,  ESPMode::ThreadSafe > { public: TThreadSafeSPMethodDelegate_FourVars_Const ( const TSharedRef< UserClass, ESPMode::ThreadSafe >& InUserObject, typename SP_METHOD_DELEGATE_INSTANCE_CLASS_FourVars_Const < UserClass, FUNC_TEMPLATE_ARGS, PAYLOAD_TEMPLATE_LIST_FourVars,  ESPMode::ThreadSafe >::FMethodPtr InMethodPtr, PAYLOAD_TEMPLATE_ARGS_FourVars  ) : SP_METHOD_DELEGATE_INSTANCE_CLASS_FourVars_Const < UserClass, FUNC_TEMPLATE_ARGS, PAYLOAD_TEMPLATE_LIST_FourVars,  ESPMode::ThreadSafe >( InUserObject, InMethodPtr, PAYLOAD_TEMPLATE_PASSIN_FourVars  ) {} };
 
-	/** Declare the user's C++ pointer-based delegate object (e.g. TRawMyDelegate) */
+	/** Declare the user's C++ pointer-based delegate instance types. */
 	template< class UserClass                                  > class TRawMethodDelegate                 : public RAW_METHOD_DELEGATE_INSTANCE_CLASS_NoVars         < UserClass, FUNC_TEMPLATE_ARGS                                  > { public: TRawMethodDelegate                ( UserClass* InUserObject, typename RAW_METHOD_DELEGATE_INSTANCE_CLASS_NoVars         < UserClass, FUNC_TEMPLATE_ARGS                                  >::FMethodPtr InMethodPtr                                  ) : RAW_METHOD_DELEGATE_INSTANCE_CLASS_NoVars         < UserClass, FUNC_TEMPLATE_ARGS                                  >( InUserObject, InMethodPtr                                    ) {} };
 	template< class UserClass                                  > class TRawMethodDelegate_Const           : public RAW_METHOD_DELEGATE_INSTANCE_CLASS_NoVars_Const   < UserClass, FUNC_TEMPLATE_ARGS                                  > { public: TRawMethodDelegate_Const          ( UserClass* InUserObject, typename RAW_METHOD_DELEGATE_INSTANCE_CLASS_NoVars_Const   < UserClass, FUNC_TEMPLATE_ARGS                                  >::FMethodPtr InMethodPtr                                  ) : RAW_METHOD_DELEGATE_INSTANCE_CLASS_NoVars_Const   < UserClass, FUNC_TEMPLATE_ARGS                                  >( InUserObject, InMethodPtr                                    ) {} };
 	template< class UserClass, PAYLOAD_TEMPLATE_DECL_OneVar    > class TRawMethodDelegate_OneVar          : public RAW_METHOD_DELEGATE_INSTANCE_CLASS_OneVar         < UserClass, FUNC_TEMPLATE_ARGS, PAYLOAD_TEMPLATE_LIST_OneVar    > { public: TRawMethodDelegate_OneVar         ( UserClass* InUserObject, typename RAW_METHOD_DELEGATE_INSTANCE_CLASS_OneVar         < UserClass, FUNC_TEMPLATE_ARGS, PAYLOAD_TEMPLATE_LIST_OneVar    >::FMethodPtr InMethodPtr, PAYLOAD_TEMPLATE_ARGS_OneVar    ) : RAW_METHOD_DELEGATE_INSTANCE_CLASS_OneVar         < UserClass, FUNC_TEMPLATE_ARGS, PAYLOAD_TEMPLATE_LIST_OneVar    >( InUserObject, InMethodPtr, PAYLOAD_TEMPLATE_PASSIN_OneVar    ) {} };
@@ -135,7 +142,14 @@ public:
 	template< class UserClass, PAYLOAD_TEMPLATE_DECL_FourVars  > class TRawMethodDelegate_FourVars        : public RAW_METHOD_DELEGATE_INSTANCE_CLASS_FourVars       < UserClass, FUNC_TEMPLATE_ARGS, PAYLOAD_TEMPLATE_LIST_FourVars  > { public: TRawMethodDelegate_FourVars       ( UserClass* InUserObject, typename RAW_METHOD_DELEGATE_INSTANCE_CLASS_FourVars       < UserClass, FUNC_TEMPLATE_ARGS, PAYLOAD_TEMPLATE_LIST_FourVars  >::FMethodPtr InMethodPtr, PAYLOAD_TEMPLATE_ARGS_FourVars  ) : RAW_METHOD_DELEGATE_INSTANCE_CLASS_FourVars       < UserClass, FUNC_TEMPLATE_ARGS, PAYLOAD_TEMPLATE_LIST_FourVars  >( InUserObject, InMethodPtr, PAYLOAD_TEMPLATE_PASSIN_FourVars  ) {} };
 	template< class UserClass, PAYLOAD_TEMPLATE_DECL_FourVars  > class TRawMethodDelegate_FourVars_Const  : public RAW_METHOD_DELEGATE_INSTANCE_CLASS_FourVars_Const < UserClass, FUNC_TEMPLATE_ARGS, PAYLOAD_TEMPLATE_LIST_FourVars  > { public: TRawMethodDelegate_FourVars_Const ( UserClass* InUserObject, typename RAW_METHOD_DELEGATE_INSTANCE_CLASS_FourVars_Const < UserClass, FUNC_TEMPLATE_ARGS, PAYLOAD_TEMPLATE_LIST_FourVars  >::FMethodPtr InMethodPtr, PAYLOAD_TEMPLATE_ARGS_FourVars  ) : RAW_METHOD_DELEGATE_INSTANCE_CLASS_FourVars_Const < UserClass, FUNC_TEMPLATE_ARGS, PAYLOAD_TEMPLATE_LIST_FourVars  >( InUserObject, InMethodPtr, PAYLOAD_TEMPLATE_PASSIN_FourVars  ) {} };
 	
-	/** Declare the user's UObject-based delegate object (e.g. TObjMyDelegate) */
+	/** Declare the user's UFunction-based delegate instance types. */
+	template< class UObjectTemplate                                  > class TUFunctionDelegateBinding                     : public UFUNCTION_DELEGATE_INSTANCE_CLASS_NoVars   < UObjectTemplate, FUNC_TEMPLATE_ARGS                                  > { public: TUFunctionDelegateBinding          ( UObjectTemplate* InUserObject, const FName& InFunctionName                                  ) : UFUNCTION_DELEGATE_INSTANCE_CLASS_NoVars   < UObjectTemplate, FUNC_TEMPLATE_ARGS                                  >( InUserObject, InFunctionName                                    ) {} };
+	template< class UObjectTemplate, PAYLOAD_TEMPLATE_DECL_OneVar    > class TUFunctionDelegateBinding_OneVar              : public UFUNCTION_DELEGATE_INSTANCE_CLASS_OneVar   < UObjectTemplate, FUNC_TEMPLATE_ARGS, PAYLOAD_TEMPLATE_LIST_OneVar    > { public: TUFunctionDelegateBinding_OneVar   ( UObjectTemplate* InUserObject, const FName& InFunctionName, PAYLOAD_TEMPLATE_ARGS_OneVar    ) : UFUNCTION_DELEGATE_INSTANCE_CLASS_OneVar   < UObjectTemplate, FUNC_TEMPLATE_ARGS, PAYLOAD_TEMPLATE_LIST_OneVar    >( InUserObject, InFunctionName, PAYLOAD_TEMPLATE_PASSIN_OneVar    ) {} };
+	template< class UObjectTemplate, PAYLOAD_TEMPLATE_DECL_TwoVars   > class TUFunctionDelegateBinding_TwoVars             : public UFUNCTION_DELEGATE_INSTANCE_CLASS_TwoVars  < UObjectTemplate, FUNC_TEMPLATE_ARGS, PAYLOAD_TEMPLATE_LIST_TwoVars   > { public: TUFunctionDelegateBinding_TwoVars  ( UObjectTemplate* InUserObject, const FName& InFunctionName, PAYLOAD_TEMPLATE_ARGS_TwoVars   ) : UFUNCTION_DELEGATE_INSTANCE_CLASS_TwoVars  < UObjectTemplate, FUNC_TEMPLATE_ARGS, PAYLOAD_TEMPLATE_LIST_TwoVars   >( InUserObject, InFunctionName, PAYLOAD_TEMPLATE_PASSIN_TwoVars   ) {} };
+	template< class UObjectTemplate, PAYLOAD_TEMPLATE_DECL_ThreeVars > class TUFunctionDelegateBinding_ThreeVars           : public UFUNCTION_DELEGATE_INSTANCE_CLASS_ThreeVars< UObjectTemplate, FUNC_TEMPLATE_ARGS, PAYLOAD_TEMPLATE_LIST_ThreeVars > { public: TUFunctionDelegateBinding_ThreeVars( UObjectTemplate* InUserObject, const FName& InFunctionName, PAYLOAD_TEMPLATE_ARGS_ThreeVars ) : UFUNCTION_DELEGATE_INSTANCE_CLASS_ThreeVars< UObjectTemplate, FUNC_TEMPLATE_ARGS, PAYLOAD_TEMPLATE_LIST_ThreeVars >( InUserObject, InFunctionName, PAYLOAD_TEMPLATE_PASSIN_ThreeVars ) {} };
+	template< class UObjectTemplate, PAYLOAD_TEMPLATE_DECL_FourVars  > class TUFunctionDelegateBinding_FourVars            : public UFUNCTION_DELEGATE_INSTANCE_CLASS_FourVars < UObjectTemplate, FUNC_TEMPLATE_ARGS, PAYLOAD_TEMPLATE_LIST_FourVars  > { public: TUFunctionDelegateBinding_FourVars ( UObjectTemplate* InUserObject, const FName& InFunctionName, PAYLOAD_TEMPLATE_ARGS_FourVars  ) : UFUNCTION_DELEGATE_INSTANCE_CLASS_FourVars < UObjectTemplate, FUNC_TEMPLATE_ARGS, PAYLOAD_TEMPLATE_LIST_FourVars  >( InUserObject, InFunctionName, PAYLOAD_TEMPLATE_PASSIN_FourVars  ) {} };
+
+	/** Declare the user's UObject-based delegate instance types. */
 	template< class UserClass                                  > class TUObjectMethodDelegate                 : public UOBJECT_METHOD_DELEGATE_INSTANCE_CLASS_NoVars         < UserClass, FUNC_TEMPLATE_ARGS                                  > { public: TUObjectMethodDelegate                ( UserClass* InUserObject, typename UOBJECT_METHOD_DELEGATE_INSTANCE_CLASS_NoVars         < UserClass, FUNC_TEMPLATE_ARGS                                  >::FMethodPtr InMethodPtr                                  ) : UOBJECT_METHOD_DELEGATE_INSTANCE_CLASS_NoVars         < UserClass, FUNC_TEMPLATE_ARGS                                  >( InUserObject, InMethodPtr                                    ) {} };
 	template< class UserClass                                  > class TUObjectMethodDelegate_Const           : public UOBJECT_METHOD_DELEGATE_INSTANCE_CLASS_NoVars_Const   < UserClass, FUNC_TEMPLATE_ARGS                                  > { public: TUObjectMethodDelegate_Const          ( UserClass* InUserObject, typename UOBJECT_METHOD_DELEGATE_INSTANCE_CLASS_NoVars_Const   < UserClass, FUNC_TEMPLATE_ARGS                                  >::FMethodPtr InMethodPtr                                  ) : UOBJECT_METHOD_DELEGATE_INSTANCE_CLASS_NoVars_Const   < UserClass, FUNC_TEMPLATE_ARGS                                  >( InUserObject, InMethodPtr                                    ) {} };
 	template< class UserClass, PAYLOAD_TEMPLATE_DECL_OneVar    > class TUObjectMethodDelegate_OneVar          : public UOBJECT_METHOD_DELEGATE_INSTANCE_CLASS_OneVar         < UserClass, FUNC_TEMPLATE_ARGS, PAYLOAD_TEMPLATE_LIST_OneVar    > { public: TUObjectMethodDelegate_OneVar         ( UserClass* InUserObject, typename UOBJECT_METHOD_DELEGATE_INSTANCE_CLASS_OneVar         < UserClass, FUNC_TEMPLATE_ARGS, PAYLOAD_TEMPLATE_LIST_OneVar    >::FMethodPtr InMethodPtr, PAYLOAD_TEMPLATE_ARGS_OneVar    ) : UOBJECT_METHOD_DELEGATE_INSTANCE_CLASS_OneVar         < UserClass, FUNC_TEMPLATE_ARGS, PAYLOAD_TEMPLATE_LIST_OneVar    >( InUserObject, InMethodPtr, PAYLOAD_TEMPLATE_PASSIN_OneVar    ) {} };
@@ -147,7 +161,7 @@ public:
 	template< class UserClass, PAYLOAD_TEMPLATE_DECL_FourVars  > class TUObjectMethodDelegate_FourVars        : public UOBJECT_METHOD_DELEGATE_INSTANCE_CLASS_FourVars       < UserClass, FUNC_TEMPLATE_ARGS, PAYLOAD_TEMPLATE_LIST_FourVars  > { public: TUObjectMethodDelegate_FourVars       ( UserClass* InUserObject, typename UOBJECT_METHOD_DELEGATE_INSTANCE_CLASS_FourVars       < UserClass, FUNC_TEMPLATE_ARGS, PAYLOAD_TEMPLATE_LIST_FourVars  >::FMethodPtr InMethodPtr, PAYLOAD_TEMPLATE_ARGS_FourVars  ) : UOBJECT_METHOD_DELEGATE_INSTANCE_CLASS_FourVars       < UserClass, FUNC_TEMPLATE_ARGS, PAYLOAD_TEMPLATE_LIST_FourVars  >( InUserObject, InMethodPtr, PAYLOAD_TEMPLATE_PASSIN_FourVars  ) {} };
 	template< class UserClass, PAYLOAD_TEMPLATE_DECL_FourVars  > class TUObjectMethodDelegate_FourVars_Const  : public UOBJECT_METHOD_DELEGATE_INSTANCE_CLASS_FourVars_Const < UserClass, FUNC_TEMPLATE_ARGS, PAYLOAD_TEMPLATE_LIST_FourVars  > { public: TUObjectMethodDelegate_FourVars_Const ( UserClass* InUserObject, typename UOBJECT_METHOD_DELEGATE_INSTANCE_CLASS_FourVars_Const < UserClass, FUNC_TEMPLATE_ARGS, PAYLOAD_TEMPLATE_LIST_FourVars  >::FMethodPtr InMethodPtr, PAYLOAD_TEMPLATE_ARGS_FourVars  ) : UOBJECT_METHOD_DELEGATE_INSTANCE_CLASS_FourVars_Const < UserClass, FUNC_TEMPLATE_ARGS, PAYLOAD_TEMPLATE_LIST_FourVars  >( InUserObject, InMethodPtr, PAYLOAD_TEMPLATE_PASSIN_FourVars  ) {} };
 	
-	/** Declare the user's static function pointer delegate object (e.g. FMyDelegate) */
+	/** Declare the user's static function pointer delegate instance types. */
 	typedef STATIC_DELEGATE_INSTANCE_CLASS_NoVars< FUNC_TEMPLATE_ARGS > FStaticDelegate;
 	template< PAYLOAD_TEMPLATE_DECL_OneVar    > class TStaticDelegate_OneVar    : public STATIC_DELEGATE_INSTANCE_CLASS_OneVar   < FUNC_TEMPLATE_ARGS, PAYLOAD_TEMPLATE_LIST_OneVar    > { public: TStaticDelegate_OneVar   ( typename STATIC_DELEGATE_INSTANCE_CLASS_OneVar   < FUNC_TEMPLATE_ARGS, PAYLOAD_TEMPLATE_LIST_OneVar    >::FFuncPtr InFuncPtr, PAYLOAD_TEMPLATE_ARGS_OneVar    ) : STATIC_DELEGATE_INSTANCE_CLASS_OneVar   < FUNC_TEMPLATE_ARGS, PAYLOAD_TEMPLATE_LIST_OneVar    >( InFuncPtr, PAYLOAD_TEMPLATE_PASSIN_OneVar    ) {} };
 	template< PAYLOAD_TEMPLATE_DECL_TwoVars   > class TStaticDelegate_TwoVars   : public STATIC_DELEGATE_INSTANCE_CLASS_TwoVars  < FUNC_TEMPLATE_ARGS, PAYLOAD_TEMPLATE_LIST_TwoVars   > { public: TStaticDelegate_TwoVars  ( typename STATIC_DELEGATE_INSTANCE_CLASS_TwoVars  < FUNC_TEMPLATE_ARGS, PAYLOAD_TEMPLATE_LIST_TwoVars   >::FFuncPtr InFuncPtr, PAYLOAD_TEMPLATE_ARGS_TwoVars   ) : STATIC_DELEGATE_INSTANCE_CLASS_TwoVars  < FUNC_TEMPLATE_ARGS, PAYLOAD_TEMPLATE_LIST_TwoVars   >( InFuncPtr, PAYLOAD_TEMPLATE_PASSIN_TwoVars   ) {} };
@@ -469,6 +483,38 @@ public:
 		return CreateThreadSafeSP( StaticCastSharedRef< UserClass >( InUserObject->AsShared() ), InFunc, PAYLOAD_TEMPLATE_PASSIN_FourVars );
 	}
 
+	/**
+	 * Static: Creates a UFunction-based member function delegate.
+	 *
+	 * UFunction delegates keep a weak reference to your object.
+	 * You can use ExecuteIfBound() to call them.
+	 */
+	template< class UObjectTemplate >
+	inline static DELEGATE_CLASS< FUNC_TEMPLATE_ARGS > CreateUFunction( UObjectTemplate* InUserObject, const FName& InFunctionName )
+	{
+		return DELEGATE_CLASS< FUNC_TEMPLATE_ARGS >( TUFunctionDelegateBinding< UObjectTemplate >::Create( InUserObject, InFunctionName ) );
+	}
+	template< class UObjectTemplate, PAYLOAD_TEMPLATE_DECL_OneVar >
+	inline static DELEGATE_CLASS< FUNC_TEMPLATE_ARGS > CreateUFunction( UObjectTemplate* InUserObject, const FName& InFunctionName, PAYLOAD_TEMPLATE_ARGS_OneVar )
+	{
+		return DELEGATE_CLASS< FUNC_TEMPLATE_ARGS >( TUFunctionDelegateBinding_OneVar< UObjectTemplate, PAYLOAD_TEMPLATE_LIST_OneVar >::Create( InUserObject, InFunctionName, PAYLOAD_TEMPLATE_PASSIN_OneVar ) );
+	}
+	template< class UObjectTemplate, PAYLOAD_TEMPLATE_DECL_TwoVars >
+	inline static DELEGATE_CLASS< FUNC_TEMPLATE_ARGS > CreateUFunction( UObjectTemplate* InUserObject, const FName& InFunctionName, PAYLOAD_TEMPLATE_ARGS_TwoVars )
+	{
+		return DELEGATE_CLASS< FUNC_TEMPLATE_ARGS >( TUFunctionDelegateBinding_TwoVars< UObjectTemplate, PAYLOAD_TEMPLATE_LIST_TwoVars >::Create( InUserObject, InFunctionName, PAYLOAD_TEMPLATE_PASSIN_TwoVars ) );
+	}
+	template< class UObjectTemplate, PAYLOAD_TEMPLATE_DECL_ThreeVars >
+	inline static DELEGATE_CLASS< FUNC_TEMPLATE_ARGS > CreateUFunction( UObjectTemplate* InUserObject, const FName& InFunctionName, PAYLOAD_TEMPLATE_ARGS_ThreeVars )
+	{
+		return DELEGATE_CLASS< FUNC_TEMPLATE_ARGS >( TUFunctionDelegateBinding_ThreeVars< UObjectTemplate, PAYLOAD_TEMPLATE_LIST_ThreeVars >::Create( InUserObject, InFunctionName, PAYLOAD_TEMPLATE_PASSIN_ThreeVars ) );
+	}
+	template< class UObjectTemplate, PAYLOAD_TEMPLATE_DECL_FourVars >
+	inline static DELEGATE_CLASS< FUNC_TEMPLATE_ARGS > CreateUFunction( UObjectTemplate* InUserObject, const FName& InFunctionName, PAYLOAD_TEMPLATE_ARGS_FourVars )
+	{
+		return DELEGATE_CLASS< FUNC_TEMPLATE_ARGS >( TUFunctionDelegateBinding_FourVars< UObjectTemplate, PAYLOAD_TEMPLATE_LIST_FourVars >::Create( InUserObject, InFunctionName, PAYLOAD_TEMPLATE_PASSIN_FourVars ) );
+	}
+
 	/**
 	 * Static: Creates a UObject-based member function delegate.
 	 *
@@ -532,11 +578,11 @@ public:
 	 * Default constructor
 	 */
 	inline DELEGATE_CLASS()
-		: Delegate( NULL )
+		: FDelegateBase(nullptr)
 	{ }
 
 	/**
-	 * Delegate instance destructor
+	 * Destructor.
 	 */
 	inline ~DELEGATE_CLASS()
 	{
@@ -544,41 +590,48 @@ public:
 	}
 
 	/**
-	 * Construct from a delegate instance pointer.  The delegate will assume ownership of the incoming delegate instance!
+	 * Creates and initializes a new instance with the given delegate instance.
+	 *
+	 * The delegate will assume ownership of the incoming delegate instance!
 	 * IMPORTANT: This is a system-internal function and you should never be using this in regular C++ code
+	 *
+	 * @param InDelegateInstance The delegate instance to assign.
 	 */
-	inline DELEGATE_CLASS( IDelegate* InitDelegate )
-		: Delegate( InitDelegate )
+	inline DELEGATE_CLASS( TDelegateInstanceInterface* InDelegateInstance )
+		: FDelegateBase(InDelegateInstance)
 	{ }
 
 	/**
-	 * Construct from existing delegate object
+	 * Creates and initializes a new instance from an existing delegate object.
 	 *
-	 * @param	OtherDelegate	Delegate object to copy from
+	 * @param Other The delegate object to copy from.
 	 */
-	inline DELEGATE_CLASS( const DELEGATE_CLASS& OtherDelegate )
-		: Delegate( NULL )
+	inline DELEGATE_CLASS( const DELEGATE_CLASS& Other )
+		: FDelegateBase(nullptr)
 	{
-		*this = OtherDelegate;
+		*this = Other;
 	}
 
 	/**
-	 * Assignment operator
+	 * Assignment operator.
 	 *
 	 * @param	OtherDelegate	Delegate object to copy from
 	 */
-	inline DELEGATE_CLASS& operator=( const DELEGATE_CLASS& OtherDelegate )
+	inline DELEGATE_CLASS& operator=( const DELEGATE_CLASS& Other )
 	{
-		if( &OtherDelegate != this )
+		if (&Other != this)
 		{
-			IDelegate* NewDelegate = NULL;
-			if( OtherDelegate.Delegate != NULL )
+			// this down-cast is OK! allows for managing invocation list in the base class without requiring virtual functions
+			TDelegateInstanceInterface* OtherInstance = (TDelegateInstanceInterface*)Other.GetDelegateInstance();
+
+			if (OtherInstance != nullptr)
 			{
-				NewDelegate = OtherDelegate.Delegate->CreateCopy();
+				SetDelegateInstance(OtherInstance->CreateCopy());
+			}
+			else
+			{
+				SetDelegateInstance(nullptr);
 			}
-
-			Unbind();
-			Delegate = NewDelegate;
 		}
 
 		return *this;
@@ -896,6 +949,38 @@ public:
 		*this = CreateThreadSafeSP( InUserObject, InFunc, PAYLOAD_TEMPLATE_PASSIN_FourVars );
 	}
 
+	/**
+	 * Creates a UFunction-based member function delegate.
+	 *
+	 * UFunction delegates keep a weak reference to your object.
+	 * You can use ExecuteIfBound() to call them.
+	 */
+	template< class UObjectTemplate >
+	inline void BindUFunction( UObjectTemplate* InUserObject, const FName& InFunctionName )
+	{
+		*this = CreateUFunction( InUserObject, InFunctionName );
+	}
+	template< class UObjectTemplate, PAYLOAD_TEMPLATE_DECL_OneVar >
+	inline void BindUFunction( UObjectTemplate* InUserObject, const FName& InFunctionName, PAYLOAD_TEMPLATE_ARGS_OneVar )
+	{
+		*this = CreateUFunction( InUserObject, InFunctionName, PAYLOAD_TEMPLATE_PASSIN_OneVar );
+	}
+	template< class UObjectTemplate, PAYLOAD_TEMPLATE_DECL_TwoVars >
+	inline void BindUFunction( UObjectTemplate* InUserObject, const FName& InFunctionName, PAYLOAD_TEMPLATE_ARGS_TwoVars )
+	{
+		*this = CreateUFunction( InUserObject, InFunctionName, PAYLOAD_TEMPLATE_PASSIN_TwoVars );
+	}
+	template< class UObjectTemplate, PAYLOAD_TEMPLATE_DECL_ThreeVars >
+	inline void BindUFunction( UObjectTemplate* InUserObject, const FName& InFunctionName, PAYLOAD_TEMPLATE_ARGS_ThreeVars )
+	{
+		*this = CreateUFunction( InUserObject, InFunctionName, PAYLOAD_TEMPLATE_PASSIN_ThreeVars );
+	}
+	template< class UObjectTemplate, PAYLOAD_TEMPLATE_DECL_FourVars >
+	inline void BindUFunction( UObjectTemplate* InUserObject, const FName& InFunctionName, PAYLOAD_TEMPLATE_ARGS_FourVars )
+	{
+		*this = CreateUFunction( InUserObject, InFunctionName, PAYLOAD_TEMPLATE_PASSIN_FourVars );
+	}
+
 	/**
 	 * Creates a UObject-based member function delegate.
 	 *
@@ -952,72 +1037,25 @@ public:
 	{
 		*this = CreateUObject( InUserObject, InFunc, PAYLOAD_TEMPLATE_PASSIN_FourVars );
 	}
+public:
 
 	/**
-	 * Unbinds this delegate
-	 */
-	inline void Unbind()
-	{
-		if( Delegate != NULL )
-		{
-			delete Delegate;
-			Delegate = NULL;
-		}
-	}
-
-	/**
-	 * Checks to see if the user object bound to this delegate is still valid
+	 * Execute the delegate.
 	 *
-	 * @return  True if the user object is still valid and it's safe to execute the function call
-	 */
-	inline bool IsBound() const
-	{
-		return Delegate != NULL && Delegate->IsSafeToExecute();
-	}
-
-	/** 
-	 * Checks to see if this delegate is bound to the given user object.
-	 *
-	 * @return	True if this delegate is bound to InUserObject, false otherwise.
-	 */
-	inline bool IsBoundToObject(void const* InUserObject) const
-	{
-		return InUserObject && Delegate != NULL && Delegate->HasSameObject(InUserObject);
-	}
-
-	/**
-	 * Checks to see if the user object bound to this delegate will ever be valid again
+	 * If the function pointer is not valid, an error will occur. Check IsBound() before
+	 * calling this method or use ExecuteIfBound() instead.
 	 *
-	 * @return  True if the user object could be valid in the future or is valid now.
-	 */
-	inline bool IsCompactable() const
-	{
-		return Delegate != NULL && Delegate->IsCompactable();
-	}
-
-	/**
-	 * Execute the delegate.  If the function pointer is not valid, an error will occur.
+	 * @see ExecuteIfBound
 	 */
 	inline RetValType Execute( FUNC_PARAM_LIST ) const
 	{
-		// If this assert goes off, Execute() was called before a function was bound to the delegate.  Consider using ExecuteIfSafe() instead.
-		checkSlow( Delegate != NULL );
-		return Delegate->Execute( FUNC_PARAM_PASSTHRU );
-	}
+		TDelegateInstanceInterface* DelegateInstance = (TDelegateInstanceInterface*)GetDelegateInstance();
 
-	/**
-	 * If this is a UObject delegate, return the UObject
-	 *
-	 * @return  the object associated with this delegate if there is one
-	 */
-	inline class UObject* GetUObject() const
-	{
-		if (Delegate && Delegate->GetType() == EDelegateInstanceType::UObjectMethod)
-		{
-			return (class UObject*)Delegate->GetRawUserObject();
-		}
+		// If this assert goes off, Execute() was called before a function was bound to the delegate.
+		// Consider using ExecuteIfSafe() instead.
+		checkSlow(DelegateInstance != nullptr);
 
-		return NULL;
+		return DelegateInstance->Execute(FUNC_PARAM_PASSTHRU);
 	}
 
 #if FUNC_IS_VOID
@@ -1031,7 +1069,7 @@ public:
 	{
 		if( IsBound() )
 		{
-			return Delegate->ExecuteIfSafe( FUNC_PARAM_PASSTHRU );
+			return ((TDelegateInstanceInterface*)GetDelegateInstance())->ExecuteIfSafe( FUNC_PARAM_PASSTHRU );
 		}
 
 		return false;
@@ -1039,19 +1077,23 @@ public:
 #endif
 
 	/**
-	 * Equality operator
+	 * Equality operator.
+	 *
+	 * @return true if this delegate equals the other, false otherwise.
 	 */
 	bool operator==( const DELEGATE_CLASS< FUNC_TEMPLATE_ARGS >& Other ) const
 	{
+		TDelegateInstanceInterface* DelegateInstance = (TDelegateInstanceInterface*)GetDelegateInstance();
+		TDelegateInstanceInterface* OtherInstance = (TDelegateInstanceInterface*)Other.GetDelegateInstance();
+
 		// The function these delegates point to must be the same
-		if( Delegate != NULL && Other.Delegate != NULL )
+		if ((DelegateInstance != nullptr) && (OtherInstance != nullptr))
 		{
-			return Delegate->IsSameFunction( *Other.Delegate );
+			return DelegateInstance->IsSameFunction(*OtherInstance);
 		}
-		else
-		{
+
 			// If neither delegate is initialized to anything yet, then we treat them as equal
-			if( Delegate == NULL && Other.Delegate == NULL )
+		if ((DelegateInstance == nullptr) && (OtherInstance == nullptr))
 			{
 				return true;
 			}
@@ -1059,48 +1101,71 @@ public:
 			// No match!
 			return false;
 		}
-	}
 
 private:
 
 	// Declare ourselves as a friend to our multi-cast counterpart, so that it can access our interface directly
 	template< FUNC_TEMPLATE_DECL_NO_SHADOW > friend class BASE_MULTICAST_DELEGATE_CLASS;
-
-	/** Pointer to the actual delegate object instance.  We need to use a pointer because the actual type of delegate class
-	    that's created will depend on the type of binding (raw, shared pointer, etc) as well as the user's object type */
-	IDelegate* Delegate;
 };
 
 
-
-// NOTE: Only delegates with no return value are valid for multi-cast
 #if FUNC_IS_VOID
 
 /**
- * Multi-cast delegate base object.   You'll use the various DECLARE_MULTICAST_DELEGATE macros to create
- * the actual delegate type, templated to the function signature the delegate is compatible with.  Then,
- * you can create an instance of that class when you want to bind a function to the delegate.
+ * Multicast delegate base class.
+	 *
+ * This class implements the functionality of multicast delegates. It is templated to the function signature
+ * that it is compatible with. Use the various DECLARE_MULTICAST_DELEGATE and DECLARE_EVENT macros to create
+ * actual delegate types.
+	 *
+ * Multicast delegates only support bindings with no return value.
  */
 template< FUNC_TEMPLATE_DECL_TYPENAME >
 class BASE_MULTICAST_DELEGATE_CLASS
+	: public FMulticastDelegateBase<>
 {
 public:
 
-	/** The actual single-cast delegate class for this multi-cast delegate */
+	/** DEPRECATED: Type definition for unicast delegate classes whose delegate instances are compatible with this delegate. */
 	typedef DELEGATE_CLASS< FUNC_TEMPLATE_ARGS > FDelegate;
 
+	/** Type definition for the shared interface of delegate instance types compatible with this delegate class. */
+	typedef DELEGATE_INSTANCE_INTERFACE_CLASS<FUNC_TEMPLATE_ARGS> TDelegateInstanceInterface;
+
+public:
+
 	/**
-	 * Adds a function delegate to this multi-cast delegate's invocation list
+	 * Adds a delegate instance to this multicast delegate's invocation list.
 	 *
-	 * @param	InDelegate	Delegate to add
+	 * This delegate will take over ownership of the given delegate instance.
+	 *
+	 * @param DelegateInstance The delegate instance to add.
 	 */
-	void Add( const FDelegate& InDelegate )
+	void Add( TDelegateInstanceInterface* DelegateInstance )
 	{
-		// Add the delegate
-		AddInternal( InDelegate );
+		if (DelegateInstance != nullptr)
+		{
+			AddDelegateInstance(*DelegateInstance);
+			CompactInvocationList();
+		}
+	}
 
-		// Then check and see if there are any stale delegates and remove them
-		CompactInvocationList();
+	/**
+	 * DEPRECATED: Adds a unicast delegate to this multi-cast delegate's invocation list.
+	 *
+	 * This method is retained for backwards compatibility.
+	 *
+	 * @param Delegate The delegate to add.
+	 */
+	void Add( const FDelegate& Delegate )
+	{
+		TDelegateInstanceInterface* DelegateInstance = (TDelegateInstanceInterface*)Delegate.GetDelegateInstance();
+
+		if (DelegateInstance != nullptr)
+		{
+			AddDelegateInstance(*DelegateInstance->CreateCopy());
+			CompactInvocationList();
+		}
 	}
 
 	/**
@@ -1427,6 +1492,40 @@ public:
 		Add( FDelegate::CreateThreadSafeSP( InUserObject, InFunc, PAYLOAD_TEMPLATE_PASSIN_FourVars ) );
 	}
 
+	/**
+	 * Adds a UFunction-based member function delegate.
+	 *
+	 * UFunction delegates keep a weak reference to your object.
+	 *
+	 * @param	InUserObject	User object to bind to
+	 * @param	InFunctionName			Class method function address
+	 */
+	template< class UObjectTemplate >
+	inline void AddUFunction( UObjectTemplate* InUserObject, const FName& InFunctionName )
+	{
+		Add( FDelegate::CreateUFunction( InUserObject, InFunctionName ) );
+	}
+	template< class UObjectTemplate, PAYLOAD_TEMPLATE_DECL_OneVar >
+	inline void AddUFunction( UObjectTemplate* InUserObject, const FName& InFunctionName, PAYLOAD_TEMPLATE_ARGS_OneVar )
+	{
+		Add( FDelegate::CreateUFunction( InUserObject, InFunctionName, PAYLOAD_TEMPLATE_PASSIN_OneVar ) );
+	}
+	template< class UObjectTemplate, PAYLOAD_TEMPLATE_DECL_TwoVars >
+	inline void AddUFunction( UObjectTemplate* InUserObject, const FName& InFunctionName, PAYLOAD_TEMPLATE_ARGS_TwoVars )
+	{
+		Add( FDelegate::CreateUFunction( InUserObject, InFunctionName, PAYLOAD_TEMPLATE_PASSIN_TwoVars ) );
+	}
+	template< class UObjectTemplate, PAYLOAD_TEMPLATE_DECL_ThreeVars >
+	inline void AddUFunction( UObjectTemplate* InUserObject, const FName& InFunctionName, PAYLOAD_TEMPLATE_ARGS_ThreeVars )
+	{
+		Add( FDelegate::CreateUFunction( InUserObject, InFunctionName, PAYLOAD_TEMPLATE_PASSIN_ThreeVars ) );
+	}
+	template< class UObjectTemplate, PAYLOAD_TEMPLATE_DECL_FourVars >
+	inline void AddUFunction( UObjectTemplate* InUserObject, const FName& InFunctionName, PAYLOAD_TEMPLATE_ARGS_FourVars )
+	{
+		Add( FDelegate::CreateUFunction( InUserObject, InFunctionName, PAYLOAD_TEMPLATE_PASSIN_FourVars ) );
+	}
+
 	/**
 	 * Adds a UObject-based member function delegate.
 	 *
@@ -1486,37 +1585,48 @@ public:
 		Add( FDelegate::CreateUObject( InUserObject, InFunc, PAYLOAD_TEMPLATE_PASSIN_FourVars ) );
 	}
 
+public:
+
 	/**
-	 *	Removes all functions from this multi-cast delegate's invocation list that are bound to the specified UserObject.
-	 *	Note that the order of the delegates may not be preserved!
+	 * Removes a delegate instance from this multi-cast delegate's invocation list (performance is O(N)).
+	 *
+	 * Note that the order of the delegate instances may not be preserved!
+	 *
+	 * @param DelegateInstance The delegate instance to remove.
 	 */
-	void RemoveAll( const void* InUserObject )
+	void Remove( const TDelegateInstanceInterface& DelegateInstance )
 	{
-		// Remove the delegate
-		RemoveByObjectInternal( InUserObject );
+		// NOTE: For delegates with PAYLOAD, the payload variables are not considered when comparing delegate objects!
+		// This means that two delegate instances bound to the same object/method, but with different payload data,
+		// cannot be bound to the same single multi-cast delegate.  We could support this by changing IsSameFunction()
+		// to do comparison of payload data, but there are some trade-offs if we do that.
 
-		// Check for any delegates that may have expired
+		RemoveDelegateInstance(DelegateInstance);
 		CompactInvocationList();
 	}
 
 	/**
-	 * Removes a function from this multi-cast delegate's invocation list (performance is O(N)).  Note that the
-	 * order of the delegates may not be preserved!
+	 * DEPRECATED: Removes a unicast delegate from this multi-cast delegate's invocation list (performance is O(N)).
+	 *
+	 * The order of the delegates may not be preserved!
+	 * This function is retained for backwards compatibility.
 	 *
-	 * @param	InDelegate	Delegate to remove
+	 * @param Delegate The delegate to remove.
 	 */
-	void Remove( const FDelegate& InDelegate )
+	void Remove( const FDelegate& Delegate )
 	{
 		// NOTE: For delegates with PAYLOAD, the payload variables are not considered when comparing delegate objects!
 		// This means that two delegate instances bound to the same object/method, but with different payload data,
 		// cannot be bound to the same single multi-cast delegate.  We could support this by changing IsSameFunction()
 		// to do comparison of payload data, but there are some trade-offs if we do that.
 
-		// Remove the delegate
-		RemoveByDelegateInternal( InDelegate );
+		TDelegateInstanceInterface* DelegateInstance = (TDelegateInstanceInterface*)Delegate.GetDelegateInstance();
 
-		// Check for any delegates that may have expired
-		CompactInvocationList();
+		if (DelegateInstance != nullptr)
+		{
+			RemoveDelegateInstance(*DelegateInstance);
+			CompactInvocationList();
+		}		
 	}
 
 	// NOTE: These direct Remove methods are only supported for multi-cast delegates with no payload attached.
@@ -1623,6 +1733,19 @@ public:
 		Remove( FDelegate::CreateThreadSafeSP( InUserObject, InFunc ) );
 	}
 
+	/**
+	 * Removes a UFunction-based member function delegate (performance is O(N)).  Note that the order of the
+	 * delegates may not be preserved!
+	 *
+	 * @param	InUserObject	User object to unbind from
+	 * @param	InFunc			Class method function address
+	 */
+	template< class UObjectTemplate >
+	inline void RemoveUFunction( UObjectTemplate* InUserObject, const FName& InFunctionName )
+	{
+		Remove( FDelegate::CreateUFunction( InUserObject, InFunctionName ) );
+	}
+
 	/**
 	 * Removes a UObject-based member function delegate (performance is O(N)).  Note that the order of the
 	 * delegates may not be preserved!
@@ -1641,208 +1764,110 @@ public:
 		Remove( FDelegate::CreateUObject( InUserObject, InFunc ) );
 	}
 
+protected:
+
 	/** 
-	 * Checks to see if any functions are bound to the given user object.
-	 *
-	 * @return	True if any functions are bound to InUserObject, false otherwise.
+	 * Hidden default constructor.
 	 */
-	inline bool IsBoundToObject(void const* InUserObject) const
-	{
-		const int32 NumFunctions = InvocationList.Num();
-		for( int32 CurFunctionIndex = 0; CurFunctionIndex < NumFunctions; ++CurFunctionIndex )
-		{
-			if (InvocationList[ CurFunctionIndex ].Delegate->HasSameObject( InUserObject ) )
-			{
-				return true;
-			}
-		}
-
-		return false;
-	}
+	inline BASE_MULTICAST_DELEGATE_CLASS( ) { }
 
 protected:
 
 	/**
-	 * Default constructor
-	 */
-	inline BASE_MULTICAST_DELEGATE_CLASS()
-		: ExpirableObjectCount( 0 )
-	{ }
-
-	/**
-	 * Checks to see if any functions are bound to this multi-cast delegate.
-	 *
-	 * This method is protected to avoid assumptions about only one delegate existing.
+	 * Adds a function delegate to this multi-cast delegate's invocation list
 	 *
-	 * @return	True if any functions are bound
+	 * @param	InDelegate	Delegate to add
 	 */
-	inline bool IsBound() const
+	void AddDelegateInstance( TDelegateInstanceInterface& InDelegateInstance )
 	{
-		return InvocationList.Num() > 0;
+#if !UE_BUILD_SHIPPING && !UE_BUILD_TEST
+		// verify that the same function isn't already bound
+		for (auto DelegateInstance : GetInvocationList())
+	{
+			// this down-cast is OK! allows for managing invocation list in the base class without requiring virtual functions
+			TDelegateInstanceInterface* DelegateInstanceInterface = (TDelegateInstanceInterface*)DelegateInstance;
+			check(!DelegateInstanceInterface->IsSameFunction(InDelegateInstance));
 	}
+#endif
 
-	/**
-	 * Removes all functions from this delegate's invocation list
-	 */
-	void Clear()
-	{
-		InvocationList.Empty();
-		ExpirableObjectCount = 0;
+		AddInternal(InDelegateInstance);
 	}
 
 	/**
-	 * Broadcasts this delegate to all bound objects, except to those that may have expired
+	 * Broadcasts this delegate to all bound objects, except to those that may have expired.
+	 *
+	 * The constness of this method is a lie, because we are removing compactable delegate intsances.
 	 */
 	void Broadcast( FUNC_PARAM_LIST ) const
 	{
-		if( InvocationList.Num() > 0 )
+		const TArray<IDelegateInstance*>& InvocationList = GetInvocationList();
+
+		if (InvocationList.Num() == 0)
 		{
-			// Create a copy of the invocation list, just in case the list is modified by one of the callbacks
-			// during the broadcast
-			typedef TArray< FDelegate, TInlineAllocator< 4 > > FInlineInvocationList;
+			return;
+		}
+
+		// copy invocation list just in case the list is modified by one of the broadcast callbacks
+		typedef TArray<IDelegateInstance*, TInlineAllocator<4>> FInlineInvocationList;
 			FInlineInvocationList InvocationListCopy = FInlineInvocationList(InvocationList);
 
 			// Invoke each bound function
-			for( typename FInlineInvocationList::TConstIterator FunctionIt( InvocationListCopy ); FunctionIt; ++FunctionIt )
+		for (IDelegateInstance* DelegateInstance : InvocationListCopy)
 			{
-				if( !FunctionIt->ExecuteIfBound( FUNC_PARAM_PASSTHRU ) )
+			// this down-cast is OK! allows for managing invocation list in the base class without requiring virtual functions
+			TDelegateInstanceInterface* DelegateInstanceInterface = (TDelegateInstanceInterface*)DelegateInstance;
+
+			// execute the delegate instance...
+			if (!DelegateInstanceInterface->ExecuteIfSafe(FUNC_PARAM_PASSTHRU))
 				{
-					if( FunctionIt->IsCompactable() )
+				// ... or remove it if compactable
+				if (DelegateInstance->IsCompactable())
 					{
-						// Function couldn't be executed and it is compactable, so remove it.  Note that because the 
-						// original list could have been modified by one of the callbacks, we have to search for the 
-						// function to remove here.
-						RemoveByDelegateInternal( *FunctionIt );
-					}
+					// search for the instance, because original list could have been modified during broadcast
+					// this is really shady, because it violates the constness of this method
+					const_cast<BASE_MULTICAST_DELEGATE_CLASS*>(this)->RemoveDelegateInstance(*DelegateInstanceInterface);
 				}
 			}
 		}
 	}
 
-private:
-
 	/**
-	 * Adds a function delegate to this multi-cast delegate's invocation list
+	 * Removes a function from this multi-cast delegate's invocation list (performance is O(N)).
 	 *
-	 * @param	InDelegate	Delegate to add
+	 * Note that the order of the delegates may not be preserved! This method will not assert if
+	 * the given delegate instance was not found in the invocation list, because some instances
+	 * hold on to weak pointers to objects that may become invalid.
+	 *
+	 * @param InDelegateInstance The delegate instance to remove.
 	 */
-	void AddInternal( const FDelegate& InDelegate )
+	void RemoveDelegateInstance( const TDelegateInstanceInterface& InDelegateInstance )
 	{
-#if !UE_BUILD_SHIPPING && !UE_BUILD_TEST
-		// Verify same function isn't already bound.
-		const int32 NumFunctions = InvocationList.Num();
-		for( int32 CurFunctionIndex = 0; CurFunctionIndex < NumFunctions; ++CurFunctionIndex )
-		{
-			// NOTE: Multi-cast delegates currently don't support binding two delegates with the same object/function
-			// that have different payload variables.  See the comment in the Remove() method for more information.
-			check( !InvocationList[ CurFunctionIndex ].Delegate->IsSameFunction( *InDelegate.Delegate ) );
-		}
-#endif
-		InvocationList.Add( InDelegate );
+		const TArray<IDelegateInstance*> InvocationList = GetInvocationList();
 
-		// Keep track of whether we have objects bound that may expire out from underneath us.
-		const EDelegateInstanceType::Type DelegateType = InDelegate.Delegate->GetType();
-		if( DelegateType == EDelegateInstanceType::SharedPointerMethod ||
-			DelegateType == EDelegateInstanceType::ThreadSafeSharedPointerMethod ||
-			DelegateType == EDelegateInstanceType::UObjectMethod )
+		for (int32 InvocationListIndex = InvocationList.Num() - 1; InvocationListIndex >= 0; --InvocationListIndex)
 		{
-			++ExpirableObjectCount;
-		}
-	}
+			// this down-cast is OK! allows for managing invocation list in the base class without requiring virtual functions
+			TDelegateInstanceInterface* DelegateInstance = (TDelegateInstanceInterface*)InvocationList[InvocationListIndex];
 
-	/**
-	 * Removes a function from this multi-cast delegate's invocation list (performance is O(N)).  Note that the
-	 * order of the delegates may not be preserved!
-	 *
-	 * @param	InDelegate	Delegate to remove
-	 */
-	void RemoveByDelegateInternal( const FDelegate& InDelegate ) const
-	{
-		const int32 NumFunctions = InvocationList.Num();
-		for( int32 CurFunctionIndex = 0; CurFunctionIndex < NumFunctions; ++CurFunctionIndex )
-		{
 			// NOTE: We must do a deep compare here, not just compare delegate pointers, because multiple
 			//       delegate pointers can refer to the exact same object and method
-			if( InvocationList[ CurFunctionIndex ].Delegate->IsSameFunction( *InDelegate.Delegate ) )
+			if (DelegateInstance->IsSameFunction(InDelegateInstance))
 			{
-				// Keep track of whether we have objects bound that may expire out from underneath us.
-				const EDelegateInstanceType::Type DelegateType = InvocationList[ CurFunctionIndex ].Delegate->GetType();
-				if( DelegateType == EDelegateInstanceType::SharedPointerMethod || 
-					DelegateType == EDelegateInstanceType::ThreadSafeSharedPointerMethod ||
-					DelegateType == EDelegateInstanceType::UObjectMethod )
-				{
-					--ExpirableObjectCount;
-				}
-
-				InvocationList.RemoveAtSwap( CurFunctionIndex );
+				RemoveByIndex(InvocationListIndex);
 
-				// No need to continue, as we never allow the same delegate to be bound twice
+				// no need to continue, as we never allow the same delegate to be bound twice
 				break;
 			}
 		}
-
-		// NOTE: We don't want to assert that the specified delegate was removed, because delegates that are
-		//       bound by weak pointer may vanish when the object goes out of scope during cleanup, which
-		//       could create unpredictable assertions.  It's best to always safely allow a delegate remove attempt.
 	}
-
-	/**
-	 *	Removes all functions from this multi-cast delegate's invocation list that are bound to the specified UserObject.
-	 *	Note that the order of the delegates may not be preserved!
-	 */
-	void RemoveByObjectInternal( const void* InUserObject )
-	{
-		const int32 NumFunctions = InvocationList.Num();
-		for( int CurFunctionIndex = NumFunctions - 1; CurFunctionIndex >= 0; --CurFunctionIndex )
-		{
-			if( InvocationList[CurFunctionIndex].Delegate->HasSameObject( InUserObject ) )
-			{
-				// Keep track of whether we have objects bound that may expire out from underneath us.
-				const EDelegateInstanceType::Type DelegateType = InvocationList[CurFunctionIndex].Delegate->GetType();
-				if( DelegateType == EDelegateInstanceType::SharedPointerMethod ||
-					DelegateType == EDelegateInstanceType::ThreadSafeSharedPointerMethod ||
-					DelegateType == EDelegateInstanceType::UObjectMethod )
-				{
-					--ExpirableObjectCount;
-				}
-
-				InvocationList.RemoveAtSwap( CurFunctionIndex );
-			}
-		}
-	}
-
-	/** Cleans up any delegates in our invocation list that have expired (performance is O(N)) */
-	void CompactInvocationList() const
-	{
-		// We only need to compact if any objects were added whose lifetime is known to us
-		if( ExpirableObjectCount > 0 )
-		{
-			for( int32 CurFunctionIndex = 0; CurFunctionIndex < InvocationList.Num(); ++CurFunctionIndex )
-			{
-				if( InvocationList[CurFunctionIndex].IsCompactable() )
-				{
-					// Remove it, and decrement loop counter so that we don't skip over an element
-					InvocationList.RemoveAtSwap( CurFunctionIndex-- );
-
-					checkSlow( ExpirableObjectCount > 0 );
-					--ExpirableObjectCount;
-				}
-			}
-		}
-	}
-
-private:
-
-	/** Ordered list functions to invoke when the Broadcast function is called */
-	typedef TArray< FDelegate > FInvocationList;
-	mutable FInvocationList InvocationList;		// Mutable so that we can housekeep list even with 'const' broadcasts
-
-	/** Tracks the number of objects in our invocation list whose lifetime is known to us.  That is,
-	    these objects could potentially expire out from underneath us. */
-	mutable int32 ExpirableObjectCount;		// Mutable so that we can housekeep list even with 'const' broadcasts
 };
 
 
+	/**
+ * Implements a multicast delegate.
+ *
+ * This class should not be instantiated directly. Use the DECLARE_MULTICAST_DELEGATE macros instead.
+	 */
 template< FUNC_TEMPLATE_DECL_TYPENAME >
 class MULTICAST_DELEGATE_CLASS : public BASE_MULTICAST_DELEGATE_CLASS< FUNC_TEMPLATE_DECL >
 {
@@ -1879,6 +1904,7 @@ private:
 	// Prevents erroneous use by other classes.
 	typedef BASE_MULTICAST_DELEGATE_CLASS< FUNC_TEMPLATE_DECL > Super;
 };
+
 #endif		// FUNC_IS_VOID
 
 
@@ -1931,7 +1957,7 @@ public:
 	template< class UserClass >
 	void __Internal_BindDynamic( UserClass* InUserObject, typename TMethodPtrResolver< UserClass >::FMethodPtr InMethodPtr, const FString& InMacroFunctionName )
 	{
-		check( InUserObject != NULL && InMethodPtr != NULL && !InMacroFunctionName.IsEmpty() );
+		check( InUserObject != nullptr && InMethodPtr != nullptr && !InMacroFunctionName.IsEmpty() );
 
 		// NOTE: We're not actually storing the incoming method pointer or calling it.  We simply require it for type-safety reasons.
 
@@ -2000,7 +2026,7 @@ public:
 	template< class UserClass >
 	bool __Internal_IsAlreadyBound( UserClass* InUserObject, typename FDelegate::template TMethodPtrResolver< UserClass >::FMethodPtr InMethodPtr, const FString& InMacroFunctionName ) const
 	{
-		check( InUserObject != NULL && InMethodPtr != NULL && !InMacroFunctionName.IsEmpty() );
+		check( InUserObject != nullptr && InMethodPtr != nullptr && !InMacroFunctionName.IsEmpty() );
 
 		// NOTE: We're not actually using the incoming method pointer or calling it.  We simply require it for type-safety reasons.
 
@@ -2023,7 +2049,7 @@ public:
 	template< class UserClass >
 	void __Internal_AddDynamic( UserClass* InUserObject, typename FDelegate::template TMethodPtrResolver< UserClass >::FMethodPtr InMethodPtr, const FString& InMacroFunctionName )
 	{
-		check( InUserObject != NULL && InMethodPtr != NULL && !InMacroFunctionName.IsEmpty() );
+		check( InUserObject != nullptr && InMethodPtr != nullptr && !InMacroFunctionName.IsEmpty() );
 
 		// NOTE: We're not actually storing the incoming method pointer or calling it.  We simply require it for type-safety reasons.
 
@@ -2046,7 +2072,7 @@ public:
 	template< class UserClass >
 	void __Internal_AddUniqueDynamic( UserClass* InUserObject, typename FDelegate::template TMethodPtrResolver< UserClass >::FMethodPtr InMethodPtr, const FString& InMacroFunctionName )
 	{
-		check( InUserObject != NULL && InMethodPtr != NULL && !InMacroFunctionName.IsEmpty() );
+		check( InUserObject != nullptr && InMethodPtr != nullptr && !InMacroFunctionName.IsEmpty() );
 
 		// NOTE: We're not actually storing the incoming method pointer or calling it.  We simply require it for type-safety reasons.
 
@@ -2069,7 +2095,7 @@ public:
 	template< class UserClass >
 	void __Internal_RemoveDynamic( UserClass* InUserObject, typename FDelegate::template TMethodPtrResolver< UserClass >::FMethodPtr InMethodPtr, const FString& InMacroFunctionName )
 	{
-		check( InUserObject != NULL && InMethodPtr != NULL && !InMacroFunctionName.IsEmpty() );
+		check( InUserObject != nullptr && InMethodPtr != nullptr && !InMacroFunctionName.IsEmpty() );
 
 		// NOTE: We're not actually storing the incoming method pointer or calling it.  We simply require it for type-safety reasons.
 
@@ -2121,6 +2147,12 @@ public:
 #undef RAW_METHOD_DELEGATE_INSTANCE_CLASS_ThreeVars
 #undef RAW_METHOD_DELEGATE_INSTANCE_CLASS_FourVars
 
+#undef UFUNCTION_DELEGATE_INSTANCE_CLASS_NoVars
+#undef UFUNCTION_DELEGATE_INSTANCE_CLASS_OneVar
+#undef UFUNCTION_DELEGATE_INSTANCE_CLASS_TwoVars
+#undef UFUNCTION_DELEGATE_INSTANCE_CLASS_ThreeVars
+#undef UFUNCTION_DELEGATE_INSTANCE_CLASS_FourVars
+
 #undef UOBJECT_METHOD_DELEGATE_INSTANCE_CLASS_NoVars
 #undef UOBJECT_METHOD_DELEGATE_INSTANCE_CLASS_OneVar
 #undef UOBJECT_METHOD_DELEGATE_INSTANCE_CLASS_TwoVars
diff --git a/Engine/Source/Runtime/Core/Public/Templates/IDelegateInstance.h b/Engine/Source/Runtime/Core/Public/Templates/IDelegateInstance.h
new file mode 100644
--- /dev/null
+++ b/Engine/Source/Runtime/Core/Public/Templates/IDelegateInstance.h
@@ -0,0 +1,106 @@
+// Copyright 1998-2014 Epic Games, Inc. All Rights Reserved.
+
+/*=============================================================================
+	IDelegateInstance.cpp: Declares the IDelegateInstance interface.
+=============================================================================*/
+
+#pragma once
+
+
+/**
+ * Types of delegate instances
+ */
+namespace EDelegateInstanceType
+{
+	enum Type
+	{
+		/** Member function pointer to method in (fast, not thread-safe) shared pointer-based class */
+		SharedPointerMethod,
+
+		/** Member function pointer to method in (conditionally thread-safe) shared pointer-based class */
+		ThreadSafeSharedPointerMethod,
+
+		/** Raw C++ member function pointer (pointer to class method) */
+		RawMethod,
+
+		/** UFunction delegate */
+		UFunction,
+
+		/** Member function pointer to method in UObject-based class */
+		UObjectMethod,
+
+		/** Raw C++ static function pointer */
+		Raw
+	};
+}
+
+
+/**
+ * Interface for delegate instances.
+ */
+class IDelegateInstance
+{
+public:
+
+	/**
+	 * Returns the name of the UFunction that this delegate instance is bound to.
+	 *
+	 * @return Name of the function, or NAME_None if not bound to a UFunction.
+	 */
+	virtual FName GetFunctionName( ) const = 0;
+
+	/**
+	 * Returns raw pointer to the delegate method.
+	 *
+	 * @return  Raw pointer to the delegate method.
+	 */
+	virtual const void* GetRawMethodPtr( ) const = 0;
+
+	/**
+	 * Returns raw pointer to UserObject,
+	 *
+	 * @return  Raw pointer to UserObject.
+	 */
+	virtual const void* GetRawUserObject( ) const = 0;
+
+	/**
+	 * Returns the type of delegate instance
+	 *
+	 * @return  Delegate instance type
+	 */
+	virtual EDelegateInstanceType::Type GetType( ) const = 0;
+
+	/**
+	 * Returns true if this delegate is bound to the specified UserObject,
+	 *
+	 * @param  InUserObject
+	 *
+	 * @return  True if delegate is bound to the specified UserObject
+	 */
+	virtual bool HasSameObject( const void* InUserObject ) const = 0;
+
+	/**
+	 * Checks to see if the user object bound to this delegate can ever be valid again.
+	 * used to compact multicast delegate arrays so they don't expand without limit.
+	 *
+	 * @return  True if the user object can never be used again
+	 */
+	virtual bool IsCompactable( ) const
+	{
+		return !IsSafeToExecute();
+	}
+
+	/**
+	 * Checks to see if the user object bound to this delegate is still valid
+	 *
+	 * @return  True if the user object is still valid and it's safe to execute the function call
+	 */
+	virtual bool IsSafeToExecute( ) const = 0;
+
+public:
+
+	/**
+	 * Virtual destructor.
+	 */
+	virtual ~IDelegateInstance( ) { }
+};
diff --git a/Engine/Source/Runtime/Core/Public/Templates/MulticastDelegateBase.h b/Engine/Source/Runtime/Core/Public/Templates/MulticastDelegateBase.h
new file mode 100644
--- /dev/null
+++ b/Engine/Source/Runtime/Core/Public/Templates/MulticastDelegateBase.h
@@ -0,0 +1,194 @@
+// Copyright 1998-2014 Epic Games, Inc. All Rights Reserved.
+
+/*=============================================================================
+	MulticastDelegateBase.h: Declares the FMulticastDelegateBase class.
+=============================================================================*/
+
+#pragma once
+
+
+/**
+ * Abstract base class for multicast delegates.
+ */
+template<typename ObjectPtrType = FWeakObjectPtr>
+class FMulticastDelegateBase
+{
+public:
+
+	/**
+	 * Removes all functions from this delegate's invocation list
+	 */
+	void Clear( )
+	{
+		for (IDelegateInstance* DelegateInstance : InvocationList)
+		{
+			delete DelegateInstance;
+		}
+
+		InvocationList.Empty();
+		ExpirableObjectCount = 0;
+	}
+
+	/**
+	 * Checks to see if any functions are bound to this multi-cast delegate.
+	 *
+	 * This method is protected to avoid assumptions about only one delegate existing.
+	 *
+	 * @return true if any functions are bound, false otherwise.
+	 */
+	inline bool IsBound( ) const
+	{
+		return (InvocationList.Num() > 0);
+	}
+
+	/** 
+	 * Checks to see if any functions are bound to the given user object.
+	 *
+	 * @return	True if any functions are bound to InUserObject, false otherwise.
+	 */
+	inline bool IsBoundToObject( void const* InUserObject ) const
+	{
+		for (IDelegateInstance* DelegateInstance : InvocationList)
+		{
+			if (DelegateInstance->HasSameObject(InUserObject))
+			{
+				return true;
+			}
+		}
+
+		return false;
+	}
+
+	/**
+	 * Removes all functions from this multi-cast delegate's invocation list that are bound to the specified UserObject.
+	 * Note that the order of the delegates may not be preserved!
+	 *
+	 * @param InUserObject The object to remove all delegates for.
+	 */
+	void RemoveAll( const void* InUserObject )
+	{
+		RemoveByObject(InUserObject);
+		CompactInvocationList();
+	}
+
+protected:
+
+	/**
+	 * Hidden default constructor.
+	 */
+	inline FMulticastDelegateBase( )
+		: ExpirableObjectCount(0)
+	{ }
+
+protected:
+
+	/**
+	 * Adds the given delegate instance to the invocation list.
+	 *
+	 * @param DelegateInstance The delegate instance to add.
+	 */
+	void AddInternal( IDelegateInstance& DelegateInstance )
+	{
+		InvocationList.Add(&DelegateInstance);
+
+		// Keep track of whether we have objects bound that may expire out from underneath us.
+		const EDelegateInstanceType::Type DelegateType = DelegateInstance.GetType();
+
+		if ((DelegateType == EDelegateInstanceType::SharedPointerMethod) ||
+			(DelegateType == EDelegateInstanceType::ThreadSafeSharedPointerMethod) ||
+			(DelegateType == EDelegateInstanceType::UObjectMethod))
+		{
+			++ExpirableObjectCount;
+		}
+	}
+
+	/**
+	 * Cleans up any delegates in our invocation list that have expired (performance is O(N))
+	 */
+	void CompactInvocationList( ) const
+	{
+		// We only need to compact if any objects were added whose lifetime is known to us
+		if (ExpirableObjectCount == 0)
+		{
+			return;
+		}
+
+		// remove expired delegates
+		for (int32 InvocationListIndex = InvocationList.Num() - 1; InvocationListIndex >= 0; --InvocationListIndex)
+		{
+			IDelegateInstance* DelegateInstance = InvocationList[InvocationListIndex];
+
+			if (DelegateInstance->IsCompactable())
+			{
+				InvocationList.RemoveAtSwap(InvocationListIndex);
+				delete DelegateInstance;
+
+				checkSlow(ExpirableObjectCount > 0);
+				--ExpirableObjectCount;
+			}
+		}
+	}
+
+	/**
+	 * Gets a read-only reference to the invocation list.
+	 *
+	 * @return The invocation list.
+	 */
+	const TArray<IDelegateInstance*>& GetInvocationList( ) const
+	{
+		return InvocationList;
+	}
+
+	/**
+	 * Removes all functions from this multi-cast delegate's invocation list that are bound to the specified UserObject.
+	 *
+	 * Note that the order of the delegates may not be preserved!
+	 *
+	 * @param InUserObject The object to remove the delegates for.
+	 */
+	void RemoveByObject( const void* InUserObject )
+	{
+		for (int32 InvocationListIndex = InvocationList.Num() - 1; InvocationListIndex >= 0; --InvocationListIndex)
+		{
+			IDelegateInstance* DelegateInstance = InvocationList[InvocationListIndex];
+
+			if (DelegateInstance->HasSameObject(InUserObject))
+			{
+				RemoveByIndex(InvocationListIndex);
+			}
+		}
+	}
+
+	/**
+	 * Removes a delegate instance given its index in the invocation list.
+	 *
+	 * @param InvocationListIndex The index of the delegate instance to remove.
+	 */
+	void RemoveByIndex( int32 InvocationListIndex )
+	{
+		IDelegateInstance* DelegateInstance = InvocationList[InvocationListIndex];
+
+		// keep track of whether we have objects bound that may expire out from underneath us
+		const EDelegateInstanceType::Type DelegateType = DelegateInstance->GetType();
+
+		if ((DelegateType == EDelegateInstanceType::SharedPointerMethod) || 
+			(DelegateType == EDelegateInstanceType::ThreadSafeSharedPointerMethod) ||
+			(DelegateType == EDelegateInstanceType::UFunction) ||
+			(DelegateType == EDelegateInstanceType::UObjectMethod))
+		{
+			--ExpirableObjectCount;
+		}
+
+		InvocationList.RemoveAtSwap(InvocationListIndex);
+		delete DelegateInstance;
+	}
+
+private:
+
+	// Mutable so that we can housekeep list even with const broadcasts
+	mutable int32 ExpirableObjectCount;
+
+	// Holds the collection of delegate instances to invoke.
+	// Mutable so that we can housekeep list even with const broadcasts
+	mutable TArray<IDelegateInstance*> InvocationList;
+};
