Delegates: Moved delegate files into their own sub-folder

diff --git a/Engine/Source/Runtime/Core/Core.Build.cs b/Engine/Source/Runtime/Core/Core.Build.cs
--- a/Engine/Source/Runtime/Core/Core.Build.cs
+++ b/Engine/Source/Runtime/Core/Core.Build.cs
@@ -17,6 +17,7 @@ public class Core : ModuleRules
 				"Runtime/Core/Public/Async",
 				"Runtime/Core/Public/Concurrency",
 				"Runtime/Core/Public/Containers",
+				"Runtime/Core/Public/Delegates",
 				"Runtime/Core/Public/GenericPlatform",
 				"Runtime/Core/Public/HAL",
 				"Runtime/Core/Public/Math",
diff --git a/Engine/Source/Runtime/Core/Public/Delegates/Delegate.h b/Engine/Source/Runtime/Core/Public/Delegates/Delegate.h
new file mode 100644
--- /dev/null
+++ b/Engine/Source/Runtime/Core/Public/Delegates/Delegate.h
@@ -0,0 +1,334 @@
+// Copyright 1998-2014 Epic Games, Inc. All Rights Reserved.
+
+// Do not remove this header guard! It ensures that DelegateCombinations.h can only be included from here.
+#ifndef __Delegate_h__
+#define __Delegate_h__
+
+#pragma once
+
+#include "SharedPointer.h"
+#include "WeakObjectPtrTemplates.h"
+#include "IDelegateInstance.h"
+#include "DelegateBase.h"
+#include "MulticastDelegateBase.h"
+
+
+/**
+ *  C++ DELEGATES
+ *  -----------------------------------------------------------------------------------------------
+ *
+ *	This system allows you to call member functions on C++ objects in a generic, yet type-safe way.
+ *  Using delegates, you can dynamically bind to a member function of an arbitrary object,
+ *	then call functions on the object, even if the caller doesn't know the object's type.
+ *
+ *	The system predefines various combinations of generic function signatures with which you can
+ *	declare a delegate type from, filling in the type names for return value and parameters with
+ *	whichever types you need.
+ *
+ *	Both single-cast and multi-cast delegates are supported, as well as "dynamic" delegates which
+ *	can be safely serialized to disk.  Additionally, delegates may define "payload" data which
+ *	will stored and passed directly to bound functions.
+ *
+ *
+ *
+ *  DELEGATE FEATURES
+ *  -----------------------------------------------------------------------------------------------
+ *
+ *	Currently we support delegate signatures using any combination of the following:
+ *   		- Functions returning a value
+ *			- Up to four "payload" variables
+ *   		- Up to eight function parameters
+ *   		- Functions declared as 'const'
+ *
+ *  Multi-cast delegates are also supported, using the 'DECLARE_MULTICAST_DELEGATE...' macros.
+ *  Multi-cast delegates allow you to attach multiple function delegates, then execute them all at
+ *  once by calling a single "Broadcast()" function.  Multi-cast delegate signatures are not allowed
+ *	to use a return value.
+ *
+ *	You can assign "payload data" to your delegates!  These are arbitrary variables that will be passed
+ *  directly to any bound function when it is invoked.  This is really useful as it allows you to store
+ *  parameters within the delegate it self at bind-time.  All delegate types (except for "dynamic") supports
+ *	payload variables automatically!
+ *
+ *	When binding to a delegate, you can pass payload data along.  This example passes two custom variables,
+ *	a bool and an int32 to a delegate.  Then when the delegate is invoked, these parameters will be passed
+ *	to your bound function.  The extra variable arguments must always be accepted after the delegate
+ *	type parameter arguments.
+ *
+ *			MyDelegate.BindStatic( &MyFunction, true, 20 );
+ *
+ *	Remember to look at the table at the bottom of this documentation comment for the macro names to
+ *	use for each function signature type.
+ *
+ *
+ *
+ *  DELEGATES EXAMPLE
+ *  -----------------------------------------------------------------------------------------------
+ *
+ *	Suppose you have a class with a method that you'd like to be able to call from anywhere: 
+ *
+ *		class FLogWriter
+ *		{
+ *			void WriteToLog( FString );
+ *		};
+ *
+ *	To call the WriteToLog function, we'll need to create a delegate type for that function's signature.
+ *	To do this, you will first declare the delegate using one of the macros below.  For example, here
+ *	is a simple delegate type:
+ *
+ *		DECLARE_DELEGATE_OneParam( FStringDelegate, FString );
+ *
+ *	This creates a delegate type called 'FStringDelegate' that takes a single parameter of type 'FString'.
+ *
+ *	Here's an example of how you'd use this 'FStringDelegate' in a class:
+ *
+ *		class FMyClass
+ *		{
+ *			FStringDelegate WriteToLogDelegate;
+ *		};
+ *
+ *	This allows your class to hold a pointer to a method in an arbitrary class.  The only thing the
+ *	class really knows about this delegate is it's function signature.
+ *
+ *	Now, to assign the delegate, simply create an instance of your delegate class, passing along the
+ *	class that owns the method as a template parameter.  You'll also pass the instance of your object
+ *	and the actual function address of the method.  So, here we'll create an instance of our 'FLogWriter'
+ *	class, then create a delegate for the 'WriteToLog' method of that object instance:
+ *
+ *		FSharedRef< FLogWriter > LogWriter( new FLogWriter() );
+ *
+ *		WriteToLogDelegate.BindSP( LogWriter, &FLogWriter::WriteToLog );
+ *
+ *	You've just dynamically bound a delegate to a method of a class!  Pretty simple, right?
+ *
+ *	Note that the 'SP' part of 'BindSP' stands for 'shared pointer', because we're binding to an
+ *  object that's owned by a shared pointer.  There are versions for different object types,
+ *	such as BindRaw() and BindUObject().  You can bind to global function pointers with BindStatic().
+ *
+ *	Now, your 'WriteToLog' method can be called by FMyClass without it even knowing anything about
+ *	the 'FLogWriter' class!  To call a your delegate, just use the 'Execute()' method:
+ *
+ *		WriteToLogDelegate.Execute( TEXT( "Delegates are spiffy!" ) );
+ *
+ *  If you call Execute() before binding a function to the delegate, an assertion will be triggered.  In
+ *  many cases, you'll instead want to do this:
+ *
+ *		WriteToLogDelegate.ExecuteIfBound( TEXT( "Only executes if a function was bound!" ) );
+ *
+ *	That's pretty much all there is to it!!  You can read below for a bit more information.
+ *
+ *
+ *
+ *	MORE INFORMATION
+ *  -----------------------------------------------------------------------------------------------
+ *
+ *	The delegate system understands certain types of objects, and additional features are enabled when
+ *  using these objects.  If you bind a delegate to a member of a UObject or shared pointer class, the
+ *  delegate system can keep a weak reference to the object, so that if the object gets destroyed out
+ *  from underneath the delegate, you'll be able to handle these cases by calling IsBound() or
+ *  ExecuteIfBound() functions.  Note the special binding syntax for the various types of supported objects.
+ *
+ *  It's perfectly safe to copy delegate objects.  Delegates can be passed around by value but this is
+ *	generally not recommended since they do have to allocate memory on the heap.  Pass them by reference
+ *  when possible!
+ *
+ *  Delegate signature declarations can exist at global scope, within a namespace or even within a class
+ *  declaration (but not function bodies.)
+ *
+ *
+ *
+ *  FUNCTION SIGNATURES
+ *  -----------------------------------------------------------------------------------------------
+ *
+ *  Use this table to find the declaration macro to use to declare your delegate.
+ *
+ *	Function signature									|	Declaration macro
+ *  ------------------------------------------------------------------------------------------------------------------------------------------------------------
+ *	void Function()										|	DECLARE_DELEGATE( DelegateName )
+ *	void Function( <Param1> )							|	DECLARE_DELEGATE_OneParam( DelegateName, Param1Type )
+ *	void Function( <Param1>, <Param2> )					|	DECLARE_DELEGATE_TwoParams( DelegateName, Param1Type, Param2Type )
+ *	void Function( <Param1>, <Param2>, ... )			|	DECLARE_DELEGATE_<Num>Params( DelegateName, Param1Type, Param2Type, ... )
+ *	<RetVal> Function()									|	DECLARE_DELEGATE_RetVal( RetValType, DelegateName )
+ *	<RetVal> Function( <Param1> )						|	DECLARE_DELEGATE_RetVal_OneParam( RetValType, DelegateName, Param1Type )
+ *	<RetVal> Function( <Param1>, <Param2> )				|	DECLARE_DELEGATE_RetVal_TwoParams( RetValType, DelegateName, Param1Type, Param2Type )
+ *	<RetVal> Function( <Param1>, <Param2>, ... )		|	DECLARE_DELEGATE_RetVal_<Num>Params( RetValType, DelegateName, Param1Type, Param2Type, ... )
+ *  ------------------------------------------------------------------------------------------------------------------------------------------------------------
+ *
+ *  Remember, there are three different delegate types you can define (any of the above signatures will work):
+ *
+ *                       Single-cast delegates:  DECLARE_DELEGATE...()
+ *                        Multi-cast delegates:  DECLARE_MULTICAST_DELEGATE...()
+ *	 Dynamic (UObject, serializable) delegates:  DECLARE_DYNAMIC_DELEGATE...()
+ */
+
+
+// This suffix is appended to all header exported delegates
+#define HEADER_GENERATED_DELEGATE_SIGNATURE_SUFFIX TEXT("__DelegateSignature")
+
+/** Helper macro that enables passing comma-separated arguments as a single macro parameter */
+#define FUNC_CONCAT( ... ) __VA_ARGS__
+
+/** Declare the user's delegate object */
+// NOTE: The last parameter is variadic and is used as the 'template args' for this delegate's classes (__VA_ARGS__)
+#define FUNC_DECLARE_DELEGATE( Suffix, DelegateName, ... ) \
+	typedef TBaseDelegate_##Suffix< __VA_ARGS__ > DelegateName;
+
+/** Declare the user's multicast delegate object */
+// NOTE: The last parameter is variadic and is used as the 'template args' for this delegate's classes (__VA_ARGS__)
+#define FUNC_DECLARE_MULTICAST_DELEGATE( Suffix, MulticastDelegateName, ... ) \
+	typedef TMulticastDelegate_##Suffix< __VA_ARGS__ > MulticastDelegateName;
+
+#define FUNC_DECLARE_EVENT( OwningType, EventName, Suffix, ... ) \
+	class EventName : public TBaseMulticastDelegate_##Suffix< __VA_ARGS__ >  { friend class OwningType; };
+
+/** Declare user's dynamic delegate, with wrapper proxy method for executing the delegate */
+#define FUNC_DECLARE_DYNAMIC_DELEGATE( Suffix, DynamicDelegateName, ExecFunction, FuncParamList, FuncParamPassThru, ... ) \
+	class DynamicDelegateName : public TBaseDynamicDelegate_##Suffix< __VA_ARGS__ > \
+	{ \
+	public: \
+		/** Default constructor */ \
+		DynamicDelegateName() \
+		{ \
+		} \
+		\
+		/** Construction from an FScriptDelegate must be explicit.  This is really only used by UObject system internals. */ \
+		explicit DynamicDelegateName( const TScriptDelegate<>& InScriptDelegate ) \
+			: TBaseDynamicDelegate_##Suffix( InScriptDelegate ) \
+		{ \
+		} \
+		\
+		/** Execute the delegate.  If the function pointer is not valid, an error will occur. */ \
+		inline void Execute( FuncParamList ) const \
+		{ \
+			/* Verify that the user object is still valid.  We only have a weak reference to it. */ \
+			checkSlow( IsBound() ); \
+			ExecFunction( FuncParamPassThru ); \
+		} \
+		/** Execute the delegate, but only if the function pointer is still valid */ \
+		inline bool ExecuteIfBound( FuncParamList ) const \
+		{ \
+			if( IsBound() ) \
+			{ \
+				ExecFunction( FuncParamPassThru ); \
+				return true; \
+			} \
+			return false; \
+		} \
+	};
+
+// Helper struct for converting uint32 return value to bool
+template <typename RetVal>
+struct TRetValCast
+{
+	inline RetVal operator()( RetVal Val ) const
+	{
+		return Val;
+	}
+	inline bool operator()( uint32 Val ) const
+	{
+		return !!Val;
+	}
+};
+
+#define FUNC_DECLARE_DYNAMIC_DELEGATE_RETVAL( Suffix, DynamicDelegateName, ExecFunction, RetValType, FuncParamList, FuncParamPassThru, ... ) \
+	class DynamicDelegateName : public TBaseDynamicDelegate_##Suffix< __VA_ARGS__ > \
+	{ \
+	public: \
+		/** Default constructor */ \
+		DynamicDelegateName() \
+		{ \
+		} \
+		\
+		/** Construction from an FScriptDelegate must be explicit.  This is really only used by UObject system internals. */ \
+		explicit DynamicDelegateName( const TScriptDelegate<>& InScriptDelegate ) \
+			: TBaseDynamicDelegate_##Suffix( InScriptDelegate ) \
+		{ \
+		} \
+		\
+		/** Execute the delegate.  If the function pointer is not valid, an error will occur. */ \
+		inline RetValType Execute( FuncParamList ) const \
+		{ \
+			/* Verify that the user object is still valid.  We only have a weak reference to it. */ \
+			checkSlow( IsBound() ); \
+			TRetValCast<RetValType> CastRetVal; \
+			return CastRetVal( ExecFunction( FuncParamPassThru ) ); \
+		} \
+	};
+
+
+/** Declare user's dynamic multi-cast delegate, with wrapper proxy method for executing the delegate */
+#define FUNC_DECLARE_DYNAMIC_MULTICAST_DELEGATE( Suffix, DynamicMulticastDelegateName, ExecFunction, FuncParamList, FuncParamPassThru, ... ) \
+class DynamicMulticastDelegateName : public TBaseDynamicMulticastDelegate_##Suffix< __VA_ARGS__ > \
+	{ \
+	public: \
+		/** Default constructor */ \
+		DynamicMulticastDelegateName() \
+		{ \
+		} \
+		\
+		/** Construction from an FMulticastScriptDelegate must be explicit.  This is really only used by UObject system internals. */ \
+		explicit DynamicMulticastDelegateName( const TMulticastScriptDelegate<>& InMulticastScriptDelegate ) \
+			: TBaseDynamicMulticastDelegate_##Suffix( InMulticastScriptDelegate ) \
+		{ \
+		} \
+		\
+		/** Broadcasts this delegate to all bound objects, except to those that may have expired */ \
+		void Broadcast( FuncParamList ) const \
+		{ \
+			ExecFunction( FuncParamPassThru ); \
+		} \
+	};
+
+
+// Simple macro chain to concatenate code text
+#define FUNC_COMBINE_ACTUAL( A, B ) A##B
+#define FUNC_COMBINE( A, B ) FUNC_COMBINE_ACTUAL( A, B )
+
+
+
+// Helper macro for calling BindDynamic() on dynamic delegates.  Automatically generates the function name string.
+#define BindDynamic( UserObject, FuncName ) __Internal_BindDynamic( UserObject, FuncName, FString( #FuncName ) )
+
+// Helper macro for calling AddDynamic() on dynamic multi-cast delegates.  Automatically generates the function name string.
+#define AddDynamic( UserObject, FuncName ) __Internal_AddDynamic( UserObject, FuncName, FString( #FuncName ) )
+
+// Helper macro for calling AddUniqueDynamic() on dynamic multi-cast delegates.  Automatically generates the function name string.
+#define AddUniqueDynamic( UserObject, FuncName ) __Internal_AddUniqueDynamic( UserObject, FuncName, FString( #FuncName ) )
+
+// Helper macro for calling RemoveDynamic() on dynamic multi-cast delegates.  Automatically generates the function name string.
+#define RemoveDynamic( UserObject, FuncName ) __Internal_RemoveDynamic( UserObject, FuncName, FString( #FuncName ) )
+
+// Helper macro for calling IsAlreadyBound() on dynamic multi-cast delegates.  Automatically generates the function name string.
+#define IsAlreadyBound( UserObject, FuncName ) __Internal_IsAlreadyBound( UserObject, FuncName, FString( #FuncName ) )
+
+
+/*********************************************************************************************************************/
+
+// We define this as a guard to prevent DelegateSignatureImpl.inl being included outside of this file
+#define FUNC_INCLUDING_INLINE_IMPL
+
+
+#ifndef UE_BUILD_DOCS
+#include "DelegateCombinations.h"
+#endif
+
+
+/*********************************************************************************************************************/
+
+#define DECLARE_DERIVED_EVENT( OwningType, BaseTypeEvent, EventName ) \
+	class EventName : public BaseTypeEvent { friend class OwningType; };
+
+
+// No longer allowed to include DelegateSignatureImpl.inl
+#undef FUNC_INCLUDING_INLINE_IMPL
+
+// Undefine temporary macros
+#undef FUNC_COMBINE_ACTUAL
+#undef FUNC_COMBINE
+
+// Simple delegate used by various utilities such as timers
+DECLARE_DELEGATE( FSimpleDelegate );
+DECLARE_MULTICAST_DELEGATE( FSimpleMulticastDelegate );
+
+
+#endif		// __Delegate_h__
diff --git a/Engine/Source/Runtime/Core/Public/Delegates/DelegateBase.h b/Engine/Source/Runtime/Core/Public/Delegates/DelegateBase.h
new file mode 100644
--- /dev/null
+++ b/Engine/Source/Runtime/Core/Public/Delegates/DelegateBase.h
@@ -0,0 +1,121 @@
+// Copyright 1998-2014 Epic Games, Inc. All Rights Reserved.
+
+/*=============================================================================
+	DelegateBase.h: Declares the FDelegateBase class.
+=============================================================================*/
+
+#pragma once
+
+
+/**
+ * Base class for unicast delegates.
+ */
+template<typename ObjectPtrType = FWeakObjectPtr>
+class FDelegateBase
+{
+public:
+
+	/**
+	 * Creates and initializes a new instance.
+	 *
+	 * @param InDelegateInstance The delegate instance to assign.
+	 */
+	FDelegateBase( IDelegateInstance* InDelegateInstance )
+		: DelegateInstance(InDelegateInstance)
+	{ }
+
+public:
+
+	/**
+	 * If this is a UFunction or UObject delegate, return the UObject
+	 *
+	 * @return  the object associated with this delegate if there is one
+	 */
+	inline class UObject* GetUObject( ) const
+	{
+		if (DelegateInstance != nullptr)
+		{
+			if ((DelegateInstance->GetType() == EDelegateInstanceType::UFunction) ||
+				(DelegateInstance->GetType() == EDelegateInstanceType::UObjectMethod))
+			{
+				return (class UObject*)DelegateInstance->GetRawUserObject();
+			}
+		}
+
+		return nullptr;
+	}
+
+	/**
+	 * Checks to see if the user object bound to this delegate is still valid
+	 *
+	 * @return  True if the user object is still valid and it's safe to execute the function call
+	 */
+	inline bool IsBound( ) const
+	{
+		return (DelegateInstance != nullptr) && (DelegateInstance->IsSafeToExecute());
+	}
+
+	/** 
+	 * Checks to see if this delegate is bound to the given user object.
+	 *
+	 * @return	True if this delegate is bound to InUserObject, false otherwise.
+	 */
+	inline bool IsBoundToObject( void const* InUserObject ) const
+	{
+		return (InUserObject != nullptr) && (DelegateInstance != nullptr) && DelegateInstance->HasSameObject(InUserObject);
+	}
+
+	/**
+	 * Checks to see if the user object bound to this delegate will ever be valid again
+	 *
+	 * @return  True if the user object could be valid in the future or is valid now.
+	 */
+	inline bool IsCompactable( ) const
+	{
+		return (DelegateInstance != nullptr) && DelegateInstance->IsCompactable();
+	}
+
+	/**
+	 * Unbinds this delegate
+	 */
+	inline void Unbind( )
+	{
+		if (DelegateInstance != nullptr)
+		{
+			delete DelegateInstance;
+			DelegateInstance = nullptr;
+		}
+	}
+
+protected:
+
+	/**
+	 * Gets the delegate instance.
+	 *
+	 * @return The delegate instance.
+	 *
+	 * @see SetDelegateInstance
+	 */
+	IDelegateInstance* GetDelegateInstance( ) const
+	{
+		return DelegateInstance;
+	}
+
+	/**
+	 * Sets the delegate instance.
+	 *
+	 * @param InDelegateInstance The delegate instance to set.
+	 *
+	 * @see GetDelegateInstance.
+	 */
+	inline void SetDelegateInstance( IDelegateInstance* InDelegateInstance )
+	{
+		Unbind();
+		DelegateInstance = InDelegateInstance;
+	}
+
+private:
+
+	// Holds the delegate instance
+	IDelegateInstance* DelegateInstance;
+};
diff --git a/Engine/Source/Runtime/Core/Public/Delegates/DelegateInstanceInterfaceImpl.inl b/Engine/Source/Runtime/Core/Public/Delegates/DelegateInstanceInterfaceImpl.inl
new file mode 100644
--- /dev/null
+++ b/Engine/Source/Runtime/Core/Public/Delegates/DelegateInstanceInterfaceImpl.inl
@@ -0,0 +1,99 @@
+// Copyright 1998-2014 Epic Games, Inc. All Rights Reserved.
+
+// Only designed to be included directly by Delegate.h
+#if !defined( __Delegate_h__ ) || !defined( FUNC_INCLUDING_INLINE_IMPL )
+	#error "This inline header must only be included by Delegate.h"
+#endif
+
+
+// NOTE: This file in re-included for EVERY delegate signature that we support.  That is, every combination
+//		 of parameter count, return value presence or other function modifier will include this file to
+//		 generate a delegate interface type and implementation type for that signature.
+
+#define DELEGATE_INSTANCE_INTERFACE_CLASS FUNC_COMBINE( IBaseDelegateInstance_, FUNC_SUFFIX )
+
+#if FUNC_HAS_PARAMS
+	#define DELEGATE_INSTANCE_COLON_PARAM_INITIALIZER_LIST : FUNC_PARAM_INITIALIZER_LIST
+#else
+	#define DELEGATE_INSTANCE_COLON_PARAM_INITIALIZER_LIST
+#endif
+
+
+/**
+ * Delegate instance base interface.  For internal use only.  This abstract class is used to interact with
+ * a delegate instance for *any* type of object (e.g. raw, shared pointer, UObject-based, etc.)
+ */
+template< FUNC_TEMPLATE_DECL_TYPENAME >
+class DELEGATE_INSTANCE_INTERFACE_CLASS
+	: public IDelegateInstance
+{
+	// Structure for UFunction call parameter marshaling.
+	struct FParms
+	{
+		FUNC_PARAM_MEMBERS
+
+#if !FUNC_IS_VOID
+		RetValType Result;
+#endif
+
+		FParms( FUNC_PARAM_LIST )
+			DELEGATE_INSTANCE_COLON_PARAM_INITIALIZER_LIST
+		{ }
+	};
+
+public:
+
+	/**
+	 * Creates a copy of the delegate instance
+	 *
+	 * @return	The newly created copy
+	 */
+	virtual DELEGATE_INSTANCE_INTERFACE_CLASS< FUNC_TEMPLATE_ARGS >* CreateCopy() = 0;
+
+	/**
+	 * Returns true if this delegate points to exactly the same object and method as the specified delegate,
+	 * even if the delegate objects themselves are different.  Also, the delegate types *must* be compatible.
+	 *
+	 * @param  InOtherDelegate
+	 *
+	 * @return  True if delegates match
+	 */
+	virtual bool IsSameFunction( const DELEGATE_INSTANCE_INTERFACE_CLASS< FUNC_TEMPLATE_ARGS >& InOtherDelegate ) const = 0;
+
+	/**
+	 * Execute the delegate.  If the function pointer is not valid, an error will occur.
+	 */
+	virtual RetValType Execute( FUNC_PARAM_LIST ) const = 0;
+
+#if FUNC_IS_VOID
+
+	/**
+	 * Execute the delegate, but only if the function pointer is still valid
+	 *
+	 * @return  Returns true if the function was executed
+	 */
+	// NOTE: Currently only delegates with no return value support ExecuteIfSafe() 
+	virtual bool ExecuteIfSafe( FUNC_PARAM_LIST ) const = 0;
+
+#endif // FUNC_IS_VOID
+
+public:
+
+	/**
+	 * Execute the delegate with parameters from the virtual machine.
+	 *
+	 * This non-virtual interface (NVI) method marshals a set of call parameters
+	 * from the UObject virtual machine to the bound delegate handler function.
+	 *
+	 * @param InParms An UFunction parameter structure.
+	 */
+	RetValType ExecuteParms( void* InParms ) const
+	{
+		FParms& Parms = *((FParms*)InParms);
+		return Execute(FUNC_PARAM_PARMS_PASSIN);
+	}
+};
+
+
+#undef DELEGATE_INSTANCE_COLON_PARAM_INITIALIZER_LIST
+#undef DELEGATE_INSTANCE_INTERFACE_CLASS
diff --git a/Engine/Source/Runtime/Core/Public/Delegates/DelegateInstancesImpl.inl b/Engine/Source/Runtime/Core/Public/Delegates/DelegateInstancesImpl.inl
new file mode 100644
--- /dev/null
+++ b/Engine/Source/Runtime/Core/Public/Delegates/DelegateInstancesImpl.inl
@@ -0,0 +1,1019 @@
+// Copyright 1998-2014 Epic Games, Inc. All Rights Reserved.
+
+/*================================================================================
+	DelegateInstancesImpl.inl: Inline implementation of delegate instance classes
+================================================================================*/
+
+
+// Only designed to be included directly by Delegate.h
+#if !defined( __Delegate_h__ ) || !defined( FUNC_INCLUDING_INLINE_IMPL )
+	#error "This inline header must only be included by Delegate.h"
+#endif
+
+
+// NOTE: This file in re-included for EVERY delegate signature that we support.  That is, every combination
+//		 of parameter count, return value presence or other function modifier will include this file to
+//		 generate a delegate interface type and implementation type for that signature.
+
+#define DELEGATE_INSTANCE_INTERFACE_CLASS FUNC_COMBINE( IBaseDelegateInstance_, FUNC_SUFFIX )
+#define SP_METHOD_DELEGATE_INSTANCE_CLASS FUNC_COMBINE( TBaseSPMethodDelegateInstance_, FUNC_COMBINE( FUNC_PAYLOAD_SUFFIX, FUNC_CONST_SUFFIX ) )
+#define RAW_METHOD_DELEGATE_INSTANCE_CLASS FUNC_COMBINE( TBaseRawMethodDelegateInstance_, FUNC_COMBINE( FUNC_PAYLOAD_SUFFIX, FUNC_CONST_SUFFIX ) )
+#define UOBJECT_METHOD_DELEGATE_INSTANCE_CLASS FUNC_COMBINE( TBaseUObjectMethodDelegateInstance_, FUNC_COMBINE( FUNC_PAYLOAD_SUFFIX, FUNC_CONST_SUFFIX ) )
+#define STATIC_DELEGATE_INSTANCE_CLASS FUNC_COMBINE( TBaseStaticDelegateInstance_, FUNC_COMBINE( FUNC_PAYLOAD_SUFFIX, FUNC_CONST_SUFFIX ) )
+#define UFUNCTION_DELEGATE_INSTANCE_CLASS FUNC_COMBINE( TBaseUFunctionDelegateInstance_, FUNC_COMBINE( FUNC_PAYLOAD_SUFFIX, FUNC_CONST_SUFFIX ) )
+
+#if FUNC_IS_CONST
+	#define DELEGATE_QUALIFIER const
+#else
+	#define DELEGATE_QUALIFIER
+#endif
+
+#if FUNC_HAS_PAYLOAD
+	#define DELEGATE_COMMA_PAYLOAD_LIST             , FUNC_PAYLOAD_LIST
+	#define DELEGATE_COMMA_PAYLOAD_PASSTHRU         , FUNC_PAYLOAD_PASSTHRU
+	#define DELEGATE_COMMA_PAYLOAD_INITIALIZER_LIST , FUNC_PAYLOAD_INITIALIZER_LIST
+	#define DELEGATE_COMMA_PAYLOAD_PASSIN           , FUNC_PAYLOAD_PASSIN
+	#if FUNC_HAS_PARAMS
+		#define DELEGATE_PARAM_PASSTHRU_COMMA_PAYLOAD_PASSIN FUNC_PARAM_PASSTHRU, FUNC_PAYLOAD_PASSIN
+		#define DELEGATE_PARAM_LIST_COMMA_PAYLOAD_LIST       FUNC_PARAM_LIST, FUNC_PAYLOAD_LIST
+		#define DELEGATE_PARMS_COLON_INITIALIZER_LIST	     : FUNC_PARAM_INITIALIZER_LIST, FUNC_PAYLOAD_INITIALIZER_LIST
+	#else
+		#define DELEGATE_PARAM_PASSTHRU_COMMA_PAYLOAD_PASSIN FUNC_PAYLOAD_PASSIN
+		#define DELEGATE_PARAM_LIST_COMMA_PAYLOAD_LIST       FUNC_PAYLOAD_LIST
+		#define DELEGATE_PARMS_COLON_INITIALIZER_LIST	     : FUNC_PAYLOAD_INITIALIZER_LIST
+	#endif
+#else
+	#define DELEGATE_COMMA_PAYLOAD_LIST
+	#define DELEGATE_COMMA_PAYLOAD_PASSTHRU
+	#define DELEGATE_COMMA_PAYLOAD_INITIALIZER_LIST
+	#define DELEGATE_COMMA_PAYLOAD_PASSIN
+	#define DELEGATE_PARAM_PASSTHRU_COMMA_PAYLOAD_PASSIN FUNC_PARAM_PASSTHRU
+	#define DELEGATE_PARAM_LIST_COMMA_PAYLOAD_LIST       FUNC_PARAM_LIST
+	#if FUNC_HAS_PARAMS
+		#define DELEGATE_PARMS_COLON_INITIALIZER_LIST	 : FUNC_PARAM_INITIALIZER_LIST
+	#else
+		#define DELEGATE_PARMS_COLON_INITIALIZER_LIST
+	#endif
+#endif
+
+
+template< class UserClass, FUNC_PAYLOAD_TEMPLATE_DECL_TYPENAME > class UOBJECT_METHOD_DELEGATE_INSTANCE_CLASS;
+template< class UserClass, FUNC_PAYLOAD_TEMPLATE_DECL_TYPENAME > class RAW_METHOD_DELEGATE_INSTANCE_CLASS;
+
+
+/**
+ * Shared pointer-based delegate instance class for methods.  For internal use only.
+ */
+template< class UserClass, FUNC_PAYLOAD_TEMPLATE_DECL_TYPENAME, ESPMode::Type SPMode >
+class SP_METHOD_DELEGATE_INSTANCE_CLASS
+	: public DELEGATE_INSTANCE_INTERFACE_CLASS< FUNC_TEMPLATE_ARGS >
+{
+	/**
+	 * Internal, inlined and non-virtual version of GetRawUserObject interface.
+	 */
+	FORCEINLINE const void* GetRawUserObjectInternal() const
+	{
+		return UserObject.Pin().Get();
+	}
+
+	/**
+	 * Internal, inlined and non-virtual version of GetRawMethod interface.
+	 */
+	FORCEINLINE const void* GetRawMethodPtrInternal() const
+	{
+		return *(const void**)&MethodPtr;
+	}
+
+public:
+
+	/** Pointer-to-member typedef */
+	typedef RetValType ( UserClass::*FMethodPtr )( DELEGATE_PARAM_LIST_COMMA_PAYLOAD_LIST ) DELEGATE_QUALIFIER;
+
+	/**
+	 * Static helper function that returns a newly-created delegate object for
+	 * the specified user class and class method
+	 *
+	 * @param  InUserObjectRef  Shared reference to the user's object that contains the class method
+	 * @param  InFunc  Member function pointer to your class method
+	 *
+	 * @return  Returns an instance of the new delegate
+	 */
+	FORCEINLINE static DELEGATE_INSTANCE_INTERFACE_CLASS< FUNC_TEMPLATE_ARGS >* Create( const TSharedPtr< UserClass, SPMode >& InUserObjectRef, FMethodPtr InFunc DELEGATE_COMMA_PAYLOAD_LIST )
+	{
+		return new SP_METHOD_DELEGATE_INSTANCE_CLASS< UserClass, FUNC_PAYLOAD_TEMPLATE_ARGS, SPMode >( InUserObjectRef, InFunc DELEGATE_COMMA_PAYLOAD_PASSTHRU );
+	}
+
+	/**
+	 * Static helper function that returns a newly-created delegate object for
+	 * the specified user class and class method.  Requires that the supplied object derives from TSharedFromThis.
+	 *
+	 * @param  InUserObject  The user's object that contains the class method.  Must derive from TSharedFromThis.
+	 * @param  InFunc  Member function pointer to your class method
+	 *
+	 * @return  Returns an instance of the new delegate
+	 */
+	FORCEINLINE static DELEGATE_INSTANCE_INTERFACE_CLASS< FUNC_TEMPLATE_ARGS >* Create( UserClass* InUserObject, FMethodPtr InFunc DELEGATE_COMMA_PAYLOAD_LIST )
+	{
+		/* We expect the incoming InUserObject to derived from TSharedFromThis */
+		TSharedRef< UserClass > UserObjectRef( StaticCastSharedRef< UserClass >( InUserObject->AsShared() ) );
+		return Create( UserObjectRef, InFunc DELEGATE_COMMA_PAYLOAD_PASSTHRU );
+	}
+
+	/**
+	 * Constructor
+	 *
+	 * @param  InUserObject  A shared reference to an arbitrary object (templated) that hosts the member function
+	 * @param  InMethodPtr  C++ member function pointer for the method to bind
+	 */
+	SP_METHOD_DELEGATE_INSTANCE_CLASS( const TSharedPtr< UserClass, SPMode >& InUserObject, FMethodPtr InMethodPtr DELEGATE_COMMA_PAYLOAD_LIST )
+		: UserObject( InUserObject ),
+		  MethodPtr( InMethodPtr )
+		  DELEGATE_COMMA_PAYLOAD_INITIALIZER_LIST
+	{
+		// NOTE: Shared pointer delegates are allowed to have a null incoming object pointer.  Weak pointers can expire,
+		//       an it is possible for a copy of a delegate instance to end up with a null pointer.
+		checkSlow( MethodPtr != NULL );
+	}
+
+	/**
+	 * Returns the type of delegate instance
+	 *
+	 * @return  Delegate instance type
+	 */
+	virtual EDelegateInstanceType::Type GetType() const OVERRIDE
+	{
+		return SPMode == ESPMode::ThreadSafe ? EDelegateInstanceType::ThreadSafeSharedPointerMethod : EDelegateInstanceType::SharedPointerMethod;
+	}
+
+	/**
+	 * Creates a copy of the delegate instance
+	 *
+	 * @return	The newly created copy
+	 */
+	virtual DELEGATE_INSTANCE_INTERFACE_CLASS< FUNC_TEMPLATE_ARGS >* CreateCopy() OVERRIDE
+	{
+		return Create( UserObject.Pin(), MethodPtr DELEGATE_COMMA_PAYLOAD_PASSIN );
+	}
+
+	/**
+	 * Returns true if this delegate points to exactly the same object and method as the specified delegate,
+	 * even if the delegate objects themselves are different.  Also, the delegate types *must* be compatible.
+	 *
+	 * @param  InOtherDelegate
+	 *
+	 * @return  True if delegates match
+	 */
+	virtual bool IsSameFunction( const DELEGATE_INSTANCE_INTERFACE_CLASS< FUNC_TEMPLATE_ARGS >& InOtherDelegate ) const OVERRIDE
+	{
+		// NOTE: Payload data is not currently considered when comparing delegate instances.  See the comment in
+		// multi-cast delegate's Remove() method for more information.
+
+		if ( InOtherDelegate.GetType() == EDelegateInstanceType::SharedPointerMethod || 
+			   InOtherDelegate.GetType() == EDelegateInstanceType::ThreadSafeSharedPointerMethod ||
+				 InOtherDelegate.GetType() == EDelegateInstanceType::RawMethod ||
+				 InOtherDelegate.GetType() == EDelegateInstanceType::UObjectMethod )
+		{
+			return GetRawMethodPtrInternal() == InOtherDelegate.GetRawMethodPtr() && UserObject.HasSameObject( InOtherDelegate.GetRawUserObject() );
+		}
+
+		return false;
+	}
+
+	/**
+	 * Returns true if this delegate is bound to the specified UserObject,
+	 *
+	 * @param  InUserObject
+	 *
+	 * @return  True if delegate is bound to the specified UserObject
+	 */
+	virtual bool HasSameObject( const void* InUserObject ) const OVERRIDE
+	{
+		return UserObject.HasSameObject( InUserObject );
+	}
+
+	virtual FName GetFunctionName() const OVERRIDE
+	{
+		return NAME_None;
+	}
+
+	virtual const void* GetRawUserObject() const OVERRIDE
+	{
+		return GetRawUserObjectInternal();
+	}
+
+	virtual const void* GetRawMethodPtr() const OVERRIDE
+	{
+		return GetRawMethodPtrInternal();
+	}
+
+	/**
+	 * Execute the delegate.  If the function pointer is not valid, an error will occur.
+	 */
+	virtual RetValType Execute( FUNC_PARAM_LIST ) const OVERRIDE
+	{
+		// Verify that the user object is still valid.  We only have a weak reference to it.
+		TSharedPtr< UserClass, SPMode > SharedUserObject( UserObject.Pin() );
+		checkSlow( SharedUserObject.IsValid() );
+
+		// Safely remove const to work around a compiler issue with instantiating template permutations for 
+		// overloaded functions that take a function pointer typedef as a member of a templated class.  In
+		// all cases where this code is actually invoked, the UserClass will already be a const pointer.
+		typedef typename TRemoveConst<UserClass>::Type MutableUserClass;
+		MutableUserClass* MutableUserObject = const_cast<MutableUserClass*>( SharedUserObject.Get() );
+
+		// Call the member function on the user's object.  And yes, this is the correct C++ syntax for calling a
+		// pointer-to-member function.
+		checkSlow( MethodPtr != NULL );
+
+		return ( MutableUserObject->*MethodPtr )( DELEGATE_PARAM_PASSTHRU_COMMA_PAYLOAD_PASSIN );
+	}
+
+	/**
+	 * Checks to see if the user object bound to this delegate is still valid
+	 *
+	 * @return  True if the object is still valid and it's safe to execute the function call
+	 */
+	virtual bool IsSafeToExecute() const OVERRIDE
+	{
+		return UserObject.IsValid();
+	}
+
+#if FUNC_IS_VOID
+	/**
+	 * Execute the delegate, but only if the function pointer is still valid
+	 *
+	 * @return  Returns true if the function was executed
+	 */
+	virtual bool ExecuteIfSafe( FUNC_PARAM_LIST ) const OVERRIDE
+	{
+		// Verify that the user object is still valid.  We only have a weak reference to it.
+		TSharedPtr< UserClass, SPMode > SharedUserObject( UserObject.Pin() );
+		if( SharedUserObject.IsValid() )
+		{
+			Execute( FUNC_PARAM_PASSTHRU );
+			return true;
+		}
+		return false;
+	}
+#endif
+
+private:
+
+	// Declare ourselves as a friend so we can access other template permutations in IsSameFunction()
+	template< class UserClassNoShadow, FUNC_PAYLOAD_TEMPLATE_DECL_NO_SHADOW, ESPMode::Type SPModeNoShadow > friend class SP_METHOD_DELEGATE_INSTANCE_CLASS;
+
+	// Declare other pointer-based delegates as a friend so IsSameFunction() can compare members
+	template< class UserClassNoShadow, FUNC_PAYLOAD_TEMPLATE_DECL_NO_SHADOW > friend class UOBJECT_METHOD_DELEGATE_INSTANCE_CLASS;
+	template< class UserClassNoShadow, FUNC_PAYLOAD_TEMPLATE_DECL_NO_SHADOW > friend class RAW_METHOD_DELEGATE_INSTANCE_CLASS;
+
+	/** Weak reference to an instance of the user's class which contains a method we would like to call. */
+	TWeakPtr< UserClass, SPMode > UserObject;
+
+	/** C++ member function pointer */
+	FMethodPtr MethodPtr;
+
+	/** Payload member variables (if any) */
+	FUNC_PAYLOAD_MEMBERS
+};
+
+
+/**
+ * Delegate type implementation class for raw C++ methods.  For internal use only.
+ */
+template< class UserClass, FUNC_PAYLOAD_TEMPLATE_DECL_TYPENAME >
+class RAW_METHOD_DELEGATE_INSTANCE_CLASS
+	: public DELEGATE_INSTANCE_INTERFACE_CLASS< FUNC_TEMPLATE_ARGS >
+{
+	// Raw method delegates to UObjects are illegal.
+	checkAtCompileTime((!CanConvertPointerFromTo<UserClass, UObjectBase>::Result), You_Cannot_Use_Raw_Method_Delegates_With_UObjects);
+
+	/**
+	 * Internal, inlined and non-virtual version of GetRawUserObject interface.
+	 */
+	FORCEINLINE const void* GetRawUserObjectInternal() const
+	{
+		return UserObject;
+	}
+
+	/**
+	 * Internal, inlined and non-virtual version of GetRawMethodPtr interface.
+	 */
+	FORCEINLINE const void* GetRawMethodPtrInternal() const
+	{
+		return *(const void**)&MethodPtr;
+	}
+
+public:
+
+	/** Pointer-to-member typedef */
+	typedef RetValType ( UserClass::*FMethodPtr )( DELEGATE_PARAM_LIST_COMMA_PAYLOAD_LIST ) DELEGATE_QUALIFIER;
+
+
+	/**
+	 * Static helper function that returns a delegate object for the specified user class and class method
+	 *
+	 * @param  InUserObject  User's object that contains the class method
+	 * @param  InFunc  Member function pointer to your class method
+	 *
+	 * @return  Returns an interface to the new delegate
+	 */
+	FORCEINLINE static DELEGATE_INSTANCE_INTERFACE_CLASS< FUNC_TEMPLATE_ARGS >* Create( UserClass* InUserObject, FMethodPtr InFunc DELEGATE_COMMA_PAYLOAD_LIST )
+	{
+		return new RAW_METHOD_DELEGATE_INSTANCE_CLASS< UserClass, FUNC_PAYLOAD_TEMPLATE_ARGS >( InUserObject, InFunc DELEGATE_COMMA_PAYLOAD_PASSTHRU );
+	}
+
+	/**
+	 * Constructor
+	 *
+	 * @param  InUserObject  An arbitrary object (templated) that hosts the member function
+	 * @param  InMethodPtr  C++ member function pointer for the method to bind
+	 */
+	RAW_METHOD_DELEGATE_INSTANCE_CLASS( UserClass* InUserObject, FMethodPtr InMethodPtr DELEGATE_COMMA_PAYLOAD_LIST )
+		: UserObject( InUserObject ),
+		  MethodPtr( InMethodPtr )
+		  DELEGATE_COMMA_PAYLOAD_INITIALIZER_LIST
+	{
+		// Non-expirable delegates must always have a non-null object pointer on creation (otherwise they could never execute.)
+		check( InUserObject != NULL && MethodPtr != NULL );
+	}
+
+	/**
+	 * Returns the type of delegate instance
+	 *
+	 * @return  Delegate instance type
+	 */
+	virtual EDelegateInstanceType::Type GetType() const OVERRIDE
+	{
+		return EDelegateInstanceType::RawMethod;
+	}
+
+	/**
+	 * Creates a copy of the delegate instance
+	 *
+	 * @return	The newly created copy
+	 */
+	virtual DELEGATE_INSTANCE_INTERFACE_CLASS< FUNC_TEMPLATE_ARGS >* CreateCopy() OVERRIDE
+	{
+		return Create( UserObject, MethodPtr DELEGATE_COMMA_PAYLOAD_PASSIN );
+	}
+
+	/**
+	 * Returns true if this delegate points to exactly the same object and method as the specified delegate,
+	 * even if the delegate objects themselves are different.  Also, the delegate types *must* be compatible.
+	 *
+	 * @param  InOtherDelegate
+	 *
+	 * @return  True if delegates match
+	 */
+	virtual bool IsSameFunction( const DELEGATE_INSTANCE_INTERFACE_CLASS< FUNC_TEMPLATE_ARGS >& InOtherDelegate ) const OVERRIDE
+	{
+		// NOTE: Payload data is not currently considered when comparing delegate instances.  See the comment in
+		// multi-cast delegate's Remove() method for more information.
+
+		if ( InOtherDelegate.GetType() == EDelegateInstanceType::RawMethod || 
+			   InOtherDelegate.GetType() == EDelegateInstanceType::UObjectMethod ||
+				 InOtherDelegate.GetType() == EDelegateInstanceType::SharedPointerMethod ||
+				 InOtherDelegate.GetType() == EDelegateInstanceType::ThreadSafeSharedPointerMethod )
+		{
+			return GetRawMethodPtrInternal() == InOtherDelegate.GetRawMethodPtr() && UserObject == InOtherDelegate.GetRawUserObject();
+		}
+
+		return false;
+	}
+
+	/**
+	 * Returns true if this delegate is bound to the specified UserObject,
+	 *
+	 * @param  InUserObject
+	 *
+	 * @return  True if delegate is bound to the specified UserObject
+	 */
+	virtual bool HasSameObject( const void* InUserObject ) const OVERRIDE
+	{
+		return UserObject == InUserObject;
+	}
+
+	virtual FName GetFunctionName() const OVERRIDE
+	{
+		return NAME_None;
+	}
+
+	virtual const void* GetRawUserObject() const OVERRIDE
+	{
+		return GetRawUserObjectInternal();
+	}
+
+	virtual const void* GetRawMethodPtr() const OVERRIDE
+	{
+		return GetRawMethodPtrInternal();
+	}
+
+	/**
+	 * Execute the delegate.  If the function pointer is not valid, an error will occur.
+	 */
+	virtual RetValType Execute( FUNC_PARAM_LIST ) const OVERRIDE
+	{
+		// Safely remove const to work around a compiler issue with instantiating template permutations for 
+		// overloaded functions that take a function pointer typedef as a member of a templated class.  In
+		// all cases where this code is actually invoked, the UserClass will already be a const pointer.
+		typedef typename TRemoveConst<UserClass>::Type MutableUserClass;
+		MutableUserClass* MutableUserObject = const_cast<MutableUserClass*>( UserObject );
+
+		// Call the member function on the user's object.  And yes, this is the correct C++ syntax for calling a
+		// pointer-to-member function.
+		checkSlow( MethodPtr != NULL );
+
+		return ( MutableUserObject->*MethodPtr )( DELEGATE_PARAM_PASSTHRU_COMMA_PAYLOAD_PASSIN );
+	}
+
+	/**
+	 * Checks to see if the user object bound to this delegate is still valid
+	 *
+	 * @return  True if the object is still valid and it's safe to execute the function call
+	 */
+	virtual bool IsSafeToExecute() const OVERRIDE
+	{
+		// We never know whether or not it is safe to deference a C++ pointer, but we have to
+		// trust the user in this case.  Prefer using a shared-pointer based delegate type instead!
+		return true;
+	}
+
+#if FUNC_IS_VOID
+	/**
+	 * Execute the delegate, but only if the function pointer is still valid
+	 *
+	 * @return  Returns true if the function was executed
+	 */
+	virtual bool ExecuteIfSafe( FUNC_PARAM_LIST ) const OVERRIDE
+	{
+		// We never know whether or not it is safe to deference a C++ pointer, but we have to
+		// trust the user in this case.  Prefer using a shared-pointer based delegate type instead!
+		Execute( FUNC_PARAM_PASSTHRU );
+		return true;
+	}
+#endif
+
+private:
+
+	// Declare other pointer-based delegates as a friend so IsSameFunction() can compare members
+	template< class UserClassNoShadow, FUNC_PAYLOAD_TEMPLATE_DECL_NO_SHADOW, ESPMode::Type SPModeNoShadow > friend class SP_METHOD_DELEGATE_INSTANCE_CLASS;
+	template< class UserClassNoShadow, FUNC_PAYLOAD_TEMPLATE_DECL_NO_SHADOW > friend class UOBJECT_METHOD_DELEGATE_INSTANCE_CLASS;
+
+	/** Pointer to the user's class which contains a method we would like to call. */
+	UserClass* UserObject;
+
+	/** C++ member function pointer */
+	FMethodPtr MethodPtr;
+
+	/** Payload member variables (if any) */
+	FUNC_PAYLOAD_MEMBERS
+};
+
+
+/**
+ * Delegate type implementation class for UObject methods.  For internal use only.
+ */
+template< class UserClass, FUNC_PAYLOAD_TEMPLATE_DECL_TYPENAME >
+class UOBJECT_METHOD_DELEGATE_INSTANCE_CLASS
+	: public DELEGATE_INSTANCE_INTERFACE_CLASS< FUNC_TEMPLATE_ARGS >
+{
+		// UObject delegates can't be used with non-UObject pointers.
+	checkAtCompileTime((CanConvertPointerFromTo<UserClass, UObjectBase>::Result), You_Cannot_Use_UObject_Method_Delegates_With_Raw_Pointers);
+
+	/**
+	 * Internal, inlined and non-virtual version of GetRawUserObject interface.
+	 */
+	FORCEINLINE const void* GetRawUserObjectInternal() const
+	{
+		return UserObject.Get();
+	}
+
+	/**
+	 * Internal, inlined and non-virtual version of GetRawMethodPtr interface.
+	 */
+	FORCEINLINE const void* GetRawMethodPtrInternal() const
+	{
+		return *(const void**)&MethodPtr;
+	}
+
+public:
+
+	/** Pointer-to-member typedef */
+	typedef RetValType ( UserClass::*FMethodPtr )( DELEGATE_PARAM_LIST_COMMA_PAYLOAD_LIST ) DELEGATE_QUALIFIER;
+
+
+	/**
+	 * Static helper function that returns a delegate object for the specified user class and class method
+	 *
+	 * @param  InUserObject  User's object that contains the class method
+	 * @param  InFunc  Member function pointer to your class method
+	 *
+	 * @return  Returns an interface to the new delegate
+	 */
+	FORCEINLINE static DELEGATE_INSTANCE_INTERFACE_CLASS< FUNC_TEMPLATE_ARGS >* Create( UserClass* InUserObject, FMethodPtr InFunc DELEGATE_COMMA_PAYLOAD_LIST )
+	{
+		return new UOBJECT_METHOD_DELEGATE_INSTANCE_CLASS< UserClass, FUNC_PAYLOAD_TEMPLATE_ARGS >( InUserObject, InFunc DELEGATE_COMMA_PAYLOAD_PASSTHRU );
+	}
+
+	/**
+	 * Constructor
+	 *
+	 * @param  InUserObject  An arbitrary object (templated) that hosts the member function
+	 * @param  InMethodPtr  C++ member function pointer for the method to bind
+	 */
+	UOBJECT_METHOD_DELEGATE_INSTANCE_CLASS( UserClass* InUserObject, FMethodPtr InMethodPtr DELEGATE_COMMA_PAYLOAD_LIST )
+		: UserObject( InUserObject ),
+		  MethodPtr( InMethodPtr )
+		  DELEGATE_COMMA_PAYLOAD_INITIALIZER_LIST
+	{
+		// NOTE: UObject delegates are allowed to have a null incoming object pointer.  UObject weak pointers can expire,
+		//       an it is possible for a copy of a delegate instance to end up with a null pointer.
+		checkSlow( MethodPtr != NULL );
+	}
+
+	/**
+	 * Returns the type of delegate instance
+	 *
+	 * @return  Delegate instance type
+	 */
+	virtual EDelegateInstanceType::Type GetType() const OVERRIDE
+	{
+		return EDelegateInstanceType::UObjectMethod;
+	}
+
+	/**
+	 * Creates a copy of the delegate instance
+	 *
+	 * @return	The newly created copy
+	 */
+	virtual DELEGATE_INSTANCE_INTERFACE_CLASS< FUNC_TEMPLATE_ARGS >* CreateCopy() OVERRIDE
+	{
+		return Create( UserObject.Get(), MethodPtr DELEGATE_COMMA_PAYLOAD_PASSIN );
+	}
+
+	/**
+	 * Returns true if this delegate points to exactly the same object and method as the specified delegate,
+	 * even if the delegate objects themselves are different.  Also, the delegate types *must* be compatible.
+	 *
+	 * @param  InOtherDelegate
+	 *
+	 * @return  True if delegates match
+	 */
+	virtual bool IsSameFunction( const DELEGATE_INSTANCE_INTERFACE_CLASS< FUNC_TEMPLATE_ARGS >& InOtherDelegate ) const OVERRIDE
+	{
+		// NOTE: Payload data is not currently considered when comparing delegate instances.  See the comment in
+		// multi-cast delegate's Remove() method for more information.
+
+		if ( InOtherDelegate.GetType() == EDelegateInstanceType::UObjectMethod || 
+			   InOtherDelegate.GetType() == EDelegateInstanceType::RawMethod ||
+				 InOtherDelegate.GetType() == EDelegateInstanceType::SharedPointerMethod ||
+				 InOtherDelegate.GetType() == EDelegateInstanceType::ThreadSafeSharedPointerMethod )
+		{
+			return GetRawMethodPtrInternal() == InOtherDelegate.GetRawMethodPtr() && UserObject.Get() == InOtherDelegate.GetRawUserObject();
+		}
+
+		return false;
+	}
+
+	/**
+	 * Returns true if this delegate is bound to the specified UserObject,
+	 *
+	 * @param  InUserObject
+	 *
+	 * @return  True if delegate is bound to the specified UserObject
+	 */
+	virtual bool HasSameObject( const void* InUserObject ) const OVERRIDE
+	{
+		return UserObject.Get() == InUserObject;
+	}
+
+	virtual FName GetFunctionName() const OVERRIDE
+	{
+		return NAME_None;
+	}
+
+	virtual const void* GetRawUserObject() const OVERRIDE
+	{
+		return GetRawUserObjectInternal();
+	}
+
+	virtual const void* GetRawMethodPtr() const OVERRIDE
+	{
+		return GetRawMethodPtrInternal();
+	}
+
+	/**
+	 * Execute the delegate.  If the function pointer is not valid, an error will occur.
+	 */
+	virtual RetValType Execute( FUNC_PARAM_LIST ) const OVERRIDE
+	{
+		// Verify that the user object is still valid.  We only have a weak reference to it.
+		checkSlow( UserObject.IsValid() );
+
+		// Safely remove const to work around a compiler issue with instantiating template permutations for 
+		// overloaded functions that take a function pointer typedef as a member of a templated class.  In
+		// all cases where this code is actually invoked, the UserClass will already be a const pointer.
+		typedef typename TRemoveConst<UserClass>::Type MutableUserClass;
+		MutableUserClass* MutableUserObject = const_cast<MutableUserClass*>( UserObject.Get() );
+
+		// Call the member function on the user's object.  And yes, this is the correct C++ syntax for calling a
+		// pointer-to-member function.
+		checkSlow( MethodPtr != NULL );
+
+		return ( MutableUserObject->*MethodPtr )( DELEGATE_PARAM_PASSTHRU_COMMA_PAYLOAD_PASSIN );
+	}
+
+	/**
+	 * Checks to see if the user object bound to this delegate is still valid
+	 *
+	 * @return  True if the object is still valid and it's safe to execute the function call
+	 */
+	virtual bool IsSafeToExecute() const OVERRIDE
+	{
+		return !!UserObject.Get();
+	}
+
+	virtual bool IsCompactable() const OVERRIDE
+	{
+		return !UserObject.Get(true);
+	}
+
+#if FUNC_IS_VOID
+	/**
+	 * Execute the delegate, but only if the function pointer is still valid
+	 *
+	 * @return  Returns true if the function was executed
+	 */
+	virtual bool ExecuteIfSafe( FUNC_PARAM_LIST ) const OVERRIDE
+	{
+		// Verify that the user object is still valid.  We only have a weak reference to it.
+		UserClass* ActualUserObject = UserObject.Get();
+		if( ActualUserObject != NULL )
+		{
+			Execute( FUNC_PARAM_PASSTHRU );
+			return true;
+		}
+		return false;
+	}
+#endif
+
+private:
+
+	// Declare other pointer-based delegates as a friend so IsSameFunction() can compare members
+	template< class UserClassNoShadow, FUNC_PAYLOAD_TEMPLATE_DECL_NO_SHADOW, ESPMode::Type SPModeNoShadow > friend class SP_METHOD_DELEGATE_INSTANCE_CLASS;
+	template< class UserClassNoShadow, FUNC_PAYLOAD_TEMPLATE_DECL_NO_SHADOW > friend class RAW_METHOD_DELEGATE_INSTANCE_CLASS;
+
+	/** Pointer to the user's class which contains a method we would like to call. */
+	TWeakObjectPtr< UserClass > UserObject;
+
+	/** C++ member function pointer */
+	FMethodPtr MethodPtr;
+
+	/** Payload member variables (if any) */
+	FUNC_PAYLOAD_MEMBERS
+};
+
+
+/**
+ * Delegate type implementation class for global functions.  For internal use only.
+ */
+template< FUNC_PAYLOAD_TEMPLATE_DECL_TYPENAME >
+class STATIC_DELEGATE_INSTANCE_CLASS
+	: public DELEGATE_INSTANCE_INTERFACE_CLASS< FUNC_TEMPLATE_ARGS >
+{
+public:
+
+	/** Static function pointer typedef */
+	// NOTE: Static C++ functions can never be const, so 'const' is always omitted here.  That means there is
+	//       usually an extra static delegate class created with a 'const' name, even though it doesn't use const.
+	typedef RetValType ( *FFuncPtr )( DELEGATE_PARAM_LIST_COMMA_PAYLOAD_LIST );
+
+	/**
+	 * Static helper function that returns a newly-created delegate object for a
+	 * static/global function pointer.
+	 *
+	 * @param  InFunc  Static function pointer
+	 *
+	 * @return  Returns a new instance of the delegate
+	 */
+	FORCEINLINE static DELEGATE_INSTANCE_INTERFACE_CLASS< FUNC_TEMPLATE_ARGS >* Create( FFuncPtr InFunc DELEGATE_COMMA_PAYLOAD_LIST )
+	{
+		return new STATIC_DELEGATE_INSTANCE_CLASS< FUNC_PAYLOAD_TEMPLATE_ARGS >( InFunc DELEGATE_COMMA_PAYLOAD_PASSTHRU );
+	}
+
+	/**
+	 * Constructor
+	 *
+	 * @param  InStaticFuncPtr  C++ function pointer
+	 */
+	STATIC_DELEGATE_INSTANCE_CLASS( FFuncPtr InStaticFuncPtr DELEGATE_COMMA_PAYLOAD_LIST )
+		: StaticFuncPtr( InStaticFuncPtr )
+		  DELEGATE_COMMA_PAYLOAD_INITIALIZER_LIST
+	{
+		check( StaticFuncPtr != NULL );
+	}
+
+	/**
+	 * Returns the type of delegate instance
+	 *
+	 * @return  Delegate instance type
+	 */
+	virtual EDelegateInstanceType::Type GetType() const OVERRIDE
+	{
+		return EDelegateInstanceType::Raw;
+	}
+
+	/**
+	 * Creates a copy of the delegate instance
+	 *
+	 * @return	The newly created copy
+	 */
+	virtual DELEGATE_INSTANCE_INTERFACE_CLASS< FUNC_TEMPLATE_ARGS >* CreateCopy() OVERRIDE
+	{
+		return Create( StaticFuncPtr DELEGATE_COMMA_PAYLOAD_PASSIN );
+	}
+
+	/**
+	 * Returns true if this delegate points to exactly the same object and method as the specified delegate,
+	 * even if the delegate objects themselves are different.  However, the delegate types *must* match.
+	 *
+	 * @param  InOtherDelegate
+	 *
+	 * @return  True if delegates match
+	 */
+	virtual bool IsSameFunction( const DELEGATE_INSTANCE_INTERFACE_CLASS< FUNC_TEMPLATE_ARGS >& InOtherDelegate ) const OVERRIDE
+	{
+		// NOTE: Payload data is not currently considered when comparing delegate instances.  See the comment in
+		// multi-cast delegate's Remove() method for more information.
+
+		if( InOtherDelegate.GetType() == EDelegateInstanceType::Raw )
+		{
+			// Downcast to our delegate type and compare
+			const STATIC_DELEGATE_INSTANCE_CLASS< FUNC_PAYLOAD_TEMPLATE_ARGS >& OtherStaticDelegate =
+				static_cast< const STATIC_DELEGATE_INSTANCE_CLASS< FUNC_PAYLOAD_TEMPLATE_ARGS >& >( InOtherDelegate );
+			if( StaticFuncPtr == OtherStaticDelegate.StaticFuncPtr )
+			{
+				return true;
+			}
+		}
+
+		return false;
+	}
+
+	/**
+	 * Returns true if this delegate is bound to the specified UserObject,
+	 *
+	 * @param  InUserObject
+	 *
+	 * @return  True if delegate is bound to the specified UserObject
+	 */
+	virtual bool HasSameObject( const void* UserObject ) const OVERRIDE
+	{
+		// Raw Delegates aren't bound to an object so they can never match
+		return false;
+	}
+
+	virtual FName GetFunctionName() const OVERRIDE
+	{
+		return NAME_None;
+	}
+
+	virtual const void* GetRawUserObject() const OVERRIDE
+	{
+		return NULL;
+	}
+
+	virtual const void* GetRawMethodPtr() const OVERRIDE
+	{
+		return *(const void**)&StaticFuncPtr;
+	}
+
+	/**
+	 * Checks to see if the user object bound to this delegate is still valid
+	 *
+	 * @return  True if the object is still valid and it's safe to execute the function call
+	 */
+	virtual bool IsSafeToExecute() const OVERRIDE
+	{
+		// Static functions are always safe to execute!
+		return true;
+	}
+
+	/**
+	 * Execute the delegate.  If the function pointer is not valid, an error will occur.
+	 */
+	virtual RetValType Execute( FUNC_PARAM_LIST ) const OVERRIDE
+	{
+		// Call the member function on the user's object.  And yes, this is the correct C++ syntax for calling a
+		// pointer-to-member function.
+		checkSlow( StaticFuncPtr != NULL );
+
+		return StaticFuncPtr( DELEGATE_PARAM_PASSTHRU_COMMA_PAYLOAD_PASSIN );
+	}
+
+#if FUNC_IS_VOID
+	/**
+	 * Execute the delegate, but only if the function pointer is still valid
+	 *
+	 * @return  Returns true if the function was executed
+	 */
+	virtual bool ExecuteIfSafe( FUNC_PARAM_LIST ) const OVERRIDE
+	{
+		Execute( FUNC_PARAM_PASSTHRU );
+		return true;
+	}
+#endif
+
+private:
+
+	/** C++ function pointer */
+	FFuncPtr StaticFuncPtr;
+
+	/** Payload member variables (if any) */
+	FUNC_PAYLOAD_MEMBERS
+};
+
+
+#if !FUNC_IS_CONST
+
+/**
+ * Delegate type implementation class for UFunctions.  For internal use only.
+ *
+ * @params UserClass Must be an UObject derived class.
+ */
+template<class UserClass, FUNC_PAYLOAD_TEMPLATE_DECL_TYPENAME>
+class UFUNCTION_DELEGATE_INSTANCE_CLASS
+	: public DELEGATE_INSTANCE_INTERFACE_CLASS<FUNC_TEMPLATE_ARGS>
+{
+	checkAtCompileTime((CanConvertPointerFromTo<UserClass, UObjectBase>::Result), You_Cannot_Use_UFunction_Delegates_With_Non_UObject_Classes);
+
+	// Structure for UFunction call parameter marshaling.
+	struct FParmsWithPayload
+	{
+		FUNC_PARAM_MEMBERS		// (Param1Type Param1; Param2Type Param2; ...)
+		FUNC_PAYLOAD_MEMBERS	// (Var1Type Var1; Var2Type Var2; ...)
+
+#if !FUNC_IS_VOID
+		RetValType Result;		// UFunction return value.
+#endif
+
+		// Constructor.
+		FParmsWithPayload( DELEGATE_PARAM_LIST_COMMA_PAYLOAD_LIST )
+			DELEGATE_PARMS_COLON_INITIALIZER_LIST
+		{ }
+	};
+
+public:
+
+	/**
+	 * Creates and initializes a new instance.
+	 *
+	 * @param InUserObject The UObject to call the function on.
+	 * @param InFunctionName The name of the function call.
+	 */
+	UFUNCTION_DELEGATE_INSTANCE_CLASS( UserClass* InUserObject, const FName& InFunctionName DELEGATE_COMMA_PAYLOAD_LIST )
+		: FunctionName(InFunctionName)
+		, UserObjectPtr(InUserObject)
+		DELEGATE_COMMA_PAYLOAD_INITIALIZER_LIST
+	{
+		check(InFunctionName != NAME_None);
+		
+		if (InUserObject != nullptr)
+		{
+			CachedFunction = UserObjectPtr->FindFunctionChecked(InFunctionName);
+		}		
+	}
+	
+public:
+
+	// Begin IDelegateInstance interface
+
+	virtual FName GetFunctionName( ) const OVERRIDE
+	{
+		return FunctionName;
+	}
+
+	virtual const void* GetRawMethodPtr( ) const OVERRIDE
+	{
+		return nullptr;
+	}
+
+	virtual const void* GetRawUserObject( ) const OVERRIDE
+	{
+		return UserObjectPtr.Get();
+	}
+
+	virtual EDelegateInstanceType::Type GetType() const OVERRIDE
+	{
+		return EDelegateInstanceType::UFunction;
+	}
+
+	virtual bool HasSameObject( const void* InUserObject ) const OVERRIDE
+	{
+		return (UserObjectPtr.Get() == InUserObject);
+	}
+
+	virtual bool IsCompactable( ) const OVERRIDE
+	{
+		return !UserObjectPtr.Get(true);
+	}
+
+	virtual bool IsSafeToExecute( ) const OVERRIDE
+	{
+		return UserObjectPtr.IsValid();
+	}
+
+	// End IDelegateInstance interface
+
+public:
+
+	// Begin DELEGATE_INSTANCE_INTERFACE_CLASS interface
+
+	virtual DELEGATE_INSTANCE_INTERFACE_CLASS<FUNC_TEMPLATE_ARGS>* CreateCopy() OVERRIDE
+	{
+		return Create(UserObjectPtr.Get(), GetFunctionName() DELEGATE_COMMA_PAYLOAD_PASSIN);
+	}
+
+	virtual RetValType Execute( FUNC_PARAM_LIST ) const OVERRIDE
+	{
+		checkSlow(IsSafeToExecute());
+
+		FParmsWithPayload Parms = FParmsWithPayload(DELEGATE_PARAM_PASSTHRU_COMMA_PAYLOAD_PASSIN);
+		UserObjectPtr->ProcessEvent(CachedFunction, &Parms);
+
+#if !FUNC_IS_VOID
+		return Parms.Result;
+#endif
+	}
+
+#if FUNC_IS_VOID
+	virtual bool ExecuteIfSafe( FUNC_PARAM_LIST ) const OVERRIDE
+	{
+		if (IsSafeToExecute())
+		{
+			Execute(FUNC_PARAM_PASSTHRU);
+
+			return true;
+		}
+
+		return false;
+	}
+#endif
+
+	virtual bool IsSameFunction( const DELEGATE_INSTANCE_INTERFACE_CLASS<FUNC_TEMPLATE_ARGS>& Other ) const OVERRIDE
+	{
+		return ((Other.GetType() == EDelegateInstanceType::UFunction) &&
+				(Other.GetRawUserObject() == GetRawUserObject()) &&
+				(Other.GetFunctionName() == GetFunctionName()));
+	}
+
+	// End DELEGATE_INSTANCE_INTERFACE_CLASS interface
+
+public:
+
+	/**
+	 * Static helper function that returns a delegate object for the specified user class and class method.
+	 *
+	 * @param InObject The user object to call the function on.
+	 * @param InFunctionName The name of the function call.
+	 *
+	 * @return  Returns an interface to the new delegate
+	 */
+	FORCEINLINE static DELEGATE_INSTANCE_INTERFACE_CLASS<FUNC_TEMPLATE_ARGS>* Create( UserClass* InUserObject, const FName& InFunctionName DELEGATE_COMMA_PAYLOAD_LIST )
+	{
+		return new UFUNCTION_DELEGATE_INSTANCE_CLASS<UserClass, FUNC_PAYLOAD_TEMPLATE_ARGS>(InUserObject, InFunctionName DELEGATE_COMMA_PAYLOAD_PASSTHRU);
+	}
+
+public:
+
+	// Holds the cached UFunction to call.
+	UFunction* CachedFunction;
+
+	// Holds the name of the function to call.
+	FName FunctionName;
+
+	// The user object to call the function on.
+	TWeakObjectPtr<UserClass> UserObjectPtr;
+
+	// Payload member variables, if any.
+	FUNC_PAYLOAD_MEMBERS
+};
+
+#endif // FUNC_IS_CONST
+
+
+#undef SP_METHOD_DELEGATE_INSTANCE_CLASS
+#undef RAW_METHOD_DELEGATE_INSTANCE_CLASS
+#undef UOBJECT_METHOD_DELEGATE_INSTANCE_CLASS
+#undef STATIC_DELEGATE_INSTANCE_CLASS
+#undef UFUNCTION_DELEGATE_INSTANCE_CLASS
+
+#undef DELEGATE_INSTANCE_INTERFACE_CLASS
+#undef DELEGATE_QUALIFIER
+#undef DELEGATE_COMMA_PAYLOAD_LIST
+#undef DELEGATE_COMMA_PAYLOAD_PASSTHRU
+#undef DELEGATE_COMMA_PAYLOAD_INITIALIZER_LIST
+#undef DELEGATE_COMMA_PAYLOAD_PASSIN
+#undef DELEGATE_PARAM_PASSTHRU_COMMA_PAYLOAD_PASSIN
+#undef DELEGATE_PARAM_LIST_COMMA_PAYLOAD_LIST
+#undef DELEGATE_PARMS_COLON_INITIALIZER_LIST
diff --git a/Engine/Source/Runtime/Core/Public/Delegates/DelegateSignatureImpl.inl b/Engine/Source/Runtime/Core/Public/Delegates/DelegateSignatureImpl.inl
new file mode 100644
--- /dev/null
+++ b/Engine/Source/Runtime/Core/Public/Delegates/DelegateSignatureImpl.inl
@@ -0,0 +1,2168 @@
+// Copyright 1998-2014 Epic Games, Inc. All Rights Reserved.
+
+// Only designed to be included directly by Delegate.h
+#if !defined( __Delegate_h__ ) || !defined( FUNC_INCLUDING_INLINE_IMPL )
+	#error "This inline header must only be included by Delegate.h"
+#endif
+
+
+// NOTE: This file is re-included for EVERY delegate signature that we support.  That is, every combination
+//		 of parameter count, return value presence or other function modifier will include this file to
+//		 generate a delegate interface type and implementation type for that signature.
+
+#define DELEGATE_INSTANCE_INTERFACE_CLASS FUNC_COMBINE( IBaseDelegateInstance_, FUNC_SUFFIX )
+
+#define SP_METHOD_DELEGATE_INSTANCE_CLASS_NoVars FUNC_COMBINE( TBaseSPMethodDelegateInstance_, FUNC_SUFFIX )
+#define SP_METHOD_DELEGATE_INSTANCE_CLASS_NoVars_Const FUNC_COMBINE( SP_METHOD_DELEGATE_INSTANCE_CLASS_NoVars, _Const )
+#define SP_METHOD_DELEGATE_INSTANCE_CLASS_OneVar FUNC_COMBINE( SP_METHOD_DELEGATE_INSTANCE_CLASS_NoVars, _OneVar )
+#define SP_METHOD_DELEGATE_INSTANCE_CLASS_OneVar_Const FUNC_COMBINE( SP_METHOD_DELEGATE_INSTANCE_CLASS_OneVar, _Const )
+#define SP_METHOD_DELEGATE_INSTANCE_CLASS_TwoVars FUNC_COMBINE( SP_METHOD_DELEGATE_INSTANCE_CLASS_NoVars, _TwoVars )
+#define SP_METHOD_DELEGATE_INSTANCE_CLASS_TwoVars_Const FUNC_COMBINE( SP_METHOD_DELEGATE_INSTANCE_CLASS_TwoVars, _Const )
+#define SP_METHOD_DELEGATE_INSTANCE_CLASS_ThreeVars FUNC_COMBINE( SP_METHOD_DELEGATE_INSTANCE_CLASS_NoVars, _ThreeVars )
+#define SP_METHOD_DELEGATE_INSTANCE_CLASS_ThreeVars_Const FUNC_COMBINE( SP_METHOD_DELEGATE_INSTANCE_CLASS_ThreeVars, _Const )
+#define SP_METHOD_DELEGATE_INSTANCE_CLASS_FourVars FUNC_COMBINE( SP_METHOD_DELEGATE_INSTANCE_CLASS_NoVars, _FourVars )
+#define SP_METHOD_DELEGATE_INSTANCE_CLASS_FourVars_Const FUNC_COMBINE( SP_METHOD_DELEGATE_INSTANCE_CLASS_FourVars, _Const )
+
+#define RAW_METHOD_DELEGATE_INSTANCE_CLASS_NoVars FUNC_COMBINE( TBaseRawMethodDelegateInstance_, FUNC_SUFFIX )
+#define RAW_METHOD_DELEGATE_INSTANCE_CLASS_NoVars_Const FUNC_COMBINE( RAW_METHOD_DELEGATE_INSTANCE_CLASS_NoVars, _Const )
+#define RAW_METHOD_DELEGATE_INSTANCE_CLASS_OneVar FUNC_COMBINE( RAW_METHOD_DELEGATE_INSTANCE_CLASS_NoVars, _OneVar )
+#define RAW_METHOD_DELEGATE_INSTANCE_CLASS_OneVar_Const FUNC_COMBINE( RAW_METHOD_DELEGATE_INSTANCE_CLASS_OneVar, _Const )
+#define RAW_METHOD_DELEGATE_INSTANCE_CLASS_TwoVars FUNC_COMBINE( RAW_METHOD_DELEGATE_INSTANCE_CLASS_NoVars, _TwoVars )
+#define RAW_METHOD_DELEGATE_INSTANCE_CLASS_TwoVars_Const FUNC_COMBINE( RAW_METHOD_DELEGATE_INSTANCE_CLASS_TwoVars, _Const )
+#define RAW_METHOD_DELEGATE_INSTANCE_CLASS_ThreeVars FUNC_COMBINE( RAW_METHOD_DELEGATE_INSTANCE_CLASS_NoVars, _ThreeVars )
+#define RAW_METHOD_DELEGATE_INSTANCE_CLASS_ThreeVars_Const FUNC_COMBINE( RAW_METHOD_DELEGATE_INSTANCE_CLASS_ThreeVars, _Const )
+#define RAW_METHOD_DELEGATE_INSTANCE_CLASS_FourVars FUNC_COMBINE( RAW_METHOD_DELEGATE_INSTANCE_CLASS_NoVars, _FourVars )
+#define RAW_METHOD_DELEGATE_INSTANCE_CLASS_FourVars_Const FUNC_COMBINE( RAW_METHOD_DELEGATE_INSTANCE_CLASS_FourVars, _Const )
+
+#define UFUNCTION_DELEGATE_INSTANCE_CLASS_NoVars FUNC_COMBINE( TBaseUFunctionDelegateInstance_, FUNC_SUFFIX )
+#define UFUNCTION_DELEGATE_INSTANCE_CLASS_OneVar FUNC_COMBINE( UFUNCTION_DELEGATE_INSTANCE_CLASS_NoVars, _OneVar )
+#define UFUNCTION_DELEGATE_INSTANCE_CLASS_TwoVars FUNC_COMBINE( UFUNCTION_DELEGATE_INSTANCE_CLASS_NoVars, _TwoVars )
+#define UFUNCTION_DELEGATE_INSTANCE_CLASS_ThreeVars FUNC_COMBINE( UFUNCTION_DELEGATE_INSTANCE_CLASS_NoVars, _ThreeVars )
+#define UFUNCTION_DELEGATE_INSTANCE_CLASS_FourVars FUNC_COMBINE( UFUNCTION_DELEGATE_INSTANCE_CLASS_NoVars, _FourVars )
+
+#define UOBJECT_METHOD_DELEGATE_INSTANCE_CLASS_NoVars FUNC_COMBINE( TBaseUObjectMethodDelegateInstance_, FUNC_SUFFIX )
+#define UOBJECT_METHOD_DELEGATE_INSTANCE_CLASS_NoVars_Const FUNC_COMBINE( UOBJECT_METHOD_DELEGATE_INSTANCE_CLASS_NoVars, _Const )
+#define UOBJECT_METHOD_DELEGATE_INSTANCE_CLASS_OneVar FUNC_COMBINE( UOBJECT_METHOD_DELEGATE_INSTANCE_CLASS_NoVars, _OneVar )
+#define UOBJECT_METHOD_DELEGATE_INSTANCE_CLASS_OneVar_Const FUNC_COMBINE( UOBJECT_METHOD_DELEGATE_INSTANCE_CLASS_OneVar, _Const )
+#define UOBJECT_METHOD_DELEGATE_INSTANCE_CLASS_TwoVars FUNC_COMBINE( UOBJECT_METHOD_DELEGATE_INSTANCE_CLASS_NoVars, _TwoVars )
+#define UOBJECT_METHOD_DELEGATE_INSTANCE_CLASS_TwoVars_Const FUNC_COMBINE( UOBJECT_METHOD_DELEGATE_INSTANCE_CLASS_TwoVars, _Const )
+#define UOBJECT_METHOD_DELEGATE_INSTANCE_CLASS_ThreeVars FUNC_COMBINE( UOBJECT_METHOD_DELEGATE_INSTANCE_CLASS_NoVars, _ThreeVars )
+#define UOBJECT_METHOD_DELEGATE_INSTANCE_CLASS_ThreeVars_Const FUNC_COMBINE( UOBJECT_METHOD_DELEGATE_INSTANCE_CLASS_ThreeVars, _Const )
+#define UOBJECT_METHOD_DELEGATE_INSTANCE_CLASS_FourVars FUNC_COMBINE( UOBJECT_METHOD_DELEGATE_INSTANCE_CLASS_NoVars, _FourVars )
+#define UOBJECT_METHOD_DELEGATE_INSTANCE_CLASS_FourVars_Const FUNC_COMBINE( UOBJECT_METHOD_DELEGATE_INSTANCE_CLASS_FourVars, _Const )
+
+#define STATIC_DELEGATE_INSTANCE_CLASS_NoVars FUNC_COMBINE( TBaseStaticDelegateInstance_, FUNC_SUFFIX )
+#define STATIC_DELEGATE_INSTANCE_CLASS_NoVars_Const FUNC_COMBINE( STATIC_DELEGATE_INSTANCE_CLASS_NoVars, _Const )
+#define STATIC_DELEGATE_INSTANCE_CLASS_OneVar FUNC_COMBINE( STATIC_DELEGATE_INSTANCE_CLASS_NoVars, _OneVar )
+#define STATIC_DELEGATE_INSTANCE_CLASS_OneVar_Const FUNC_COMBINE( STATIC_DELEGATE_INSTANCE_CLASS_OneVar, _Const )
+#define STATIC_DELEGATE_INSTANCE_CLASS_TwoVars FUNC_COMBINE( STATIC_DELEGATE_INSTANCE_CLASS_NoVars, _TwoVars )
+#define STATIC_DELEGATE_INSTANCE_CLASS_TwoVars_Const FUNC_COMBINE( STATIC_DELEGATE_INSTANCE_CLASS_TwoVars, _Const )
+#define STATIC_DELEGATE_INSTANCE_CLASS_ThreeVars FUNC_COMBINE( STATIC_DELEGATE_INSTANCE_CLASS_NoVars, _ThreeVars )
+#define STATIC_DELEGATE_INSTANCE_CLASS_ThreeVars_Const FUNC_COMBINE( STATIC_DELEGATE_INSTANCE_CLASS_ThreeVars, _Const )
+#define STATIC_DELEGATE_INSTANCE_CLASS_FourVars FUNC_COMBINE( STATIC_DELEGATE_INSTANCE_CLASS_NoVars, _FourVars )
+#define STATIC_DELEGATE_INSTANCE_CLASS_FourVars_Const FUNC_COMBINE( STATIC_DELEGATE_INSTANCE_CLASS_FourVars, _Const )
+
+#define DELEGATE_CLASS FUNC_COMBINE( TBaseDelegate_, FUNC_SUFFIX )
+#define BASE_MULTICAST_DELEGATE_CLASS FUNC_COMBINE( TBaseMulticastDelegate_, FUNC_SUFFIX )
+#define MULTICAST_DELEGATE_CLASS FUNC_COMBINE( TMulticastDelegate_, FUNC_SUFFIX )
+#define EVENT_CLASS FUNC_COMBINE( TEvent_, FUNC_SUFFIX )
+#define DYNAMIC_DELEGATE_CLASS FUNC_COMBINE( TBaseDynamicDelegate_, FUNC_SUFFIX )
+#define DYNAMIC_MULTICAST_DELEGATE_CLASS FUNC_COMBINE( TBaseDynamicMulticastDelegate_, FUNC_SUFFIX )
+
+#define PAYLOAD_TEMPLATE_DECL_OneVar typename Var1Type
+#define PAYLOAD_TEMPLATE_LIST_OneVar Var1Type
+#define PAYLOAD_TEMPLATE_ARGS_OneVar Var1Type Var1
+#define PAYLOAD_TEMPLATE_PASSIN_OneVar Var1
+
+#define PAYLOAD_TEMPLATE_DECL_TwoVars PAYLOAD_TEMPLATE_DECL_OneVar, typename Var2Type
+#define PAYLOAD_TEMPLATE_LIST_TwoVars PAYLOAD_TEMPLATE_LIST_OneVar, Var2Type
+#define PAYLOAD_TEMPLATE_ARGS_TwoVars PAYLOAD_TEMPLATE_ARGS_OneVar, Var2Type Var2
+#define PAYLOAD_TEMPLATE_PASSIN_TwoVars PAYLOAD_TEMPLATE_PASSIN_OneVar, Var2
+
+#define PAYLOAD_TEMPLATE_DECL_ThreeVars PAYLOAD_TEMPLATE_DECL_TwoVars, typename Var3Type
+#define PAYLOAD_TEMPLATE_LIST_ThreeVars PAYLOAD_TEMPLATE_LIST_TwoVars, Var3Type
+#define PAYLOAD_TEMPLATE_ARGS_ThreeVars PAYLOAD_TEMPLATE_ARGS_TwoVars, Var3Type Var3
+#define PAYLOAD_TEMPLATE_PASSIN_ThreeVars PAYLOAD_TEMPLATE_PASSIN_TwoVars, Var3
+
+#define PAYLOAD_TEMPLATE_DECL_FourVars PAYLOAD_TEMPLATE_DECL_ThreeVars, typename Var4Type
+#define PAYLOAD_TEMPLATE_LIST_FourVars PAYLOAD_TEMPLATE_LIST_ThreeVars, Var4Type
+#define PAYLOAD_TEMPLATE_ARGS_FourVars PAYLOAD_TEMPLATE_ARGS_ThreeVars, Var4Type Var4
+#define PAYLOAD_TEMPLATE_PASSIN_FourVars PAYLOAD_TEMPLATE_PASSIN_ThreeVars, Var4
+
+
+/**
+ * Unicast delegate base object.
+ *
+ * Use the various DECLARE_DELEGATE macros to create the actual delegate type, templated to
+ * the function signature the delegate is compatible with. Then, you can create an instance
+ * of that class when you want to bind a function to the delegate.
+ */
+template< FUNC_TEMPLATE_DECL_TYPENAME >
+class DELEGATE_CLASS
+	: public FDelegateBase<>
+{
+public:
+
+	/** Type definition for the shared interface of delegate instance types compatible with this delegate class. */
+	typedef DELEGATE_INSTANCE_INTERFACE_CLASS<FUNC_TEMPLATE_ARGS> TDelegateInstanceInterface;
+
+	/** Declare the user's "fast" shared pointer-based delegate instance types. */
+	template< class UserClass                                  > class TSPMethodDelegate                 : public SP_METHOD_DELEGATE_INSTANCE_CLASS_NoVars         < UserClass, FUNC_TEMPLATE_ARGS,                                  ESPMode::Fast > { public: TSPMethodDelegate                ( const TSharedRef< UserClass, ESPMode::Fast >& InUserObject, typename SP_METHOD_DELEGATE_INSTANCE_CLASS_NoVars         < UserClass, FUNC_TEMPLATE_ARGS,                                  ESPMode::Fast >::FMethodPtr InMethodPtr                                  ) : SP_METHOD_DELEGATE_INSTANCE_CLASS_NoVars         < UserClass, FUNC_TEMPLATE_ARGS,                                  ESPMode::Fast >( InUserObject, InMethodPtr                                    ) {} };
+	template< class UserClass                                  > class TSPMethodDelegate_Const           : public SP_METHOD_DELEGATE_INSTANCE_CLASS_NoVars_Const   < UserClass, FUNC_TEMPLATE_ARGS,                                  ESPMode::Fast > { public: TSPMethodDelegate_Const          ( const TSharedRef< UserClass, ESPMode::Fast >& InUserObject, typename SP_METHOD_DELEGATE_INSTANCE_CLASS_NoVars_Const   < UserClass, FUNC_TEMPLATE_ARGS,                                  ESPMode::Fast >::FMethodPtr InMethodPtr                                  ) : SP_METHOD_DELEGATE_INSTANCE_CLASS_NoVars_Const   < UserClass, FUNC_TEMPLATE_ARGS,                                  ESPMode::Fast >( InUserObject, InMethodPtr                                    ) {} };
+	template< class UserClass, PAYLOAD_TEMPLATE_DECL_OneVar    > class TSPMethodDelegate_OneVar          : public SP_METHOD_DELEGATE_INSTANCE_CLASS_OneVar         < UserClass, FUNC_TEMPLATE_ARGS, PAYLOAD_TEMPLATE_LIST_OneVar,    ESPMode::Fast > { public: TSPMethodDelegate_OneVar         ( const TSharedRef< UserClass, ESPMode::Fast >& InUserObject, typename SP_METHOD_DELEGATE_INSTANCE_CLASS_OneVar         < UserClass, FUNC_TEMPLATE_ARGS, PAYLOAD_TEMPLATE_LIST_OneVar,    ESPMode::Fast >::FMethodPtr InMethodPtr, PAYLOAD_TEMPLATE_ARGS_OneVar    ) : SP_METHOD_DELEGATE_INSTANCE_CLASS_OneVar         < UserClass, FUNC_TEMPLATE_ARGS, PAYLOAD_TEMPLATE_LIST_OneVar,    ESPMode::Fast >( InUserObject, InMethodPtr, PAYLOAD_TEMPLATE_PASSIN_OneVar    ) {} };
+	template< class UserClass, PAYLOAD_TEMPLATE_DECL_OneVar    > class TSPMethodDelegate_OneVar_Const    : public SP_METHOD_DELEGATE_INSTANCE_CLASS_OneVar_Const   < UserClass, FUNC_TEMPLATE_ARGS, PAYLOAD_TEMPLATE_LIST_OneVar,    ESPMode::Fast > { public: TSPMethodDelegate_OneVar_Const   ( const TSharedRef< UserClass, ESPMode::Fast >& InUserObject, typename SP_METHOD_DELEGATE_INSTANCE_CLASS_OneVar_Const   < UserClass, FUNC_TEMPLATE_ARGS, PAYLOAD_TEMPLATE_LIST_OneVar,    ESPMode::Fast >::FMethodPtr InMethodPtr, PAYLOAD_TEMPLATE_ARGS_OneVar    ) : SP_METHOD_DELEGATE_INSTANCE_CLASS_OneVar_Const   < UserClass, FUNC_TEMPLATE_ARGS, PAYLOAD_TEMPLATE_LIST_OneVar,    ESPMode::Fast >( InUserObject, InMethodPtr, PAYLOAD_TEMPLATE_PASSIN_OneVar    ) {} };
+	template< class UserClass, PAYLOAD_TEMPLATE_DECL_TwoVars   > class TSPMethodDelegate_TwoVars         : public SP_METHOD_DELEGATE_INSTANCE_CLASS_TwoVars        < UserClass, FUNC_TEMPLATE_ARGS, PAYLOAD_TEMPLATE_LIST_TwoVars,   ESPMode::Fast > { public: TSPMethodDelegate_TwoVars        ( const TSharedRef< UserClass, ESPMode::Fast >& InUserObject, typename SP_METHOD_DELEGATE_INSTANCE_CLASS_TwoVars        < UserClass, FUNC_TEMPLATE_ARGS, PAYLOAD_TEMPLATE_LIST_TwoVars,   ESPMode::Fast >::FMethodPtr InMethodPtr, PAYLOAD_TEMPLATE_ARGS_TwoVars   ) : SP_METHOD_DELEGATE_INSTANCE_CLASS_TwoVars        < UserClass, FUNC_TEMPLATE_ARGS, PAYLOAD_TEMPLATE_LIST_TwoVars,   ESPMode::Fast >( InUserObject, InMethodPtr, PAYLOAD_TEMPLATE_PASSIN_TwoVars   ) {} };
+	template< class UserClass, PAYLOAD_TEMPLATE_DECL_TwoVars   > class TSPMethodDelegate_TwoVars_Const   : public SP_METHOD_DELEGATE_INSTANCE_CLASS_TwoVars_Const  < UserClass, FUNC_TEMPLATE_ARGS, PAYLOAD_TEMPLATE_LIST_TwoVars,   ESPMode::Fast > { public: TSPMethodDelegate_TwoVars_Const  ( const TSharedRef< UserClass, ESPMode::Fast >& InUserObject, typename SP_METHOD_DELEGATE_INSTANCE_CLASS_TwoVars_Const  < UserClass, FUNC_TEMPLATE_ARGS, PAYLOAD_TEMPLATE_LIST_TwoVars,   ESPMode::Fast >::FMethodPtr InMethodPtr, PAYLOAD_TEMPLATE_ARGS_TwoVars   ) : SP_METHOD_DELEGATE_INSTANCE_CLASS_TwoVars_Const  < UserClass, FUNC_TEMPLATE_ARGS, PAYLOAD_TEMPLATE_LIST_TwoVars,   ESPMode::Fast >( InUserObject, InMethodPtr, PAYLOAD_TEMPLATE_PASSIN_TwoVars   ) {} };
+	template< class UserClass, PAYLOAD_TEMPLATE_DECL_ThreeVars > class TSPMethodDelegate_ThreeVars       : public SP_METHOD_DELEGATE_INSTANCE_CLASS_ThreeVars      < UserClass, FUNC_TEMPLATE_ARGS, PAYLOAD_TEMPLATE_LIST_ThreeVars, ESPMode::Fast > { public: TSPMethodDelegate_ThreeVars      ( const TSharedRef< UserClass, ESPMode::Fast >& InUserObject, typename SP_METHOD_DELEGATE_INSTANCE_CLASS_ThreeVars      < UserClass, FUNC_TEMPLATE_ARGS, PAYLOAD_TEMPLATE_LIST_ThreeVars, ESPMode::Fast >::FMethodPtr InMethodPtr, PAYLOAD_TEMPLATE_ARGS_ThreeVars ) : SP_METHOD_DELEGATE_INSTANCE_CLASS_ThreeVars      < UserClass, FUNC_TEMPLATE_ARGS, PAYLOAD_TEMPLATE_LIST_ThreeVars, ESPMode::Fast >( InUserObject, InMethodPtr, PAYLOAD_TEMPLATE_PASSIN_ThreeVars ) {} };
+	template< class UserClass, PAYLOAD_TEMPLATE_DECL_ThreeVars > class TSPMethodDelegate_ThreeVars_Const : public SP_METHOD_DELEGATE_INSTANCE_CLASS_ThreeVars_Const< UserClass, FUNC_TEMPLATE_ARGS, PAYLOAD_TEMPLATE_LIST_ThreeVars, ESPMode::Fast > { public: TSPMethodDelegate_ThreeVars_Const( const TSharedRef< UserClass, ESPMode::Fast >& InUserObject, typename SP_METHOD_DELEGATE_INSTANCE_CLASS_ThreeVars_Const< UserClass, FUNC_TEMPLATE_ARGS, PAYLOAD_TEMPLATE_LIST_ThreeVars, ESPMode::Fast >::FMethodPtr InMethodPtr, PAYLOAD_TEMPLATE_ARGS_ThreeVars ) : SP_METHOD_DELEGATE_INSTANCE_CLASS_ThreeVars_Const< UserClass, FUNC_TEMPLATE_ARGS, PAYLOAD_TEMPLATE_LIST_ThreeVars, ESPMode::Fast >( InUserObject, InMethodPtr, PAYLOAD_TEMPLATE_PASSIN_ThreeVars ) {} };
+	template< class UserClass, PAYLOAD_TEMPLATE_DECL_FourVars  > class TSPMethodDelegate_FourVars        : public SP_METHOD_DELEGATE_INSTANCE_CLASS_FourVars       < UserClass, FUNC_TEMPLATE_ARGS, PAYLOAD_TEMPLATE_LIST_FourVars,  ESPMode::Fast > { public: TSPMethodDelegate_FourVars       ( const TSharedRef< UserClass, ESPMode::Fast >& InUserObject, typename SP_METHOD_DELEGATE_INSTANCE_CLASS_FourVars       < UserClass, FUNC_TEMPLATE_ARGS, PAYLOAD_TEMPLATE_LIST_FourVars,  ESPMode::Fast >::FMethodPtr InMethodPtr, PAYLOAD_TEMPLATE_ARGS_FourVars  ) : SP_METHOD_DELEGATE_INSTANCE_CLASS_FourVars       < UserClass, FUNC_TEMPLATE_ARGS, PAYLOAD_TEMPLATE_LIST_FourVars,  ESPMode::Fast >( InUserObject, InMethodPtr, PAYLOAD_TEMPLATE_PASSIN_FourVars  ) {} };
+	template< class UserClass, PAYLOAD_TEMPLATE_DECL_FourVars  > class TSPMethodDelegate_FourVars_Const  : public SP_METHOD_DELEGATE_INSTANCE_CLASS_FourVars_Const < UserClass, FUNC_TEMPLATE_ARGS, PAYLOAD_TEMPLATE_LIST_FourVars,  ESPMode::Fast > { public: TSPMethodDelegate_FourVars_Const ( const TSharedRef< UserClass, ESPMode::Fast >& InUserObject, typename SP_METHOD_DELEGATE_INSTANCE_CLASS_FourVars_Const < UserClass, FUNC_TEMPLATE_ARGS, PAYLOAD_TEMPLATE_LIST_FourVars,  ESPMode::Fast >::FMethodPtr InMethodPtr, PAYLOAD_TEMPLATE_ARGS_FourVars  ) : SP_METHOD_DELEGATE_INSTANCE_CLASS_FourVars_Const < UserClass, FUNC_TEMPLATE_ARGS, PAYLOAD_TEMPLATE_LIST_FourVars,  ESPMode::Fast >( InUserObject, InMethodPtr, PAYLOAD_TEMPLATE_PASSIN_FourVars  ) {} };
+
+	/** Declare the user's "thread-safe" shared pointer-based delegate instance types. */
+	template< class UserClass                                  > class TThreadSafeSPMethodDelegate                 : public SP_METHOD_DELEGATE_INSTANCE_CLASS_NoVars         < UserClass, FUNC_TEMPLATE_ARGS,                                  ESPMode::ThreadSafe > { public: TThreadSafeSPMethodDelegate                ( const TSharedRef< UserClass, ESPMode::ThreadSafe >& InUserObject, typename SP_METHOD_DELEGATE_INSTANCE_CLASS_NoVars         < UserClass, FUNC_TEMPLATE_ARGS,                                  ESPMode::ThreadSafe >::FMethodPtr InMethodPtr                                  ) : SP_METHOD_DELEGATE_INSTANCE_CLASS_NoVars         < UserClass, FUNC_TEMPLATE_ARGS,                                  ESPMode::ThreadSafe >( InUserObject, InMethodPtr                                    ) {} };
+	template< class UserClass                                  > class TThreadSafeSPMethodDelegate_Const           : public SP_METHOD_DELEGATE_INSTANCE_CLASS_NoVars_Const   < UserClass, FUNC_TEMPLATE_ARGS,                                  ESPMode::ThreadSafe > { public: TThreadSafeSPMethodDelegate_Const          ( const TSharedRef< UserClass, ESPMode::ThreadSafe >& InUserObject, typename SP_METHOD_DELEGATE_INSTANCE_CLASS_NoVars_Const   < UserClass, FUNC_TEMPLATE_ARGS,                                  ESPMode::ThreadSafe >::FMethodPtr InMethodPtr                                  ) : SP_METHOD_DELEGATE_INSTANCE_CLASS_NoVars_Const   < UserClass, FUNC_TEMPLATE_ARGS,                                  ESPMode::ThreadSafe >( InUserObject, InMethodPtr                                    ) {} };
+	template< class UserClass, PAYLOAD_TEMPLATE_DECL_OneVar    > class TThreadSafeSPMethodDelegate_OneVar          : public SP_METHOD_DELEGATE_INSTANCE_CLASS_OneVar         < UserClass, FUNC_TEMPLATE_ARGS, PAYLOAD_TEMPLATE_LIST_OneVar,    ESPMode::ThreadSafe > { public: TThreadSafeSPMethodDelegate_OneVar         ( const TSharedRef< UserClass, ESPMode::ThreadSafe >& InUserObject, typename SP_METHOD_DELEGATE_INSTANCE_CLASS_OneVar         < UserClass, FUNC_TEMPLATE_ARGS, PAYLOAD_TEMPLATE_LIST_OneVar,    ESPMode::ThreadSafe >::FMethodPtr InMethodPtr, PAYLOAD_TEMPLATE_ARGS_OneVar    ) : SP_METHOD_DELEGATE_INSTANCE_CLASS_OneVar         < UserClass, FUNC_TEMPLATE_ARGS, PAYLOAD_TEMPLATE_LIST_OneVar,    ESPMode::ThreadSafe >( InUserObject, InMethodPtr, PAYLOAD_TEMPLATE_PASSIN_OneVar    ) {} };
+	template< class UserClass, PAYLOAD_TEMPLATE_DECL_OneVar    > class TThreadSafeSPMethodDelegate_OneVar_Const    : public SP_METHOD_DELEGATE_INSTANCE_CLASS_OneVar_Const   < UserClass, FUNC_TEMPLATE_ARGS, PAYLOAD_TEMPLATE_LIST_OneVar,    ESPMode::ThreadSafe > { public: TThreadSafeSPMethodDelegate_OneVar_Const   ( const TSharedRef< UserClass, ESPMode::ThreadSafe >& InUserObject, typename SP_METHOD_DELEGATE_INSTANCE_CLASS_OneVar_Const   < UserClass, FUNC_TEMPLATE_ARGS, PAYLOAD_TEMPLATE_LIST_OneVar,    ESPMode::ThreadSafe >::FMethodPtr InMethodPtr, PAYLOAD_TEMPLATE_ARGS_OneVar    ) : SP_METHOD_DELEGATE_INSTANCE_CLASS_OneVar_Const   < UserClass, FUNC_TEMPLATE_ARGS, PAYLOAD_TEMPLATE_LIST_OneVar,    ESPMode::ThreadSafe >( InUserObject, InMethodPtr, PAYLOAD_TEMPLATE_PASSIN_OneVar    ) {} };
+	template< class UserClass, PAYLOAD_TEMPLATE_DECL_TwoVars   > class TThreadSafeSPMethodDelegate_TwoVars         : public SP_METHOD_DELEGATE_INSTANCE_CLASS_TwoVars        < UserClass, FUNC_TEMPLATE_ARGS, PAYLOAD_TEMPLATE_LIST_TwoVars,   ESPMode::ThreadSafe > { public: TThreadSafeSPMethodDelegate_TwoVars        ( const TSharedRef< UserClass, ESPMode::ThreadSafe >& InUserObject, typename SP_METHOD_DELEGATE_INSTANCE_CLASS_TwoVars        < UserClass, FUNC_TEMPLATE_ARGS, PAYLOAD_TEMPLATE_LIST_TwoVars,   ESPMode::ThreadSafe >::FMethodPtr InMethodPtr, PAYLOAD_TEMPLATE_ARGS_TwoVars   ) : SP_METHOD_DELEGATE_INSTANCE_CLASS_TwoVars        < UserClass, FUNC_TEMPLATE_ARGS, PAYLOAD_TEMPLATE_LIST_TwoVars,   ESPMode::ThreadSafe >( InUserObject, InMethodPtr, PAYLOAD_TEMPLATE_PASSIN_TwoVars   ) {} };
+	template< class UserClass, PAYLOAD_TEMPLATE_DECL_TwoVars   > class TThreadSafeSPMethodDelegate_TwoVars_Const   : public SP_METHOD_DELEGATE_INSTANCE_CLASS_TwoVars_Const  < UserClass, FUNC_TEMPLATE_ARGS, PAYLOAD_TEMPLATE_LIST_TwoVars,   ESPMode::ThreadSafe > { public: TThreadSafeSPMethodDelegate_TwoVars_Const  ( const TSharedRef< UserClass, ESPMode::ThreadSafe >& InUserObject, typename SP_METHOD_DELEGATE_INSTANCE_CLASS_TwoVars_Const  < UserClass, FUNC_TEMPLATE_ARGS, PAYLOAD_TEMPLATE_LIST_TwoVars,   ESPMode::ThreadSafe >::FMethodPtr InMethodPtr, PAYLOAD_TEMPLATE_ARGS_TwoVars   ) : SP_METHOD_DELEGATE_INSTANCE_CLASS_TwoVars_Const  < UserClass, FUNC_TEMPLATE_ARGS, PAYLOAD_TEMPLATE_LIST_TwoVars,   ESPMode::ThreadSafe >( InUserObject, InMethodPtr, PAYLOAD_TEMPLATE_PASSIN_TwoVars   ) {} };
+	template< class UserClass, PAYLOAD_TEMPLATE_DECL_ThreeVars > class TThreadSafeSPMethodDelegate_ThreeVars       : public SP_METHOD_DELEGATE_INSTANCE_CLASS_ThreeVars      < UserClass, FUNC_TEMPLATE_ARGS, PAYLOAD_TEMPLATE_LIST_ThreeVars, ESPMode::ThreadSafe > { public: TThreadSafeSPMethodDelegate_ThreeVars      ( const TSharedRef< UserClass, ESPMode::ThreadSafe >& InUserObject, typename SP_METHOD_DELEGATE_INSTANCE_CLASS_ThreeVars      < UserClass, FUNC_TEMPLATE_ARGS, PAYLOAD_TEMPLATE_LIST_ThreeVars, ESPMode::ThreadSafe >::FMethodPtr InMethodPtr, PAYLOAD_TEMPLATE_ARGS_ThreeVars ) : SP_METHOD_DELEGATE_INSTANCE_CLASS_ThreeVars      < UserClass, FUNC_TEMPLATE_ARGS, PAYLOAD_TEMPLATE_LIST_ThreeVars, ESPMode::ThreadSafe >( InUserObject, InMethodPtr, PAYLOAD_TEMPLATE_PASSIN_ThreeVars ) {} };
+	template< class UserClass, PAYLOAD_TEMPLATE_DECL_ThreeVars > class TThreadSafeSPMethodDelegate_ThreeVars_Const : public SP_METHOD_DELEGATE_INSTANCE_CLASS_ThreeVars_Const< UserClass, FUNC_TEMPLATE_ARGS, PAYLOAD_TEMPLATE_LIST_ThreeVars, ESPMode::ThreadSafe > { public: TThreadSafeSPMethodDelegate_ThreeVars_Const( const TSharedRef< UserClass, ESPMode::ThreadSafe >& InUserObject, typename SP_METHOD_DELEGATE_INSTANCE_CLASS_ThreeVars_Const< UserClass, FUNC_TEMPLATE_ARGS, PAYLOAD_TEMPLATE_LIST_ThreeVars, ESPMode::ThreadSafe >::FMethodPtr InMethodPtr, PAYLOAD_TEMPLATE_ARGS_ThreeVars ) : SP_METHOD_DELEGATE_INSTANCE_CLASS_ThreeVars_Const< UserClass, FUNC_TEMPLATE_ARGS, PAYLOAD_TEMPLATE_LIST_ThreeVars, ESPMode::ThreadSafe >( InUserObject, InMethodPtr, PAYLOAD_TEMPLATE_PASSIN_ThreeVars ) {} };
+	template< class UserClass, PAYLOAD_TEMPLATE_DECL_FourVars  > class TThreadSafeSPMethodDelegate_FourVars        : public SP_METHOD_DELEGATE_INSTANCE_CLASS_FourVars       < UserClass, FUNC_TEMPLATE_ARGS, PAYLOAD_TEMPLATE_LIST_FourVars,  ESPMode::ThreadSafe > { public: TThreadSafeSPMethodDelegate_FourVars       ( const TSharedRef< UserClass, ESPMode::ThreadSafe >& InUserObject, typename SP_METHOD_DELEGATE_INSTANCE_CLASS_FourVars       < UserClass, FUNC_TEMPLATE_ARGS, PAYLOAD_TEMPLATE_LIST_FourVars,  ESPMode::ThreadSafe >::FMethodPtr InMethodPtr, PAYLOAD_TEMPLATE_ARGS_FourVars  ) : SP_METHOD_DELEGATE_INSTANCE_CLASS_FourVars       < UserClass, FUNC_TEMPLATE_ARGS, PAYLOAD_TEMPLATE_LIST_FourVars,  ESPMode::ThreadSafe >( InUserObject, InMethodPtr, PAYLOAD_TEMPLATE_PASSIN_FourVars  ) {} };
+	template< class UserClass, PAYLOAD_TEMPLATE_DECL_FourVars  > class TThreadSafeSPMethodDelegate_FourVars_Const  : public SP_METHOD_DELEGATE_INSTANCE_CLASS_FourVars_Const < UserClass, FUNC_TEMPLATE_ARGS, PAYLOAD_TEMPLATE_LIST_FourVars,  ESPMode::ThreadSafe > { public: TThreadSafeSPMethodDelegate_FourVars_Const ( const TSharedRef< UserClass, ESPMode::ThreadSafe >& InUserObject, typename SP_METHOD_DELEGATE_INSTANCE_CLASS_FourVars_Const < UserClass, FUNC_TEMPLATE_ARGS, PAYLOAD_TEMPLATE_LIST_FourVars,  ESPMode::ThreadSafe >::FMethodPtr InMethodPtr, PAYLOAD_TEMPLATE_ARGS_FourVars  ) : SP_METHOD_DELEGATE_INSTANCE_CLASS_FourVars_Const < UserClass, FUNC_TEMPLATE_ARGS, PAYLOAD_TEMPLATE_LIST_FourVars,  ESPMode::ThreadSafe >( InUserObject, InMethodPtr, PAYLOAD_TEMPLATE_PASSIN_FourVars  ) {} };
+
+	/** Declare the user's C++ pointer-based delegate instance types. */
+	template< class UserClass                                  > class TRawMethodDelegate                 : public RAW_METHOD_DELEGATE_INSTANCE_CLASS_NoVars         < UserClass, FUNC_TEMPLATE_ARGS                                  > { public: TRawMethodDelegate                ( UserClass* InUserObject, typename RAW_METHOD_DELEGATE_INSTANCE_CLASS_NoVars         < UserClass, FUNC_TEMPLATE_ARGS                                  >::FMethodPtr InMethodPtr                                  ) : RAW_METHOD_DELEGATE_INSTANCE_CLASS_NoVars         < UserClass, FUNC_TEMPLATE_ARGS                                  >( InUserObject, InMethodPtr                                    ) {} };
+	template< class UserClass                                  > class TRawMethodDelegate_Const           : public RAW_METHOD_DELEGATE_INSTANCE_CLASS_NoVars_Const   < UserClass, FUNC_TEMPLATE_ARGS                                  > { public: TRawMethodDelegate_Const          ( UserClass* InUserObject, typename RAW_METHOD_DELEGATE_INSTANCE_CLASS_NoVars_Const   < UserClass, FUNC_TEMPLATE_ARGS                                  >::FMethodPtr InMethodPtr                                  ) : RAW_METHOD_DELEGATE_INSTANCE_CLASS_NoVars_Const   < UserClass, FUNC_TEMPLATE_ARGS                                  >( InUserObject, InMethodPtr                                    ) {} };
+	template< class UserClass, PAYLOAD_TEMPLATE_DECL_OneVar    > class TRawMethodDelegate_OneVar          : public RAW_METHOD_DELEGATE_INSTANCE_CLASS_OneVar         < UserClass, FUNC_TEMPLATE_ARGS, PAYLOAD_TEMPLATE_LIST_OneVar    > { public: TRawMethodDelegate_OneVar         ( UserClass* InUserObject, typename RAW_METHOD_DELEGATE_INSTANCE_CLASS_OneVar         < UserClass, FUNC_TEMPLATE_ARGS, PAYLOAD_TEMPLATE_LIST_OneVar    >::FMethodPtr InMethodPtr, PAYLOAD_TEMPLATE_ARGS_OneVar    ) : RAW_METHOD_DELEGATE_INSTANCE_CLASS_OneVar         < UserClass, FUNC_TEMPLATE_ARGS, PAYLOAD_TEMPLATE_LIST_OneVar    >( InUserObject, InMethodPtr, PAYLOAD_TEMPLATE_PASSIN_OneVar    ) {} };
+	template< class UserClass, PAYLOAD_TEMPLATE_DECL_OneVar    > class TRawMethodDelegate_OneVar_Const    : public RAW_METHOD_DELEGATE_INSTANCE_CLASS_OneVar_Const   < UserClass, FUNC_TEMPLATE_ARGS, PAYLOAD_TEMPLATE_LIST_OneVar    > { public: TRawMethodDelegate_OneVar_Const   ( UserClass* InUserObject, typename RAW_METHOD_DELEGATE_INSTANCE_CLASS_OneVar_Const   < UserClass, FUNC_TEMPLATE_ARGS, PAYLOAD_TEMPLATE_LIST_OneVar    >::FMethodPtr InMethodPtr, PAYLOAD_TEMPLATE_ARGS_OneVar    ) : RAW_METHOD_DELEGATE_INSTANCE_CLASS_OneVar_Const   < UserClass, FUNC_TEMPLATE_ARGS, PAYLOAD_TEMPLATE_LIST_OneVar    >( InUserObject, InMethodPtr, PAYLOAD_TEMPLATE_PASSIN_OneVar    ) {} };
+	template< class UserClass, PAYLOAD_TEMPLATE_DECL_TwoVars   > class TRawMethodDelegate_TwoVars         : public RAW_METHOD_DELEGATE_INSTANCE_CLASS_TwoVars        < UserClass, FUNC_TEMPLATE_ARGS, PAYLOAD_TEMPLATE_LIST_TwoVars   > { public: TRawMethodDelegate_TwoVars        ( UserClass* InUserObject, typename RAW_METHOD_DELEGATE_INSTANCE_CLASS_TwoVars        < UserClass, FUNC_TEMPLATE_ARGS, PAYLOAD_TEMPLATE_LIST_TwoVars   >::FMethodPtr InMethodPtr, PAYLOAD_TEMPLATE_ARGS_TwoVars   ) : RAW_METHOD_DELEGATE_INSTANCE_CLASS_TwoVars        < UserClass, FUNC_TEMPLATE_ARGS, PAYLOAD_TEMPLATE_LIST_TwoVars   >( InUserObject, InMethodPtr, PAYLOAD_TEMPLATE_PASSIN_TwoVars   ) {} };
+	template< class UserClass, PAYLOAD_TEMPLATE_DECL_TwoVars   > class TRawMethodDelegate_TwoVars_Const   : public RAW_METHOD_DELEGATE_INSTANCE_CLASS_TwoVars_Const  < UserClass, FUNC_TEMPLATE_ARGS, PAYLOAD_TEMPLATE_LIST_TwoVars   > { public: TRawMethodDelegate_TwoVars_Const  ( UserClass* InUserObject, typename RAW_METHOD_DELEGATE_INSTANCE_CLASS_TwoVars_Const  < UserClass, FUNC_TEMPLATE_ARGS, PAYLOAD_TEMPLATE_LIST_TwoVars   >::FMethodPtr InMethodPtr, PAYLOAD_TEMPLATE_ARGS_TwoVars   ) : RAW_METHOD_DELEGATE_INSTANCE_CLASS_TwoVars_Const  < UserClass, FUNC_TEMPLATE_ARGS, PAYLOAD_TEMPLATE_LIST_TwoVars   >( InUserObject, InMethodPtr, PAYLOAD_TEMPLATE_PASSIN_TwoVars   ) {} };
+	template< class UserClass, PAYLOAD_TEMPLATE_DECL_ThreeVars > class TRawMethodDelegate_ThreeVars       : public RAW_METHOD_DELEGATE_INSTANCE_CLASS_ThreeVars      < UserClass, FUNC_TEMPLATE_ARGS, PAYLOAD_TEMPLATE_LIST_ThreeVars > { public: TRawMethodDelegate_ThreeVars      ( UserClass* InUserObject, typename RAW_METHOD_DELEGATE_INSTANCE_CLASS_ThreeVars      < UserClass, FUNC_TEMPLATE_ARGS, PAYLOAD_TEMPLATE_LIST_ThreeVars >::FMethodPtr InMethodPtr, PAYLOAD_TEMPLATE_ARGS_ThreeVars ) : RAW_METHOD_DELEGATE_INSTANCE_CLASS_ThreeVars      < UserClass, FUNC_TEMPLATE_ARGS, PAYLOAD_TEMPLATE_LIST_ThreeVars >( InUserObject, InMethodPtr, PAYLOAD_TEMPLATE_PASSIN_ThreeVars ) {} };
+	template< class UserClass, PAYLOAD_TEMPLATE_DECL_ThreeVars > class TRawMethodDelegate_ThreeVars_Const : public RAW_METHOD_DELEGATE_INSTANCE_CLASS_ThreeVars_Const< UserClass, FUNC_TEMPLATE_ARGS, PAYLOAD_TEMPLATE_LIST_ThreeVars > { public: TRawMethodDelegate_ThreeVars_Const( UserClass* InUserObject, typename RAW_METHOD_DELEGATE_INSTANCE_CLASS_ThreeVars_Const< UserClass, FUNC_TEMPLATE_ARGS, PAYLOAD_TEMPLATE_LIST_ThreeVars >::FMethodPtr InMethodPtr, PAYLOAD_TEMPLATE_ARGS_ThreeVars ) : RAW_METHOD_DELEGATE_INSTANCE_CLASS_ThreeVars_Const< UserClass, FUNC_TEMPLATE_ARGS, PAYLOAD_TEMPLATE_LIST_ThreeVars >( InUserObject, InMethodPtr, PAYLOAD_TEMPLATE_PASSIN_ThreeVars ) {} };
+	template< class UserClass, PAYLOAD_TEMPLATE_DECL_FourVars  > class TRawMethodDelegate_FourVars        : public RAW_METHOD_DELEGATE_INSTANCE_CLASS_FourVars       < UserClass, FUNC_TEMPLATE_ARGS, PAYLOAD_TEMPLATE_LIST_FourVars  > { public: TRawMethodDelegate_FourVars       ( UserClass* InUserObject, typename RAW_METHOD_DELEGATE_INSTANCE_CLASS_FourVars       < UserClass, FUNC_TEMPLATE_ARGS, PAYLOAD_TEMPLATE_LIST_FourVars  >::FMethodPtr InMethodPtr, PAYLOAD_TEMPLATE_ARGS_FourVars  ) : RAW_METHOD_DELEGATE_INSTANCE_CLASS_FourVars       < UserClass, FUNC_TEMPLATE_ARGS, PAYLOAD_TEMPLATE_LIST_FourVars  >( InUserObject, InMethodPtr, PAYLOAD_TEMPLATE_PASSIN_FourVars  ) {} };
+	template< class UserClass, PAYLOAD_TEMPLATE_DECL_FourVars  > class TRawMethodDelegate_FourVars_Const  : public RAW_METHOD_DELEGATE_INSTANCE_CLASS_FourVars_Const < UserClass, FUNC_TEMPLATE_ARGS, PAYLOAD_TEMPLATE_LIST_FourVars  > { public: TRawMethodDelegate_FourVars_Const ( UserClass* InUserObject, typename RAW_METHOD_DELEGATE_INSTANCE_CLASS_FourVars_Const < UserClass, FUNC_TEMPLATE_ARGS, PAYLOAD_TEMPLATE_LIST_FourVars  >::FMethodPtr InMethodPtr, PAYLOAD_TEMPLATE_ARGS_FourVars  ) : RAW_METHOD_DELEGATE_INSTANCE_CLASS_FourVars_Const < UserClass, FUNC_TEMPLATE_ARGS, PAYLOAD_TEMPLATE_LIST_FourVars  >( InUserObject, InMethodPtr, PAYLOAD_TEMPLATE_PASSIN_FourVars  ) {} };
+	
+	/** Declare the user's UFunction-based delegate instance types. */
+	template< class UObjectTemplate                                  > class TUFunctionDelegateBinding                     : public UFUNCTION_DELEGATE_INSTANCE_CLASS_NoVars   < UObjectTemplate, FUNC_TEMPLATE_ARGS                                  > { public: TUFunctionDelegateBinding          ( UObjectTemplate* InUserObject, const FName& InFunctionName                                  ) : UFUNCTION_DELEGATE_INSTANCE_CLASS_NoVars   < UObjectTemplate, FUNC_TEMPLATE_ARGS                                  >( InUserObject, InFunctionName                                    ) {} };
+	template< class UObjectTemplate, PAYLOAD_TEMPLATE_DECL_OneVar    > class TUFunctionDelegateBinding_OneVar              : public UFUNCTION_DELEGATE_INSTANCE_CLASS_OneVar   < UObjectTemplate, FUNC_TEMPLATE_ARGS, PAYLOAD_TEMPLATE_LIST_OneVar    > { public: TUFunctionDelegateBinding_OneVar   ( UObjectTemplate* InUserObject, const FName& InFunctionName, PAYLOAD_TEMPLATE_ARGS_OneVar    ) : UFUNCTION_DELEGATE_INSTANCE_CLASS_OneVar   < UObjectTemplate, FUNC_TEMPLATE_ARGS, PAYLOAD_TEMPLATE_LIST_OneVar    >( InUserObject, InFunctionName, PAYLOAD_TEMPLATE_PASSIN_OneVar    ) {} };
+	template< class UObjectTemplate, PAYLOAD_TEMPLATE_DECL_TwoVars   > class TUFunctionDelegateBinding_TwoVars             : public UFUNCTION_DELEGATE_INSTANCE_CLASS_TwoVars  < UObjectTemplate, FUNC_TEMPLATE_ARGS, PAYLOAD_TEMPLATE_LIST_TwoVars   > { public: TUFunctionDelegateBinding_TwoVars  ( UObjectTemplate* InUserObject, const FName& InFunctionName, PAYLOAD_TEMPLATE_ARGS_TwoVars   ) : UFUNCTION_DELEGATE_INSTANCE_CLASS_TwoVars  < UObjectTemplate, FUNC_TEMPLATE_ARGS, PAYLOAD_TEMPLATE_LIST_TwoVars   >( InUserObject, InFunctionName, PAYLOAD_TEMPLATE_PASSIN_TwoVars   ) {} };
+	template< class UObjectTemplate, PAYLOAD_TEMPLATE_DECL_ThreeVars > class TUFunctionDelegateBinding_ThreeVars           : public UFUNCTION_DELEGATE_INSTANCE_CLASS_ThreeVars< UObjectTemplate, FUNC_TEMPLATE_ARGS, PAYLOAD_TEMPLATE_LIST_ThreeVars > { public: TUFunctionDelegateBinding_ThreeVars( UObjectTemplate* InUserObject, const FName& InFunctionName, PAYLOAD_TEMPLATE_ARGS_ThreeVars ) : UFUNCTION_DELEGATE_INSTANCE_CLASS_ThreeVars< UObjectTemplate, FUNC_TEMPLATE_ARGS, PAYLOAD_TEMPLATE_LIST_ThreeVars >( InUserObject, InFunctionName, PAYLOAD_TEMPLATE_PASSIN_ThreeVars ) {} };
+	template< class UObjectTemplate, PAYLOAD_TEMPLATE_DECL_FourVars  > class TUFunctionDelegateBinding_FourVars            : public UFUNCTION_DELEGATE_INSTANCE_CLASS_FourVars < UObjectTemplate, FUNC_TEMPLATE_ARGS, PAYLOAD_TEMPLATE_LIST_FourVars  > { public: TUFunctionDelegateBinding_FourVars ( UObjectTemplate* InUserObject, const FName& InFunctionName, PAYLOAD_TEMPLATE_ARGS_FourVars  ) : UFUNCTION_DELEGATE_INSTANCE_CLASS_FourVars < UObjectTemplate, FUNC_TEMPLATE_ARGS, PAYLOAD_TEMPLATE_LIST_FourVars  >( InUserObject, InFunctionName, PAYLOAD_TEMPLATE_PASSIN_FourVars  ) {} };
+
+	/** Declare the user's UObject-based delegate instance types. */
+	template< class UserClass                                  > class TUObjectMethodDelegate                 : public UOBJECT_METHOD_DELEGATE_INSTANCE_CLASS_NoVars         < UserClass, FUNC_TEMPLATE_ARGS                                  > { public: TUObjectMethodDelegate                ( UserClass* InUserObject, typename UOBJECT_METHOD_DELEGATE_INSTANCE_CLASS_NoVars         < UserClass, FUNC_TEMPLATE_ARGS                                  >::FMethodPtr InMethodPtr                                  ) : UOBJECT_METHOD_DELEGATE_INSTANCE_CLASS_NoVars         < UserClass, FUNC_TEMPLATE_ARGS                                  >( InUserObject, InMethodPtr                                    ) {} };
+	template< class UserClass                                  > class TUObjectMethodDelegate_Const           : public UOBJECT_METHOD_DELEGATE_INSTANCE_CLASS_NoVars_Const   < UserClass, FUNC_TEMPLATE_ARGS                                  > { public: TUObjectMethodDelegate_Const          ( UserClass* InUserObject, typename UOBJECT_METHOD_DELEGATE_INSTANCE_CLASS_NoVars_Const   < UserClass, FUNC_TEMPLATE_ARGS                                  >::FMethodPtr InMethodPtr                                  ) : UOBJECT_METHOD_DELEGATE_INSTANCE_CLASS_NoVars_Const   < UserClass, FUNC_TEMPLATE_ARGS                                  >( InUserObject, InMethodPtr                                    ) {} };
+	template< class UserClass, PAYLOAD_TEMPLATE_DECL_OneVar    > class TUObjectMethodDelegate_OneVar          : public UOBJECT_METHOD_DELEGATE_INSTANCE_CLASS_OneVar         < UserClass, FUNC_TEMPLATE_ARGS, PAYLOAD_TEMPLATE_LIST_OneVar    > { public: TUObjectMethodDelegate_OneVar         ( UserClass* InUserObject, typename UOBJECT_METHOD_DELEGATE_INSTANCE_CLASS_OneVar         < UserClass, FUNC_TEMPLATE_ARGS, PAYLOAD_TEMPLATE_LIST_OneVar    >::FMethodPtr InMethodPtr, PAYLOAD_TEMPLATE_ARGS_OneVar    ) : UOBJECT_METHOD_DELEGATE_INSTANCE_CLASS_OneVar         < UserClass, FUNC_TEMPLATE_ARGS, PAYLOAD_TEMPLATE_LIST_OneVar    >( InUserObject, InMethodPtr, PAYLOAD_TEMPLATE_PASSIN_OneVar    ) {} };
+	template< class UserClass, PAYLOAD_TEMPLATE_DECL_OneVar    > class TUObjectMethodDelegate_OneVar_Const    : public UOBJECT_METHOD_DELEGATE_INSTANCE_CLASS_OneVar_Const   < UserClass, FUNC_TEMPLATE_ARGS, PAYLOAD_TEMPLATE_LIST_OneVar    > { public: TUObjectMethodDelegate_OneVar_Const   ( UserClass* InUserObject, typename UOBJECT_METHOD_DELEGATE_INSTANCE_CLASS_OneVar_Const   < UserClass, FUNC_TEMPLATE_ARGS, PAYLOAD_TEMPLATE_LIST_OneVar    >::FMethodPtr InMethodPtr, PAYLOAD_TEMPLATE_ARGS_OneVar    ) : UOBJECT_METHOD_DELEGATE_INSTANCE_CLASS_OneVar_Const   < UserClass, FUNC_TEMPLATE_ARGS, PAYLOAD_TEMPLATE_LIST_OneVar    >( InUserObject, InMethodPtr, PAYLOAD_TEMPLATE_PASSIN_OneVar    ) {} };
+	template< class UserClass, PAYLOAD_TEMPLATE_DECL_TwoVars   > class TUObjectMethodDelegate_TwoVars         : public UOBJECT_METHOD_DELEGATE_INSTANCE_CLASS_TwoVars        < UserClass, FUNC_TEMPLATE_ARGS, PAYLOAD_TEMPLATE_LIST_TwoVars   > { public: TUObjectMethodDelegate_TwoVars        ( UserClass* InUserObject, typename UOBJECT_METHOD_DELEGATE_INSTANCE_CLASS_TwoVars        < UserClass, FUNC_TEMPLATE_ARGS, PAYLOAD_TEMPLATE_LIST_TwoVars   >::FMethodPtr InMethodPtr, PAYLOAD_TEMPLATE_ARGS_TwoVars   ) : UOBJECT_METHOD_DELEGATE_INSTANCE_CLASS_TwoVars        < UserClass, FUNC_TEMPLATE_ARGS, PAYLOAD_TEMPLATE_LIST_TwoVars   >( InUserObject, InMethodPtr, PAYLOAD_TEMPLATE_PASSIN_TwoVars   ) {} };
+	template< class UserClass, PAYLOAD_TEMPLATE_DECL_TwoVars   > class TUObjectMethodDelegate_TwoVars_Const   : public UOBJECT_METHOD_DELEGATE_INSTANCE_CLASS_TwoVars_Const  < UserClass, FUNC_TEMPLATE_ARGS, PAYLOAD_TEMPLATE_LIST_TwoVars   > { public: TUObjectMethodDelegate_TwoVars_Const  ( UserClass* InUserObject, typename UOBJECT_METHOD_DELEGATE_INSTANCE_CLASS_TwoVars_Const  < UserClass, FUNC_TEMPLATE_ARGS, PAYLOAD_TEMPLATE_LIST_TwoVars   >::FMethodPtr InMethodPtr, PAYLOAD_TEMPLATE_ARGS_TwoVars   ) : UOBJECT_METHOD_DELEGATE_INSTANCE_CLASS_TwoVars_Const  < UserClass, FUNC_TEMPLATE_ARGS, PAYLOAD_TEMPLATE_LIST_TwoVars   >( InUserObject, InMethodPtr, PAYLOAD_TEMPLATE_PASSIN_TwoVars   ) {} };
+	template< class UserClass, PAYLOAD_TEMPLATE_DECL_ThreeVars > class TUObjectMethodDelegate_ThreeVars       : public UOBJECT_METHOD_DELEGATE_INSTANCE_CLASS_ThreeVars      < UserClass, FUNC_TEMPLATE_ARGS, PAYLOAD_TEMPLATE_LIST_ThreeVars > { public: TUObjectMethodDelegate_ThreeVars      ( UserClass* InUserObject, typename UOBJECT_METHOD_DELEGATE_INSTANCE_CLASS_ThreeVars      < UserClass, FUNC_TEMPLATE_ARGS, PAYLOAD_TEMPLATE_LIST_ThreeVars >::FMethodPtr InMethodPtr, PAYLOAD_TEMPLATE_ARGS_ThreeVars ) : UOBJECT_METHOD_DELEGATE_INSTANCE_CLASS_ThreeVars      < UserClass, FUNC_TEMPLATE_ARGS, PAYLOAD_TEMPLATE_LIST_ThreeVars >( InUserObject, InMethodPtr, PAYLOAD_TEMPLATE_PASSIN_ThreeVars ) {} };
+	template< class UserClass, PAYLOAD_TEMPLATE_DECL_ThreeVars > class TUObjectMethodDelegate_ThreeVars_Const : public UOBJECT_METHOD_DELEGATE_INSTANCE_CLASS_ThreeVars_Const< UserClass, FUNC_TEMPLATE_ARGS, PAYLOAD_TEMPLATE_LIST_ThreeVars > { public: TUObjectMethodDelegate_ThreeVars_Const( UserClass* InUserObject, typename UOBJECT_METHOD_DELEGATE_INSTANCE_CLASS_ThreeVars_Const< UserClass, FUNC_TEMPLATE_ARGS, PAYLOAD_TEMPLATE_LIST_ThreeVars >::FMethodPtr InMethodPtr, PAYLOAD_TEMPLATE_ARGS_ThreeVars ) : UOBJECT_METHOD_DELEGATE_INSTANCE_CLASS_ThreeVars_Const< UserClass, FUNC_TEMPLATE_ARGS, PAYLOAD_TEMPLATE_LIST_ThreeVars >( InUserObject, InMethodPtr, PAYLOAD_TEMPLATE_PASSIN_ThreeVars ) {} };
+	template< class UserClass, PAYLOAD_TEMPLATE_DECL_FourVars  > class TUObjectMethodDelegate_FourVars        : public UOBJECT_METHOD_DELEGATE_INSTANCE_CLASS_FourVars       < UserClass, FUNC_TEMPLATE_ARGS, PAYLOAD_TEMPLATE_LIST_FourVars  > { public: TUObjectMethodDelegate_FourVars       ( UserClass* InUserObject, typename UOBJECT_METHOD_DELEGATE_INSTANCE_CLASS_FourVars       < UserClass, FUNC_TEMPLATE_ARGS, PAYLOAD_TEMPLATE_LIST_FourVars  >::FMethodPtr InMethodPtr, PAYLOAD_TEMPLATE_ARGS_FourVars  ) : UOBJECT_METHOD_DELEGATE_INSTANCE_CLASS_FourVars       < UserClass, FUNC_TEMPLATE_ARGS, PAYLOAD_TEMPLATE_LIST_FourVars  >( InUserObject, InMethodPtr, PAYLOAD_TEMPLATE_PASSIN_FourVars  ) {} };
+	template< class UserClass, PAYLOAD_TEMPLATE_DECL_FourVars  > class TUObjectMethodDelegate_FourVars_Const  : public UOBJECT_METHOD_DELEGATE_INSTANCE_CLASS_FourVars_Const < UserClass, FUNC_TEMPLATE_ARGS, PAYLOAD_TEMPLATE_LIST_FourVars  > { public: TUObjectMethodDelegate_FourVars_Const ( UserClass* InUserObject, typename UOBJECT_METHOD_DELEGATE_INSTANCE_CLASS_FourVars_Const < UserClass, FUNC_TEMPLATE_ARGS, PAYLOAD_TEMPLATE_LIST_FourVars  >::FMethodPtr InMethodPtr, PAYLOAD_TEMPLATE_ARGS_FourVars  ) : UOBJECT_METHOD_DELEGATE_INSTANCE_CLASS_FourVars_Const < UserClass, FUNC_TEMPLATE_ARGS, PAYLOAD_TEMPLATE_LIST_FourVars  >( InUserObject, InMethodPtr, PAYLOAD_TEMPLATE_PASSIN_FourVars  ) {} };
+	
+	/** Declare the user's static function pointer delegate instance types. */
+	typedef STATIC_DELEGATE_INSTANCE_CLASS_NoVars< FUNC_TEMPLATE_ARGS > FStaticDelegate;
+	template< PAYLOAD_TEMPLATE_DECL_OneVar    > class TStaticDelegate_OneVar    : public STATIC_DELEGATE_INSTANCE_CLASS_OneVar   < FUNC_TEMPLATE_ARGS, PAYLOAD_TEMPLATE_LIST_OneVar    > { public: TStaticDelegate_OneVar   ( typename STATIC_DELEGATE_INSTANCE_CLASS_OneVar   < FUNC_TEMPLATE_ARGS, PAYLOAD_TEMPLATE_LIST_OneVar    >::FFuncPtr InFuncPtr, PAYLOAD_TEMPLATE_ARGS_OneVar    ) : STATIC_DELEGATE_INSTANCE_CLASS_OneVar   < FUNC_TEMPLATE_ARGS, PAYLOAD_TEMPLATE_LIST_OneVar    >( InFuncPtr, PAYLOAD_TEMPLATE_PASSIN_OneVar    ) {} };
+	template< PAYLOAD_TEMPLATE_DECL_TwoVars   > class TStaticDelegate_TwoVars   : public STATIC_DELEGATE_INSTANCE_CLASS_TwoVars  < FUNC_TEMPLATE_ARGS, PAYLOAD_TEMPLATE_LIST_TwoVars   > { public: TStaticDelegate_TwoVars  ( typename STATIC_DELEGATE_INSTANCE_CLASS_TwoVars  < FUNC_TEMPLATE_ARGS, PAYLOAD_TEMPLATE_LIST_TwoVars   >::FFuncPtr InFuncPtr, PAYLOAD_TEMPLATE_ARGS_TwoVars   ) : STATIC_DELEGATE_INSTANCE_CLASS_TwoVars  < FUNC_TEMPLATE_ARGS, PAYLOAD_TEMPLATE_LIST_TwoVars   >( InFuncPtr, PAYLOAD_TEMPLATE_PASSIN_TwoVars   ) {} };
+	template< PAYLOAD_TEMPLATE_DECL_ThreeVars > class TStaticDelegate_ThreeVars : public STATIC_DELEGATE_INSTANCE_CLASS_ThreeVars< FUNC_TEMPLATE_ARGS, PAYLOAD_TEMPLATE_LIST_ThreeVars > { public: TStaticDelegate_ThreeVars( typename STATIC_DELEGATE_INSTANCE_CLASS_ThreeVars< FUNC_TEMPLATE_ARGS, PAYLOAD_TEMPLATE_LIST_ThreeVars >::FFuncPtr InFuncPtr, PAYLOAD_TEMPLATE_ARGS_ThreeVars ) : STATIC_DELEGATE_INSTANCE_CLASS_ThreeVars< FUNC_TEMPLATE_ARGS, PAYLOAD_TEMPLATE_LIST_ThreeVars >( InFuncPtr, PAYLOAD_TEMPLATE_PASSIN_ThreeVars ) {} };
+	template< PAYLOAD_TEMPLATE_DECL_FourVars  > class TStaticDelegate_FourVars  : public STATIC_DELEGATE_INSTANCE_CLASS_FourVars < FUNC_TEMPLATE_ARGS, PAYLOAD_TEMPLATE_LIST_FourVars  > { public: TStaticDelegate_FourVars ( typename STATIC_DELEGATE_INSTANCE_CLASS_FourVars < FUNC_TEMPLATE_ARGS, PAYLOAD_TEMPLATE_LIST_FourVars  >::FFuncPtr InFuncPtr, PAYLOAD_TEMPLATE_ARGS_FourVars  ) : STATIC_DELEGATE_INSTANCE_CLASS_FourVars < FUNC_TEMPLATE_ARGS, PAYLOAD_TEMPLATE_LIST_FourVars  >( InFuncPtr, PAYLOAD_TEMPLATE_PASSIN_FourVars  ) {} };
+
+public:
+
+	/**
+	 * Static: Creates a raw C++ pointer global function delegate
+	 */
+	inline static DELEGATE_CLASS< FUNC_TEMPLATE_ARGS > CreateStatic( typename FStaticDelegate::FFuncPtr InFunc )
+	{
+		return DELEGATE_CLASS< FUNC_TEMPLATE_ARGS >( FStaticDelegate::Create( InFunc ) );
+	}
+	template< PAYLOAD_TEMPLATE_DECL_OneVar >
+	inline static DELEGATE_CLASS< FUNC_TEMPLATE_ARGS > CreateStatic( typename TStaticDelegate_OneVar< PAYLOAD_TEMPLATE_LIST_OneVar >::FFuncPtr InFunc, PAYLOAD_TEMPLATE_ARGS_OneVar )
+	{
+		return DELEGATE_CLASS< FUNC_TEMPLATE_ARGS >( TStaticDelegate_OneVar< PAYLOAD_TEMPLATE_LIST_OneVar >::Create( InFunc, PAYLOAD_TEMPLATE_PASSIN_OneVar ) );
+	}
+	template< PAYLOAD_TEMPLATE_DECL_TwoVars >
+	inline static DELEGATE_CLASS< FUNC_TEMPLATE_ARGS > CreateStatic( typename TStaticDelegate_TwoVars< PAYLOAD_TEMPLATE_LIST_TwoVars >::FFuncPtr InFunc, PAYLOAD_TEMPLATE_ARGS_TwoVars )
+	{
+		return DELEGATE_CLASS< FUNC_TEMPLATE_ARGS >( TStaticDelegate_TwoVars< PAYLOAD_TEMPLATE_LIST_TwoVars >::Create( InFunc, PAYLOAD_TEMPLATE_PASSIN_TwoVars ) );
+	}
+	template< PAYLOAD_TEMPLATE_DECL_ThreeVars >
+	inline static DELEGATE_CLASS< FUNC_TEMPLATE_ARGS > CreateStatic( typename TStaticDelegate_ThreeVars< PAYLOAD_TEMPLATE_LIST_ThreeVars >::FFuncPtr InFunc, PAYLOAD_TEMPLATE_ARGS_ThreeVars )
+	{
+		return DELEGATE_CLASS< FUNC_TEMPLATE_ARGS >( TStaticDelegate_ThreeVars< PAYLOAD_TEMPLATE_LIST_ThreeVars >::Create( InFunc, PAYLOAD_TEMPLATE_PASSIN_ThreeVars ) );
+	}
+	template< PAYLOAD_TEMPLATE_DECL_FourVars >
+	inline static DELEGATE_CLASS< FUNC_TEMPLATE_ARGS > CreateStatic( typename TStaticDelegate_FourVars< PAYLOAD_TEMPLATE_LIST_FourVars >::FFuncPtr InFunc, PAYLOAD_TEMPLATE_ARGS_FourVars )
+	{
+		return DELEGATE_CLASS< FUNC_TEMPLATE_ARGS >( TStaticDelegate_FourVars< PAYLOAD_TEMPLATE_LIST_FourVars >::Create( InFunc, PAYLOAD_TEMPLATE_PASSIN_FourVars ) );
+	}
+
+	/**
+	 * Static: Creates a raw C++ pointer member function delegate.
+	 *
+	 * Raw pointer doesn't use any sort of reference, so may be unsafe to call if the object was
+	 * deleted out from underneath your delegate. Be careful when calling Execute()!
+	 */
+	template< class UserClass >
+	inline static DELEGATE_CLASS< FUNC_TEMPLATE_ARGS > CreateRaw( UserClass* InUserObject, typename TRawMethodDelegate< UserClass >::FMethodPtr InFunc )
+	{
+		return DELEGATE_CLASS< FUNC_TEMPLATE_ARGS >( TRawMethodDelegate< UserClass >::Create( InUserObject, InFunc ) );
+	}
+	template< class UserClass >
+	inline static DELEGATE_CLASS< FUNC_TEMPLATE_ARGS > CreateRaw( UserClass* InUserObject, typename TRawMethodDelegate_Const< UserClass >::FMethodPtr InFunc )
+	{
+		return DELEGATE_CLASS< FUNC_TEMPLATE_ARGS >( TRawMethodDelegate_Const< UserClass >::Create( InUserObject, InFunc ) );
+	}
+	template< class UserClass, PAYLOAD_TEMPLATE_DECL_OneVar >
+	inline static DELEGATE_CLASS< FUNC_TEMPLATE_ARGS > CreateRaw( UserClass* InUserObject, typename TRawMethodDelegate_OneVar< UserClass, PAYLOAD_TEMPLATE_LIST_OneVar >::FMethodPtr InFunc, PAYLOAD_TEMPLATE_ARGS_OneVar )
+	{
+		return DELEGATE_CLASS< FUNC_TEMPLATE_ARGS >( TRawMethodDelegate_OneVar< UserClass, PAYLOAD_TEMPLATE_LIST_OneVar >::Create( InUserObject, InFunc, PAYLOAD_TEMPLATE_PASSIN_OneVar ) );
+	}
+	template< class UserClass, PAYLOAD_TEMPLATE_DECL_OneVar >
+	inline static DELEGATE_CLASS< FUNC_TEMPLATE_ARGS > CreateRaw( UserClass* InUserObject, typename TRawMethodDelegate_OneVar_Const< UserClass, PAYLOAD_TEMPLATE_LIST_OneVar >::FMethodPtr InFunc, PAYLOAD_TEMPLATE_ARGS_OneVar )
+	{
+		return DELEGATE_CLASS< FUNC_TEMPLATE_ARGS >( TRawMethodDelegate_OneVar_Const< UserClass, PAYLOAD_TEMPLATE_LIST_OneVar >::Create( InUserObject, InFunc, PAYLOAD_TEMPLATE_PASSIN_OneVar ) );
+	}
+	template< class UserClass, PAYLOAD_TEMPLATE_DECL_TwoVars >
+	inline static DELEGATE_CLASS< FUNC_TEMPLATE_ARGS > CreateRaw( UserClass* InUserObject, typename TRawMethodDelegate_TwoVars< UserClass, PAYLOAD_TEMPLATE_LIST_TwoVars >::FMethodPtr InFunc, PAYLOAD_TEMPLATE_ARGS_TwoVars )
+	{
+		return DELEGATE_CLASS< FUNC_TEMPLATE_ARGS >( TRawMethodDelegate_TwoVars< UserClass, PAYLOAD_TEMPLATE_LIST_TwoVars >::Create( InUserObject, InFunc, PAYLOAD_TEMPLATE_PASSIN_TwoVars ) );
+	}
+	template< class UserClass, PAYLOAD_TEMPLATE_DECL_TwoVars >
+	inline static DELEGATE_CLASS< FUNC_TEMPLATE_ARGS > CreateRaw( UserClass* InUserObject, typename TRawMethodDelegate_TwoVars_Const< UserClass, PAYLOAD_TEMPLATE_LIST_TwoVars >::FMethodPtr InFunc, PAYLOAD_TEMPLATE_ARGS_TwoVars )
+	{
+		return DELEGATE_CLASS< FUNC_TEMPLATE_ARGS >( TRawMethodDelegate_TwoVars_Const< UserClass, PAYLOAD_TEMPLATE_LIST_TwoVars >::Create( InUserObject, InFunc, PAYLOAD_TEMPLATE_PASSIN_TwoVars ) );
+	}
+	template< class UserClass, PAYLOAD_TEMPLATE_DECL_ThreeVars >
+	inline static DELEGATE_CLASS< FUNC_TEMPLATE_ARGS > CreateRaw( UserClass* InUserObject, typename TRawMethodDelegate_ThreeVars< UserClass, PAYLOAD_TEMPLATE_LIST_ThreeVars >::FMethodPtr InFunc, PAYLOAD_TEMPLATE_ARGS_ThreeVars )
+	{
+		return DELEGATE_CLASS< FUNC_TEMPLATE_ARGS >( TRawMethodDelegate_ThreeVars< UserClass, PAYLOAD_TEMPLATE_LIST_ThreeVars >::Create( InUserObject, InFunc, PAYLOAD_TEMPLATE_PASSIN_ThreeVars ) );
+	}
+	template< class UserClass, PAYLOAD_TEMPLATE_DECL_ThreeVars >
+	inline static DELEGATE_CLASS< FUNC_TEMPLATE_ARGS > CreateRaw( UserClass* InUserObject, typename TRawMethodDelegate_ThreeVars_Const< UserClass, PAYLOAD_TEMPLATE_LIST_ThreeVars >::FMethodPtr InFunc, PAYLOAD_TEMPLATE_ARGS_ThreeVars )
+	{
+		return DELEGATE_CLASS< FUNC_TEMPLATE_ARGS >( TRawMethodDelegate_ThreeVars_Const< UserClass, PAYLOAD_TEMPLATE_LIST_ThreeVars >::Create( InUserObject, InFunc, PAYLOAD_TEMPLATE_PASSIN_ThreeVars ) );
+	}
+	template< class UserClass, PAYLOAD_TEMPLATE_DECL_FourVars >
+	inline static DELEGATE_CLASS< FUNC_TEMPLATE_ARGS > CreateRaw( UserClass* InUserObject, typename TRawMethodDelegate_FourVars< UserClass, PAYLOAD_TEMPLATE_LIST_FourVars >::FMethodPtr InFunc, PAYLOAD_TEMPLATE_ARGS_FourVars )
+	{
+		return DELEGATE_CLASS< FUNC_TEMPLATE_ARGS >( TRawMethodDelegate_FourVars< UserClass, PAYLOAD_TEMPLATE_LIST_FourVars >::Create( InUserObject, InFunc, PAYLOAD_TEMPLATE_PASSIN_FourVars ) );
+	}
+	template< class UserClass, PAYLOAD_TEMPLATE_DECL_FourVars >
+	inline static DELEGATE_CLASS< FUNC_TEMPLATE_ARGS > CreateRaw( UserClass* InUserObject, typename TRawMethodDelegate_FourVars_Const< UserClass, PAYLOAD_TEMPLATE_LIST_FourVars >::FMethodPtr InFunc, PAYLOAD_TEMPLATE_ARGS_FourVars )
+	{
+		return DELEGATE_CLASS< FUNC_TEMPLATE_ARGS >( TRawMethodDelegate_FourVars_Const< UserClass, PAYLOAD_TEMPLATE_LIST_FourVars >::Create( InUserObject, InFunc, PAYLOAD_TEMPLATE_PASSIN_FourVars ) );
+	}
+	
+	/**
+	 * Static: Creates a shared pointer-based (fast, not thread-safe) member function delegate.
+	 *
+	 * Shared pointer delegates keep a weak reference to your object.
+	 * You can use ExecuteIfBound() to call them.
+	 */
+	template< class UserClass >
+	inline static DELEGATE_CLASS< FUNC_TEMPLATE_ARGS > CreateSP( const TSharedRef< UserClass, ESPMode::Fast >& InUserObjectRef, typename TSPMethodDelegate< UserClass >::FMethodPtr InFunc )
+	{
+		return DELEGATE_CLASS< FUNC_TEMPLATE_ARGS >( TSPMethodDelegate< UserClass >::Create( InUserObjectRef, InFunc ) );
+	}
+	template< class UserClass >
+	inline static DELEGATE_CLASS< FUNC_TEMPLATE_ARGS > CreateSP( const TSharedRef< UserClass, ESPMode::Fast >& InUserObjectRef, typename TSPMethodDelegate_Const< UserClass >::FMethodPtr InFunc )
+	{
+		return DELEGATE_CLASS< FUNC_TEMPLATE_ARGS >( TSPMethodDelegate_Const< UserClass >::Create( InUserObjectRef, InFunc ) );
+	}
+	template< class UserClass, PAYLOAD_TEMPLATE_DECL_OneVar >
+	inline static DELEGATE_CLASS< FUNC_TEMPLATE_ARGS > CreateSP( const TSharedRef< UserClass, ESPMode::Fast >& InUserObjectRef, typename TSPMethodDelegate_OneVar< UserClass, PAYLOAD_TEMPLATE_LIST_OneVar >::FMethodPtr InFunc, PAYLOAD_TEMPLATE_ARGS_OneVar )
+	{
+		return DELEGATE_CLASS< FUNC_TEMPLATE_ARGS >( TSPMethodDelegate_OneVar< UserClass, PAYLOAD_TEMPLATE_LIST_OneVar >::Create( InUserObjectRef, InFunc, PAYLOAD_TEMPLATE_PASSIN_OneVar ) );
+	}
+	template< class UserClass, PAYLOAD_TEMPLATE_DECL_OneVar >
+	inline static DELEGATE_CLASS< FUNC_TEMPLATE_ARGS > CreateSP( const TSharedRef< UserClass, ESPMode::Fast >& InUserObjectRef, typename TSPMethodDelegate_OneVar_Const< UserClass, PAYLOAD_TEMPLATE_LIST_OneVar >::FMethodPtr InFunc, PAYLOAD_TEMPLATE_ARGS_OneVar )
+	{
+		return DELEGATE_CLASS< FUNC_TEMPLATE_ARGS >( TSPMethodDelegate_OneVar_Const< UserClass, PAYLOAD_TEMPLATE_LIST_OneVar >::Create( InUserObjectRef, InFunc, PAYLOAD_TEMPLATE_PASSIN_OneVar ) );
+	}
+	template< class UserClass, PAYLOAD_TEMPLATE_DECL_TwoVars >
+	inline static DELEGATE_CLASS< FUNC_TEMPLATE_ARGS > CreateSP( const TSharedRef< UserClass, ESPMode::Fast >& InUserObjectRef, typename TSPMethodDelegate_TwoVars< UserClass, PAYLOAD_TEMPLATE_LIST_TwoVars >::FMethodPtr InFunc, PAYLOAD_TEMPLATE_ARGS_TwoVars )
+	{
+		return DELEGATE_CLASS< FUNC_TEMPLATE_ARGS >( TSPMethodDelegate_TwoVars< UserClass, PAYLOAD_TEMPLATE_LIST_TwoVars >::Create( InUserObjectRef, InFunc, PAYLOAD_TEMPLATE_PASSIN_TwoVars ) );
+	}
+	template< class UserClass, PAYLOAD_TEMPLATE_DECL_TwoVars >
+	inline static DELEGATE_CLASS< FUNC_TEMPLATE_ARGS > CreateSP( const TSharedRef< UserClass, ESPMode::Fast >& InUserObjectRef, typename TSPMethodDelegate_TwoVars_Const< UserClass, PAYLOAD_TEMPLATE_LIST_TwoVars >::FMethodPtr InFunc, PAYLOAD_TEMPLATE_ARGS_TwoVars )
+	{
+		return DELEGATE_CLASS< FUNC_TEMPLATE_ARGS >( TSPMethodDelegate_TwoVars_Const< UserClass, PAYLOAD_TEMPLATE_LIST_TwoVars >::Create( InUserObjectRef, InFunc, PAYLOAD_TEMPLATE_PASSIN_TwoVars ) );
+	}
+	template< class UserClass, PAYLOAD_TEMPLATE_DECL_ThreeVars >
+	inline static DELEGATE_CLASS< FUNC_TEMPLATE_ARGS > CreateSP( const TSharedRef< UserClass, ESPMode::Fast >& InUserObjectRef, typename TSPMethodDelegate_ThreeVars< UserClass, PAYLOAD_TEMPLATE_LIST_ThreeVars >::FMethodPtr InFunc, PAYLOAD_TEMPLATE_ARGS_ThreeVars )
+	{
+		return DELEGATE_CLASS< FUNC_TEMPLATE_ARGS >( TSPMethodDelegate_ThreeVars< UserClass, PAYLOAD_TEMPLATE_LIST_ThreeVars >::Create( InUserObjectRef, InFunc, PAYLOAD_TEMPLATE_PASSIN_ThreeVars ) );
+	}
+	template< class UserClass, PAYLOAD_TEMPLATE_DECL_ThreeVars >
+	inline static DELEGATE_CLASS< FUNC_TEMPLATE_ARGS > CreateSP( const TSharedRef< UserClass, ESPMode::Fast >& InUserObjectRef, typename TSPMethodDelegate_ThreeVars_Const< UserClass, PAYLOAD_TEMPLATE_LIST_ThreeVars >::FMethodPtr InFunc, PAYLOAD_TEMPLATE_ARGS_ThreeVars )
+	{
+		return DELEGATE_CLASS< FUNC_TEMPLATE_ARGS >( TSPMethodDelegate_ThreeVars_Const< UserClass, PAYLOAD_TEMPLATE_LIST_ThreeVars >::Create( InUserObjectRef, InFunc, PAYLOAD_TEMPLATE_PASSIN_ThreeVars ) );
+	}
+	template< class UserClass, PAYLOAD_TEMPLATE_DECL_FourVars >
+	inline static DELEGATE_CLASS< FUNC_TEMPLATE_ARGS > CreateSP( const TSharedRef< UserClass, ESPMode::Fast >& InUserObjectRef, typename TSPMethodDelegate_FourVars< UserClass, PAYLOAD_TEMPLATE_LIST_FourVars >::FMethodPtr InFunc, PAYLOAD_TEMPLATE_ARGS_FourVars )
+	{
+		return DELEGATE_CLASS< FUNC_TEMPLATE_ARGS >( TSPMethodDelegate_FourVars< UserClass, PAYLOAD_TEMPLATE_LIST_FourVars >::Create( InUserObjectRef, InFunc, PAYLOAD_TEMPLATE_PASSIN_FourVars ) );
+	}
+	template< class UserClass, PAYLOAD_TEMPLATE_DECL_FourVars >
+	inline static DELEGATE_CLASS< FUNC_TEMPLATE_ARGS > CreateSP( const TSharedRef< UserClass, ESPMode::Fast >& InUserObjectRef, typename TSPMethodDelegate_FourVars_Const< UserClass, PAYLOAD_TEMPLATE_LIST_FourVars >::FMethodPtr InFunc, PAYLOAD_TEMPLATE_ARGS_FourVars )
+	{
+		return DELEGATE_CLASS< FUNC_TEMPLATE_ARGS >( TSPMethodDelegate_FourVars_Const< UserClass, PAYLOAD_TEMPLATE_LIST_FourVars >::Create( InUserObjectRef, InFunc, PAYLOAD_TEMPLATE_PASSIN_FourVars ) );
+	}
+
+	/**
+	 * Static: Creates a shared pointer-based (fast, not thread-safe) member function delegate.
+	 *
+	 * Shared pointer delegates keep a weak reference to your object.
+	 * You can use ExecuteIfBound() to call them.
+	 */
+	template< class UserClass >
+	inline static DELEGATE_CLASS< FUNC_TEMPLATE_ARGS > CreateSP( UserClass* InUserObject, typename TSPMethodDelegate< UserClass >::FMethodPtr InFunc )
+	{
+		return CreateSP( StaticCastSharedRef< UserClass >( InUserObject->AsShared() ), InFunc );
+	}
+	template< class UserClass >
+	inline static DELEGATE_CLASS< FUNC_TEMPLATE_ARGS > CreateSP( UserClass* InUserObject, typename TSPMethodDelegate_Const< UserClass >::FMethodPtr InFunc )
+	{
+		return CreateSP( StaticCastSharedRef< UserClass >( InUserObject->AsShared() ), InFunc );
+	}
+	template< class UserClass, PAYLOAD_TEMPLATE_DECL_OneVar >
+	inline static DELEGATE_CLASS< FUNC_TEMPLATE_ARGS > CreateSP( UserClass* InUserObject, typename TSPMethodDelegate_OneVar< UserClass, PAYLOAD_TEMPLATE_LIST_OneVar >::FMethodPtr InFunc, PAYLOAD_TEMPLATE_ARGS_OneVar )
+	{
+		return CreateSP( StaticCastSharedRef< UserClass >( InUserObject->AsShared() ), InFunc, PAYLOAD_TEMPLATE_PASSIN_OneVar );
+	}
+	template< class UserClass, PAYLOAD_TEMPLATE_DECL_OneVar >
+	inline static DELEGATE_CLASS< FUNC_TEMPLATE_ARGS > CreateSP( UserClass* InUserObject, typename TSPMethodDelegate_OneVar_Const< UserClass, PAYLOAD_TEMPLATE_LIST_OneVar >::FMethodPtr InFunc, PAYLOAD_TEMPLATE_ARGS_OneVar )
+	{
+		return CreateSP( StaticCastSharedRef< UserClass >( InUserObject->AsShared() ), InFunc, PAYLOAD_TEMPLATE_PASSIN_OneVar );
+	}
+	template< class UserClass, PAYLOAD_TEMPLATE_DECL_TwoVars >
+	inline static DELEGATE_CLASS< FUNC_TEMPLATE_ARGS > CreateSP( UserClass* InUserObject, typename TSPMethodDelegate_TwoVars< UserClass, PAYLOAD_TEMPLATE_LIST_TwoVars >::FMethodPtr InFunc, PAYLOAD_TEMPLATE_ARGS_TwoVars )
+	{
+		return CreateSP( StaticCastSharedRef< UserClass >( InUserObject->AsShared() ), InFunc, PAYLOAD_TEMPLATE_PASSIN_TwoVars );
+	}
+	template< class UserClass, PAYLOAD_TEMPLATE_DECL_TwoVars >
+	inline static DELEGATE_CLASS< FUNC_TEMPLATE_ARGS > CreateSP( UserClass* InUserObject, typename TSPMethodDelegate_TwoVars_Const< UserClass, PAYLOAD_TEMPLATE_LIST_TwoVars >::FMethodPtr InFunc, PAYLOAD_TEMPLATE_ARGS_TwoVars )
+	{
+		return CreateSP( StaticCastSharedRef< UserClass >( InUserObject->AsShared() ), InFunc, PAYLOAD_TEMPLATE_PASSIN_TwoVars );
+	}
+	template< class UserClass, PAYLOAD_TEMPLATE_DECL_ThreeVars >
+	inline static DELEGATE_CLASS< FUNC_TEMPLATE_ARGS > CreateSP( UserClass* InUserObject, typename TSPMethodDelegate_ThreeVars< UserClass, PAYLOAD_TEMPLATE_LIST_ThreeVars >::FMethodPtr InFunc, PAYLOAD_TEMPLATE_ARGS_ThreeVars )
+	{
+		return CreateSP( StaticCastSharedRef< UserClass >( InUserObject->AsShared() ), InFunc, PAYLOAD_TEMPLATE_PASSIN_ThreeVars );
+	}
+	template< class UserClass, PAYLOAD_TEMPLATE_DECL_ThreeVars >
+	inline static DELEGATE_CLASS< FUNC_TEMPLATE_ARGS > CreateSP( UserClass* InUserObject, typename TSPMethodDelegate_ThreeVars_Const< UserClass, PAYLOAD_TEMPLATE_LIST_ThreeVars >::FMethodPtr InFunc, PAYLOAD_TEMPLATE_ARGS_ThreeVars )
+	{
+		return CreateSP( StaticCastSharedRef< UserClass >( InUserObject->AsShared() ), InFunc, PAYLOAD_TEMPLATE_PASSIN_ThreeVars );
+	}
+	template< class UserClass, PAYLOAD_TEMPLATE_DECL_FourVars >
+	inline static DELEGATE_CLASS< FUNC_TEMPLATE_ARGS > CreateSP( UserClass* InUserObject, typename TSPMethodDelegate_FourVars< UserClass, PAYLOAD_TEMPLATE_LIST_FourVars >::FMethodPtr InFunc, PAYLOAD_TEMPLATE_ARGS_FourVars )
+	{
+		return CreateSP( StaticCastSharedRef< UserClass >( InUserObject->AsShared() ), InFunc, PAYLOAD_TEMPLATE_PASSIN_FourVars );
+	}
+	template< class UserClass, PAYLOAD_TEMPLATE_DECL_FourVars >
+	inline static DELEGATE_CLASS< FUNC_TEMPLATE_ARGS > CreateSP( UserClass* InUserObject, typename TSPMethodDelegate_FourVars_Const< UserClass, PAYLOAD_TEMPLATE_LIST_FourVars >::FMethodPtr InFunc, PAYLOAD_TEMPLATE_ARGS_FourVars )
+	{
+		return CreateSP( StaticCastSharedRef< UserClass >( InUserObject->AsShared() ), InFunc, PAYLOAD_TEMPLATE_PASSIN_FourVars );
+	}
+	
+	/**
+	 * Static: Creates a shared pointer-based (slower, conditionally thread-safe) member function delegate.
+	 *
+	 * Shared pointer delegates keep a weak reference to your object.
+	 * You can use ExecuteIfBound() to call them.
+	 */
+	template< class UserClass >
+	inline static DELEGATE_CLASS< FUNC_TEMPLATE_ARGS > CreateThreadSafeSP( const TSharedRef< UserClass, ESPMode::ThreadSafe >& InUserObjectRef, typename TThreadSafeSPMethodDelegate< UserClass >::FMethodPtr InFunc )
+	{
+		return DELEGATE_CLASS< FUNC_TEMPLATE_ARGS >( TThreadSafeSPMethodDelegate< UserClass >::Create( InUserObjectRef, InFunc ) );
+	}
+	template< class UserClass >
+	inline static DELEGATE_CLASS< FUNC_TEMPLATE_ARGS > CreateThreadSafeSP( const TSharedRef< UserClass, ESPMode::ThreadSafe >& InUserObjectRef, typename TThreadSafeSPMethodDelegate_Const< UserClass >::FMethodPtr InFunc )
+	{
+		return DELEGATE_CLASS< FUNC_TEMPLATE_ARGS >( TThreadSafeSPMethodDelegate_Const< UserClass >::Create( InUserObjectRef, InFunc ) );
+	}
+	template< class UserClass, PAYLOAD_TEMPLATE_DECL_OneVar >
+	inline static DELEGATE_CLASS< FUNC_TEMPLATE_ARGS > CreateThreadSafeSP( const TSharedRef< UserClass, ESPMode::ThreadSafe >& InUserObjectRef, typename TThreadSafeSPMethodDelegate_OneVar< UserClass, PAYLOAD_TEMPLATE_LIST_OneVar >::FMethodPtr InFunc, PAYLOAD_TEMPLATE_ARGS_OneVar )
+	{
+		return DELEGATE_CLASS< FUNC_TEMPLATE_ARGS >( TThreadSafeSPMethodDelegate_OneVar< UserClass, PAYLOAD_TEMPLATE_LIST_OneVar >::Create( InUserObjectRef, InFunc, PAYLOAD_TEMPLATE_PASSIN_OneVar ) );
+	}
+	template< class UserClass, PAYLOAD_TEMPLATE_DECL_OneVar >
+	inline static DELEGATE_CLASS< FUNC_TEMPLATE_ARGS > CreateThreadSafeSP( const TSharedRef< UserClass, ESPMode::ThreadSafe >& InUserObjectRef, typename TThreadSafeSPMethodDelegate_OneVar_Const< UserClass, PAYLOAD_TEMPLATE_LIST_OneVar >::FMethodPtr InFunc, PAYLOAD_TEMPLATE_ARGS_OneVar )
+	{
+		return DELEGATE_CLASS< FUNC_TEMPLATE_ARGS >( TThreadSafeSPMethodDelegate_OneVar_Const< UserClass, PAYLOAD_TEMPLATE_LIST_OneVar >::Create( InUserObjectRef, InFunc, PAYLOAD_TEMPLATE_PASSIN_OneVar ) );
+	}
+	template< class UserClass, PAYLOAD_TEMPLATE_DECL_TwoVars >
+	inline static DELEGATE_CLASS< FUNC_TEMPLATE_ARGS > CreateThreadSafeSP( const TSharedRef< UserClass, ESPMode::ThreadSafe >& InUserObjectRef, typename TThreadSafeSPMethodDelegate_TwoVars< UserClass, PAYLOAD_TEMPLATE_LIST_TwoVars >::FMethodPtr InFunc, PAYLOAD_TEMPLATE_ARGS_TwoVars )
+	{
+		return DELEGATE_CLASS< FUNC_TEMPLATE_ARGS >( TThreadSafeSPMethodDelegate_TwoVars< UserClass, PAYLOAD_TEMPLATE_LIST_TwoVars >::Create( InUserObjectRef, InFunc, PAYLOAD_TEMPLATE_PASSIN_TwoVars ) );
+	}
+	template< class UserClass, PAYLOAD_TEMPLATE_DECL_TwoVars >
+	inline static DELEGATE_CLASS< FUNC_TEMPLATE_ARGS > CreateThreadSafeSP( const TSharedRef< UserClass, ESPMode::ThreadSafe >& InUserObjectRef, typename TThreadSafeSPMethodDelegate_TwoVars_Const< UserClass, PAYLOAD_TEMPLATE_LIST_TwoVars >::FMethodPtr InFunc, PAYLOAD_TEMPLATE_ARGS_TwoVars )
+	{
+		return DELEGATE_CLASS< FUNC_TEMPLATE_ARGS >( TThreadSafeSPMethodDelegate_TwoVars_Const< UserClass, PAYLOAD_TEMPLATE_LIST_TwoVars >::Create( InUserObjectRef, InFunc, PAYLOAD_TEMPLATE_PASSIN_TwoVars ) );
+	}
+	template< class UserClass, PAYLOAD_TEMPLATE_DECL_ThreeVars >
+	inline static DELEGATE_CLASS< FUNC_TEMPLATE_ARGS > CreateThreadSafeSP( const TSharedRef< UserClass, ESPMode::ThreadSafe >& InUserObjectRef, typename TThreadSafeSPMethodDelegate_ThreeVars< UserClass, PAYLOAD_TEMPLATE_LIST_ThreeVars >::FMethodPtr InFunc, PAYLOAD_TEMPLATE_ARGS_ThreeVars )
+	{
+		return DELEGATE_CLASS< FUNC_TEMPLATE_ARGS >( TThreadSafeSPMethodDelegate_ThreeVars< UserClass, PAYLOAD_TEMPLATE_LIST_ThreeVars >::Create( InUserObjectRef, InFunc, PAYLOAD_TEMPLATE_PASSIN_ThreeVars ) );
+	}
+	template< class UserClass, PAYLOAD_TEMPLATE_DECL_ThreeVars >
+	inline static DELEGATE_CLASS< FUNC_TEMPLATE_ARGS > CreateThreadSafeSP( const TSharedRef< UserClass, ESPMode::ThreadSafe >& InUserObjectRef, typename TThreadSafeSPMethodDelegate_ThreeVars_Const< UserClass, PAYLOAD_TEMPLATE_LIST_ThreeVars >::FMethodPtr InFunc, PAYLOAD_TEMPLATE_ARGS_ThreeVars )
+	{
+		return DELEGATE_CLASS< FUNC_TEMPLATE_ARGS >( TThreadSafeSPMethodDelegate_ThreeVars_Const< UserClass, PAYLOAD_TEMPLATE_LIST_ThreeVars >::Create( InUserObjectRef, InFunc, PAYLOAD_TEMPLATE_PASSIN_ThreeVars ) );
+	}
+	template< class UserClass, PAYLOAD_TEMPLATE_DECL_FourVars >
+	inline static DELEGATE_CLASS< FUNC_TEMPLATE_ARGS > CreateThreadSafeSP( const TSharedRef< UserClass, ESPMode::ThreadSafe >& InUserObjectRef, typename TThreadSafeSPMethodDelegate_FourVars< UserClass, PAYLOAD_TEMPLATE_LIST_FourVars >::FMethodPtr InFunc, PAYLOAD_TEMPLATE_ARGS_FourVars )
+	{
+		return DELEGATE_CLASS< FUNC_TEMPLATE_ARGS >( TThreadSafeSPMethodDelegate_FourVars< UserClass, PAYLOAD_TEMPLATE_LIST_FourVars >::Create( InUserObjectRef, InFunc, PAYLOAD_TEMPLATE_PASSIN_FourVars ) );
+	}
+	template< class UserClass, PAYLOAD_TEMPLATE_DECL_FourVars >
+	inline static DELEGATE_CLASS< FUNC_TEMPLATE_ARGS > CreateThreadSafeSP( const TSharedRef< UserClass, ESPMode::ThreadSafe >& InUserObjectRef, typename TThreadSafeSPMethodDelegate_FourVars_Const< UserClass, PAYLOAD_TEMPLATE_LIST_FourVars >::FMethodPtr InFunc, PAYLOAD_TEMPLATE_ARGS_FourVars )
+	{
+		return DELEGATE_CLASS< FUNC_TEMPLATE_ARGS >( TThreadSafeSPMethodDelegate_FourVars_Const< UserClass, PAYLOAD_TEMPLATE_LIST_FourVars >::Create( InUserObjectRef, InFunc, PAYLOAD_TEMPLATE_PASSIN_FourVars ) );
+	}
+
+	/**
+	 * Static: Creates a shared pointer-based (slower, conditionally thread-safe) member function delegate.
+	 *
+	 * Shared pointer delegates keep a weak reference to your object.
+	 * You can use ExecuteIfBound() to call them.
+	 */
+	template< class UserClass >
+	inline static DELEGATE_CLASS< FUNC_TEMPLATE_ARGS > CreateThreadSafeSP( UserClass* InUserObject, typename TThreadSafeSPMethodDelegate< UserClass >::FMethodPtr InFunc )
+	{
+		return CreateThreadSafeSP( StaticCastSharedRef< UserClass >( InUserObject->AsShared() ), InFunc );
+	}
+	template< class UserClass >
+	inline static DELEGATE_CLASS< FUNC_TEMPLATE_ARGS > CreateThreadSafeSP( UserClass* InUserObject, typename TThreadSafeSPMethodDelegate_Const< UserClass >::FMethodPtr InFunc )
+	{
+		return CreateThreadSafeSP( StaticCastSharedRef< UserClass >( InUserObject->AsShared() ), InFunc );
+	}
+	template< class UserClass, PAYLOAD_TEMPLATE_DECL_OneVar >
+	inline static DELEGATE_CLASS< FUNC_TEMPLATE_ARGS > CreateThreadSafeSP( UserClass* InUserObject, typename TThreadSafeSPMethodDelegate_OneVar< UserClass, PAYLOAD_TEMPLATE_LIST_OneVar >::FMethodPtr InFunc, PAYLOAD_TEMPLATE_ARGS_OneVar )
+	{
+		return CreateThreadSafeSP( StaticCastSharedRef< UserClass >( InUserObject->AsShared() ), InFunc, PAYLOAD_TEMPLATE_PASSIN_OneVar );
+	}
+	template< class UserClass, PAYLOAD_TEMPLATE_DECL_OneVar >
+	inline static DELEGATE_CLASS< FUNC_TEMPLATE_ARGS > CreateThreadSafeSP( UserClass* InUserObject, typename TThreadSafeSPMethodDelegate_OneVar_Const< UserClass, PAYLOAD_TEMPLATE_LIST_OneVar >::FMethodPtr InFunc, PAYLOAD_TEMPLATE_ARGS_OneVar )
+	{
+		return CreateThreadSafeSP( StaticCastSharedRef< UserClass >( InUserObject->AsShared() ), InFunc, PAYLOAD_TEMPLATE_PASSIN_OneVar );
+	}
+	template< class UserClass, PAYLOAD_TEMPLATE_DECL_TwoVars >
+	inline static DELEGATE_CLASS< FUNC_TEMPLATE_ARGS > CreateThreadSafeSP( UserClass* InUserObject, typename TThreadSafeSPMethodDelegate_TwoVars< UserClass, PAYLOAD_TEMPLATE_LIST_TwoVars >::FMethodPtr InFunc, PAYLOAD_TEMPLATE_ARGS_TwoVars )
+	{
+		return CreateThreadSafeSP( StaticCastSharedRef< UserClass >( InUserObject->AsShared() ), InFunc, PAYLOAD_TEMPLATE_PASSIN_TwoVars );
+	}
+	template< class UserClass, PAYLOAD_TEMPLATE_DECL_TwoVars >
+	inline static DELEGATE_CLASS< FUNC_TEMPLATE_ARGS > CreateThreadSafeSP( UserClass* InUserObject, typename TThreadSafeSPMethodDelegate_TwoVars_Const< UserClass, PAYLOAD_TEMPLATE_LIST_TwoVars >::FMethodPtr InFunc, PAYLOAD_TEMPLATE_ARGS_TwoVars )
+	{
+		return CreateThreadSafeSP( StaticCastSharedRef< UserClass >( InUserObject->AsShared() ), InFunc, PAYLOAD_TEMPLATE_PASSIN_TwoVars );
+	}
+	template< class UserClass, PAYLOAD_TEMPLATE_DECL_ThreeVars >
+	inline static DELEGATE_CLASS< FUNC_TEMPLATE_ARGS > CreateThreadSafeSP( UserClass* InUserObject, typename TThreadSafeSPMethodDelegate_ThreeVars< UserClass, PAYLOAD_TEMPLATE_LIST_ThreeVars >::FMethodPtr InFunc, PAYLOAD_TEMPLATE_ARGS_ThreeVars )
+	{
+		return CreateThreadSafeSP( StaticCastSharedRef< UserClass >( InUserObject->AsShared() ), InFunc, PAYLOAD_TEMPLATE_PASSIN_ThreeVars );
+	}
+	template< class UserClass, PAYLOAD_TEMPLATE_DECL_ThreeVars >
+	inline static DELEGATE_CLASS< FUNC_TEMPLATE_ARGS > CreateThreadSafeSP( UserClass* InUserObject, typename TThreadSafeSPMethodDelegate_ThreeVars_Const< UserClass, PAYLOAD_TEMPLATE_LIST_ThreeVars >::FMethodPtr InFunc, PAYLOAD_TEMPLATE_ARGS_ThreeVars )
+	{
+		return CreateThreadSafeSP( StaticCastSharedRef< UserClass >( InUserObject->AsShared() ), InFunc, PAYLOAD_TEMPLATE_PASSIN_ThreeVars );
+	}
+	template< class UserClass, PAYLOAD_TEMPLATE_DECL_FourVars >
+	inline static DELEGATE_CLASS< FUNC_TEMPLATE_ARGS > CreateThreadSafeSP( UserClass* InUserObject, typename TThreadSafeSPMethodDelegate_FourVars< UserClass, PAYLOAD_TEMPLATE_LIST_FourVars >::FMethodPtr InFunc, PAYLOAD_TEMPLATE_ARGS_FourVars )
+	{
+		return CreateThreadSafeSP( StaticCastSharedRef< UserClass >( InUserObject->AsShared() ), InFunc, PAYLOAD_TEMPLATE_PASSIN_FourVars );
+	}
+	template< class UserClass, PAYLOAD_TEMPLATE_DECL_FourVars >
+	inline static DELEGATE_CLASS< FUNC_TEMPLATE_ARGS > CreateThreadSafeSP( UserClass* InUserObject, typename TThreadSafeSPMethodDelegate_FourVars_Const< UserClass, PAYLOAD_TEMPLATE_LIST_FourVars >::FMethodPtr InFunc, PAYLOAD_TEMPLATE_ARGS_FourVars )
+	{
+		return CreateThreadSafeSP( StaticCastSharedRef< UserClass >( InUserObject->AsShared() ), InFunc, PAYLOAD_TEMPLATE_PASSIN_FourVars );
+	}
+
+	/**
+	 * Static: Creates a UFunction-based member function delegate.
+	 *
+	 * UFunction delegates keep a weak reference to your object.
+	 * You can use ExecuteIfBound() to call them.
+	 */
+	template< class UObjectTemplate >
+	inline static DELEGATE_CLASS< FUNC_TEMPLATE_ARGS > CreateUFunction( UObjectTemplate* InUserObject, const FName& InFunctionName )
+	{
+		return DELEGATE_CLASS< FUNC_TEMPLATE_ARGS >( TUFunctionDelegateBinding< UObjectTemplate >::Create( InUserObject, InFunctionName ) );
+	}
+	template< class UObjectTemplate, PAYLOAD_TEMPLATE_DECL_OneVar >
+	inline static DELEGATE_CLASS< FUNC_TEMPLATE_ARGS > CreateUFunction( UObjectTemplate* InUserObject, const FName& InFunctionName, PAYLOAD_TEMPLATE_ARGS_OneVar )
+	{
+		return DELEGATE_CLASS< FUNC_TEMPLATE_ARGS >( TUFunctionDelegateBinding_OneVar< UObjectTemplate, PAYLOAD_TEMPLATE_LIST_OneVar >::Create( InUserObject, InFunctionName, PAYLOAD_TEMPLATE_PASSIN_OneVar ) );
+	}
+	template< class UObjectTemplate, PAYLOAD_TEMPLATE_DECL_TwoVars >
+	inline static DELEGATE_CLASS< FUNC_TEMPLATE_ARGS > CreateUFunction( UObjectTemplate* InUserObject, const FName& InFunctionName, PAYLOAD_TEMPLATE_ARGS_TwoVars )
+	{
+		return DELEGATE_CLASS< FUNC_TEMPLATE_ARGS >( TUFunctionDelegateBinding_TwoVars< UObjectTemplate, PAYLOAD_TEMPLATE_LIST_TwoVars >::Create( InUserObject, InFunctionName, PAYLOAD_TEMPLATE_PASSIN_TwoVars ) );
+	}
+	template< class UObjectTemplate, PAYLOAD_TEMPLATE_DECL_ThreeVars >
+	inline static DELEGATE_CLASS< FUNC_TEMPLATE_ARGS > CreateUFunction( UObjectTemplate* InUserObject, const FName& InFunctionName, PAYLOAD_TEMPLATE_ARGS_ThreeVars )
+	{
+		return DELEGATE_CLASS< FUNC_TEMPLATE_ARGS >( TUFunctionDelegateBinding_ThreeVars< UObjectTemplate, PAYLOAD_TEMPLATE_LIST_ThreeVars >::Create( InUserObject, InFunctionName, PAYLOAD_TEMPLATE_PASSIN_ThreeVars ) );
+	}
+	template< class UObjectTemplate, PAYLOAD_TEMPLATE_DECL_FourVars >
+	inline static DELEGATE_CLASS< FUNC_TEMPLATE_ARGS > CreateUFunction( UObjectTemplate* InUserObject, const FName& InFunctionName, PAYLOAD_TEMPLATE_ARGS_FourVars )
+	{
+		return DELEGATE_CLASS< FUNC_TEMPLATE_ARGS >( TUFunctionDelegateBinding_FourVars< UObjectTemplate, PAYLOAD_TEMPLATE_LIST_FourVars >::Create( InUserObject, InFunctionName, PAYLOAD_TEMPLATE_PASSIN_FourVars ) );
+	}
+
+	/**
+	 * Static: Creates a UObject-based member function delegate.
+	 *
+	 * UObject delegates keep a weak reference to your object.
+	 * You can use ExecuteIfBound() to call them.
+	 */
+	template< class UserClass >
+	inline static DELEGATE_CLASS< FUNC_TEMPLATE_ARGS > CreateUObject( UserClass* InUserObject, typename TUObjectMethodDelegate< UserClass >::FMethodPtr InFunc )
+	{
+		return DELEGATE_CLASS< FUNC_TEMPLATE_ARGS >( TUObjectMethodDelegate< UserClass >::Create( InUserObject, InFunc ) );
+	}
+	template< class UserClass >
+	inline static DELEGATE_CLASS< FUNC_TEMPLATE_ARGS > CreateUObject( UserClass* InUserObject, typename TUObjectMethodDelegate_Const< UserClass >::FMethodPtr InFunc )
+	{
+		return DELEGATE_CLASS< FUNC_TEMPLATE_ARGS >( TUObjectMethodDelegate_Const< UserClass >::Create( InUserObject, InFunc ) );
+	}
+	template< class UserClass, PAYLOAD_TEMPLATE_DECL_OneVar >
+	inline static DELEGATE_CLASS< FUNC_TEMPLATE_ARGS > CreateUObject( UserClass* InUserObject, typename TUObjectMethodDelegate_OneVar< UserClass, PAYLOAD_TEMPLATE_LIST_OneVar >::FMethodPtr InFunc, PAYLOAD_TEMPLATE_ARGS_OneVar )
+	{
+		return DELEGATE_CLASS< FUNC_TEMPLATE_ARGS >( TUObjectMethodDelegate_OneVar< UserClass, PAYLOAD_TEMPLATE_LIST_OneVar >::Create( InUserObject, InFunc, PAYLOAD_TEMPLATE_PASSIN_OneVar ) );
+	}
+	template< class UserClass, PAYLOAD_TEMPLATE_DECL_OneVar >
+	inline static DELEGATE_CLASS< FUNC_TEMPLATE_ARGS > CreateUObject( UserClass* InUserObject, typename TUObjectMethodDelegate_OneVar_Const< UserClass, PAYLOAD_TEMPLATE_LIST_OneVar >::FMethodPtr InFunc, PAYLOAD_TEMPLATE_ARGS_OneVar )
+	{
+		return DELEGATE_CLASS< FUNC_TEMPLATE_ARGS >( TUObjectMethodDelegate_OneVar_Const< UserClass, PAYLOAD_TEMPLATE_LIST_OneVar >::Create( InUserObject, InFunc, PAYLOAD_TEMPLATE_PASSIN_OneVar ) );
+	}
+	template< class UserClass, PAYLOAD_TEMPLATE_DECL_TwoVars >
+	inline static DELEGATE_CLASS< FUNC_TEMPLATE_ARGS > CreateUObject( UserClass* InUserObject, typename TUObjectMethodDelegate_TwoVars< UserClass, PAYLOAD_TEMPLATE_LIST_TwoVars >::FMethodPtr InFunc, PAYLOAD_TEMPLATE_ARGS_TwoVars )
+	{
+		return DELEGATE_CLASS< FUNC_TEMPLATE_ARGS >( TUObjectMethodDelegate_TwoVars< UserClass, PAYLOAD_TEMPLATE_LIST_TwoVars >::Create( InUserObject, InFunc, PAYLOAD_TEMPLATE_PASSIN_TwoVars ) );
+	}
+	template< class UserClass, PAYLOAD_TEMPLATE_DECL_TwoVars >
+	inline static DELEGATE_CLASS< FUNC_TEMPLATE_ARGS > CreateUObject( UserClass* InUserObject, typename TUObjectMethodDelegate_TwoVars_Const< UserClass, PAYLOAD_TEMPLATE_LIST_TwoVars >::FMethodPtr InFunc, PAYLOAD_TEMPLATE_ARGS_TwoVars )
+	{
+		return DELEGATE_CLASS< FUNC_TEMPLATE_ARGS >( TUObjectMethodDelegate_TwoVars_Const< UserClass, PAYLOAD_TEMPLATE_LIST_TwoVars >::Create( InUserObject, InFunc, PAYLOAD_TEMPLATE_PASSIN_TwoVars ) );
+	}
+	template< class UserClass, PAYLOAD_TEMPLATE_DECL_ThreeVars >
+	inline static DELEGATE_CLASS< FUNC_TEMPLATE_ARGS > CreateUObject( UserClass* InUserObject, typename TUObjectMethodDelegate_ThreeVars< UserClass, PAYLOAD_TEMPLATE_LIST_ThreeVars >::FMethodPtr InFunc, PAYLOAD_TEMPLATE_ARGS_ThreeVars )
+	{
+		return DELEGATE_CLASS< FUNC_TEMPLATE_ARGS >( TUObjectMethodDelegate_ThreeVars< UserClass, PAYLOAD_TEMPLATE_LIST_ThreeVars >::Create( InUserObject, InFunc, PAYLOAD_TEMPLATE_PASSIN_ThreeVars ) );
+	}
+	template< class UserClass, PAYLOAD_TEMPLATE_DECL_ThreeVars >
+	inline static DELEGATE_CLASS< FUNC_TEMPLATE_ARGS > CreateUObject( UserClass* InUserObject, typename TUObjectMethodDelegate_ThreeVars_Const< UserClass, PAYLOAD_TEMPLATE_LIST_ThreeVars >::FMethodPtr InFunc, PAYLOAD_TEMPLATE_ARGS_ThreeVars )
+	{
+		return DELEGATE_CLASS< FUNC_TEMPLATE_ARGS >( TUObjectMethodDelegate_ThreeVars_Const< UserClass, PAYLOAD_TEMPLATE_LIST_ThreeVars >::Create( InUserObject, InFunc, PAYLOAD_TEMPLATE_PASSIN_ThreeVars ) );
+	}
+	template< class UserClass, PAYLOAD_TEMPLATE_DECL_FourVars >
+	inline static DELEGATE_CLASS< FUNC_TEMPLATE_ARGS > CreateUObject( UserClass* InUserObject, typename TUObjectMethodDelegate_FourVars< UserClass, PAYLOAD_TEMPLATE_LIST_FourVars >::FMethodPtr InFunc, PAYLOAD_TEMPLATE_ARGS_FourVars )
+	{
+		return DELEGATE_CLASS< FUNC_TEMPLATE_ARGS >( TUObjectMethodDelegate_FourVars< UserClass, PAYLOAD_TEMPLATE_LIST_FourVars >::Create( InUserObject, InFunc, PAYLOAD_TEMPLATE_PASSIN_FourVars ) );
+	}
+	template< class UserClass, PAYLOAD_TEMPLATE_DECL_FourVars >
+	inline static DELEGATE_CLASS< FUNC_TEMPLATE_ARGS > CreateUObject( UserClass* InUserObject, typename TUObjectMethodDelegate_FourVars_Const< UserClass, PAYLOAD_TEMPLATE_LIST_FourVars >::FMethodPtr InFunc, PAYLOAD_TEMPLATE_ARGS_FourVars )
+	{
+		return DELEGATE_CLASS< FUNC_TEMPLATE_ARGS >( TUObjectMethodDelegate_FourVars_Const< UserClass, PAYLOAD_TEMPLATE_LIST_FourVars >::Create( InUserObject, InFunc, PAYLOAD_TEMPLATE_PASSIN_FourVars ) );
+	}
+
+public:
+
+	/**
+	 * Default constructor
+	 */
+	inline DELEGATE_CLASS()
+		: FDelegateBase(nullptr)
+	{ }
+
+	/**
+	 * Destructor.
+	 */
+	inline ~DELEGATE_CLASS()
+	{
+		Unbind();
+	}
+
+	/**
+	 * Creates and initializes a new instance with the given delegate instance.
+	 *
+	 * The delegate will assume ownership of the incoming delegate instance!
+	 * IMPORTANT: This is a system-internal function and you should never be using this in regular C++ code
+	 *
+	 * @param InDelegateInstance The delegate instance to assign.
+	 */
+	inline DELEGATE_CLASS( TDelegateInstanceInterface* InDelegateInstance )
+		: FDelegateBase(InDelegateInstance)
+	{ }
+
+	/**
+	 * Creates and initializes a new instance from an existing delegate object.
+	 *
+	 * @param Other The delegate object to copy from.
+	 */
+	inline DELEGATE_CLASS( const DELEGATE_CLASS& Other )
+		: FDelegateBase(nullptr)
+	{
+		*this = Other;
+	}
+
+	/**
+	 * Assignment operator.
+	 *
+	 * @param	OtherDelegate	Delegate object to copy from
+	 */
+	inline DELEGATE_CLASS& operator=( const DELEGATE_CLASS& Other )
+	{
+		if (&Other != this)
+		{
+			// this down-cast is OK! allows for managing invocation list in the base class without requiring virtual functions
+			TDelegateInstanceInterface* OtherInstance = (TDelegateInstanceInterface*)Other.GetDelegateInstance();
+
+			if (OtherInstance != nullptr)
+			{
+				SetDelegateInstance(OtherInstance->CreateCopy());
+			}
+			else
+			{
+				SetDelegateInstance(nullptr);
+			}
+		}
+
+		return *this;
+	}
+
+public:
+
+	/**
+	 * Creates a raw C++ pointer global function delegate
+	 */
+	inline void BindStatic( typename FStaticDelegate::FFuncPtr InFunc )
+	{
+		*this = CreateStatic( InFunc );
+	}
+	template< PAYLOAD_TEMPLATE_DECL_OneVar >
+	inline void BindStatic( typename TStaticDelegate_OneVar< PAYLOAD_TEMPLATE_LIST_OneVar >::FFuncPtr InFunc, PAYLOAD_TEMPLATE_ARGS_OneVar )
+	{
+		*this = CreateStatic( InFunc, PAYLOAD_TEMPLATE_PASSIN_OneVar );
+	}
+	template< PAYLOAD_TEMPLATE_DECL_TwoVars >
+	inline void BindStatic( typename TStaticDelegate_TwoVars< PAYLOAD_TEMPLATE_LIST_TwoVars >::FFuncPtr InFunc, PAYLOAD_TEMPLATE_ARGS_TwoVars )
+	{
+		*this = CreateStatic( InFunc, PAYLOAD_TEMPLATE_PASSIN_TwoVars );
+	}
+	template< PAYLOAD_TEMPLATE_DECL_ThreeVars >
+	inline void BindStatic( typename TStaticDelegate_ThreeVars< PAYLOAD_TEMPLATE_LIST_ThreeVars >::FFuncPtr InFunc, PAYLOAD_TEMPLATE_ARGS_ThreeVars )
+	{
+		*this = CreateStatic( InFunc, PAYLOAD_TEMPLATE_PASSIN_ThreeVars );
+	}
+	template< PAYLOAD_TEMPLATE_DECL_FourVars >
+	inline void BindStatic( typename TStaticDelegate_FourVars< PAYLOAD_TEMPLATE_LIST_FourVars >::FFuncPtr InFunc, PAYLOAD_TEMPLATE_ARGS_FourVars )
+	{
+		*this = CreateStatic( InFunc, PAYLOAD_TEMPLATE_PASSIN_FourVars );
+	}
+
+	/**
+	 * Binds a raw C++ pointer delegate.
+	 *
+	 * Raw pointer doesn't use any sort of reference, so may be unsafe to call if the object was
+	 * deleted out from underneath your delegate. Be careful when calling Execute()!
+	 */
+	template< class UserClass >
+	inline void BindRaw( UserClass* InUserObject, typename TRawMethodDelegate< UserClass >::FMethodPtr InFunc )
+	{
+		*this = CreateRaw( InUserObject, InFunc );
+	}
+	template< class UserClass >
+	inline void BindRaw( UserClass* InUserObject, typename TRawMethodDelegate_Const< UserClass >::FMethodPtr InFunc )
+	{
+		*this = CreateRaw( InUserObject, InFunc );
+	}
+	template< class UserClass, PAYLOAD_TEMPLATE_DECL_OneVar >
+	inline void BindRaw( UserClass* InUserObject, typename TRawMethodDelegate_OneVar< UserClass, PAYLOAD_TEMPLATE_LIST_OneVar >::FMethodPtr InFunc, PAYLOAD_TEMPLATE_ARGS_OneVar )
+	{
+		*this = CreateRaw( InUserObject, InFunc, PAYLOAD_TEMPLATE_PASSIN_OneVar );
+	}
+	template< class UserClass, PAYLOAD_TEMPLATE_DECL_OneVar >
+	inline void BindRaw( UserClass* InUserObject, typename TRawMethodDelegate_OneVar_Const< UserClass, PAYLOAD_TEMPLATE_LIST_OneVar >::FMethodPtr InFunc, PAYLOAD_TEMPLATE_ARGS_OneVar )
+	{
+		*this = CreateRaw( InUserObject, InFunc, PAYLOAD_TEMPLATE_PASSIN_OneVar );
+	}
+	template< class UserClass, PAYLOAD_TEMPLATE_DECL_TwoVars >
+	inline void BindRaw( UserClass* InUserObject, typename TRawMethodDelegate_TwoVars< UserClass, PAYLOAD_TEMPLATE_LIST_TwoVars >::FMethodPtr InFunc, PAYLOAD_TEMPLATE_ARGS_TwoVars )
+	{
+		*this = CreateRaw( InUserObject, InFunc, PAYLOAD_TEMPLATE_PASSIN_TwoVars );
+	}
+	template< class UserClass, PAYLOAD_TEMPLATE_DECL_TwoVars >
+	inline void BindRaw( UserClass* InUserObject, typename TRawMethodDelegate_TwoVars_Const< UserClass, PAYLOAD_TEMPLATE_LIST_TwoVars >::FMethodPtr InFunc, PAYLOAD_TEMPLATE_ARGS_TwoVars )
+	{
+		*this = CreateRaw( InUserObject, InFunc, PAYLOAD_TEMPLATE_PASSIN_TwoVars );
+	}
+	template< class UserClass, PAYLOAD_TEMPLATE_DECL_ThreeVars >
+	inline void BindRaw( UserClass* InUserObject, typename TRawMethodDelegate_ThreeVars< UserClass, PAYLOAD_TEMPLATE_LIST_ThreeVars >::FMethodPtr InFunc, PAYLOAD_TEMPLATE_ARGS_ThreeVars )
+	{
+		*this = CreateRaw( InUserObject, InFunc, PAYLOAD_TEMPLATE_PASSIN_ThreeVars );
+	}
+	template< class UserClass, PAYLOAD_TEMPLATE_DECL_ThreeVars >
+	inline void BindRaw( UserClass* InUserObject, typename TRawMethodDelegate_ThreeVars_Const< UserClass, PAYLOAD_TEMPLATE_LIST_ThreeVars >::FMethodPtr InFunc, PAYLOAD_TEMPLATE_ARGS_ThreeVars )
+	{
+		*this = CreateRaw( InUserObject, InFunc, PAYLOAD_TEMPLATE_PASSIN_ThreeVars );
+	}
+	template< class UserClass, PAYLOAD_TEMPLATE_DECL_FourVars >
+	inline void BindRaw( UserClass* InUserObject, typename TRawMethodDelegate_FourVars< UserClass, PAYLOAD_TEMPLATE_LIST_FourVars >::FMethodPtr InFunc, PAYLOAD_TEMPLATE_ARGS_FourVars )
+	{
+		*this = CreateRaw( InUserObject, InFunc, PAYLOAD_TEMPLATE_PASSIN_FourVars );
+	}
+	template< class UserClass, PAYLOAD_TEMPLATE_DECL_FourVars >
+	inline void BindRaw( UserClass* InUserObject, typename TRawMethodDelegate_FourVars_Const< UserClass, PAYLOAD_TEMPLATE_LIST_FourVars >::FMethodPtr InFunc, PAYLOAD_TEMPLATE_ARGS_FourVars )
+	{
+		*this = CreateRaw( InUserObject, InFunc, PAYLOAD_TEMPLATE_PASSIN_FourVars );
+	}
+
+	/**
+	 * Binds a shared pointer-based (fast, not thread-safe) member function delegate.  Shared pointer delegates keep a weak reference to your object.  You can use ExecuteIfBound() to call them.
+	 */
+	template< class UserClass >
+	inline void BindSP( const TSharedRef< UserClass, ESPMode::Fast >& InUserObjectRef, typename TSPMethodDelegate< UserClass >::FMethodPtr InFunc )
+	{
+		*this = CreateSP( InUserObjectRef, InFunc );
+	}
+	template< class UserClass >
+	inline void BindSP( const TSharedRef< UserClass, ESPMode::Fast >& InUserObjectRef, typename TSPMethodDelegate_Const< UserClass >::FMethodPtr InFunc )
+	{
+		*this = CreateSP( InUserObjectRef, InFunc );
+	}
+	template< class UserClass, PAYLOAD_TEMPLATE_DECL_OneVar >
+	inline void BindSP( const TSharedRef< UserClass, ESPMode::Fast >& InUserObjectRef, typename TSPMethodDelegate_OneVar< UserClass, PAYLOAD_TEMPLATE_LIST_OneVar >::FMethodPtr InFunc, PAYLOAD_TEMPLATE_ARGS_OneVar )
+	{
+		*this = CreateSP( InUserObjectRef, InFunc, PAYLOAD_TEMPLATE_PASSIN_OneVar );
+	}
+	template< class UserClass, PAYLOAD_TEMPLATE_DECL_OneVar >
+	inline void BindSP( const TSharedRef< UserClass, ESPMode::Fast >& InUserObjectRef, typename TSPMethodDelegate_OneVar_Const< UserClass, PAYLOAD_TEMPLATE_LIST_OneVar >::FMethodPtr InFunc, PAYLOAD_TEMPLATE_ARGS_OneVar )
+	{
+		*this = CreateSP( InUserObjectRef, InFunc, PAYLOAD_TEMPLATE_PASSIN_OneVar );
+	}
+	template< class UserClass, PAYLOAD_TEMPLATE_DECL_TwoVars >
+	inline void BindSP( const TSharedRef< UserClass, ESPMode::Fast >& InUserObjectRef, typename TSPMethodDelegate_TwoVars< UserClass, PAYLOAD_TEMPLATE_LIST_TwoVars >::FMethodPtr InFunc, PAYLOAD_TEMPLATE_ARGS_TwoVars )
+	{
+		*this = CreateSP( InUserObjectRef, InFunc, PAYLOAD_TEMPLATE_PASSIN_TwoVars );
+	}
+	template< class UserClass, PAYLOAD_TEMPLATE_DECL_TwoVars >
+	inline void BindSP( const TSharedRef< UserClass, ESPMode::Fast >& InUserObjectRef, typename TSPMethodDelegate_TwoVars_Const< UserClass, PAYLOAD_TEMPLATE_LIST_TwoVars >::FMethodPtr InFunc, PAYLOAD_TEMPLATE_ARGS_TwoVars )
+	{
+		*this = CreateSP( InUserObjectRef, InFunc, PAYLOAD_TEMPLATE_PASSIN_TwoVars );
+	}
+	template< class UserClass, PAYLOAD_TEMPLATE_DECL_ThreeVars >
+	inline void BindSP( const TSharedRef< UserClass, ESPMode::Fast >& InUserObjectRef, typename TSPMethodDelegate_ThreeVars< UserClass, PAYLOAD_TEMPLATE_LIST_ThreeVars >::FMethodPtr InFunc, PAYLOAD_TEMPLATE_ARGS_ThreeVars )
+	{
+		*this = CreateSP( InUserObjectRef, InFunc, PAYLOAD_TEMPLATE_PASSIN_ThreeVars );
+	}
+	template< class UserClass, PAYLOAD_TEMPLATE_DECL_ThreeVars >
+	inline void BindSP( const TSharedRef< UserClass, ESPMode::Fast >& InUserObjectRef, typename TSPMethodDelegate_ThreeVars_Const< UserClass, PAYLOAD_TEMPLATE_LIST_ThreeVars >::FMethodPtr InFunc, PAYLOAD_TEMPLATE_ARGS_ThreeVars )
+	{
+		*this = CreateSP( InUserObjectRef, InFunc, PAYLOAD_TEMPLATE_PASSIN_ThreeVars );
+	}
+	template< class UserClass, PAYLOAD_TEMPLATE_DECL_FourVars >
+	inline void BindSP( const TSharedRef< UserClass, ESPMode::Fast >& InUserObjectRef, typename TSPMethodDelegate_FourVars< UserClass, PAYLOAD_TEMPLATE_LIST_FourVars >::FMethodPtr InFunc, PAYLOAD_TEMPLATE_ARGS_FourVars )
+	{
+		*this = CreateSP( InUserObjectRef, InFunc, PAYLOAD_TEMPLATE_PASSIN_FourVars );
+	}
+	template< class UserClass, PAYLOAD_TEMPLATE_DECL_FourVars >
+	inline void BindSP( const TSharedRef< UserClass, ESPMode::Fast >& InUserObjectRef, typename TSPMethodDelegate_FourVars_Const< UserClass, PAYLOAD_TEMPLATE_LIST_FourVars >::FMethodPtr InFunc, PAYLOAD_TEMPLATE_ARGS_FourVars )
+	{
+		*this = CreateSP( InUserObjectRef, InFunc, PAYLOAD_TEMPLATE_PASSIN_FourVars );
+	}
+
+	/**
+	 * Creates a shared pointer-based (fast, not thread-safe) member function delegate.
+	 *
+	 * Shared pointer delegates keep a weak reference to your object.
+	 * You can use ExecuteIfBound() to call them.
+	 */
+	template< class UserClass >
+	inline void BindSP( UserClass* InUserObject, typename TSPMethodDelegate< UserClass >::FMethodPtr InFunc )
+	{
+		*this = CreateSP( InUserObject, InFunc );
+	}
+	template< class UserClass >
+	inline void BindSP( UserClass* InUserObject, typename TSPMethodDelegate_Const< UserClass >::FMethodPtr InFunc )
+	{
+		*this = CreateSP( InUserObject, InFunc );
+	}
+	template< class UserClass, PAYLOAD_TEMPLATE_DECL_OneVar >
+	inline void BindSP( UserClass* InUserObject, typename TSPMethodDelegate_OneVar< UserClass, PAYLOAD_TEMPLATE_LIST_OneVar >::FMethodPtr InFunc, PAYLOAD_TEMPLATE_ARGS_OneVar )
+	{
+		*this = CreateSP( InUserObject, InFunc, PAYLOAD_TEMPLATE_PASSIN_OneVar );
+	}
+	template< class UserClass, PAYLOAD_TEMPLATE_DECL_OneVar >
+	inline void BindSP( UserClass* InUserObject, typename TSPMethodDelegate_OneVar_Const< UserClass, PAYLOAD_TEMPLATE_LIST_OneVar >::FMethodPtr InFunc, PAYLOAD_TEMPLATE_ARGS_OneVar )
+	{
+		*this = CreateSP( InUserObject, InFunc, PAYLOAD_TEMPLATE_PASSIN_OneVar );
+	}
+	template< class UserClass, PAYLOAD_TEMPLATE_DECL_TwoVars >
+	inline void BindSP( UserClass* InUserObject, typename TSPMethodDelegate_TwoVars< UserClass, PAYLOAD_TEMPLATE_LIST_TwoVars >::FMethodPtr InFunc, PAYLOAD_TEMPLATE_ARGS_TwoVars )
+	{
+		*this = CreateSP( InUserObject, InFunc, PAYLOAD_TEMPLATE_PASSIN_TwoVars );
+	}
+	template< class UserClass, PAYLOAD_TEMPLATE_DECL_TwoVars >
+	inline void BindSP( UserClass* InUserObject, typename TSPMethodDelegate_TwoVars_Const< UserClass, PAYLOAD_TEMPLATE_LIST_TwoVars >::FMethodPtr InFunc, PAYLOAD_TEMPLATE_ARGS_TwoVars )
+	{
+		*this = CreateSP( InUserObject, InFunc, PAYLOAD_TEMPLATE_PASSIN_TwoVars );
+	}
+	template< class UserClass, PAYLOAD_TEMPLATE_DECL_ThreeVars >
+	inline void BindSP( UserClass* InUserObject, typename TSPMethodDelegate_ThreeVars< UserClass, PAYLOAD_TEMPLATE_LIST_ThreeVars >::FMethodPtr InFunc, PAYLOAD_TEMPLATE_ARGS_ThreeVars )
+	{
+		*this = CreateSP( InUserObject, InFunc, PAYLOAD_TEMPLATE_PASSIN_ThreeVars );
+	}
+	template< class UserClass, PAYLOAD_TEMPLATE_DECL_ThreeVars >
+	inline void BindSP( UserClass* InUserObject, typename TSPMethodDelegate_ThreeVars_Const< UserClass, PAYLOAD_TEMPLATE_LIST_ThreeVars >::FMethodPtr InFunc, PAYLOAD_TEMPLATE_ARGS_ThreeVars )
+	{
+		*this = CreateSP( InUserObject, InFunc, PAYLOAD_TEMPLATE_PASSIN_ThreeVars );
+	}
+	template< class UserClass, PAYLOAD_TEMPLATE_DECL_FourVars >
+	inline void BindSP( UserClass* InUserObject, typename TSPMethodDelegate_FourVars< UserClass, PAYLOAD_TEMPLATE_LIST_FourVars >::FMethodPtr InFunc, PAYLOAD_TEMPLATE_ARGS_FourVars )
+	{
+		*this = CreateSP( InUserObject, InFunc, PAYLOAD_TEMPLATE_PASSIN_FourVars );
+	}
+	template< class UserClass, PAYLOAD_TEMPLATE_DECL_FourVars >
+	inline void BindSP( UserClass* InUserObject, typename TSPMethodDelegate_FourVars_Const< UserClass, PAYLOAD_TEMPLATE_LIST_FourVars >::FMethodPtr InFunc, PAYLOAD_TEMPLATE_ARGS_FourVars )
+	{
+		*this = CreateSP( InUserObject, InFunc, PAYLOAD_TEMPLATE_PASSIN_FourVars );
+	}
+
+	/**
+	 * Binds a shared pointer-based (slower, conditionally thread-safe) member function delegate.
+	 *
+	 * Shared pointer delegates keep a weak reference to your object.
+	 * You can use ExecuteIfBound() to call them.
+	 */
+	template< class UserClass >
+	inline void BindThreadSafeSP( const TSharedRef< UserClass, ESPMode::ThreadSafe >& InUserObjectRef, typename TThreadSafeSPMethodDelegate< UserClass >::FMethodPtr InFunc )
+	{
+		*this = CreateThreadSafeSP( InUserObjectRef, InFunc );
+	}
+	template< class UserClass >
+	inline void BindThreadSafeSP( const TSharedRef< UserClass, ESPMode::ThreadSafe >& InUserObjectRef, typename TThreadSafeSPMethodDelegate_Const< UserClass >::FMethodPtr InFunc )
+	{
+		*this = CreateThreadSafeSP( InUserObjectRef, InFunc );
+	}
+	template< class UserClass, PAYLOAD_TEMPLATE_DECL_OneVar >
+	inline void BindThreadSafeSP( const TSharedRef< UserClass, ESPMode::ThreadSafe >& InUserObjectRef, typename TThreadSafeSPMethodDelegate_OneVar< UserClass, PAYLOAD_TEMPLATE_LIST_OneVar >::FMethodPtr InFunc, PAYLOAD_TEMPLATE_ARGS_OneVar )
+	{
+		*this = CreateThreadSafeSP( InUserObjectRef, InFunc, PAYLOAD_TEMPLATE_PASSIN_OneVar );
+	}
+	template< class UserClass, PAYLOAD_TEMPLATE_DECL_OneVar >
+	inline void BindThreadSafeSP( const TSharedRef< UserClass, ESPMode::ThreadSafe >& InUserObjectRef, typename TThreadSafeSPMethodDelegate_OneVar_Const< UserClass, PAYLOAD_TEMPLATE_LIST_OneVar >::FMethodPtr InFunc, PAYLOAD_TEMPLATE_ARGS_OneVar )
+	{
+		*this = CreateThreadSafeSP( InUserObjectRef, InFunc, PAYLOAD_TEMPLATE_PASSIN_OneVar );
+	}
+	template< class UserClass, PAYLOAD_TEMPLATE_DECL_TwoVars >
+	inline void BindThreadSafeSP( const TSharedRef< UserClass, ESPMode::ThreadSafe >& InUserObjectRef, typename TThreadSafeSPMethodDelegate_TwoVars< UserClass, PAYLOAD_TEMPLATE_LIST_TwoVars >::FMethodPtr InFunc, PAYLOAD_TEMPLATE_ARGS_TwoVars )
+	{
+		*this = CreateThreadSafeSP( InUserObjectRef, InFunc, PAYLOAD_TEMPLATE_PASSIN_TwoVars );
+	}
+	template< class UserClass, PAYLOAD_TEMPLATE_DECL_TwoVars >
+	inline void BindThreadSafeSP( const TSharedRef< UserClass, ESPMode::ThreadSafe >& InUserObjectRef, typename TThreadSafeSPMethodDelegate_TwoVars_Const< UserClass, PAYLOAD_TEMPLATE_LIST_TwoVars >::FMethodPtr InFunc, PAYLOAD_TEMPLATE_ARGS_TwoVars )
+	{
+		*this = CreateThreadSafeSP( InUserObjectRef, InFunc, PAYLOAD_TEMPLATE_PASSIN_TwoVars );
+	}
+	template< class UserClass, PAYLOAD_TEMPLATE_DECL_ThreeVars >
+	inline void BindThreadSafeSP( const TSharedRef< UserClass, ESPMode::ThreadSafe >& InUserObjectRef, typename TThreadSafeSPMethodDelegate_ThreeVars< UserClass, PAYLOAD_TEMPLATE_LIST_ThreeVars >::FMethodPtr InFunc, PAYLOAD_TEMPLATE_ARGS_ThreeVars )
+	{
+		*this = CreateThreadSafeSP( InUserObjectRef, InFunc, PAYLOAD_TEMPLATE_PASSIN_ThreeVars );
+	}
+	template< class UserClass, PAYLOAD_TEMPLATE_DECL_ThreeVars >
+	inline void BindThreadSafeSP( const TSharedRef< UserClass, ESPMode::ThreadSafe >& InUserObjectRef, typename TThreadSafeSPMethodDelegate_ThreeVars_Const< UserClass, PAYLOAD_TEMPLATE_LIST_ThreeVars >::FMethodPtr InFunc, PAYLOAD_TEMPLATE_ARGS_ThreeVars )
+	{
+		*this = CreateThreadSafeSP( InUserObjectRef, InFunc, PAYLOAD_TEMPLATE_PASSIN_ThreeVars );
+	}
+	template< class UserClass, PAYLOAD_TEMPLATE_DECL_FourVars >
+	inline void BindThreadSafeSP( const TSharedRef< UserClass, ESPMode::ThreadSafe >& InUserObjectRef, typename TThreadSafeSPMethodDelegate_FourVars< UserClass, PAYLOAD_TEMPLATE_LIST_FourVars >::FMethodPtr InFunc, PAYLOAD_TEMPLATE_ARGS_FourVars )
+	{
+		*this = CreateThreadSafeSP( InUserObjectRef, InFunc, PAYLOAD_TEMPLATE_PASSIN_FourVars );
+	}
+	template< class UserClass, PAYLOAD_TEMPLATE_DECL_FourVars >
+	inline void BindThreadSafeSP( const TSharedRef< UserClass, ESPMode::ThreadSafe >& InUserObjectRef, typename TThreadSafeSPMethodDelegate_FourVars_Const< UserClass, PAYLOAD_TEMPLATE_LIST_FourVars >::FMethodPtr InFunc, PAYLOAD_TEMPLATE_ARGS_FourVars )
+	{
+		*this = CreateThreadSafeSP( InUserObjectRef, InFunc, PAYLOAD_TEMPLATE_PASSIN_FourVars );
+	}
+
+	/**
+	 * Creates a shared pointer-based (slower, conditionally thread-safe) member function delegate.
+	 *
+	 * Shared pointer delegates keep a weak reference to your object.
+	 * You can use ExecuteIfBound() to call them.
+	 */
+	template< class UserClass >
+	inline void BindThreadSafeSP( UserClass* InUserObject, typename TThreadSafeSPMethodDelegate< UserClass >::FMethodPtr InFunc )
+	{
+		*this = CreateThreadSafeSP( InUserObject, InFunc );
+	}
+	template< class UserClass >
+	inline void BindThreadSafeSP( UserClass* InUserObject, typename TThreadSafeSPMethodDelegate_Const< UserClass >::FMethodPtr InFunc )
+	{
+		*this = CreateThreadSafeSP( InUserObject, InFunc );
+	}
+	template< class UserClass, PAYLOAD_TEMPLATE_DECL_OneVar >
+	inline void BindThreadSafeSP( UserClass* InUserObject, typename TThreadSafeSPMethodDelegate_OneVar< UserClass, PAYLOAD_TEMPLATE_LIST_OneVar >::FMethodPtr InFunc, PAYLOAD_TEMPLATE_ARGS_OneVar )
+	{
+		*this = CreateThreadSafeSP( InUserObject, InFunc, PAYLOAD_TEMPLATE_PASSIN_OneVar );
+	}
+	template< class UserClass, PAYLOAD_TEMPLATE_DECL_OneVar >
+	inline void BindThreadSafeSP( UserClass* InUserObject, typename TThreadSafeSPMethodDelegate_OneVar_Const< UserClass, PAYLOAD_TEMPLATE_LIST_OneVar >::FMethodPtr InFunc, PAYLOAD_TEMPLATE_ARGS_OneVar )
+	{
+		*this = CreateThreadSafeSP( InUserObject, InFunc, PAYLOAD_TEMPLATE_PASSIN_OneVar );
+	}
+	template< class UserClass, PAYLOAD_TEMPLATE_DECL_TwoVars >
+	inline void BindThreadSafeSP( UserClass* InUserObject, typename TThreadSafeSPMethodDelegate_TwoVars< UserClass, PAYLOAD_TEMPLATE_LIST_TwoVars >::FMethodPtr InFunc, PAYLOAD_TEMPLATE_ARGS_TwoVars )
+	{
+		*this = CreateThreadSafeSP( InUserObject, InFunc, PAYLOAD_TEMPLATE_PASSIN_TwoVars );
+	}
+	template< class UserClass, PAYLOAD_TEMPLATE_DECL_TwoVars >
+	inline void BindThreadSafeSP( UserClass* InUserObject, typename TThreadSafeSPMethodDelegate_TwoVars_Const< UserClass, PAYLOAD_TEMPLATE_LIST_TwoVars >::FMethodPtr InFunc, PAYLOAD_TEMPLATE_ARGS_TwoVars )
+	{
+		*this = CreateThreadSafeSP( InUserObject, InFunc, PAYLOAD_TEMPLATE_PASSIN_TwoVars );
+	}
+	template< class UserClass, PAYLOAD_TEMPLATE_DECL_ThreeVars >
+	inline void BindThreadSafeSP( UserClass* InUserObject, typename TThreadSafeSPMethodDelegate_ThreeVars< UserClass, PAYLOAD_TEMPLATE_LIST_ThreeVars >::FMethodPtr InFunc, PAYLOAD_TEMPLATE_ARGS_ThreeVars )
+	{
+		*this = CreateThreadSafeSP( InUserObject, InFunc, PAYLOAD_TEMPLATE_PASSIN_ThreeVars );
+	}
+	template< class UserClass, PAYLOAD_TEMPLATE_DECL_ThreeVars >
+	inline void BindThreadSafeSP( UserClass* InUserObject, typename TThreadSafeSPMethodDelegate_ThreeVars_Const< UserClass, PAYLOAD_TEMPLATE_LIST_ThreeVars >::FMethodPtr InFunc, PAYLOAD_TEMPLATE_ARGS_ThreeVars )
+	{
+		*this = CreateThreadSafeSP( InUserObject, InFunc, PAYLOAD_TEMPLATE_PASSIN_ThreeVars );
+	}
+	template< class UserClass, PAYLOAD_TEMPLATE_DECL_FourVars >
+	inline void BindThreadSafeSP( UserClass* InUserObject, typename TThreadSafeSPMethodDelegate_FourVars< UserClass, PAYLOAD_TEMPLATE_LIST_FourVars >::FMethodPtr InFunc, PAYLOAD_TEMPLATE_ARGS_FourVars )
+	{
+		*this = CreateThreadSafeSP( InUserObject, InFunc, PAYLOAD_TEMPLATE_PASSIN_FourVars );
+	}
+	template< class UserClass, PAYLOAD_TEMPLATE_DECL_FourVars >
+	inline void BindThreadSafeSP( UserClass* InUserObject, typename TThreadSafeSPMethodDelegate_FourVars_Const< UserClass, PAYLOAD_TEMPLATE_LIST_FourVars >::FMethodPtr InFunc, PAYLOAD_TEMPLATE_ARGS_FourVars )
+	{
+		*this = CreateThreadSafeSP( InUserObject, InFunc, PAYLOAD_TEMPLATE_PASSIN_FourVars );
+	}
+
+	/**
+	 * Creates a UFunction-based member function delegate.
+	 *
+	 * UFunction delegates keep a weak reference to your object.
+	 * You can use ExecuteIfBound() to call them.
+	 */
+	template< class UObjectTemplate >
+	inline void BindUFunction( UObjectTemplate* InUserObject, const FName& InFunctionName )
+	{
+		*this = CreateUFunction( InUserObject, InFunctionName );
+	}
+	template< class UObjectTemplate, PAYLOAD_TEMPLATE_DECL_OneVar >
+	inline void BindUFunction( UObjectTemplate* InUserObject, const FName& InFunctionName, PAYLOAD_TEMPLATE_ARGS_OneVar )
+	{
+		*this = CreateUFunction( InUserObject, InFunctionName, PAYLOAD_TEMPLATE_PASSIN_OneVar );
+	}
+	template< class UObjectTemplate, PAYLOAD_TEMPLATE_DECL_TwoVars >
+	inline void BindUFunction( UObjectTemplate* InUserObject, const FName& InFunctionName, PAYLOAD_TEMPLATE_ARGS_TwoVars )
+	{
+		*this = CreateUFunction( InUserObject, InFunctionName, PAYLOAD_TEMPLATE_PASSIN_TwoVars );
+	}
+	template< class UObjectTemplate, PAYLOAD_TEMPLATE_DECL_ThreeVars >
+	inline void BindUFunction( UObjectTemplate* InUserObject, const FName& InFunctionName, PAYLOAD_TEMPLATE_ARGS_ThreeVars )
+	{
+		*this = CreateUFunction( InUserObject, InFunctionName, PAYLOAD_TEMPLATE_PASSIN_ThreeVars );
+	}
+	template< class UObjectTemplate, PAYLOAD_TEMPLATE_DECL_FourVars >
+	inline void BindUFunction( UObjectTemplate* InUserObject, const FName& InFunctionName, PAYLOAD_TEMPLATE_ARGS_FourVars )
+	{
+		*this = CreateUFunction( InUserObject, InFunctionName, PAYLOAD_TEMPLATE_PASSIN_FourVars );
+	}
+
+	/**
+	 * Creates a UObject-based member function delegate.
+	 *
+	 * UObject delegates keep a weak reference to your object.
+	 * You can use ExecuteIfBound() to call them.
+	 */
+	template< class UserClass >
+	inline void BindUObject( UserClass* InUserObject, typename TUObjectMethodDelegate< UserClass >::FMethodPtr InFunc )
+	{
+		*this = CreateUObject( InUserObject, InFunc );
+	}
+	template< class UserClass >
+	inline void BindUObject( UserClass* InUserObject, typename TUObjectMethodDelegate_Const< UserClass >::FMethodPtr InFunc )
+	{
+		*this = CreateUObject( InUserObject, InFunc );
+	}
+	template< class UserClass, PAYLOAD_TEMPLATE_DECL_OneVar >
+	inline void BindUObject( UserClass* InUserObject, typename TUObjectMethodDelegate_OneVar< UserClass, PAYLOAD_TEMPLATE_LIST_OneVar >::FMethodPtr InFunc, PAYLOAD_TEMPLATE_ARGS_OneVar )
+	{
+		*this = CreateUObject( InUserObject, InFunc, PAYLOAD_TEMPLATE_PASSIN_OneVar );
+	}
+	template< class UserClass, PAYLOAD_TEMPLATE_DECL_OneVar >
+	inline void BindUObject( UserClass* InUserObject, typename TUObjectMethodDelegate_OneVar_Const< UserClass, PAYLOAD_TEMPLATE_LIST_OneVar >::FMethodPtr InFunc, PAYLOAD_TEMPLATE_ARGS_OneVar )
+	{
+		*this = CreateUObject( InUserObject, InFunc, PAYLOAD_TEMPLATE_PASSIN_OneVar );
+	}
+	template< class UserClass, PAYLOAD_TEMPLATE_DECL_TwoVars >
+	inline void BindUObject( UserClass* InUserObject, typename TUObjectMethodDelegate_TwoVars< UserClass, PAYLOAD_TEMPLATE_LIST_TwoVars >::FMethodPtr InFunc, PAYLOAD_TEMPLATE_ARGS_TwoVars )
+	{
+		*this = CreateUObject( InUserObject, InFunc, PAYLOAD_TEMPLATE_PASSIN_TwoVars );
+	}
+	template< class UserClass, PAYLOAD_TEMPLATE_DECL_TwoVars >
+	inline void BindUObject( UserClass* InUserObject, typename TUObjectMethodDelegate_TwoVars_Const< UserClass, PAYLOAD_TEMPLATE_LIST_TwoVars >::FMethodPtr InFunc, PAYLOAD_TEMPLATE_ARGS_TwoVars )
+	{
+		*this = CreateUObject( InUserObject, InFunc, PAYLOAD_TEMPLATE_PASSIN_TwoVars );
+	}
+	template< class UserClass, PAYLOAD_TEMPLATE_DECL_ThreeVars >
+	inline void BindUObject( UserClass* InUserObject, typename TUObjectMethodDelegate_ThreeVars< UserClass, PAYLOAD_TEMPLATE_LIST_ThreeVars >::FMethodPtr InFunc, PAYLOAD_TEMPLATE_ARGS_ThreeVars )
+	{
+		*this = CreateUObject( InUserObject, InFunc, PAYLOAD_TEMPLATE_PASSIN_ThreeVars );
+	}
+	template< class UserClass, PAYLOAD_TEMPLATE_DECL_ThreeVars >
+	inline void BindUObject( UserClass* InUserObject, typename TUObjectMethodDelegate_ThreeVars_Const< UserClass, PAYLOAD_TEMPLATE_LIST_ThreeVars >::FMethodPtr InFunc, PAYLOAD_TEMPLATE_ARGS_ThreeVars )
+	{
+		*this = CreateUObject( InUserObject, InFunc, PAYLOAD_TEMPLATE_PASSIN_ThreeVars );
+	}
+	template< class UserClass, PAYLOAD_TEMPLATE_DECL_FourVars >
+	inline void BindUObject( UserClass* InUserObject, typename TUObjectMethodDelegate_FourVars< UserClass, PAYLOAD_TEMPLATE_LIST_FourVars >::FMethodPtr InFunc, PAYLOAD_TEMPLATE_ARGS_FourVars )
+	{
+		*this = CreateUObject( InUserObject, InFunc, PAYLOAD_TEMPLATE_PASSIN_FourVars );
+	}
+	template< class UserClass, PAYLOAD_TEMPLATE_DECL_FourVars >
+	inline void BindUObject( UserClass* InUserObject, typename TUObjectMethodDelegate_FourVars_Const< UserClass, PAYLOAD_TEMPLATE_LIST_FourVars >::FMethodPtr InFunc, PAYLOAD_TEMPLATE_ARGS_FourVars )
+	{
+		*this = CreateUObject( InUserObject, InFunc, PAYLOAD_TEMPLATE_PASSIN_FourVars );
+	}
+public:
+
+	/**
+	 * Execute the delegate.
+	 *
+	 * If the function pointer is not valid, an error will occur. Check IsBound() before
+	 * calling this method or use ExecuteIfBound() instead.
+	 *
+	 * @see ExecuteIfBound
+	 */
+	inline RetValType Execute( FUNC_PARAM_LIST ) const
+	{
+		TDelegateInstanceInterface* DelegateInstance = (TDelegateInstanceInterface*)GetDelegateInstance();
+
+		// If this assert goes off, Execute() was called before a function was bound to the delegate.
+		// Consider using ExecuteIfSafe() instead.
+		checkSlow(DelegateInstance != nullptr);
+
+		return DelegateInstance->Execute(FUNC_PARAM_PASSTHRU);
+	}
+
+#if FUNC_IS_VOID
+	/**
+	 * Execute the delegate, but only if the function pointer is still valid
+	 *
+	 * @return  Returns true if the function was executed
+	 */
+	// NOTE: Currently only delegates with no return value support ExecuteIfBound() 
+	inline bool ExecuteIfBound( FUNC_PARAM_LIST ) const
+	{
+		if( IsBound() )
+		{
+			return ((TDelegateInstanceInterface*)GetDelegateInstance())->ExecuteIfSafe( FUNC_PARAM_PASSTHRU );
+		}
+
+		return false;
+	}
+#endif
+
+	/**
+	 * Equality operator.
+	 *
+	 * @return true if this delegate equals the other, false otherwise.
+	 */
+	bool operator==( const DELEGATE_CLASS< FUNC_TEMPLATE_ARGS >& Other ) const
+	{
+		TDelegateInstanceInterface* DelegateInstance = (TDelegateInstanceInterface*)GetDelegateInstance();
+		TDelegateInstanceInterface* OtherInstance = (TDelegateInstanceInterface*)Other.GetDelegateInstance();
+
+		// The function these delegates point to must be the same
+		if ((DelegateInstance != nullptr) && (OtherInstance != nullptr))
+		{
+			return DelegateInstance->IsSameFunction(*OtherInstance);
+		}
+
+			// If neither delegate is initialized to anything yet, then we treat them as equal
+		if ((DelegateInstance == nullptr) && (OtherInstance == nullptr))
+			{
+				return true;
+			}
+
+			// No match!
+			return false;
+		}
+
+private:
+
+	// Declare ourselves as a friend to our multi-cast counterpart, so that it can access our interface directly
+	template< FUNC_TEMPLATE_DECL_NO_SHADOW > friend class BASE_MULTICAST_DELEGATE_CLASS;
+};
+
+
+#if FUNC_IS_VOID
+
+/**
+ * Multicast delegate base class.
+	 *
+ * This class implements the functionality of multicast delegates. It is templated to the function signature
+ * that it is compatible with. Use the various DECLARE_MULTICAST_DELEGATE and DECLARE_EVENT macros to create
+ * actual delegate types.
+	 *
+ * Multicast delegates only support bindings with no return value.
+ */
+template< FUNC_TEMPLATE_DECL_TYPENAME >
+class BASE_MULTICAST_DELEGATE_CLASS
+	: public FMulticastDelegateBase<>
+{
+public:
+
+	/** DEPRECATED: Type definition for unicast delegate classes whose delegate instances are compatible with this delegate. */
+	typedef DELEGATE_CLASS< FUNC_TEMPLATE_ARGS > FDelegate;
+
+	/** Type definition for the shared interface of delegate instance types compatible with this delegate class. */
+	typedef DELEGATE_INSTANCE_INTERFACE_CLASS<FUNC_TEMPLATE_ARGS> TDelegateInstanceInterface;
+
+public:
+
+	/**
+	 * Adds a delegate instance to this multicast delegate's invocation list.
+	 *
+	 * This delegate will take over ownership of the given delegate instance.
+	 *
+	 * @param DelegateInstance The delegate instance to add.
+	 */
+	void Add( TDelegateInstanceInterface* DelegateInstance )
+	{
+		if (DelegateInstance != nullptr)
+		{
+			AddDelegateInstance(*DelegateInstance);
+			CompactInvocationList();
+		}
+	}
+
+	/**
+	 * DEPRECATED: Adds a unicast delegate to this multi-cast delegate's invocation list.
+	 *
+	 * This method is retained for backwards compatibility.
+	 *
+	 * @param Delegate The delegate to add.
+	 */
+	void Add( const FDelegate& Delegate )
+	{
+		TDelegateInstanceInterface* DelegateInstance = (TDelegateInstanceInterface*)Delegate.GetDelegateInstance();
+
+		if (DelegateInstance != nullptr)
+		{
+			AddDelegateInstance(*DelegateInstance->CreateCopy());
+			CompactInvocationList();
+		}
+	}
+
+	/**
+	 * Adds a raw C++ pointer global function delegate
+	 *
+	 * @param	InFunc	Function pointer
+	 */
+	inline void AddStatic( typename FDelegate::FStaticDelegate::FFuncPtr InFunc )
+	{
+		Add( FDelegate::CreateStatic( InFunc ) );
+	}
+	template< PAYLOAD_TEMPLATE_DECL_OneVar >
+	inline void AddStatic( typename FDelegate::template TStaticDelegate_OneVar< PAYLOAD_TEMPLATE_LIST_OneVar >::FFuncPtr InFunc, PAYLOAD_TEMPLATE_ARGS_OneVar )
+	{
+		Add( FDelegate::CreateStatic( InFunc, PAYLOAD_TEMPLATE_PASSIN_OneVar ) );
+	}
+	template< PAYLOAD_TEMPLATE_DECL_TwoVars >
+	inline void AddStatic( typename FDelegate::template TStaticDelegate_TwoVars< PAYLOAD_TEMPLATE_LIST_TwoVars >::FFuncPtr InFunc, PAYLOAD_TEMPLATE_ARGS_TwoVars )
+	{
+		Add( FDelegate::CreateStatic( InFunc, PAYLOAD_TEMPLATE_PASSIN_TwoVars ) );
+	}
+	template< PAYLOAD_TEMPLATE_DECL_ThreeVars >
+	inline void AddStatic( typename FDelegate::template TStaticDelegate_ThreeVars< PAYLOAD_TEMPLATE_LIST_ThreeVars >::FFuncPtr InFunc, PAYLOAD_TEMPLATE_ARGS_ThreeVars )
+	{
+		Add( FDelegate::CreateStatic( InFunc, PAYLOAD_TEMPLATE_PASSIN_ThreeVars ) );
+	}
+	template< PAYLOAD_TEMPLATE_DECL_FourVars >
+	inline void AddStatic( typename FDelegate::template TStaticDelegate_FourVars< PAYLOAD_TEMPLATE_LIST_FourVars >::FFuncPtr InFunc, PAYLOAD_TEMPLATE_ARGS_FourVars )
+	{
+		Add( FDelegate::CreateStatic( InFunc, PAYLOAD_TEMPLATE_PASSIN_FourVars ) );
+	}
+
+	/**
+	 * Adds a raw C++ pointer delegate.
+	 *
+	 * Raw pointer doesn't use any sort of reference, so may be unsafe to call if the object was
+	 * deleted out from underneath your delegate. Be careful when calling Execute()!
+	 *
+	 * @param	InUserObject	User object to bind to
+	 * @param	InFunc			Class method function address
+	 */
+	template< class UserClass >
+	inline void AddRaw( UserClass* InUserObject, typename FDelegate::template TRawMethodDelegate< UserClass >::FMethodPtr InFunc )
+	{
+		Add( FDelegate::CreateRaw( InUserObject, InFunc ) );
+	}
+	template< class UserClass >
+	inline void AddRaw( UserClass* InUserObject, typename FDelegate::template TRawMethodDelegate_Const< UserClass >::FMethodPtr InFunc )
+	{
+		Add( FDelegate::CreateRaw( InUserObject, InFunc ) );
+	}
+	template< class UserClass, PAYLOAD_TEMPLATE_DECL_OneVar >
+	inline void AddRaw( UserClass* InUserObject, typename FDelegate::template TRawMethodDelegate_OneVar< UserClass, PAYLOAD_TEMPLATE_LIST_OneVar >::FMethodPtr InFunc, PAYLOAD_TEMPLATE_ARGS_OneVar )
+	{
+		Add( FDelegate::CreateRaw( InUserObject, InFunc, PAYLOAD_TEMPLATE_PASSIN_OneVar ) );
+	}
+	template< class UserClass, PAYLOAD_TEMPLATE_DECL_OneVar >
+	inline void AddRaw( UserClass* InUserObject, typename FDelegate::template TRawMethodDelegate_OneVar_Const< UserClass, PAYLOAD_TEMPLATE_LIST_OneVar >::FMethodPtr InFunc, PAYLOAD_TEMPLATE_ARGS_OneVar )
+	{
+		Add( FDelegate::CreateRaw( InUserObject, InFunc, PAYLOAD_TEMPLATE_PASSIN_OneVar ) );
+	}
+	template< class UserClass, PAYLOAD_TEMPLATE_DECL_TwoVars >
+	inline void AddRaw( UserClass* InUserObject, typename FDelegate::template TRawMethodDelegate_TwoVars< UserClass, PAYLOAD_TEMPLATE_LIST_TwoVars >::FMethodPtr InFunc, PAYLOAD_TEMPLATE_ARGS_TwoVars )
+	{
+		Add( FDelegate::CreateRaw( InUserObject, InFunc, PAYLOAD_TEMPLATE_PASSIN_TwoVars ) );
+	}
+	template< class UserClass, PAYLOAD_TEMPLATE_DECL_TwoVars >
+	inline void AddRaw( UserClass* InUserObject, typename FDelegate::template TRawMethodDelegate_TwoVars_Const< UserClass, PAYLOAD_TEMPLATE_LIST_TwoVars >::FMethodPtr InFunc, PAYLOAD_TEMPLATE_ARGS_TwoVars )
+	{
+		Add( FDelegate::CreateRaw( InUserObject, InFunc, PAYLOAD_TEMPLATE_PASSIN_TwoVars ) );
+	}
+	template< class UserClass, PAYLOAD_TEMPLATE_DECL_ThreeVars >
+	inline void AddRaw( UserClass* InUserObject, typename FDelegate::template TRawMethodDelegate_ThreeVars< UserClass, PAYLOAD_TEMPLATE_LIST_ThreeVars >::FMethodPtr InFunc, PAYLOAD_TEMPLATE_ARGS_ThreeVars )
+	{
+		Add( FDelegate::CreateRaw( InUserObject, InFunc, PAYLOAD_TEMPLATE_PASSIN_ThreeVars ) );
+	}
+	template< class UserClass, PAYLOAD_TEMPLATE_DECL_ThreeVars >
+	inline void AddRaw( UserClass* InUserObject, typename FDelegate::template TRawMethodDelegate_ThreeVars_Const< UserClass, PAYLOAD_TEMPLATE_LIST_ThreeVars >::FMethodPtr InFunc, PAYLOAD_TEMPLATE_ARGS_ThreeVars )
+	{
+		Add( FDelegate::CreateRaw( InUserObject, InFunc, PAYLOAD_TEMPLATE_PASSIN_ThreeVars ) );
+	}
+	template< class UserClass, PAYLOAD_TEMPLATE_DECL_FourVars >
+	inline void AddRaw( UserClass* InUserObject, typename FDelegate::template TRawMethodDelegate_FourVars< UserClass, PAYLOAD_TEMPLATE_LIST_FourVars >::FMethodPtr InFunc, PAYLOAD_TEMPLATE_ARGS_FourVars )
+	{
+		Add( FDelegate::CreateRaw( InUserObject, InFunc, PAYLOAD_TEMPLATE_PASSIN_FourVars ) );
+	}
+	template< class UserClass, PAYLOAD_TEMPLATE_DECL_FourVars >
+	inline void AddRaw( UserClass* InUserObject, typename FDelegate::template TRawMethodDelegate_FourVars_Const< UserClass, PAYLOAD_TEMPLATE_LIST_FourVars >::FMethodPtr InFunc, PAYLOAD_TEMPLATE_ARGS_FourVars )
+	{
+		Add( FDelegate::CreateRaw( InUserObject, InFunc, PAYLOAD_TEMPLATE_PASSIN_FourVars ) );
+	}
+
+	/**
+	 * Adds a shared pointer-based (fast, not thread-safe) member function delegate.
+	 *
+	 * Shared pointer delegates keep a weak reference to your object.
+	 *
+	 * @param	InUserObjectRef	User object to bind to
+	 * @param	InFunc			Class method function address
+	 */
+	template< class UserClass >
+	inline void AddSP( const TSharedRef< UserClass, ESPMode::Fast >& InUserObjectRef, typename FDelegate::template TSPMethodDelegate< UserClass >::FMethodPtr InFunc )
+	{
+		Add( FDelegate::CreateSP( InUserObjectRef, InFunc ) );
+	}
+	template< class UserClass >
+	inline void AddSP( const TSharedRef< UserClass, ESPMode::Fast >& InUserObjectRef, typename FDelegate::template TSPMethodDelegate_Const< UserClass >::FMethodPtr InFunc )
+	{
+		Add( FDelegate::CreateSP( InUserObjectRef, InFunc ) );
+	}
+	template< class UserClass, PAYLOAD_TEMPLATE_DECL_OneVar >
+	inline void AddSP( const TSharedRef< UserClass, ESPMode::Fast >& InUserObjectRef, typename FDelegate::template TSPMethodDelegate_OneVar< UserClass, PAYLOAD_TEMPLATE_LIST_OneVar >::FMethodPtr InFunc, PAYLOAD_TEMPLATE_ARGS_OneVar )
+	{
+		Add( FDelegate::CreateSP( InUserObjectRef, InFunc, PAYLOAD_TEMPLATE_PASSIN_OneVar ) );
+	}
+	template< class UserClass, PAYLOAD_TEMPLATE_DECL_OneVar >
+	inline void AddSP( const TSharedRef< UserClass, ESPMode::Fast >& InUserObjectRef, typename FDelegate::template TSPMethodDelegate_OneVar_Const< UserClass, PAYLOAD_TEMPLATE_LIST_OneVar >::FMethodPtr InFunc, PAYLOAD_TEMPLATE_ARGS_OneVar )
+	{
+		Add( FDelegate::CreateSP( InUserObjectRef, InFunc, PAYLOAD_TEMPLATE_PASSIN_OneVar ) );
+	}
+	template< class UserClass, PAYLOAD_TEMPLATE_DECL_TwoVars >
+	inline void AddSP( const TSharedRef< UserClass, ESPMode::Fast >& InUserObjectRef, typename FDelegate::template TSPMethodDelegate_TwoVars< UserClass, PAYLOAD_TEMPLATE_LIST_TwoVars >::FMethodPtr InFunc, PAYLOAD_TEMPLATE_ARGS_TwoVars )
+	{
+		Add( FDelegate::CreateSP( InUserObjectRef, InFunc, PAYLOAD_TEMPLATE_PASSIN_TwoVars ) );
+	}
+	template< class UserClass, PAYLOAD_TEMPLATE_DECL_TwoVars >
+	inline void AddSP( const TSharedRef< UserClass, ESPMode::Fast >& InUserObjectRef, typename FDelegate::template TSPMethodDelegate_TwoVars_Const< UserClass, PAYLOAD_TEMPLATE_LIST_TwoVars >::FMethodPtr InFunc, PAYLOAD_TEMPLATE_ARGS_TwoVars )
+	{
+		Add( FDelegate::CreateSP( InUserObjectRef, InFunc, PAYLOAD_TEMPLATE_PASSIN_TwoVars ) );
+	}
+	template< class UserClass, PAYLOAD_TEMPLATE_DECL_ThreeVars >
+	inline void AddSP( const TSharedRef< UserClass, ESPMode::Fast >& InUserObjectRef, typename FDelegate::template TSPMethodDelegate_ThreeVars< UserClass, PAYLOAD_TEMPLATE_LIST_ThreeVars >::FMethodPtr InFunc, PAYLOAD_TEMPLATE_ARGS_ThreeVars )
+	{
+		Add( FDelegate::CreateSP( InUserObjectRef, InFunc, PAYLOAD_TEMPLATE_PASSIN_ThreeVars ) );
+	}
+	template< class UserClass, PAYLOAD_TEMPLATE_DECL_ThreeVars >
+	inline void AddSP( const TSharedRef< UserClass, ESPMode::Fast >& InUserObjectRef, typename FDelegate::template TSPMethodDelegate_ThreeVars_Const< UserClass, PAYLOAD_TEMPLATE_LIST_ThreeVars >::FMethodPtr InFunc, PAYLOAD_TEMPLATE_ARGS_ThreeVars )
+	{
+		Add( FDelegate::CreateSP( InUserObjectRef, InFunc, PAYLOAD_TEMPLATE_PASSIN_ThreeVars ) );
+	}
+	template< class UserClass, PAYLOAD_TEMPLATE_DECL_FourVars >
+	inline void AddSP( const TSharedRef< UserClass, ESPMode::Fast >& InUserObjectRef, typename FDelegate::template TSPMethodDelegate_FourVars< UserClass, PAYLOAD_TEMPLATE_LIST_FourVars >::FMethodPtr InFunc, PAYLOAD_TEMPLATE_ARGS_FourVars )
+	{
+		Add( FDelegate::CreateSP( InUserObjectRef, InFunc, PAYLOAD_TEMPLATE_PASSIN_FourVars ) );
+	}
+	template< class UserClass, PAYLOAD_TEMPLATE_DECL_FourVars >
+	inline void AddSP( const TSharedRef< UserClass, ESPMode::Fast >& InUserObjectRef, typename FDelegate::template TSPMethodDelegate_FourVars_Const< UserClass, PAYLOAD_TEMPLATE_LIST_FourVars >::FMethodPtr InFunc, PAYLOAD_TEMPLATE_ARGS_FourVars )
+	{
+		Add( FDelegate::CreateSP( InUserObjectRef, InFunc, PAYLOAD_TEMPLATE_PASSIN_FourVars ) );
+	}
+
+	/**
+	 * Adds a shared pointer-based (fast, not thread-safe) member function delegate.
+	 *
+	 * Shared pointer delegates keep a weak reference to your object.
+	 *
+	 * @param	InUserObject	User object to bind to
+	 * @param	InFunc			Class method function address
+	 */
+	template< class UserClass >
+	inline void AddSP( UserClass* InUserObject, typename FDelegate::template TSPMethodDelegate< UserClass >::FMethodPtr InFunc )
+	{
+		Add( FDelegate::CreateSP( InUserObject, InFunc ) );
+	}
+	template< class UserClass >
+	inline void AddSP( UserClass* InUserObject, typename FDelegate::template TSPMethodDelegate_Const< UserClass >::FMethodPtr InFunc )
+	{
+		Add( FDelegate::CreateSP( InUserObject, InFunc ) );
+	}
+	template< class UserClass, PAYLOAD_TEMPLATE_DECL_OneVar >
+	inline void AddSP( UserClass* InUserObject, typename FDelegate::template TSPMethodDelegate_OneVar< UserClass, PAYLOAD_TEMPLATE_LIST_OneVar >::FMethodPtr InFunc, PAYLOAD_TEMPLATE_ARGS_OneVar )
+	{
+		Add( FDelegate::CreateSP( InUserObject, InFunc, PAYLOAD_TEMPLATE_PASSIN_OneVar ) );
+	}
+	template< class UserClass, PAYLOAD_TEMPLATE_DECL_OneVar >
+	inline void AddSP( UserClass* InUserObject, typename FDelegate::template TSPMethodDelegate_OneVar_Const< UserClass, PAYLOAD_TEMPLATE_LIST_OneVar >::FMethodPtr InFunc, PAYLOAD_TEMPLATE_ARGS_OneVar )
+	{
+		Add( FDelegate::CreateSP( InUserObject, InFunc, PAYLOAD_TEMPLATE_PASSIN_OneVar ) );
+	}
+	template< class UserClass, PAYLOAD_TEMPLATE_DECL_TwoVars >
+	inline void AddSP( UserClass* InUserObject, typename FDelegate::template TSPMethodDelegate_TwoVars< UserClass, PAYLOAD_TEMPLATE_LIST_TwoVars >::FMethodPtr InFunc, PAYLOAD_TEMPLATE_ARGS_TwoVars )
+	{
+		Add( FDelegate::CreateSP( InUserObject, InFunc, PAYLOAD_TEMPLATE_PASSIN_TwoVars ) );
+	}
+	template< class UserClass, PAYLOAD_TEMPLATE_DECL_TwoVars >
+	inline void AddSP( UserClass* InUserObject, typename FDelegate::template TSPMethodDelegate_TwoVars_Const< UserClass, PAYLOAD_TEMPLATE_LIST_TwoVars >::FMethodPtr InFunc, PAYLOAD_TEMPLATE_ARGS_TwoVars )
+	{
+		Add( FDelegate::CreateSP( InUserObject, InFunc, PAYLOAD_TEMPLATE_PASSIN_TwoVars ) );
+	}
+	template< class UserClass, PAYLOAD_TEMPLATE_DECL_ThreeVars >
+	inline void AddSP( UserClass* InUserObject, typename FDelegate::template TSPMethodDelegate_ThreeVars< UserClass, PAYLOAD_TEMPLATE_LIST_ThreeVars >::FMethodPtr InFunc, PAYLOAD_TEMPLATE_ARGS_ThreeVars )
+	{
+		Add( FDelegate::CreateSP( InUserObject, InFunc, PAYLOAD_TEMPLATE_PASSIN_ThreeVars ) );
+	}
+	template< class UserClass, PAYLOAD_TEMPLATE_DECL_ThreeVars >
+	inline void AddSP( UserClass* InUserObject, typename FDelegate::template TSPMethodDelegate_ThreeVars_Const< UserClass, PAYLOAD_TEMPLATE_LIST_ThreeVars >::FMethodPtr InFunc, PAYLOAD_TEMPLATE_ARGS_ThreeVars )
+	{
+		Add( FDelegate::CreateSP( InUserObject, InFunc, PAYLOAD_TEMPLATE_PASSIN_ThreeVars ) );
+	}
+	template< class UserClass, PAYLOAD_TEMPLATE_DECL_FourVars >
+	inline void AddSP( UserClass* InUserObject, typename FDelegate::template TSPMethodDelegate_FourVars< UserClass, PAYLOAD_TEMPLATE_LIST_FourVars >::FMethodPtr InFunc, PAYLOAD_TEMPLATE_ARGS_FourVars )
+	{
+		Add( FDelegate::CreateSP( InUserObject, InFunc, PAYLOAD_TEMPLATE_PASSIN_FourVars ) );
+	}
+	template< class UserClass, PAYLOAD_TEMPLATE_DECL_FourVars >
+	inline void AddSP( UserClass* InUserObject, typename FDelegate::template TSPMethodDelegate_FourVars_Const< UserClass, PAYLOAD_TEMPLATE_LIST_FourVars >::FMethodPtr InFunc, PAYLOAD_TEMPLATE_ARGS_FourVars )
+	{
+		Add( FDelegate::CreateSP( InUserObject, InFunc, PAYLOAD_TEMPLATE_PASSIN_FourVars ) );
+	}
+
+	/**
+	 * Adds a shared pointer-based (slower, conditionally thread-safe) member function delegate.  Shared pointer delegates keep a weak reference to your object.
+	 *
+	 * @param	InUserObjectRef	User object to bind to
+	 * @param	InFunc			Class method function address
+	 */
+	template< class UserClass >
+	inline void AddThreadSafeSP( const TSharedRef< UserClass, ESPMode::ThreadSafe >& InUserObjectRef, typename FDelegate::template TThreadSafeSPMethodDelegate< UserClass >::FMethodPtr InFunc )
+	{
+		Add( FDelegate::CreateThreadSafeSP( InUserObjectRef, InFunc ) );
+	}
+	template< class UserClass >
+	inline void AddThreadSafeSP( const TSharedRef< UserClass, ESPMode::ThreadSafe >& InUserObjectRef, typename FDelegate::template TThreadSafeSPMethodDelegate_Const< UserClass >::FMethodPtr InFunc )
+	{
+		Add( FDelegate::CreateThreadSafeSP( InUserObjectRef, InFunc ) );
+	}
+	template< class UserClass, PAYLOAD_TEMPLATE_DECL_OneVar >
+	inline void AddThreadSafeSP( const TSharedRef< UserClass, ESPMode::ThreadSafe >& InUserObjectRef, typename FDelegate::template TThreadSafeSPMethodDelegate_OneVar< UserClass, PAYLOAD_TEMPLATE_LIST_OneVar >::FMethodPtr InFunc, PAYLOAD_TEMPLATE_ARGS_OneVar )
+	{
+		Add( FDelegate::CreateThreadSafeSP( InUserObjectRef, InFunc, PAYLOAD_TEMPLATE_PASSIN_OneVar ) );
+	}
+	template< class UserClass, PAYLOAD_TEMPLATE_DECL_OneVar >
+	inline void AddThreadSafeSP( const TSharedRef< UserClass, ESPMode::ThreadSafe >& InUserObjectRef, typename FDelegate::template TThreadSafeSPMethodDelegate_OneVar_Const< UserClass, PAYLOAD_TEMPLATE_LIST_OneVar >::FMethodPtr InFunc, PAYLOAD_TEMPLATE_ARGS_OneVar )
+	{
+		Add( FDelegate::CreateThreadSafeSP( InUserObjectRef, InFunc, PAYLOAD_TEMPLATE_PASSIN_OneVar ) );
+	}
+	template< class UserClass, PAYLOAD_TEMPLATE_DECL_TwoVars >
+	inline void AddThreadSafeSP( const TSharedRef< UserClass, ESPMode::ThreadSafe >& InUserObjectRef, typename FDelegate::template TThreadSafeSPMethodDelegate_TwoVars< UserClass, PAYLOAD_TEMPLATE_LIST_TwoVars >::FMethodPtr InFunc, PAYLOAD_TEMPLATE_ARGS_TwoVars )
+	{
+		Add( FDelegate::CreateThreadSafeSP( InUserObjectRef, InFunc, PAYLOAD_TEMPLATE_PASSIN_TwoVars ) );
+	}
+	template< class UserClass, PAYLOAD_TEMPLATE_DECL_TwoVars >
+	inline void AddThreadSafeSP( const TSharedRef< UserClass, ESPMode::ThreadSafe >& InUserObjectRef, typename FDelegate::template TThreadSafeSPMethodDelegate_TwoVars_Const< UserClass, PAYLOAD_TEMPLATE_LIST_TwoVars >::FMethodPtr InFunc, PAYLOAD_TEMPLATE_ARGS_TwoVars )
+	{
+		Add( FDelegate::CreateThreadSafeSP( InUserObjectRef, InFunc, PAYLOAD_TEMPLATE_PASSIN_TwoVars ) );
+	}
+	template< class UserClass, PAYLOAD_TEMPLATE_DECL_ThreeVars >
+	inline void AddThreadSafeSP( const TSharedRef< UserClass, ESPMode::ThreadSafe >& InUserObjectRef, typename FDelegate::template TThreadSafeSPMethodDelegate_ThreeVars< UserClass, PAYLOAD_TEMPLATE_LIST_ThreeVars >::FMethodPtr InFunc, PAYLOAD_TEMPLATE_ARGS_ThreeVars )
+	{
+		Add( FDelegate::CreateThreadSafeSP( InUserObjectRef, InFunc, PAYLOAD_TEMPLATE_PASSIN_ThreeVars ) );
+	}
+	template< class UserClass, PAYLOAD_TEMPLATE_DECL_ThreeVars >
+	inline void AddThreadSafeSP( const TSharedRef< UserClass, ESPMode::ThreadSafe >& InUserObjectRef, typename FDelegate::template TThreadSafeSPMethodDelegate_ThreeVars_Const< UserClass, PAYLOAD_TEMPLATE_LIST_ThreeVars >::FMethodPtr InFunc, PAYLOAD_TEMPLATE_ARGS_ThreeVars )
+	{
+		Add( FDelegate::CreateThreadSafeSP( InUserObjectRef, InFunc, PAYLOAD_TEMPLATE_PASSIN_ThreeVars ) );
+	}
+	template< class UserClass, PAYLOAD_TEMPLATE_DECL_FourVars >
+	inline void AddThreadSafeSP( const TSharedRef< UserClass, ESPMode::ThreadSafe >& InUserObjectRef, typename FDelegate::template TThreadSafeSPMethodDelegate_FourVars< UserClass, PAYLOAD_TEMPLATE_LIST_FourVars >::FMethodPtr InFunc, PAYLOAD_TEMPLATE_ARGS_FourVars )
+	{
+		Add( FDelegate::CreateThreadSafeSP( InUserObjectRef, InFunc, PAYLOAD_TEMPLATE_PASSIN_FourVars ) );
+	}
+	template< class UserClass, PAYLOAD_TEMPLATE_DECL_FourVars >
+	inline void AddThreadSafeSP( const TSharedRef< UserClass, ESPMode::ThreadSafe >& InUserObjectRef, typename FDelegate::template TThreadSafeSPMethodDelegate_FourVars_Const< UserClass, PAYLOAD_TEMPLATE_LIST_FourVars >::FMethodPtr InFunc, PAYLOAD_TEMPLATE_ARGS_FourVars )
+	{
+		Add( FDelegate::CreateThreadSafeSP( InUserObjectRef, InFunc, PAYLOAD_TEMPLATE_PASSIN_FourVars ) );
+	}
+
+	/**
+	 * Adds a shared pointer-based (slower, conditionally thread-safe) member function delegate.
+	 *
+	 * Shared pointer delegates keep a weak reference to your object.
+	 *
+	 * @param	InUserObject	User object to bind to
+	 * @param	InFunc			Class method function address
+	 */
+	template< class UserClass >
+	inline void AddThreadSafeSP( UserClass* InUserObject, typename FDelegate::template TThreadSafeSPMethodDelegate< UserClass >::FMethodPtr InFunc )
+	{
+		Add( FDelegate::CreateThreadSafeSP( InUserObject, InFunc ) );
+	}
+	template< class UserClass >
+	inline void AddThreadSafeSP( UserClass* InUserObject, typename FDelegate::template TThreadSafeSPMethodDelegate_Const< UserClass >::FMethodPtr InFunc )
+	{
+		Add( FDelegate::CreateThreadSafeSP( InUserObject, InFunc ) );
+	}
+	template< class UserClass, PAYLOAD_TEMPLATE_DECL_OneVar >
+	inline void AddThreadSafeSP( UserClass* InUserObject, typename FDelegate::template TThreadSafeSPMethodDelegate_OneVar< UserClass, PAYLOAD_TEMPLATE_LIST_OneVar >::FMethodPtr InFunc, PAYLOAD_TEMPLATE_ARGS_OneVar )
+	{
+		Add( FDelegate::CreateThreadSafeSP( InUserObject, InFunc, PAYLOAD_TEMPLATE_PASSIN_OneVar ) );
+	}
+	template< class UserClass, PAYLOAD_TEMPLATE_DECL_OneVar >
+	inline void AddThreadSafeSP( UserClass* InUserObject, typename FDelegate::template TThreadSafeSPMethodDelegate_OneVar_Const< UserClass, PAYLOAD_TEMPLATE_LIST_OneVar >::FMethodPtr InFunc, PAYLOAD_TEMPLATE_ARGS_OneVar )
+	{
+		Add( FDelegate::CreateThreadSafeSP( InUserObject, InFunc, PAYLOAD_TEMPLATE_PASSIN_OneVar ) );
+	}
+	template< class UserClass, PAYLOAD_TEMPLATE_DECL_TwoVars >
+	inline void AddThreadSafeSP( UserClass* InUserObject, typename FDelegate::template TThreadSafeSPMethodDelegate_TwoVars< UserClass, PAYLOAD_TEMPLATE_LIST_TwoVars >::FMethodPtr InFunc, PAYLOAD_TEMPLATE_ARGS_TwoVars )
+	{
+		Add( FDelegate::CreateThreadSafeSP( InUserObject, InFunc, PAYLOAD_TEMPLATE_PASSIN_TwoVars ) );
+	}
+	template< class UserClass, PAYLOAD_TEMPLATE_DECL_TwoVars >
+	inline void AddThreadSafeSP( UserClass* InUserObject, typename FDelegate::template TThreadSafeSPMethodDelegate_TwoVars_Const< UserClass, PAYLOAD_TEMPLATE_LIST_TwoVars >::FMethodPtr InFunc, PAYLOAD_TEMPLATE_ARGS_TwoVars )
+	{
+		Add( FDelegate::CreateThreadSafeSP( InUserObject, InFunc, PAYLOAD_TEMPLATE_PASSIN_TwoVars ) );
+	}
+	template< class UserClass, PAYLOAD_TEMPLATE_DECL_ThreeVars >
+	inline void AddThreadSafeSP( UserClass* InUserObject, typename FDelegate::template TThreadSafeSPMethodDelegate_ThreeVars< UserClass, PAYLOAD_TEMPLATE_LIST_ThreeVars >::FMethodPtr InFunc, PAYLOAD_TEMPLATE_ARGS_ThreeVars )
+	{
+		Add( FDelegate::CreateThreadSafeSP( InUserObject, InFunc, PAYLOAD_TEMPLATE_PASSIN_ThreeVars ) );
+	}
+	template< class UserClass, PAYLOAD_TEMPLATE_DECL_ThreeVars >
+	inline void AddThreadSafeSP( UserClass* InUserObject, typename FDelegate::template TThreadSafeSPMethodDelegate_ThreeVars_Const< UserClass, PAYLOAD_TEMPLATE_LIST_ThreeVars >::FMethodPtr InFunc, PAYLOAD_TEMPLATE_ARGS_ThreeVars )
+	{
+		Add( FDelegate::CreateThreadSafeSP( InUserObject, InFunc, PAYLOAD_TEMPLATE_PASSIN_ThreeVars ) );
+	}
+	template< class UserClass, PAYLOAD_TEMPLATE_DECL_FourVars >
+	inline void AddThreadSafeSP( UserClass* InUserObject, typename FDelegate::template TThreadSafeSPMethodDelegate_FourVars< UserClass, PAYLOAD_TEMPLATE_LIST_FourVars >::FMethodPtr InFunc, PAYLOAD_TEMPLATE_ARGS_FourVars )
+	{
+		Add( FDelegate::CreateThreadSafeSP( InUserObject, InFunc, PAYLOAD_TEMPLATE_PASSIN_FourVars ) );
+	}
+	template< class UserClass, PAYLOAD_TEMPLATE_DECL_FourVars >
+	inline void AddThreadSafeSP( UserClass* InUserObject, typename FDelegate::template TThreadSafeSPMethodDelegate_FourVars_Const< UserClass, PAYLOAD_TEMPLATE_LIST_FourVars >::FMethodPtr InFunc, PAYLOAD_TEMPLATE_ARGS_FourVars )
+	{
+		Add( FDelegate::CreateThreadSafeSP( InUserObject, InFunc, PAYLOAD_TEMPLATE_PASSIN_FourVars ) );
+	}
+
+	/**
+	 * Adds a UFunction-based member function delegate.
+	 *
+	 * UFunction delegates keep a weak reference to your object.
+	 *
+	 * @param	InUserObject	User object to bind to
+	 * @param	InFunctionName			Class method function address
+	 */
+	template< class UObjectTemplate >
+	inline void AddUFunction( UObjectTemplate* InUserObject, const FName& InFunctionName )
+	{
+		Add( FDelegate::CreateUFunction( InUserObject, InFunctionName ) );
+	}
+	template< class UObjectTemplate, PAYLOAD_TEMPLATE_DECL_OneVar >
+	inline void AddUFunction( UObjectTemplate* InUserObject, const FName& InFunctionName, PAYLOAD_TEMPLATE_ARGS_OneVar )
+	{
+		Add( FDelegate::CreateUFunction( InUserObject, InFunctionName, PAYLOAD_TEMPLATE_PASSIN_OneVar ) );
+	}
+	template< class UObjectTemplate, PAYLOAD_TEMPLATE_DECL_TwoVars >
+	inline void AddUFunction( UObjectTemplate* InUserObject, const FName& InFunctionName, PAYLOAD_TEMPLATE_ARGS_TwoVars )
+	{
+		Add( FDelegate::CreateUFunction( InUserObject, InFunctionName, PAYLOAD_TEMPLATE_PASSIN_TwoVars ) );
+	}
+	template< class UObjectTemplate, PAYLOAD_TEMPLATE_DECL_ThreeVars >
+	inline void AddUFunction( UObjectTemplate* InUserObject, const FName& InFunctionName, PAYLOAD_TEMPLATE_ARGS_ThreeVars )
+	{
+		Add( FDelegate::CreateUFunction( InUserObject, InFunctionName, PAYLOAD_TEMPLATE_PASSIN_ThreeVars ) );
+	}
+	template< class UObjectTemplate, PAYLOAD_TEMPLATE_DECL_FourVars >
+	inline void AddUFunction( UObjectTemplate* InUserObject, const FName& InFunctionName, PAYLOAD_TEMPLATE_ARGS_FourVars )
+	{
+		Add( FDelegate::CreateUFunction( InUserObject, InFunctionName, PAYLOAD_TEMPLATE_PASSIN_FourVars ) );
+	}
+
+	/**
+	 * Adds a UObject-based member function delegate.
+	 *
+	 * UObject delegates keep a weak reference to your object.
+	 *
+	 * @param	InUserObject	User object to bind to
+	 * @param	InFunc			Class method function address
+	 */
+	template< class UserClass >
+	inline void AddUObject( UserClass* InUserObject, typename FDelegate::template TUObjectMethodDelegate< UserClass >::FMethodPtr InFunc )
+	{
+		Add( FDelegate::CreateUObject( InUserObject, InFunc ) );
+	}
+	template< class UserClass >
+	inline void AddUObject( UserClass* InUserObject, typename FDelegate::template TUObjectMethodDelegate_Const< UserClass >::FMethodPtr InFunc )
+	{
+		Add( FDelegate::CreateUObject( InUserObject, InFunc ) );
+	}
+	template< class UserClass, PAYLOAD_TEMPLATE_DECL_OneVar >
+	inline void AddUObject( UserClass* InUserObject, typename FDelegate::template TUObjectMethodDelegate_OneVar< UserClass, PAYLOAD_TEMPLATE_LIST_OneVar >::FMethodPtr InFunc, PAYLOAD_TEMPLATE_ARGS_OneVar )
+	{
+		Add( FDelegate::CreateUObject( InUserObject, InFunc, PAYLOAD_TEMPLATE_PASSIN_OneVar ) );
+	}
+	template< class UserClass, PAYLOAD_TEMPLATE_DECL_OneVar >
+	inline void AddUObject( UserClass* InUserObject, typename FDelegate::template TUObjectMethodDelegate_OneVar_Const< UserClass, PAYLOAD_TEMPLATE_LIST_OneVar >::FMethodPtr InFunc, PAYLOAD_TEMPLATE_ARGS_OneVar )
+	{
+		Add( FDelegate::CreateUObject( InUserObject, InFunc, PAYLOAD_TEMPLATE_PASSIN_OneVar ) );
+	}
+	template< class UserClass, PAYLOAD_TEMPLATE_DECL_TwoVars >
+	inline void AddUObject( UserClass* InUserObject, typename FDelegate::template TUObjectMethodDelegate_TwoVars< UserClass, PAYLOAD_TEMPLATE_LIST_TwoVars >::FMethodPtr InFunc, PAYLOAD_TEMPLATE_ARGS_TwoVars )
+	{
+		Add( FDelegate::CreateUObject( InUserObject, InFunc, PAYLOAD_TEMPLATE_PASSIN_TwoVars ) );
+	}
+	template< class UserClass, PAYLOAD_TEMPLATE_DECL_TwoVars >
+	inline void AddUObject( UserClass* InUserObject, typename FDelegate::template TUObjectMethodDelegate_TwoVars_Const< UserClass, PAYLOAD_TEMPLATE_LIST_TwoVars >::FMethodPtr InFunc, PAYLOAD_TEMPLATE_ARGS_TwoVars )
+	{
+		Add( FDelegate::CreateUObject( InUserObject, InFunc, PAYLOAD_TEMPLATE_PASSIN_TwoVars ) );
+	}
+	template< class UserClass, PAYLOAD_TEMPLATE_DECL_ThreeVars >
+	inline void AddUObject( UserClass* InUserObject, typename FDelegate::template TUObjectMethodDelegate_ThreeVars< UserClass, PAYLOAD_TEMPLATE_LIST_ThreeVars >::FMethodPtr InFunc, PAYLOAD_TEMPLATE_ARGS_ThreeVars )
+	{
+		Add( FDelegate::CreateUObject( InUserObject, InFunc, PAYLOAD_TEMPLATE_PASSIN_ThreeVars ) );
+	}
+	template< class UserClass, PAYLOAD_TEMPLATE_DECL_ThreeVars >
+	inline void AddUObject( UserClass* InUserObject, typename FDelegate::template TUObjectMethodDelegate_ThreeVars_Const< UserClass, PAYLOAD_TEMPLATE_LIST_ThreeVars >::FMethodPtr InFunc, PAYLOAD_TEMPLATE_ARGS_ThreeVars )
+	{
+		Add( FDelegate::CreateUObject( InUserObject, InFunc, PAYLOAD_TEMPLATE_PASSIN_ThreeVars ) );
+	}
+	template< class UserClass, PAYLOAD_TEMPLATE_DECL_FourVars >
+	inline void AddUObject( UserClass* InUserObject, typename FDelegate::template TUObjectMethodDelegate_FourVars< UserClass, PAYLOAD_TEMPLATE_LIST_FourVars >::FMethodPtr InFunc, PAYLOAD_TEMPLATE_ARGS_FourVars )
+	{
+		Add( FDelegate::CreateUObject( InUserObject, InFunc, PAYLOAD_TEMPLATE_PASSIN_FourVars ) );
+	}
+	template< class UserClass, PAYLOAD_TEMPLATE_DECL_FourVars >
+	inline void AddUObject( UserClass* InUserObject, typename FDelegate::template TUObjectMethodDelegate_FourVars_Const< UserClass, PAYLOAD_TEMPLATE_LIST_FourVars >::FMethodPtr InFunc, PAYLOAD_TEMPLATE_ARGS_FourVars )
+	{
+		Add( FDelegate::CreateUObject( InUserObject, InFunc, PAYLOAD_TEMPLATE_PASSIN_FourVars ) );
+	}
+
+public:
+
+	/**
+	 * Removes a delegate instance from this multi-cast delegate's invocation list (performance is O(N)).
+	 *
+	 * Note that the order of the delegate instances may not be preserved!
+	 *
+	 * @param DelegateInstance The delegate instance to remove.
+	 */
+	void Remove( const TDelegateInstanceInterface& DelegateInstance )
+	{
+		// NOTE: For delegates with PAYLOAD, the payload variables are not considered when comparing delegate objects!
+		// This means that two delegate instances bound to the same object/method, but with different payload data,
+		// cannot be bound to the same single multi-cast delegate.  We could support this by changing IsSameFunction()
+		// to do comparison of payload data, but there are some trade-offs if we do that.
+
+		RemoveDelegateInstance(DelegateInstance);
+		CompactInvocationList();
+	}
+
+	/**
+	 * DEPRECATED: Removes a unicast delegate from this multi-cast delegate's invocation list (performance is O(N)).
+	 *
+	 * The order of the delegates may not be preserved!
+	 * This function is retained for backwards compatibility.
+	 *
+	 * @param Delegate The delegate to remove.
+	 */
+	void Remove( const FDelegate& Delegate )
+	{
+		// NOTE: For delegates with PAYLOAD, the payload variables are not considered when comparing delegate objects!
+		// This means that two delegate instances bound to the same object/method, but with different payload data,
+		// cannot be bound to the same single multi-cast delegate.  We could support this by changing IsSameFunction()
+		// to do comparison of payload data, but there are some trade-offs if we do that.
+
+		TDelegateInstanceInterface* DelegateInstance = (TDelegateInstanceInterface*)Delegate.GetDelegateInstance();
+
+		if (DelegateInstance != nullptr)
+		{
+			RemoveDelegateInstance(*DelegateInstance);
+			CompactInvocationList();
+		}		
+	}
+
+	// NOTE: These direct Remove methods are only supported for multi-cast delegates with no payload attached.
+	// See the comment in the multi-cast delegate Remove() method above for more details.
+
+	/**
+	 * Removes a raw C++ pointer global function delegate (performance is O(N)).  Note that the order of the
+	 * delegates may not be preserved!
+	 *
+	 * @param	InFunc	Function pointer
+	 */
+	inline void RemoveStatic( typename FDelegate::FStaticDelegate::FFuncPtr InFunc )
+	{
+		Remove( FDelegate::CreateStatic( InFunc ) );
+	}
+
+	/**
+	 * Removes a raw C++ pointer delegate (performance is O(N)).  Note that the order of the delegates may not
+	 * be preserved!
+	 *
+	 * @param	InUserObject	User object to unbind from
+	 * @param	InFunc			Class method function address
+	 */
+	template< class UserClass >
+	inline void RemoveRaw( UserClass* InUserObject, typename FDelegate::template TRawMethodDelegate< UserClass >::FMethodPtr InFunc )
+	{
+		Remove( FDelegate::CreateRaw( InUserObject, InFunc ) );
+	}
+	template< class UserClass >
+	inline void RemoveRaw( UserClass* InUserObject, typename FDelegate::template TRawMethodDelegate_Const< UserClass >::FMethodPtr InFunc )
+	{
+		Remove( FDelegate::CreateRaw( InUserObject, InFunc ) );
+	}
+
+	/**
+	 * Removes a shared pointer-based member function delegate (performance is O(N)).  Note that the order of
+	 * the delegates may not be preserved!
+	 *
+	 * @param	InUserObjectRef	User object to unbind from
+	 * @param	InFunc			Class method function address
+	 */
+	template< class UserClass >
+	inline void RemoveSP( const TSharedRef< UserClass, ESPMode::Fast >& InUserObjectRef, typename FDelegate::template TSPMethodDelegate< UserClass >::FMethodPtr InFunc )
+	{
+		Remove( FDelegate::CreateSP( InUserObjectRef, InFunc ) );
+	}
+	template< class UserClass >
+	inline void RemoveSP( const TSharedRef< UserClass, ESPMode::Fast >& InUserObjectRef, typename FDelegate::template TSPMethodDelegate_Const< UserClass >::FMethodPtr InFunc )
+	{
+		Remove( FDelegate::CreateSP( InUserObjectRef, InFunc ) );
+	}
+
+	/**
+	 * Removes a shared pointer-based member function delegate (performance is O(N)).  Note that the order of
+	 * the delegates may not be preserved!
+	 *
+	 * @param	InUserObject	User object to unbind from
+	 * @param	InFunc			Class method function address
+	 */
+	template< class UserClass >
+	inline void RemoveSP( UserClass* InUserObject, typename FDelegate::template TSPMethodDelegate< UserClass >::FMethodPtr InFunc )
+	{
+		Remove( FDelegate::CreateSP( InUserObject, InFunc ) );
+	}
+	template< class UserClass >
+	inline void RemoveSP( UserClass* InUserObject, typename FDelegate::template TSPMethodDelegate_Const< UserClass >::FMethodPtr InFunc )
+	{
+		Remove( FDelegate::CreateSP( InUserObject, InFunc ) );
+	}
+
+	/**
+	 * Removes a shared pointer-based member function delegate (performance is O(N)).  Note that the order of
+	 * the delegates may not be preserved!
+	 *
+	 * @param	InUserObjectRef	User object to unbind from
+	 * @param	InFunc			Class method function address
+	 */
+	template< class UserClass >
+	inline void RemoveThreadSafeSP( const TSharedRef< UserClass, ESPMode::ThreadSafe >& InUserObjectRef, typename FDelegate::template TThreadSafeSPMethodDelegate< UserClass >::FMethodPtr InFunc )
+	{
+		Remove( FDelegate::CreateThreadSafeSP( InUserObjectRef, InFunc ) );
+	}
+	template< class UserClass >
+	inline void RemoveThreadSafeSP( const TSharedRef< UserClass, ESPMode::ThreadSafe >& InUserObjectRef, typename FDelegate::template TThreadSafeSPMethodDelegate_Const< UserClass >::FMethodPtr InFunc )
+	{
+		Remove( FDelegate::CreateThreadSafeSP( InUserObjectRef, InFunc ) );
+	}
+
+	/**
+	 * Removes a shared pointer-based member function delegate (performance is O(N)).  Note that the order of
+	 * the delegates may not be preserved!
+	 *
+	 * @param	InUserObject	User object to unbind from
+	 * @param	InFunc			Class method function address
+	 */
+	template< class UserClass >
+	inline void RemoveThreadSafeSP( UserClass* InUserObject, typename FDelegate::template TThreadSafeSPMethodDelegate< UserClass >::FMethodPtr InFunc )
+	{
+		Remove( FDelegate::CreateThreadSafeSP( InUserObject, InFunc ) );
+	}
+	template< class UserClass >
+	inline void RemoveThreadSafeSP( UserClass* InUserObject, typename FDelegate::template TThreadSafeSPMethodDelegate_Const< UserClass >::FMethodPtr InFunc )
+	{
+		Remove( FDelegate::CreateThreadSafeSP( InUserObject, InFunc ) );
+	}
+
+	/**
+	 * Removes a UFunction-based member function delegate (performance is O(N)).  Note that the order of the
+	 * delegates may not be preserved!
+	 *
+	 * @param	InUserObject	User object to unbind from
+	 * @param	InFunc			Class method function address
+	 */
+	template< class UObjectTemplate >
+	inline void RemoveUFunction( UObjectTemplate* InUserObject, const FName& InFunctionName )
+	{
+		Remove( FDelegate::CreateUFunction( InUserObject, InFunctionName ) );
+	}
+
+	/**
+	 * Removes a UObject-based member function delegate (performance is O(N)).  Note that the order of the
+	 * delegates may not be preserved!
+	 *
+	 * @param	InUserObject	User object to unbind from
+	 * @param	InFunc			Class method function address
+	 */
+	template< class UserClass >
+	inline void RemoveUObject( UserClass* InUserObject, typename FDelegate::template TUObjectMethodDelegate< UserClass >::FMethodPtr InFunc )
+	{
+		Remove( FDelegate::CreateUObject( InUserObject, InFunc ) );
+	}
+	template< class UserClass >
+	inline void RemoveUObject( UserClass* InUserObject, typename FDelegate::template TUObjectMethodDelegate_Const< UserClass >::FMethodPtr InFunc )
+	{
+		Remove( FDelegate::CreateUObject( InUserObject, InFunc ) );
+	}
+
+protected:
+
+	/** 
+	 * Hidden default constructor.
+	 */
+	inline BASE_MULTICAST_DELEGATE_CLASS( ) { }
+
+protected:
+
+	/**
+	 * Adds a function delegate to this multi-cast delegate's invocation list
+	 *
+	 * @param	InDelegate	Delegate to add
+	 */
+	void AddDelegateInstance( TDelegateInstanceInterface& InDelegateInstance )
+	{
+#if !UE_BUILD_SHIPPING && !UE_BUILD_TEST
+		// verify that the same function isn't already bound
+		for (auto DelegateInstance : GetInvocationList())
+	{
+			// this down-cast is OK! allows for managing invocation list in the base class without requiring virtual functions
+			TDelegateInstanceInterface* DelegateInstanceInterface = (TDelegateInstanceInterface*)DelegateInstance;
+			check(!DelegateInstanceInterface->IsSameFunction(InDelegateInstance));
+	}
+#endif
+
+		AddInternal(InDelegateInstance);
+	}
+
+	/**
+	 * Broadcasts this delegate to all bound objects, except to those that may have expired.
+	 *
+	 * The constness of this method is a lie, because we are removing compactable delegate intsances.
+	 */
+	void Broadcast( FUNC_PARAM_LIST ) const
+	{
+		const TArray<IDelegateInstance*>& InvocationList = GetInvocationList();
+
+		if (InvocationList.Num() == 0)
+		{
+			return;
+		}
+
+		// copy invocation list just in case the list is modified by one of the broadcast callbacks
+		typedef TArray<IDelegateInstance*, TInlineAllocator<4>> FInlineInvocationList;
+			FInlineInvocationList InvocationListCopy = FInlineInvocationList(InvocationList);
+
+			// Invoke each bound function
+		for (IDelegateInstance* DelegateInstance : InvocationListCopy)
+			{
+			// this down-cast is OK! allows for managing invocation list in the base class without requiring virtual functions
+			TDelegateInstanceInterface* DelegateInstanceInterface = (TDelegateInstanceInterface*)DelegateInstance;
+
+			// execute the delegate instance...
+			if (!DelegateInstanceInterface->ExecuteIfSafe(FUNC_PARAM_PASSTHRU))
+				{
+				// ... or remove it if compactable
+				if (DelegateInstance->IsCompactable())
+					{
+					// search for the instance, because original list could have been modified during broadcast
+					// this is really shady, because it violates the constness of this method
+					const_cast<BASE_MULTICAST_DELEGATE_CLASS*>(this)->RemoveDelegateInstance(*DelegateInstanceInterface);
+				}
+			}
+		}
+	}
+
+	/**
+	 * Removes a function from this multi-cast delegate's invocation list (performance is O(N)).
+	 *
+	 * Note that the order of the delegates may not be preserved! This method will not assert if
+	 * the given delegate instance was not found in the invocation list, because some instances
+	 * hold on to weak pointers to objects that may become invalid.
+	 *
+	 * @param InDelegateInstance The delegate instance to remove.
+	 */
+	void RemoveDelegateInstance( const TDelegateInstanceInterface& InDelegateInstance )
+	{
+		const TArray<IDelegateInstance*> InvocationList = GetInvocationList();
+
+		for (int32 InvocationListIndex = InvocationList.Num() - 1; InvocationListIndex >= 0; --InvocationListIndex)
+		{
+			// this down-cast is OK! allows for managing invocation list in the base class without requiring virtual functions
+			TDelegateInstanceInterface* DelegateInstance = (TDelegateInstanceInterface*)InvocationList[InvocationListIndex];
+
+			// NOTE: We must do a deep compare here, not just compare delegate pointers, because multiple
+			//       delegate pointers can refer to the exact same object and method
+			if (DelegateInstance->IsSameFunction(InDelegateInstance))
+			{
+				RemoveByIndex(InvocationListIndex);
+
+				// no need to continue, as we never allow the same delegate to be bound twice
+				break;
+			}
+		}
+	}
+};
+
+
+	/**
+ * Implements a multicast delegate.
+ *
+ * This class should not be instantiated directly. Use the DECLARE_MULTICAST_DELEGATE macros instead.
+	 */
+template< FUNC_TEMPLATE_DECL_TYPENAME >
+class MULTICAST_DELEGATE_CLASS : public BASE_MULTICAST_DELEGATE_CLASS< FUNC_TEMPLATE_DECL >
+{
+public:
+
+	/**
+	 * Checks to see if any functions are bound to this multi-cast delegate
+	 *
+	 * @return	True if any functions are bound
+	 */
+	inline bool IsBound() const
+	{
+		return Super::IsBound();
+	}
+
+	/**
+	 * Removes all functions from this delegate's invocation list
+	 */
+	inline void Clear()
+	{
+		return Super::Clear();
+	}
+
+	/**
+	 * Broadcasts this delegate to all bound objects, except to those that may have expired
+	 */
+	inline void Broadcast( FUNC_PARAM_LIST ) const
+	{
+		return Super::Broadcast( FUNC_PARAM_PASSTHRU );
+	}
+
+private: 
+	
+	// Prevents erroneous use by other classes.
+	typedef BASE_MULTICAST_DELEGATE_CLASS< FUNC_TEMPLATE_DECL > Super;
+};
+
+#endif		// FUNC_IS_VOID
+
+
+/**
+ * Dynamic delegate base object (UObject-based, serializable).  You'll use the various DECLARE_DYNAMIC_DELEGATE
+ * macros to create the actual delegate type, templated to the function signature the delegate is compatible with.
+ * Then, you can create an instance of that class when you want to assign functions to the delegate.
+ */
+template< FUNC_TEMPLATE_DECL_TYPENAME, typename TWeakPtr = FWeakObjectPtr >
+class DYNAMIC_DELEGATE_CLASS
+	: public TScriptDelegate<TWeakPtr>
+{
+public:
+
+	/**
+	 * Default constructor
+	 */
+	DYNAMIC_DELEGATE_CLASS() { }
+
+	/**
+	 * Construction from an FScriptDelegate must be explicit.  This is really only used by UObject system internals.
+	 *
+	 * @param	InScriptDelegate	The delegate to construct from by copying
+	 */
+	explicit DYNAMIC_DELEGATE_CLASS( const TScriptDelegate<TWeakPtr>& InScriptDelegate )
+		: TScriptDelegate<TWeakPtr>( InScriptDelegate )
+	{ }
+
+	/**
+	 * Templated helper class to define a typedef for user's method pointer, then used below
+	 */
+	template< class UserClass >
+	class TMethodPtrResolver
+	{
+	public:
+
+		typedef RetValType ( UserClass::*FMethodPtr )( FUNC_PARAM_LIST );
+	};
+
+	/**
+	 * Binds a UObject instance and a UObject method address to this delegate.
+	 *
+	 * @param	InUserObject			UObject instance
+	 * @param	InMethodPtr				Member function address pointer
+	 * @param	InMacroFunctionName		Name of member function, including class name (generated by a macro)
+	 *
+	 * NOTE:  Do not call this function directly.  Instead, call BindDynamic() which is a macro proxy function that
+	 *        automatically sets the function name string for the caller.
+	 */
+	template< class UserClass >
+	void __Internal_BindDynamic( UserClass* InUserObject, typename TMethodPtrResolver< UserClass >::FMethodPtr InMethodPtr, const FString& InMacroFunctionName )
+	{
+		check( InUserObject != nullptr && InMethodPtr != nullptr && !InMacroFunctionName.IsEmpty() );
+
+		// NOTE: We're not actually storing the incoming method pointer or calling it.  We simply require it for type-safety reasons.
+
+		// NOTE: If you hit a compile error on the following line, it means you're trying to use a non-UObject type
+		//       with this delegate, which is not supported
+		this->Object = InUserObject;
+
+		// Store the function name.  The incoming function name was generated by a macro and includes the method's class name.
+		// We'll need to strip off that class prefix and just store the function name by itself.
+		this->FunctionName = *InMacroFunctionName.Mid( InMacroFunctionName.Find( TEXT( "::" ), ESearchCase::CaseSensitive, ESearchDir::FromEnd) + 2 );
+
+		ensureMsgf(this->IsBound(), TEXT("Unable to bind delegate to '%s' (function might not be marked as a UFUNCTION)"), *InMacroFunctionName);
+	}
+
+	friend uint32 GetTypeHash(const DYNAMIC_DELEGATE_CLASS& Key)
+	{
+		return FCrc::MemCrc_DEPRECATED(&Key,sizeof(Key));
+	}
+
+	// NOTE:  Execute() method must be defined in derived classes
+
+	// NOTE:  ExecuteIfBound() method must be defined in derived classes
+};
+
+
+/**
+ * Dynamic multi-cast delegate base object (UObject-based, serializable).  You'll use the various
+ * DECLARE_DYNAMIC_MULTICAST_DELEGATE macros to create the actual delegate type, templated to the function
+ * signature the delegate is compatible with.   Then, you can create an instance of that class when you
+ * want to assign functions to the delegate.
+ */
+template< FUNC_TEMPLATE_DECL_TYPENAME, typename TWeakPtr = FWeakObjectPtr >
+class DYNAMIC_MULTICAST_DELEGATE_CLASS
+	: public TMulticastScriptDelegate<TWeakPtr>
+{
+public:
+
+	/** The actual single-cast delegate class for this multi-cast delegate */
+	typedef DYNAMIC_DELEGATE_CLASS< FUNC_TEMPLATE_ARGS > FDelegate;
+
+	/**
+	 * Default constructor
+	 */
+	DYNAMIC_MULTICAST_DELEGATE_CLASS() { }
+
+	/**
+	 * Construction from an FMulticastScriptDelegate must be explicit.  This is really only used by UObject system internals.
+	 *
+	 * @param	InScriptDelegate	The delegate to construct from by copying
+	 */
+	explicit DYNAMIC_MULTICAST_DELEGATE_CLASS( const TMulticastScriptDelegate<TWeakPtr>& InMulticastScriptDelegate )
+		: TMulticastScriptDelegate<TWeakPtr>( InMulticastScriptDelegate )
+	{ }
+
+	/**
+	 * Tests if a UObject instance and a UObject method address pair are already bound to this multi-cast delegate.
+	 *
+	 * @param	InUserObject			UObject instance
+	 * @param	InMethodPtr				Member function address pointer
+	 * @param	InMacroFunctionName		Name of member function, including class name (generated by a macro)
+	 * @return	True if the instance/method is already bound.
+	 *
+	 * NOTE:  Do not call this function directly.  Instead, call IsAlreadyBound() which is a macro proxy function that
+	 *        automatically sets the function name string for the caller.
+	 */
+	template< class UserClass >
+	bool __Internal_IsAlreadyBound( UserClass* InUserObject, typename FDelegate::template TMethodPtrResolver< UserClass >::FMethodPtr InMethodPtr, const FString& InMacroFunctionName ) const
+	{
+		check( InUserObject != nullptr && InMethodPtr != nullptr && !InMacroFunctionName.IsEmpty() );
+
+		// NOTE: We're not actually using the incoming method pointer or calling it.  We simply require it for type-safety reasons.
+
+		FDelegate NewDelegate;
+		NewDelegate.__Internal_BindDynamic( InUserObject, InMethodPtr, InMacroFunctionName );
+
+		return this->Contains( NewDelegate );
+	}
+
+	/**
+	 * Binds a UObject instance and a UObject method address to this multi-cast delegate.
+	 *
+	 * @param	InUserObject			UObject instance
+	 * @param	InMethodPtr				Member function address pointer
+	 * @param	InMacroFunctionName		Name of member function, including class name (generated by a macro)
+	 *
+	 * NOTE:  Do not call this function directly.  Instead, call AddDynamic() which is a macro proxy function that
+	 *        automatically sets the function name string for the caller.
+	 */
+	template< class UserClass >
+	void __Internal_AddDynamic( UserClass* InUserObject, typename FDelegate::template TMethodPtrResolver< UserClass >::FMethodPtr InMethodPtr, const FString& InMacroFunctionName )
+	{
+		check( InUserObject != nullptr && InMethodPtr != nullptr && !InMacroFunctionName.IsEmpty() );
+
+		// NOTE: We're not actually storing the incoming method pointer or calling it.  We simply require it for type-safety reasons.
+
+		FDelegate NewDelegate;
+		NewDelegate.__Internal_BindDynamic( InUserObject, InMethodPtr, InMacroFunctionName );
+
+		this->Add( NewDelegate );
+	}
+
+	/**
+	 * Binds a UObject instance and a UObject method address to this multi-cast delegate, but only if it hasn't been bound before.
+	 *
+	 * @param	InUserObject			UObject instance
+	 * @param	InMethodPtr				Member function address pointer
+	 * @param	InMacroFunctionName		Name of member function, including class name (generated by a macro)
+	 *
+	 * NOTE:  Do not call this function directly.  Instead, call AddUniqueDynamic() which is a macro proxy function that
+	 *        automatically sets the function name string for the caller.
+	 */
+	template< class UserClass >
+	void __Internal_AddUniqueDynamic( UserClass* InUserObject, typename FDelegate::template TMethodPtrResolver< UserClass >::FMethodPtr InMethodPtr, const FString& InMacroFunctionName )
+	{
+		check( InUserObject != nullptr && InMethodPtr != nullptr && !InMacroFunctionName.IsEmpty() );
+
+		// NOTE: We're not actually storing the incoming method pointer or calling it.  We simply require it for type-safety reasons.
+
+		FDelegate NewDelegate;
+		NewDelegate.__Internal_BindDynamic( InUserObject, InMethodPtr, InMacroFunctionName );
+
+		this->AddUnique( NewDelegate );
+	}
+
+	/**
+	 * Unbinds a UObject instance and a UObject method address from this multi-cast delegate.
+	 *
+	 * @param	InUserObject			UObject instance
+	 * @param	InMethodPtr				Member function address pointer
+	 * @param	InMacroFunctionName		Name of member function, including class name (generated by a macro)
+	 *
+	 * NOTE:  Do not call this function directly.  Instead, call RemoveDynamic() which is a macro proxy function that
+	 *        automatically sets the function name string for the caller.
+	 */
+	template< class UserClass >
+	void __Internal_RemoveDynamic( UserClass* InUserObject, typename FDelegate::template TMethodPtrResolver< UserClass >::FMethodPtr InMethodPtr, const FString& InMacroFunctionName )
+	{
+		check( InUserObject != nullptr && InMethodPtr != nullptr && !InMacroFunctionName.IsEmpty() );
+
+		// NOTE: We're not actually storing the incoming method pointer or calling it.  We simply require it for type-safety reasons.
+
+		FDelegate NewDelegate;
+		NewDelegate.__Internal_BindDynamic( InUserObject, InMethodPtr, InMacroFunctionName );
+
+		this->Remove( NewDelegate );
+	}
+
+	// NOTE:  Broadcast() method must be defined in derived classes
+};
+
+
+#undef PAYLOAD_TEMPLATE_DECL_OneVar
+#undef PAYLOAD_TEMPLATE_LIST_OneVar
+#undef PAYLOAD_TEMPLATE_ARGS_OneVar
+#undef PAYLOAD_TEMPLATE_PASSIN_OneVar
+
+#undef PAYLOAD_TEMPLATE_DECL_TwoVars
+#undef PAYLOAD_TEMPLATE_LIST_TwoVars
+#undef PAYLOAD_TEMPLATE_ARGS_TwoVars
+#undef PAYLOAD_TEMPLATE_PASSIN_TwoVars
+
+#undef PAYLOAD_TEMPLATE_DECL_ThreeVars
+#undef PAYLOAD_TEMPLATE_LIST_ThreeVars
+#undef PAYLOAD_TEMPLATE_ARGS_ThreeVars
+#undef PAYLOAD_TEMPLATE_PASSIN_ThreeVars
+
+#undef PAYLOAD_TEMPLATE_DECL_FourVars
+#undef PAYLOAD_TEMPLATE_LIST_FourVars
+#undef PAYLOAD_TEMPLATE_ARGS_FourVars
+#undef PAYLOAD_TEMPLATE_PASSIN_FourVars
+
+#undef DYNAMIC_MULTICAST_DELEGATE_CLASS
+#undef DYNAMIC_DELEGATE_CLASS
+#undef EVENT_CLASS
+#undef MULTICAST_DELEGATE_CLASS
+#undef DELEGATE_CLASS
+
+#undef SP_METHOD_DELEGATE_INSTANCE_CLASS_NoVars
+#undef SP_METHOD_DELEGATE_INSTANCE_CLASS_OneVar
+#undef SP_METHOD_DELEGATE_INSTANCE_CLASS_TwoVars
+#undef SP_METHOD_DELEGATE_INSTANCE_CLASS_ThreeVars
+#undef SP_METHOD_DELEGATE_INSTANCE_CLASS_FourVars
+
+#undef RAW_METHOD_DELEGATE_INSTANCE_CLASS_NoVars
+#undef RAW_METHOD_DELEGATE_INSTANCE_CLASS_OneVar
+#undef RAW_METHOD_DELEGATE_INSTANCE_CLASS_TwoVars
+#undef RAW_METHOD_DELEGATE_INSTANCE_CLASS_ThreeVars
+#undef RAW_METHOD_DELEGATE_INSTANCE_CLASS_FourVars
+
+#undef UFUNCTION_DELEGATE_INSTANCE_CLASS_NoVars
+#undef UFUNCTION_DELEGATE_INSTANCE_CLASS_OneVar
+#undef UFUNCTION_DELEGATE_INSTANCE_CLASS_TwoVars
+#undef UFUNCTION_DELEGATE_INSTANCE_CLASS_ThreeVars
+#undef UFUNCTION_DELEGATE_INSTANCE_CLASS_FourVars
+
+#undef UOBJECT_METHOD_DELEGATE_INSTANCE_CLASS_NoVars
+#undef UOBJECT_METHOD_DELEGATE_INSTANCE_CLASS_OneVar
+#undef UOBJECT_METHOD_DELEGATE_INSTANCE_CLASS_TwoVars
+#undef UOBJECT_METHOD_DELEGATE_INSTANCE_CLASS_ThreeVars
+#undef UOBJECT_METHOD_DELEGATE_INSTANCE_CLASS_FourVars
+
+#undef STATIC_DELEGATE_INSTANCE_CLASS_NoVars
+#undef STATIC_DELEGATE_INSTANCE_CLASS_OneVar
+#undef STATIC_DELEGATE_INSTANCE_CLASS_TwoVars
+#undef STATIC_DELEGATE_INSTANCE_CLASS_ThreeVars
+#undef STATIC_DELEGATE_INSTANCE_CLASS_FourVars
+
+#undef DELEGATE_INSTANCE_INTERFACE_CLASS
diff --git a/Engine/Source/Runtime/Core/Public/Delegates/IDelegateInstance.h b/Engine/Source/Runtime/Core/Public/Delegates/IDelegateInstance.h
new file mode 100644
--- /dev/null
+++ b/Engine/Source/Runtime/Core/Public/Delegates/IDelegateInstance.h
@@ -0,0 +1,106 @@
+// Copyright 1998-2014 Epic Games, Inc. All Rights Reserved.
+
+/*=============================================================================
+	IDelegateInstance.cpp: Declares the IDelegateInstance interface.
+=============================================================================*/
+
+#pragma once
+
+
+/**
+ * Types of delegate instances
+ */
+namespace EDelegateInstanceType
+{
+	enum Type
+	{
+		/** Member function pointer to method in (fast, not thread-safe) shared pointer-based class */
+		SharedPointerMethod,
+
+		/** Member function pointer to method in (conditionally thread-safe) shared pointer-based class */
+		ThreadSafeSharedPointerMethod,
+
+		/** Raw C++ member function pointer (pointer to class method) */
+		RawMethod,
+
+		/** UFunction delegate */
+		UFunction,
+
+		/** Member function pointer to method in UObject-based class */
+		UObjectMethod,
+
+		/** Raw C++ static function pointer */
+		Raw
+	};
+}
+
+
+/**
+ * Interface for delegate instances.
+ */
+class IDelegateInstance
+{
+public:
+
+	/**
+	 * Returns the name of the UFunction that this delegate instance is bound to.
+	 *
+	 * @return Name of the function, or NAME_None if not bound to a UFunction.
+	 */
+	virtual FName GetFunctionName( ) const = 0;
+
+	/**
+	 * Returns raw pointer to the delegate method.
+	 *
+	 * @return  Raw pointer to the delegate method.
+	 */
+	virtual const void* GetRawMethodPtr( ) const = 0;
+
+	/**
+	 * Returns raw pointer to UserObject,
+	 *
+	 * @return  Raw pointer to UserObject.
+	 */
+	virtual const void* GetRawUserObject( ) const = 0;
+
+	/**
+	 * Returns the type of delegate instance
+	 *
+	 * @return  Delegate instance type
+	 */
+	virtual EDelegateInstanceType::Type GetType( ) const = 0;
+
+	/**
+	 * Returns true if this delegate is bound to the specified UserObject,
+	 *
+	 * @param  InUserObject
+	 *
+	 * @return  True if delegate is bound to the specified UserObject
+	 */
+	virtual bool HasSameObject( const void* InUserObject ) const = 0;
+
+	/**
+	 * Checks to see if the user object bound to this delegate can ever be valid again.
+	 * used to compact multicast delegate arrays so they don't expand without limit.
+	 *
+	 * @return  True if the user object can never be used again
+	 */
+	virtual bool IsCompactable( ) const
+	{
+		return !IsSafeToExecute();
+	}
+
+	/**
+	 * Checks to see if the user object bound to this delegate is still valid
+	 *
+	 * @return  True if the user object is still valid and it's safe to execute the function call
+	 */
+	virtual bool IsSafeToExecute( ) const = 0;
+
+public:
+
+	/**
+	 * Virtual destructor.
+	 */
+	virtual ~IDelegateInstance( ) { }
+};
diff --git a/Engine/Source/Runtime/Core/Public/Delegates/MulticastDelegateBase.h b/Engine/Source/Runtime/Core/Public/Delegates/MulticastDelegateBase.h
new file mode 100644
--- /dev/null
+++ b/Engine/Source/Runtime/Core/Public/Delegates/MulticastDelegateBase.h
@@ -0,0 +1,194 @@
+// Copyright 1998-2014 Epic Games, Inc. All Rights Reserved.
+
+/*=============================================================================
+	MulticastDelegateBase.h: Declares the FMulticastDelegateBase class.
+=============================================================================*/
+
+#pragma once
+
+
+/**
+ * Abstract base class for multicast delegates.
+ */
+template<typename ObjectPtrType = FWeakObjectPtr>
+class FMulticastDelegateBase
+{
+public:
+
+	/**
+	 * Removes all functions from this delegate's invocation list
+	 */
+	void Clear( )
+	{
+		for (IDelegateInstance* DelegateInstance : InvocationList)
+		{
+			delete DelegateInstance;
+		}
+
+		InvocationList.Empty();
+		ExpirableObjectCount = 0;
+	}
+
+	/**
+	 * Checks to see if any functions are bound to this multi-cast delegate.
+	 *
+	 * This method is protected to avoid assumptions about only one delegate existing.
+	 *
+	 * @return true if any functions are bound, false otherwise.
+	 */
+	inline bool IsBound( ) const
+	{
+		return (InvocationList.Num() > 0);
+	}
+
+	/** 
+	 * Checks to see if any functions are bound to the given user object.
+	 *
+	 * @return	True if any functions are bound to InUserObject, false otherwise.
+	 */
+	inline bool IsBoundToObject( void const* InUserObject ) const
+	{
+		for (IDelegateInstance* DelegateInstance : InvocationList)
+		{
+			if (DelegateInstance->HasSameObject(InUserObject))
+			{
+				return true;
+			}
+		}
+
+		return false;
+	}
+
+	/**
+	 * Removes all functions from this multi-cast delegate's invocation list that are bound to the specified UserObject.
+	 * Note that the order of the delegates may not be preserved!
+	 *
+	 * @param InUserObject The object to remove all delegates for.
+	 */
+	void RemoveAll( const void* InUserObject )
+	{
+		RemoveByObject(InUserObject);
+		CompactInvocationList();
+	}
+
+protected:
+
+	/**
+	 * Hidden default constructor.
+	 */
+	inline FMulticastDelegateBase( )
+		: ExpirableObjectCount(0)
+	{ }
+
+protected:
+
+	/**
+	 * Adds the given delegate instance to the invocation list.
+	 *
+	 * @param DelegateInstance The delegate instance to add.
+	 */
+	void AddInternal( IDelegateInstance& DelegateInstance )
+	{
+		InvocationList.Add(&DelegateInstance);
+
+		// Keep track of whether we have objects bound that may expire out from underneath us.
+		const EDelegateInstanceType::Type DelegateType = DelegateInstance.GetType();
+
+		if ((DelegateType == EDelegateInstanceType::SharedPointerMethod) ||
+			(DelegateType == EDelegateInstanceType::ThreadSafeSharedPointerMethod) ||
+			(DelegateType == EDelegateInstanceType::UObjectMethod))
+		{
+			++ExpirableObjectCount;
+		}
+	}
+
+	/**
+	 * Cleans up any delegates in our invocation list that have expired (performance is O(N))
+	 */
+	void CompactInvocationList( ) const
+	{
+		// We only need to compact if any objects were added whose lifetime is known to us
+		if (ExpirableObjectCount == 0)
+		{
+			return;
+		}
+
+		// remove expired delegates
+		for (int32 InvocationListIndex = InvocationList.Num() - 1; InvocationListIndex >= 0; --InvocationListIndex)
+		{
+			IDelegateInstance* DelegateInstance = InvocationList[InvocationListIndex];
+
+			if (DelegateInstance->IsCompactable())
+			{
+				InvocationList.RemoveAtSwap(InvocationListIndex);
+				delete DelegateInstance;
+
+				checkSlow(ExpirableObjectCount > 0);
+				--ExpirableObjectCount;
+			}
+		}
+	}
+
+	/**
+	 * Gets a read-only reference to the invocation list.
+	 *
+	 * @return The invocation list.
+	 */
+	const TArray<IDelegateInstance*>& GetInvocationList( ) const
+	{
+		return InvocationList;
+	}
+
+	/**
+	 * Removes all functions from this multi-cast delegate's invocation list that are bound to the specified UserObject.
+	 *
+	 * Note that the order of the delegates may not be preserved!
+	 *
+	 * @param InUserObject The object to remove the delegates for.
+	 */
+	void RemoveByObject( const void* InUserObject )
+	{
+		for (int32 InvocationListIndex = InvocationList.Num() - 1; InvocationListIndex >= 0; --InvocationListIndex)
+		{
+			IDelegateInstance* DelegateInstance = InvocationList[InvocationListIndex];
+
+			if (DelegateInstance->HasSameObject(InUserObject))
+			{
+				RemoveByIndex(InvocationListIndex);
+			}
+		}
+	}
+
+	/**
+	 * Removes a delegate instance given its index in the invocation list.
+	 *
+	 * @param InvocationListIndex The index of the delegate instance to remove.
+	 */
+	void RemoveByIndex( int32 InvocationListIndex )
+	{
+		IDelegateInstance* DelegateInstance = InvocationList[InvocationListIndex];
+
+		// keep track of whether we have objects bound that may expire out from underneath us
+		const EDelegateInstanceType::Type DelegateType = DelegateInstance->GetType();
+
+		if ((DelegateType == EDelegateInstanceType::SharedPointerMethod) || 
+			(DelegateType == EDelegateInstanceType::ThreadSafeSharedPointerMethod) ||
+			(DelegateType == EDelegateInstanceType::UFunction) ||
+			(DelegateType == EDelegateInstanceType::UObjectMethod))
+		{
+			--ExpirableObjectCount;
+		}
+
+		InvocationList.RemoveAtSwap(InvocationListIndex);
+		delete DelegateInstance;
+	}
+
+private:
+
+	// Mutable so that we can housekeep list even with const broadcasts
+	mutable int32 ExpirableObjectCount;
+
+	// Holds the collection of delegate instances to invoke.
+	// Mutable so that we can housekeep list even with const broadcasts
+	mutable TArray<IDelegateInstance*> InvocationList;
+};
diff --git a/Engine/Source/Runtime/Core/Public/Templates/Delegate.h b/Engine/Source/Runtime/Core/Public/Templates/Delegate.h
deleted file mode 100644
diff --git a/Engine/Source/Runtime/Core/Public/Templates/DelegateBase.h b/Engine/Source/Runtime/Core/Public/Templates/DelegateBase.h
deleted file mode 100644
diff --git a/Engine/Source/Runtime/Core/Public/Templates/DelegateCombinations.h b/Engine/Source/Runtime/Core/Public/Templates/DelegateCombinations.h
deleted file mode 100644
diff --git a/Engine/Source/Runtime/Core/Public/Templates/DelegateInstanceInterfaceImpl.inl b/Engine/Source/Runtime/Core/Public/Templates/DelegateInstanceInterfaceImpl.inl
deleted file mode 100644
diff --git a/Engine/Source/Runtime/Core/Public/Templates/DelegateInstancesImpl.inl b/Engine/Source/Runtime/Core/Public/Templates/DelegateInstancesImpl.inl
deleted file mode 100644
diff --git a/Engine/Source/Runtime/Core/Public/Templates/DelegateSignatureImpl.inl b/Engine/Source/Runtime/Core/Public/Templates/DelegateSignatureImpl.inl
deleted file mode 100644
diff --git a/Engine/Source/Runtime/Core/Public/Templates/IDelegateInstance.h b/Engine/Source/Runtime/Core/Public/Templates/IDelegateInstance.h
deleted file mode 100644
diff --git a/Engine/Source/Runtime/Core/Public/Templates/MulticastDelegateBase.h b/Engine/Source/Runtime/Core/Public/Templates/MulticastDelegateBase.h
deleted file mode 100644
