Fixed TAttribute::Create not working with static function pointers

It was trying to call a constructor which took the function pointer as well as a bool, but it actually seemed to be calling one of the templated constructors instead.

I've simplified the code by removing the constructor which took a function pointer and instead have the Create version which takes a function pointer call FGetter::CreateStatic itself. This simplifies the code and seems to resolve the ambiguity.

diff --git a/Engine/Source/Runtime/Core/Public/Misc/Attribute.h b/Engine/Source/Runtime/Core/Public/Misc/Attribute.h
--- a/Engine/Source/Runtime/Core/Public/Misc/Attribute.h
+++ b/Engine/Source/Runtime/Core/Public/Misc/Attribute.h
@@ -88,10 +88,10 @@ public:
 	 *
 	 * @param  InGetter		Delegate to bind
 	 */
-	static TAttribute< ObjectType > Create( const FGetter& InGetter )
+	static TAttribute Create( const FGetter& InGetter )
 	{
 		const bool bExplicitConstructor = true;
-		return TAttribute< ObjectType >( InGetter, bExplicitConstructor );
+		return TAttribute( InGetter, bExplicitConstructor );
 	}
 
 	
@@ -102,10 +102,10 @@ public:
 	 *
 	 * @param  InFuncPtr Member function to bind.  The function's structure (return value, arguments, etc) must match IBoundAttributeDelegate's definition.
 	 */
-	static TAttribute< ObjectType > Create( typename FGetter::FStaticDelegate::FFuncPtr InFuncPtr )
+	static TAttribute Create( typename FGetter::FStaticDelegate::FFuncPtr InFuncPtr )
 	{
 		const bool bExplicitConstructor = true;
-		return TAttribute< ObjectType >( InFuncPtr, bExplicitConstructor );
+		return TAttribute( FGetter::CreateStatic( InFuncPtr ), bExplicitConstructor );
 	}
 
 
@@ -281,16 +281,9 @@ public:
 		return InOther.Get() != Get();
 	}
 
-
-	/** Special explicit constructor for TAttribute::Create() */
-	TAttribute( typename FGetter::FStaticDelegate::FFuncPtr InFuncPtr, bool bExplicitConstructor )
-		: Value()
-		, Getter( FGetter::CreateStatic( InFuncPtr ) )
-		, bIsSet( true )
-	{
-	}
-
 	
+private:
+
 	/** Special explicit constructor for TAttribute::Create() */
 	TAttribute( const FGetter& InGetter, bool bExplicitConstructor )
 		: Value()
@@ -299,9 +292,6 @@ public:
 	{
 	}
 
-
-private:
-
 	// We declare ourselves as a friend (templated using OtherType) so we can access members as needed
     template< class OtherType > friend class TAttribute;
 
