Fix handling of CGAssociateMouseAndMouseCursorPosition & CGWarpMouseCursorPosition to avoid overhead of calling each frame.

diff --git a/Engine/Source/Runtime/Core/Private/Mac/MacApplication.cpp b/Engine/Source/Runtime/Core/Private/Mac/MacApplication.cpp
--- a/Engine/Source/Runtime/Core/Private/Mac/MacApplication.cpp
+++ b/Engine/Source/Runtime/Core/Private/Mac/MacApplication.cpp
@@ -325,8 +325,6 @@ void FMacApplication::ProcessEvent( NSEvent* Event )
 				}
 				else
 				{
-					// Call this here or else the cursor gets stuck in place after returning from high precision mode
-					CGAssociateMouseAndMouseCursorPosition( true );
 					FVector2D CurrentPosition = MacCursor->GetPosition();
 					if( MacCursor->UpdateCursorClipping( CurrentPosition ) )
 					{
@@ -763,8 +761,7 @@ void FMacApplication::SetHighPrecisionMouseMode( const bool Enable, const TShare
 {
 	bUsingHighPrecisionMouseInput = Enable;
 	HighPrecisionMousePos = static_cast<FMacCursor*>( Cursor.Get() )->GetPosition();
-	
-	CGAssociateMouseAndMouseCursorPosition( !Enable );
+	static_cast<FMacCursor*>( Cursor.Get() )->AssociateMouseAndCursorPosition( !Enable );
 }
 
 FModifierKeysState FMacApplication::GetModifierKeys() const
diff --git a/Engine/Source/Runtime/Core/Private/Mac/MacCursor.cpp b/Engine/Source/Runtime/Core/Private/Mac/MacCursor.cpp
--- a/Engine/Source/Runtime/Core/Private/Mac/MacCursor.cpp
+++ b/Engine/Source/Runtime/Core/Private/Mac/MacCursor.cpp
@@ -7,6 +7,7 @@
 
 FMacCursor::FMacCursor()
 :	bIsVisible(true)
+,	bAssociateMouseCursor(false)
 ,	MouseWarpDelta(FVector2D::ZeroVector)
 {
 	SCOPED_AUTORELEASE_POOL;
@@ -267,7 +268,23 @@ void FMacCursor::UpdateVisibility()
 
 void FMacCursor::WarpCursor( const int32 X, const int32 Y )
 {
+	// Apple suppress mouse events for 0.25 seconds after a call to Warp, unless we call CGAssociateMouseAndMouseCursorPosition.
+	// Previously there was CGSetLocalEventsSuppressionInterval to explicitly control this behaviour but that is deprecated.
+	// The replacement CGEventSourceSetLocalEventsSuppressionInterval isn't useful because it is unclear how to obtain the correct event source.
+	// Instead, when we want the warp to be visible we need to disassociate mouse & cursor...
+	if( bAssociateMouseCursor )
+	{
+		CGAssociateMouseAndMouseCursorPosition( false );
+	}
+	
+	// Perform the warp as normal
 	CGWarpMouseCursorPosition( NSMakePoint( FMath::TruncToInt( X ), FMath::TruncToInt( Y ) ) );
+	
+	// And then reassociate the mouse cursor, which forces the mouse events to come through.
+	if( bAssociateMouseCursor )
+	{
+		CGAssociateMouseAndMouseCursorPosition( true );
+	}
 }
 
 FVector2D FMacCursor::GetMouseWarpDelta( bool const bClearAccumulatedDelta )
@@ -276,3 +293,12 @@ FVector2D FMacCursor::GetMouseWarpDelta( bool const bClearAccumulatedDelta )
 	MouseWarpDelta = bClearAccumulatedDelta ? FVector2D::ZeroVector : MouseWarpDelta;
 	return Result;
 }
+
+void FMacCursor::AssociateMouseAndCursorPosition( bool const bEnable )
+{
+	if( bEnable != bAssociateMouseCursor )
+	{
+		bAssociateMouseCursor = bEnable;
+		CGAssociateMouseAndMouseCursorPosition( bEnable );
+	}
+}
diff --git a/Engine/Source/Runtime/Core/Public/Mac/MacCursor.h b/Engine/Source/Runtime/Core/Public/Mac/MacCursor.h
--- a/Engine/Source/Runtime/Core/Public/Mac/MacCursor.h
+++ b/Engine/Source/Runtime/Core/Public/Mac/MacCursor.h
@@ -37,6 +37,7 @@ public:
 	
 	FVector2D GetMouseWarpDelta( bool const bClearAccumulatedDelta );
 
+	void AssociateMouseAndCursorPosition( bool const bEnable );
 
 private:
 
@@ -48,6 +49,7 @@ private:
 	FIntRect CusorClipRect;
 
 	bool bIsVisible;
+	bool bAssociateMouseCursor;
 	NSCursor* CurrentCursor;
 	
 	FVector2D MouseWarpDelta;
