fixed crash when changing filter to Static Mesh (TTP# 334946)

diff --git a/Engine/Source/Runtime/Core/Public/Delegates/DelegateSignatureImpl.inl b/Engine/Source/Runtime/Core/Public/Delegates/DelegateSignatureImpl.inl
--- a/Engine/Source/Runtime/Core/Public/Delegates/DelegateSignatureImpl.inl
+++ b/Engine/Source/Runtime/Core/Public/Delegates/DelegateSignatureImpl.inl
@@ -1783,11 +1783,11 @@ protected:
 #if !UE_BUILD_SHIPPING && !UE_BUILD_TEST
 		// verify that the same function isn't already bound
 		for (auto DelegateInstance : GetInvocationList())
-	{
+		{
 			// this down-cast is OK! allows for managing invocation list in the base class without requiring virtual functions
 			TDelegateInstanceInterface* DelegateInstanceInterface = (TDelegateInstanceInterface*)DelegateInstance;
 			check(!DelegateInstanceInterface->IsSameFunction(InDelegateInstance));
-	}
+		}
 #endif
 
 		AddInternal(InDelegateInstance);
@@ -1809,20 +1809,20 @@ protected:
 
 		// copy invocation list just in case the list is modified by one of the broadcast callbacks
 		typedef TArray<IDelegateInstance*, TInlineAllocator<4>> FInlineInvocationList;
-			FInlineInvocationList InvocationListCopy = FInlineInvocationList(InvocationList);
+		FInlineInvocationList InvocationListCopy = FInlineInvocationList(InvocationList);
 
-			// Invoke each bound function
+		// Invoke each bound function
 		for (IDelegateInstance* DelegateInstance : InvocationListCopy)
-			{
+		{
 			// this down-cast is OK! allows for managing invocation list in the base class without requiring virtual functions
 			TDelegateInstanceInterface* DelegateInstanceInterface = (TDelegateInstanceInterface*)DelegateInstance;
 
 			// execute the delegate instance...
 			if (!DelegateInstanceInterface->ExecuteIfSafe(FUNC_PARAM_PASSTHRU))
-				{
+			{
 				// ... or remove it if compactable
 				if (DelegateInstance->IsCompactable())
-					{
+				{
 					// search for the instance, because original list could have been modified during broadcast
 					// this is really shady, because it violates the constness of this method
 					const_cast<BASE_MULTICAST_DELEGATE_CLASS*>(this)->RemoveDelegateInstance(*DelegateInstanceInterface);
@@ -1844,7 +1844,7 @@ protected:
 	{
 		const TArray<IDelegateInstance*> InvocationList = GetInvocationList();
 
-		for (int32 InvocationListIndex = InvocationList.Num() - 1; InvocationListIndex >= 0; --InvocationListIndex)
+		for (int32 InvocationListIndex = 0; InvocationListIndex < InvocationList.Num(); ++InvocationListIndex)
 		{
 			// this down-cast is OK! allows for managing invocation list in the base class without requiring virtual functions
 			TDelegateInstanceInterface* DelegateInstance = (TDelegateInstanceInterface*)InvocationList[InvocationListIndex];
