Delegates: temp fix for delegates being removed incorrectly

diff --git a/Engine/Source/Runtime/Core/Public/Delegates/DelegateSignatureImpl.inl b/Engine/Source/Runtime/Core/Public/Delegates/DelegateSignatureImpl.inl
--- a/Engine/Source/Runtime/Core/Public/Delegates/DelegateSignatureImpl.inl
+++ b/Engine/Source/Runtime/Core/Public/Delegates/DelegateSignatureImpl.inl
@@ -1145,7 +1145,7 @@ public:
 	{
 		if (DelegateInstance != nullptr)
 		{
-			AddDelegateInstance(*DelegateInstance);
+			AddDelegateInstance(DelegateInstance);
 			CompactInvocationList();
 		}
 	}
@@ -1163,7 +1163,7 @@ public:
 
 		if (DelegateInstance != nullptr)
 		{
-			AddDelegateInstance(*DelegateInstance->CreateCopy());
+			AddDelegateInstance(DelegateInstance->CreateCopy());
 			CompactInvocationList();
 		}
 	}
@@ -1778,19 +1778,19 @@ protected:
 	 *
 	 * @param	InDelegate	Delegate to add
 	 */
-	void AddDelegateInstance( TDelegateInstanceInterface& InDelegateInstance )
+	void AddDelegateInstance( TDelegateInstanceInterface* InDelegateInstance )
 	{
 #if !UE_BUILD_SHIPPING && !UE_BUILD_TEST
 		// verify that the same function isn't already bound
-		for (auto DelegateInstance : GetInvocationList())
+		for (IDelegateInstancePtr DelegateInstance : GetInvocationList())
 		{
 			// this down-cast is OK! allows for managing invocation list in the base class without requiring virtual functions
-			TDelegateInstanceInterface* DelegateInstanceInterface = (TDelegateInstanceInterface*)DelegateInstance;
-			check(!DelegateInstanceInterface->IsSameFunction(InDelegateInstance));
+			TDelegateInstanceInterface* DelegateInstanceInterface = (TDelegateInstanceInterface*)DelegateInstance.Get();
+			check(!DelegateInstanceInterface->IsSameFunction(*InDelegateInstance));
 		}
 #endif
 
-		AddInternal(InDelegateInstance);
+		AddInternal(MakeShareable(InDelegateInstance));
 	}
 
 	/**
@@ -1800,7 +1800,7 @@ protected:
 	 */
 	void Broadcast( FUNC_PARAM_LIST ) const
 	{
-		const TArray<IDelegateInstance*>& InvocationList = GetInvocationList();
+		const TArray<IDelegateInstancePtr>& InvocationList = GetInvocationList();
 
 		if (InvocationList.Num() == 0)
 		{
@@ -1808,14 +1808,14 @@ protected:
 		}
 
 		// copy invocation list just in case the list is modified by one of the broadcast callbacks
-		typedef TArray<IDelegateInstance*, TInlineAllocator<4>> FInlineInvocationList;
+		typedef TArray<IDelegateInstancePtr, TInlineAllocator<4>> FInlineInvocationList;
 		FInlineInvocationList InvocationListCopy = FInlineInvocationList(InvocationList);
 
 		// Invoke each bound function
-		for (IDelegateInstance* DelegateInstance : InvocationListCopy)
+		for (IDelegateInstancePtr DelegateInstance : InvocationListCopy)
 		{
 			// this down-cast is OK! allows for managing invocation list in the base class without requiring virtual functions
-			TDelegateInstanceInterface* DelegateInstanceInterface = (TDelegateInstanceInterface*)DelegateInstance;
+			TDelegateInstanceInterface* DelegateInstanceInterface = (TDelegateInstanceInterface*)DelegateInstance.Get();
 
 			// execute the delegate instance...
 			if (!DelegateInstanceInterface->ExecuteIfSafe(FUNC_PARAM_PASSTHRU))
@@ -1825,7 +1825,7 @@ protected:
 				{
 					// search for the instance, because original list could have been modified during broadcast
 					// this is really shady, because it violates the constness of this method
-					const_cast<BASE_MULTICAST_DELEGATE_CLASS*>(this)->RemoveDelegateInstance(*DelegateInstanceInterface);
+					const_cast<BASE_MULTICAST_DELEGATE_CLASS*>(this)->RemoveInternal(DelegateInstance);
 				}
 			}
 		}
@@ -1842,12 +1842,12 @@ protected:
 	 */
 	void RemoveDelegateInstance( const TDelegateInstanceInterface& InDelegateInstance )
 	{
-		const TArray<IDelegateInstance*> InvocationList = GetInvocationList();
+		const TArray<IDelegateInstancePtr> InvocationList = GetInvocationList();
 
 		for (int32 InvocationListIndex = 0; InvocationListIndex < InvocationList.Num(); ++InvocationListIndex)
 		{
 			// this down-cast is OK! allows for managing invocation list in the base class without requiring virtual functions
-			TDelegateInstanceInterface* DelegateInstance = (TDelegateInstanceInterface*)InvocationList[InvocationListIndex];
+			TDelegateInstanceInterface* DelegateInstance = (TDelegateInstanceInterface*)InvocationList[InvocationListIndex].Get();
 
 			// NOTE: We must do a deep compare here, not just compare delegate pointers, because multiple
 			//       delegate pointers can refer to the exact same object and method
diff --git a/Engine/Source/Runtime/Core/Public/Delegates/IDelegateInstance.h b/Engine/Source/Runtime/Core/Public/Delegates/IDelegateInstance.h
--- a/Engine/Source/Runtime/Core/Public/Delegates/IDelegateInstance.h
+++ b/Engine/Source/Runtime/Core/Public/Delegates/IDelegateInstance.h
@@ -35,6 +35,17 @@ namespace EDelegateInstanceType
 }
 
 
+/**
+ * Type definition for shared pointers to instances of IDelegateInstance.
+ */
+typedef TSharedPtr<class IDelegateInstance> IDelegateInstancePtr;
+
+/**
+ * Type definition for shared references to instances of IDelegateInstance.
+ */
+typedef TSharedRef<class IDelegateInstance> IDelegateInstanceRef;
+
+
 /**
  * Interface for delegate instances.
  */
diff --git a/Engine/Source/Runtime/Core/Public/Delegates/MulticastDelegateBase.h b/Engine/Source/Runtime/Core/Public/Delegates/MulticastDelegateBase.h
--- a/Engine/Source/Runtime/Core/Public/Delegates/MulticastDelegateBase.h
+++ b/Engine/Source/Runtime/Core/Public/Delegates/MulticastDelegateBase.h
@@ -20,11 +20,6 @@ public:
 	 */
 	void Clear( )
 	{
-		for (IDelegateInstance* DelegateInstance : InvocationList)
-		{
-			delete DelegateInstance;
-		}
-
 		InvocationList.Empty();
 		ExpirableObjectCount = 0;
 	}
@@ -48,7 +43,7 @@ public:
 	 */
 	inline bool IsBoundToObject( void const* InUserObject ) const
 	{
-		for (IDelegateInstance* DelegateInstance : InvocationList)
+		for (IDelegateInstancePtr DelegateInstance : InvocationList)
 		{
 			if (DelegateInstance->HasSameObject(InUserObject))
 			{
@@ -87,12 +82,12 @@ protected:
 	 *
 	 * @param DelegateInstance The delegate instance to add.
 	 */
-	void AddInternal( IDelegateInstance& DelegateInstance )
+	void AddInternal( const IDelegateInstanceRef& DelegateInstance )
 	{
-		InvocationList.Add(&DelegateInstance);
+		InvocationList.Add(DelegateInstance);
 
 		// Keep track of whether we have objects bound that may expire out from underneath us.
-		const EDelegateInstanceType::Type DelegateType = DelegateInstance.GetType();
+		const EDelegateInstanceType::Type DelegateType = DelegateInstance->GetType();
 
 		if ((DelegateType == EDelegateInstanceType::SharedPointerMethod) ||
 			(DelegateType == EDelegateInstanceType::ThreadSafeSharedPointerMethod) ||
@@ -116,12 +111,11 @@ protected:
 		// remove expired delegates
 		for (int32 InvocationListIndex = InvocationList.Num() - 1; InvocationListIndex >= 0; --InvocationListIndex)
 		{
-			IDelegateInstance* DelegateInstance = InvocationList[InvocationListIndex];
+			IDelegateInstancePtr DelegateInstance = InvocationList[InvocationListIndex];
 
 			if (DelegateInstance->IsCompactable())
 			{
 				InvocationList.RemoveAtSwap(InvocationListIndex);
-				delete DelegateInstance;
 
 				checkSlow(ExpirableObjectCount > 0);
 				--ExpirableObjectCount;
@@ -134,7 +128,7 @@ protected:
 	 *
 	 * @return The invocation list.
 	 */
-	const TArray<IDelegateInstance*>& GetInvocationList( ) const
+	const TArray<IDelegateInstancePtr>& GetInvocationList( ) const
 	{
 		return InvocationList;
 	}
@@ -150,7 +144,7 @@ protected:
 	{
 		for (int32 InvocationListIndex = InvocationList.Num() - 1; InvocationListIndex >= 0; --InvocationListIndex)
 		{
-			IDelegateInstance* DelegateInstance = InvocationList[InvocationListIndex];
+			IDelegateInstancePtr DelegateInstance = InvocationList[InvocationListIndex];
 
 			if (DelegateInstance->HasSameObject(InUserObject))
 			{
@@ -166,7 +160,7 @@ protected:
 	 */
 	void RemoveByIndex( int32 InvocationListIndex )
 	{
-		IDelegateInstance* DelegateInstance = InvocationList[InvocationListIndex];
+		IDelegateInstancePtr DelegateInstance = InvocationList[InvocationListIndex];
 
 		// keep track of whether we have objects bound that may expire out from underneath us
 		const EDelegateInstanceType::Type DelegateType = DelegateInstance->GetType();
@@ -180,7 +174,21 @@ protected:
 		}
 
 		InvocationList.RemoveAtSwap(InvocationListIndex);
-		delete DelegateInstance;
+	}
+
+	/**
+	 * Removes the given delegate instance from the invocation list.
+	 *
+	 * @param DelegateInstance The delegate instance to remove.
+	 */
+	void RemoveInternal( const IDelegateInstancePtr& DelegateInstance )
+	{
+		int32 DelegateIndex;
+		
+		if (InvocationList.Find(DelegateInstance, DelegateIndex))
+		{
+			RemoveByIndex(DelegateIndex);
+		}
 	}
 
 private:
@@ -189,6 +197,6 @@ private:
 	mutable int32 ExpirableObjectCount;
 
 	// Holds the collection of delegate instances to invoke.
-	// Mutable so that we can housekeep list even with const broadcasts
-	mutable TArray<IDelegateInstance*> InvocationList;
+	// Mutable so that we can do housekeeping during const broadcasts.
+	mutable TArray<IDelegateInstancePtr> InvocationList;
 };
