Added motion support for IOS

diff --git a/Engine/Source/Runtime/Core/Core.Build.cs b/Engine/Source/Runtime/Core/Core.Build.cs
--- a/Engine/Source/Runtime/Core/Core.Build.cs
+++ b/Engine/Source/Runtime/Core/Core.Build.cs
@@ -87,7 +87,7 @@ public class Core : ModuleRules
 			AddThirdPartyPrivateStaticDependencies(Target, 
 				"zlib"
 				);
-			PublicFrameworks.AddRange(new string[] { "UIKit", "Foundation", "AudioToolbox", "AVFoundation", "GameKit", "StoreKit", "CoreVideo", "CoreMedia"});
+			PublicFrameworks.AddRange(new string[] { "UIKit", "Foundation", "AudioToolbox", "AVFoundation", "GameKit", "StoreKit", "CoreVideo", "CoreMedia", "CoreMotion"});
 
 			bool bSupportAdvertising = true;
 
diff --git a/Engine/Source/Runtime/Core/Private/IOS/IOSInputInterface.cpp b/Engine/Source/Runtime/Core/Private/IOS/IOSInputInterface.cpp
--- a/Engine/Source/Runtime/Core/Private/IOS/IOSInputInterface.cpp
+++ b/Engine/Source/Runtime/Core/Private/IOS/IOSInputInterface.cpp
@@ -16,6 +16,8 @@ TSharedRef< FIOSInputInterface > FIOSInputInterface::Create(  const TSharedRef<
 FIOSInputInterface::FIOSInputInterface( const TSharedRef< FGenericApplicationMessageHandler >& InMessageHandler )
 	: MessageHandler( InMessageHandler )
 {
+	MotionManager = nil;
+	ReferenceAttitude = nil;
 }
 
 void FIOSInputInterface::SetMessageHandler( const TSharedRef< FGenericApplicationMessageHandler >& InMessageHandler )
@@ -27,6 +29,67 @@ void FIOSInputInterface::Tick( float DeltaTime )
 {
 }
 
+
+void ModifyVectorByOrientation(FVector& Vec, bool bIsRotation)
+{
+	UIInterfaceOrientation Orientation = [[IOSAppDelegate GetDelegate].IOSController interfaceOrientation];
+	switch (Orientation)
+	{
+	case UIInterfaceOrientationPortrait:
+		// this is the base orientation, so nothing to do
+		break;
+
+	case UIInterfaceOrientationPortraitUpsideDown:
+		if (bIsRotation)
+		{
+			// negate roll and pitch
+			Vec.X = -Vec.X;
+			Vec.Z = -Vec.Z;
+		}
+		else
+		{
+			// negate x/y
+			Vec.X = -Vec.X;
+			Vec.Y = -Vec.Y;
+		}
+		break;
+
+	case UIInterfaceOrientationLandscapeRight:
+		if (bIsRotation)
+		{
+			// swap and negate (as needed) roll and pitch
+			float Temp = Vec.X;
+			Vec.X = -Vec.Z;
+			Vec.Z = Temp;
+		}
+		else
+		{
+			// swap and negate (as needed) x and y
+			float Temp = Vec.X;
+			Vec.X = -Vec.Y;
+			Vec.Y = Temp;
+		}
+		break;
+
+	case UIInterfaceOrientationLandscapeLeft:
+		if (bIsRotation)
+		{
+			// swap and negate (as needed) roll and pitch
+			float Temp = Vec.X;
+			Vec.X = Vec.Z;
+			Vec.Z = -Temp;
+		}
+		else
+		{
+			// swap and negate (as needed) x and y
+			float Temp = Vec.X;
+			Vec.X = Vec.Y;
+			Vec.Y = -Temp;
+		}
+		break;
+	}
+}
+
 void FIOSInputInterface::SendControllerEvents()
 {
 	FScopeLock Lock(&CriticalSection);
@@ -51,6 +114,27 @@ void FIOSInputInterface::SendControllerEvents()
 	}
 
 	FIOSInputInterface::TouchInputStack.Empty(0);
+
+
+	// Update motion controls.
+	FVector Attitude;
+	FVector RotationRate;
+	FVector Gravity;
+	FVector Acceleration;
+
+	GetMovementData(Attitude, RotationRate, Gravity, Acceleration);
+
+	// Fix-up yaw to match directions
+	Attitude.Y = -Attitude.Y;
+	RotationRate.Y = -RotationRate.Y;
+
+	// munge the vectors based on the orientation
+	ModifyVectorByOrientation(Attitude, true);
+	ModifyVectorByOrientation(RotationRate, true);
+	ModifyVectorByOrientation(Gravity, false);
+	ModifyVectorByOrientation(Acceleration, false);
+
+	MessageHandler->OnMotionDetected(Attitude, RotationRate, Gravity, Acceleration, 0);
 }
 
 void FIOSInputInterface::QueueTouchInput(TArray<TouchInput> InTouchEvents)
@@ -58,4 +142,124 @@ void FIOSInputInterface::QueueTouchInput(TArray<TouchInput> InTouchEvents)
 	FScopeLock Lock(&CriticalSection);
 
 	FIOSInputInterface::TouchInputStack.Append(InTouchEvents);
+}
+
+void FIOSInputInterface::GetMovementData(FVector& Attitude, FVector& RotationRate, FVector& Gravity, FVector& Acceleration)
+{
+	// initialize on first use
+	if (MotionManager == nil)
+	{
+		// Look to see if we can create the motion manager
+		MotionManager = [[CMMotionManager alloc] init];
+
+		// Check to see if the device supports full motion (gyro + accelerometer)
+		if (MotionManager.deviceMotionAvailable)
+		{
+			MotionManager.deviceMotionUpdateInterval = 0.02;
+
+			// Start the Device updating motion
+			[MotionManager startDeviceMotionUpdates];
+		}
+		else
+		{
+			[MotionManager startAccelerometerUpdates];
+			CenterPitch = CenterPitch = 0;
+			bIsCalibrationRequested = false;
+		}
+	}
+
+	// do we have full motion data?
+	if (MotionManager.deviceMotionActive)
+	{
+		// Grab the values
+		CMAttitude* CurrentAttitude = MotionManager.deviceMotion.attitude;
+		CMRotationRate CurrentRotationRate = MotionManager.deviceMotion.rotationRate;
+		CMAcceleration CurrentGravity = MotionManager.deviceMotion.gravity;
+		CMAcceleration CurrentUserAcceleration = MotionManager.deviceMotion.userAcceleration;
+
+		// apply a reference attitude if we have been calibrated away from default
+		if (ReferenceAttitude)
+		{
+			[CurrentAttitude multiplyByInverseOfAttitude : ReferenceAttitude];
+		}
+
+		// convert to UE3
+		Attitude = FVector(float(CurrentAttitude.pitch), float(CurrentAttitude.yaw), float(CurrentAttitude.roll));
+		RotationRate = FVector(float(CurrentRotationRate.x), float(CurrentRotationRate.y), float(CurrentRotationRate.z));
+		Gravity = FVector(float(CurrentGravity.x), float(CurrentGravity.y), float(CurrentGravity.z));
+		Acceleration = FVector(float(CurrentUserAcceleration.x), float(CurrentUserAcceleration.y), float(CurrentUserAcceleration.z));
+	}
+	else
+	{
+		// get the plain accleration
+		CMAcceleration RawAcceleration = [MotionManager accelerometerData].acceleration;
+		FVector NewAcceleration(RawAcceleration.x, RawAcceleration.y, RawAcceleration.z);
+
+		// storage for keeping the accelerometer values over time (for filtering)
+		static bool bFirstAccel = true;
+
+		// how much of the previous frame's acceleration to keep
+		const float VectorFilter = bFirstAccel ? 0.0f : 0.85f;
+		bFirstAccel = false;
+
+		// apply new accelerometer values to last frames
+		FilteredAccelerometer = FilteredAccelerometer * VectorFilter + (1.0f - VectorFilter) * NewAcceleration;
+
+		// create an normalized acceleration vector
+		FVector FinalAcceleration = -FilteredAccelerometer.SafeNormal();
+
+		// calculate Roll/Pitch
+		float CurrentPitch = FMath::Atan2(FinalAcceleration.Y, FinalAcceleration.Z);
+		float CurrentRoll = -FMath::Atan2(FinalAcceleration.X, FinalAcceleration.Z);
+
+		// if we want to calibrate, use the current values as center
+		if (bIsCalibrationRequested)
+		{
+			CenterPitch = CurrentPitch;
+			CenterRoll = CurrentRoll;
+			bIsCalibrationRequested = false;
+		}
+
+		CurrentPitch -= CenterPitch;
+		CurrentRoll -= CenterRoll;
+
+		Attitude = FVector(CurrentPitch, 0, CurrentRoll);
+		RotationRate = FVector(LastPitch - CurrentPitch, 0, LastRoll - CurrentRoll);
+		Gravity = FVector(0, 0, 0);
+
+		// use the raw acceleration for acceleration
+		Acceleration = NewAcceleration;
+
+		// remember for next time (for rotation rate)
+		LastPitch = CurrentPitch;
+		LastRoll = CurrentRoll;
+	}
+}
+
+void FIOSInputInterface::CalibrateMotion()
+{
+	// If we are using the motion manager, grab a reference frame.  Note, once you set the Attitude Reference frame
+	// all additional reference information will come from it
+	if (MotionManager.deviceMotionActive)
+	{
+		ReferenceAttitude = [MotionManager.deviceMotion.attitude retain];
+	}
+	else
+	{
+		bIsCalibrationRequested = true;
+	}
+}
+
+bool FIOSInputInterface::Exec(UWorld* InWorld, const TCHAR* Cmd, FOutputDevice& Ar)
+{
+	// Keep track whether the command was handled or not.
+	bool bHandledCommand = false;
+
+	if (FParse::Command(&Cmd, TEXT("CALIBRATEMOTION")))
+	{
+		CalibrateMotion();
+		bHandledCommand = true;
+	}
+
+	return bHandledCommand;
 }
\ No newline at end of file
diff --git a/Engine/Source/Runtime/Core/Public/IOS/IOSInputInterface.h b/Engine/Source/Runtime/Core/Public/IOS/IOSInputInterface.h
--- a/Engine/Source/Runtime/Core/Public/IOS/IOSInputInterface.h
+++ b/Engine/Source/Runtime/Core/Public/IOS/IOSInputInterface.h
@@ -2,6 +2,8 @@
 
 #pragma once
 
+#import <CoreMotion/CoreMotion.h>
+
 enum TouchType
 {
 	TouchBegan,
@@ -20,7 +22,7 @@ struct TouchInput
 /**
  * Interface class for IOS input devices
  */
-class FIOSInputInterface
+class FIOSInputInterface : public FExec
 {
 public:
 
@@ -28,7 +30,7 @@ public:
 
 public:
 
-	~FIOSInputInterface() {}
+	virtual ~FIOSInputInterface() {}
 
 	void SetMessageHandler( const TSharedRef< FGenericApplicationMessageHandler >& InMessageHandler );
 
@@ -42,9 +44,29 @@ public:
 
 	static void QueueTouchInput(TArray<TouchInput> InTouchEvents);
 
+	// Begin Exec Interface
+	virtual bool Exec(UWorld* InWorld, const TCHAR* Cmd, FOutputDevice& Ar) OVERRIDE;
+	// End Exec Interface
+
 private:
 
 	FIOSInputInterface( const TSharedRef< FGenericApplicationMessageHandler >& InMessageHandler );
+	
+	/**
+	 * Get the current Movement data from the device
+	 *
+	 * @param Attitude The current Roll/Pitch/Yaw of the device
+	 * @param RotationRate The current rate of change of the attitude
+	 * @param Gravity A vector that describes the direction of gravity for the device
+	 * @param Acceleration returns the current acceleration of the device
+	 */
+	void GetMovementData(FVector& Attitude, FVector& RotationRate, FVector& Gravity, FVector& Acceleration);
+
+	/**
+	 * Calibrate the devices motion
+	 */
+	void CalibrateMotion();
+
 private:
 
 	static TArray<TouchInput> TouchInputStack;
@@ -54,4 +76,30 @@ private:
 
 	TSharedRef< FGenericApplicationMessageHandler > MessageHandler;
 
+
+private:
+
+	/** Access to the ios devices motion */
+	CMMotionManager* MotionManager;
+
+	/** Access to the ios devices tilt information */
+	CMAttitude* ReferenceAttitude;
+
+	/** Last frames roll, for calculating rate */
+	float LastRoll;
+
+	/** Last frames pitch, for calculating rate */
+	float LastPitch;
+
+	/** True if a calibration is requested */
+	bool bIsCalibrationRequested;
+
+	/** The center roll value for tilt calibration */
+	float CenterRoll;
+
+	/** The center pitch value for tilt calibration */
+	float CenterPitch;
+
+	/** When using just acceleration (without full motion) we store a frame of accel data to filter by */
+	FVector FilteredAccelerometer;
 };
\ No newline at end of file
