Bring back previous version of FPlatformMath API with deprecated functions marked accordingly.

diff --git a/Engine/Source/Runtime/Core/Public/GenericPlatform/GenericPlatformMath.h b/Engine/Source/Runtime/Core/Public/GenericPlatform/GenericPlatformMath.h
--- a/Engine/Source/Runtime/Core/Public/GenericPlatform/GenericPlatformMath.h
+++ b/Engine/Source/Runtime/Core/Public/GenericPlatform/GenericPlatformMath.h
@@ -23,7 +23,7 @@ struct FGenericPlatformMath
 		return (int32)F;
 	}
 
-	DEPRECATED(4.2, "FPlatformMath::Trunc is deprecated, use FPlatformMath::TruncToInt instead")
+	DEPRECATED(4.2, "FPlatformMath::Trunc is deprecated, use FPlatformMath::TruncToInt instead.")
 	static CONSTEXPR FORCEINLINE int32 Trunc(float F)
 	{
 		return TruncToInt(F);
@@ -39,7 +39,7 @@ struct FGenericPlatformMath
 		return (float)TruncToInt(F);
 	}
 
-	DEPRECATED(4.2, "FPlatformMath::TruncFloat is deprecated, use FPlatformMath::TruncToFloat instead")
+	DEPRECATED(4.2, "FPlatformMath::TruncFloat is deprecated, use FPlatformMath::TruncToFloat instead.")
 	static CONSTEXPR FORCEINLINE float TruncFloat(float F)
 	{
 		return TruncToFloat(F);
@@ -70,11 +70,33 @@ struct FGenericPlatformMath
 	 * @param F		Floating point value to convert
 	 * @return		An integer less or equal to 'F'.
 	 */
+	DEPRECATED(4.2, "FPlatformMath::FloorDouble is deprecated, use FPlatformMath::FloorToDouble instead.")
 	static FORCEINLINE double FloorDouble(double F)
+	{
+		return FloorToDouble(F);
+	}
+
+	/**
+	* Converts a double to a less or equal integer.
+	* @param F		Floating point value to convert
+	* @return		The nearest integer value to 'F'.
+	*/
+	static FORCEINLINE double FloorToDouble(double F)
 	{
 		return floor(F);
 	}
 
+	/**
+	* Converts a float to a less or equal integer.
+	* @param F		Floating point value to convert
+	* @return		An integer less or equal to 'F'.
+	*/
+	DEPRECATED(4.2, "FPlatformMath::Floor is deprecated, use FPlatformMath::FloorToInt instead.")
+	static FORCEINLINE int32 Floor(float F)
+	{
+		return FloorToInt(F);
+	}
+
 	/**
 	 * Converts a float to the nearest integer. Rounds up when the fraction is .5
 	 * @param F		Floating point value to convert
@@ -95,6 +117,17 @@ struct FGenericPlatformMath
 		return FloorToFloat(F + 0.5f);
 	}
 
+	/**
+	* Converts a float to the nearest integer. Rounds up when the fraction is .5
+	* @param F		Floating point value to convert
+	* @return		The nearest integer to 'F'.
+	*/
+	DEPRECATED(4.2, "FPlatformMath::Round is deprecated, use FPlatformMath::RoundToInt instead.")
+	static FORCEINLINE int32 Round(float F)
+	{
+		return RoundToInt(F);
+	}
+
 	/**
 	* Converts a float to the nearest greater or equal integer.
 	* @param F		Floating point value to convert
@@ -115,6 +148,17 @@ struct FGenericPlatformMath
 		return ceilf(F);
 	}
 
+	/**
+	* Converts a float to a greater or equal integer.
+	* @param F		Floating point value to convert
+	* @return		An integer greater or equal to 'F'.
+	*/
+	DEPRECATED(4.2, "FPlatformMath::Ceil is deprecated, use FPlatformMath::CeilToInt instead.")
+	static FORCEINLINE int32 Ceil(float F)
+	{
+		return CeilToInt(F);
+	}
+
 	/**
 	* Returns signed fractional part of a float.
 	* @param Value	Floating point value to convert
diff --git a/Engine/Source/Runtime/Core/Public/Math/UnrealMathUtility.h b/Engine/Source/Runtime/Core/Public/Math/UnrealMathUtility.h
--- a/Engine/Source/Runtime/Core/Public/Math/UnrealMathUtility.h
+++ b/Engine/Source/Runtime/Core/Public/Math/UnrealMathUtility.h
@@ -241,7 +241,7 @@ struct FMath : public FPlatformMath
 		if( Grid==0.0 )	return Location;
 		else			
 		{
-			return FloorDouble((Location + 0.5*Grid)/Grid)*Grid;
+			return FloorToDouble((Location + 0.5*Grid)/Grid)*Grid;
 		}
 	}
 
diff --git a/Engine/Source/Runtime/Slate/Public/Widgets/Input/SSpinBox.h b/Engine/Source/Runtime/Slate/Public/Widgets/Input/SSpinBox.h
--- a/Engine/Source/Runtime/Slate/Public/Widgets/Input/SSpinBox.h
+++ b/Engine/Source/Runtime/Slate/Public/Widgets/Input/SSpinBox.h
@@ -536,7 +536,7 @@ protected:
 			if (TIsIntegralType<NumericType>::Value)
 			{
 				// Round floating point to avoid 'jumping' when spinning down integral values 
-				ValueAttribute.Set( (NumericType)FMath::FloorDouble(NewValue+0.5) );
+				ValueAttribute.Set( (NumericType)FMath::FloorToDouble(NewValue+0.5) );
 			}
 			else
 			{
