UMG - TAttribute now allows you to bind a UFunction to it directly, using Max's delegate changes.

diff --git a/Engine/Source/Runtime/Core/Public/Misc/Attribute.h b/Engine/Source/Runtime/Core/Public/Misc/Attribute.h
--- a/Engine/Source/Runtime/Core/Public/Misc/Attribute.h
+++ b/Engine/Source/Runtime/Core/Public/Misc/Attribute.h
@@ -134,8 +134,8 @@ public:
 	 *
 	 * @return  The attribute's value
 	 */
- 	const ObjectType& Get() const
- 	{
+	const ObjectType& Get() const
+	{
 		// If we have a getter delegate, then we'll call that to generate the value
 		if( Getter.IsBound() )
 		{
@@ -149,7 +149,7 @@ public:
 
 		// Return the stored value
 		return Value;
- 	}
+	}
 
 	/**
 	 * Gets the attribute's current value. The attribute may not be set, in which case use the default value provided.
@@ -253,6 +253,39 @@ public:
 	}
 
 
+	/**
+	 * Binds an arbitrary function that will be called to generate this attribute's value on demand.
+	 * After binding, the attribute will no longer have a value that can be accessed directly, and instead the bound
+	 * function will always be called to generate the value.
+	 *
+	 * @param  InUserObject  Instance of the class that contains the member function you want to bind.
+	 * @param  InFunctionName Member function name to bind.
+	 */
+	template< class SourceType >
+	void BindUFunction( SourceType* InUserObject, const FName& InFunctionName )
+	{
+		bIsSet = true;
+		Getter.BindUFunction(InUserObject, InFunctionName);
+	}
+
+
+	/**
+	 * Creates an attribute by binding an arbitrary function that will be called to generate this attribute's value on demand.
+	 * After binding, the attribute will no longer have a value that can be accessed directly, and instead the bound
+	 * function will always be called to generate the value.
+	 *
+	 * @param  InUserObject  Instance of the class that contains the member function you want to bind.
+	 * @param  InFunctionName Member function name to bind.
+	 */
+	template< class SourceType >
+	static TAttribute< ObjectType > Create(SourceType* InUserObject, const FName& InFunctionName)
+	{
+		TAttribute< ObjectType > Attrib;
+		Attrib.BindUFunction<SourceType>(InUserObject, InFunctionName);
+		return Attrib;
+	}
+
+
 	/**
 	 * Checks to see if this attribute has a 'getter' function bound
 	 *
@@ -293,13 +326,13 @@ private:
 	}
 
 	// We declare ourselves as a friend (templated using OtherType) so we can access members as needed
-    template< class OtherType > friend class TAttribute;
+	template< class OtherType > friend class TAttribute;
 
 	/** Current value.  Mutable so that we can cache the value locally when using a bound Getter (allows const ref return value.) */
 	mutable ObjectType Value;
 
 	/** Bound member function for this attribute (may be NULL if no function is bound.)  When set, all attempts
-	    to read the attribute's value will instead call this delegate to generate the value. */
+		to read the attribute's value will instead call this delegate to generate the value. */
 	/** Our attribute's 'getter' delegate */
 	FGetter Getter;
 
