Add markup language which can be used to report high-level user feedback from external tools. Change UBT to output the correct syntax when generating project files.

diff --git a/Engine/Source/Programs/UnrealBuildTool/System/ProjectFileGenerator.cs b/Engine/Source/Programs/UnrealBuildTool/System/ProjectFileGenerator.cs
--- a/Engine/Source/Programs/UnrealBuildTool/System/ProjectFileGenerator.cs
+++ b/Engine/Source/Programs/UnrealBuildTool/System/ProjectFileGenerator.cs
@@ -1223,15 +1223,9 @@ namespace UnrealBuildTool
 			var bSuccess = true;
 			if( ShouldGenerateIntelliSenseData() && TargetFiles.Count > 0 )
 			{
-				try
+				using(ProgressWriter Progress = new ProgressWriter("Binding Binding IntelliSense data..."))
 				{
-					Console.Write( "Binding IntelliSense data... " );
-
-					int ProgressChars = 0;
-					int TargetIndex = 0;
-					Utils.DisplayProgress( TargetIndex, TargetFiles.Count - 1, ref ProgressChars );
-
-					for( ; TargetIndex < TargetFiles.Count; ++TargetIndex)
+					for(int TargetIndex = 0; TargetIndex < TargetFiles.Count; ++TargetIndex)
 					{
 						var CurTarget = TargetFiles[ TargetIndex ];
 						var TargetName = Utils.GetFilenameWithoutAnyExtensions( CurTarget );	// Twice, to remove both extensions from *.Target.cs file
@@ -1249,13 +1243,9 @@ namespace UnrealBuildTool
 						}
 
 						// Display progress
-						Utils.DisplayProgress( TargetIndex, TargetFiles.Count - 1, ref ProgressChars );
+						Progress.Write(TargetIndex + 1, TargetFiles.Count);
 					}
 				}
-				finally
-				{
-					Console.WriteLine();
-				}
 			}
 
 			return bSuccess;
@@ -1863,17 +1853,11 @@ namespace UnrealBuildTool
 		/// <returns>True if successful</returns>
 		protected virtual bool WriteProjectFiles()
 		{
-			try
+			using(ProgressWriter Progress = new ProgressWriter("Writing project files..."))
 			{
-				Console.Write( "Writing project files... " );
-
 				var TotalProjectFileCount = GeneratedProjectFiles.Count + 1;	// +1 for the master project file, which we'll save next
 
-				int ProgressChars = 0;
-				int ProjectFileIndex = 0;
-				Utils.DisplayProgress( ProjectFileIndex, TotalProjectFileCount, ref ProgressChars );
-
-				for( ; ProjectFileIndex < GeneratedProjectFiles.Count; ++ProjectFileIndex )
+				for(int ProjectFileIndex = 0 ; ProjectFileIndex < GeneratedProjectFiles.Count; ++ProjectFileIndex )
 				{
 					var CurProject = GeneratedProjectFiles[ ProjectFileIndex ];
 					if( !CurProject.WriteProjectFile(
@@ -1883,19 +1867,12 @@ namespace UnrealBuildTool
 						return false;
 					}
 
-					Utils.DisplayProgress( ProjectFileIndex, TotalProjectFileCount, ref ProgressChars );
+					Progress.Write(ProjectFileIndex + 1, TotalProjectFileCount);
 				}
 
 				WriteMasterProjectFile( UBTProject: UBTProject );
-				++ProjectFileIndex;
-
-				Utils.DisplayProgress( ProjectFileIndex, TotalProjectFileCount, ref ProgressChars );
+				Progress.Write(TotalProjectFileCount, TotalProjectFileCount);
 			}
-			finally
-			{
-				Console.WriteLine();
-			}
-
 			return true;
 		}
 
diff --git a/Engine/Source/Programs/UnrealBuildTool/System/UnrealBuildTool.cs b/Engine/Source/Programs/UnrealBuildTool/System/UnrealBuildTool.cs
--- a/Engine/Source/Programs/UnrealBuildTool/System/UnrealBuildTool.cs
+++ b/Engine/Source/Programs/UnrealBuildTool/System/UnrealBuildTool.cs
@@ -791,6 +791,10 @@ namespace UnrealBuildTool
 						{
 							// Skip -define
 						}
+						else if (LowercaseArg == "-progress")
+						{
+							ProgressWriter.bWriteMarkup = true;
+						}
 						else if (CheckPlatform.ToString().ToLowerInvariant() == LowercaseArg)
 						{
 							// It's the platform set...
diff --git a/Engine/Source/Programs/UnrealBuildTool/System/Utils.cs b/Engine/Source/Programs/UnrealBuildTool/System/Utils.cs
--- a/Engine/Source/Programs/UnrealBuildTool/System/Utils.cs
+++ b/Engine/Source/Programs/UnrealBuildTool/System/Utils.cs
@@ -562,28 +562,6 @@ namespace UnrealBuildTool
 		}
 
 
-		/// <summary>
-		/// Backspaces the specified number of characters, then displays a progress percentage value to the console
-		/// </summary>
-		/// <param name="Numerator">Progress numerator</param>
-		/// <param name="Denominator">Progress denominator</param>
-		/// <param name="NumCharsToBackspaceOver">Number of characters to backspace before writing the text.  This value will be updated with the length of the new progress string.  The first time progress is displayed, you should pass 0 for this value.</param>
-		public static void DisplayProgress( int Numerator, int Denominator, ref int NumCharsToBackspaceOver )
-		{
-			// Backspace over previous progress value
-			while( NumCharsToBackspaceOver-- > 0 )
-			{
-				Console.Write( "\b" );
-			}
-
-			// Display updated progress string and keep track of how long it was
-			float ProgressValue = Denominator > 0 ? ( (float)Numerator / (float)Denominator ) : 1.0f;
-			var ProgressString = String.Format( "{0}%", Math.Round( ProgressValue * 100.0f ) );
-			NumCharsToBackspaceOver = ProgressString.Length;
-			Console.Write( ProgressString );
-		}
-
-
 		/*
 		 * Read and write classes with xml specifiers
 		 */
@@ -886,6 +864,58 @@ namespace UnrealBuildTool
 		}
 	}
 
+	/// <summary>
+	/// Class to display an incrementing progress percentage. Handles progress markup and direct console output.
+	/// </summary>
+	public class ProgressWriter : IDisposable
+	{
+		public static bool bWriteMarkup = false;
+
+		string Message;
+		int NumCharsToBackspaceOver;
+
+		public ProgressWriter(string InMessage)
+		{
+			Message = InMessage;
+			if(!bWriteMarkup)
+			{
+				Console.Write(Message + " ");
+			}
+			Write(0, 100);
+		}
+
+		public void Dispose()
+		{
+			if(!bWriteMarkup)
+			{
+				Console.WriteLine();
+			}
+		}
+
+		public void Write(int Numerator, int Denominator)
+		{
+			float ProgressValue = Denominator > 0 ? ((float)Numerator / (float)Denominator) : 1.0f;
+			string ProgressString = String.Format("{0}%", Math.Round(ProgressValue * 100.0f));
+
+			if (bWriteMarkup)
+			{
+				Log.WriteLine(TraceEventType.Information, "@progress '{0}' {1}", Message, ProgressString);
+			}
+			else
+			{
+				// Backspace over previous progress value
+				while (NumCharsToBackspaceOver-- > 0)
+				{
+					Console.Write("\b");
+				}
+
+				// Display updated progress string and keep track of how long it was
+				NumCharsToBackspaceOver = ProgressString.Length;
+				Console.Write(ProgressString);
+			}
+		}
+	}
+
 	/// <summary>
 	/// Verbosity filter
 	/// </summary>
diff --git a/Engine/Source/Runtime/Core/Private/Misc/FeedbackContextMarkup.cpp b/Engine/Source/Runtime/Core/Private/Misc/FeedbackContextMarkup.cpp
new file mode 100644
--- /dev/null
+++ b/Engine/Source/Runtime/Core/Private/Misc/FeedbackContextMarkup.cpp
@@ -0,0 +1,178 @@
+// Copyright 1998-2014 Epic Games, Inc. All Rights Reserved.
+
+#include "CorePrivate.h"
+#include "FeedbackContextMarkup.h"
+#include "MonitoredProcess.h"
+
+bool FFeedbackContextMarkup::ParseCommand(const FString& Line, FFeedbackContext* Warn)
+{
+	const TCHAR *Text = *Line;
+	if(ReadToken(Text, TEXT("@progress")))
+	{
+		FString Status;
+		bool bHaveStatus = ReadString(Text, Status);
+
+		int32 Numerator, Denominator;
+		bool bHaveProgress = ReadProgress(Text, Numerator, Denominator);
+
+		if(*Text == 0)
+		{
+			if(bHaveProgress && bHaveStatus)
+			{
+				Warn->StatusUpdate(Numerator, Denominator, FText::FromString(Status));
+				return true;
+			}
+			if(bHaveProgress)
+			{
+				Warn->UpdateProgress(Numerator, Denominator);
+				return true;
+			}
+		}
+	}
+	return false;
+}
+
+bool FFeedbackContextMarkup::PipeProcessOutput(const FText& Description, const FString& URL, const FString& Params, FFeedbackContext* Warn, int32* OutExitCode)
+{
+	bool bRes;
+
+	// Create a read and write pipe for the child process
+	void* PipeRead = NULL;
+	void* PipeWrite = NULL;
+	verify(FPlatformProcess::CreatePipe(PipeRead, PipeWrite));
+
+	// Start the slow task
+	Warn->BeginSlowTask(Description, true, true);
+
+	// Create the process
+	FProcHandle ProcessHandle = FPlatformProcess::CreateProc(*URL, *Params, false, true, true, NULL, 0, NULL, PipeWrite);
+	if(ProcessHandle.IsValid())
+	{
+		FString BufferedText;
+		for(bool bProcessFinished = false; !bProcessFinished; )
+		{
+			bProcessFinished = FPlatformProcess::GetProcReturnCode(ProcessHandle, OutExitCode);
+		
+			if(!bProcessFinished && Warn->ReceivedUserCancel())
+			{
+				FPlatformProcess::TerminateProc(ProcessHandle);
+				bProcessFinished = true;
+			}
+
+			BufferedText += FPlatformProcess::ReadPipe(PipeRead);
+
+			int32 EndOfLineIdx;
+			while(BufferedText.FindChar('\n', EndOfLineIdx))
+			{
+				FString Line = BufferedText.Left(EndOfLineIdx);
+				Line.RemoveFromEnd(TEXT("\r"));
+
+				if(!ParseCommand(Line, Warn))
+				{
+					Warn->Log(*Line);
+				}
+
+				BufferedText = BufferedText.Mid(EndOfLineIdx + 1);
+			}
+
+			FPlatformProcess::Sleep(0.1f);
+		}
+		ProcessHandle.Reset();
+		bRes = true;
+	}
+	else
+	{
+		Warn->Logf(ELogVerbosity::Error, TEXT("Couldn't create process '%s'"), *URL);
+		bRes = false;
+	}
+
+	// Finish the slow task
+	Warn->EndSlowTask();
+
+	// Close the pipes
+	FPlatformProcess::ClosePipe(0, PipeRead);
+	FPlatformProcess::ClosePipe(0, PipeWrite);
+	return bRes;
+}
+
+bool FFeedbackContextMarkup::ReadToken(const TCHAR*& Text, const TCHAR* Token)
+{
+	int32 Length = FCString::Strlen(Token);
+	if(FCString::Strncmp(Text, Token, Length) == 0)
+	{
+		if(Text[Length] == 0 || FChar::IsWhitespace(Text[Length]))
+		{
+			Text = SkipWhitespace(Text + Length);
+			return true;
+		}
+	}
+	return false;
+}
+
+bool FFeedbackContextMarkup::ReadProgress(const TCHAR*& Text, int32& OutNumerator, int32& OutDenominator)
+{
+	uint32 FirstInteger;
+	if(ReadInteger(Text, FirstInteger))
+	{
+		if(*Text == '%')
+		{
+			OutNumerator = FirstInteger;
+			OutDenominator = 100;
+			Text = SkipWhitespace(Text + 1);
+			return true;
+		}
+		else if(*Text == '/')
+		{
+			Text++;
+
+			uint32 SecondInteger;
+			if(ReadInteger(Text, SecondInteger))
+			{
+				OutNumerator = FirstInteger;
+				OutDenominator = SecondInteger;
+				Text = SkipWhitespace(Text);
+				return true;
+			}
+		}
+	}
+	return false;
+}
+
+bool FFeedbackContextMarkup::ReadInteger(const TCHAR*& Text, uint32& OutInteger)
+{
+	if(FChar::IsDigit(*Text))
+	{
+		TCHAR *End;
+		OutInteger = FCString::Strtoui64(Text, &End, 10);
+		Text = End;
+		while(FChar::IsWhitespace(*Text)) Text++;
+		return true;
+	}
+	return false;
+}
+
+bool FFeedbackContextMarkup::ReadString(const TCHAR*& Text, FString& OutString)
+{
+	if(*Text == '\'' || *Text == '\"')
+	{
+		for(const TCHAR *End = Text + 1; *End; End++)
+		{
+			if(*End == *Text)
+			{
+				OutString = FString(End - (Text + 1), Text + 1);
+				do { End++; } while(FChar::IsWhitespace(*End));
+				return true;
+			}
+		}
+	}
+	return false;
+}
+
+const TCHAR *FFeedbackContextMarkup::SkipWhitespace(const TCHAR* Text)
+{
+	while(FChar::IsWhitespace(*Text))
+	{
+		Text++;
+	}
+	return Text;
+}
diff --git a/Engine/Source/Runtime/Core/Public/Misc/FeedbackContextMarkup.h b/Engine/Source/Runtime/Core/Public/Misc/FeedbackContextMarkup.h
new file mode 100644
--- /dev/null
+++ b/Engine/Source/Runtime/Core/Public/Misc/FeedbackContextMarkup.h
@@ -0,0 +1,47 @@
+// Copyright 1998-2014 Epic Games, Inc. All Rights Reserved.
+
+/*=============================================================================================
+	FeedbackContextMarkup.h: Functions to parse markup from external processes and control
+	a feedback context.
+==============================================================================================*/
+
+#pragma once
+
+/**
+ * A text-based markup language can be used to allow external processes to control the state 
+ * machine, allowing command-line utilities to give graphical user feedback.
+ *
+ * Example markup:
+ *
+ * Update the progress of the current operation:
+ *		@progress 10/20
+ *		@progress 50%
+ *
+ * Update the progress and set a status message for the current operation:
+ *		@progress 'Compiling source code...' 50%
+ */
+class CORE_API FFeedbackContextMarkup
+{
+public:
+	// Markup stack manipulation
+	static bool ParseCommand(const FString& Line, FFeedbackContext* Warn);
+
+	// Utility functions for dealing with external processes
+	static bool PipeProcessOutput(const FText& Description, const FString& URL, const FString& Params, FFeedbackContext* Warn, int32* OutExitCode);
+
+private:
+	/** Try to read a single exact-match token from the input stream. Must be followed by whitespace or EOL. */
+	static bool ReadToken(const TCHAR *&Text, const TCHAR *Token);
+
+	/** Read a progress value from the input stream. Valid forms are <Numerator>/<Denominator> or <Value>% */
+	static bool ReadProgress(const TCHAR *&Text, int32 &OutNumerator, int32 &OutDenominator);
+
+	/** Read an integer from the input stream */
+	static bool ReadInteger(const TCHAR *&Text, uint32 &OutInteger);
+
+	/** Read a string from the input stream */
+	static bool ReadString(const TCHAR *&Text, FString &OutString);
+
+	/** Skip over a sequence of whitespace characters in the input stream and returns a pointer to the next character */
+	static const TCHAR *SkipWhitespace(const TCHAR *Text);
+};
