Delegates: better handling of removal of functions from multicast delegates

diff --git a/Engine/Source/Runtime/Core/Public/Delegates/DelegateBase.h b/Engine/Source/Runtime/Core/Public/Delegates/DelegateBase.h
--- a/Engine/Source/Runtime/Core/Public/Delegates/DelegateBase.h
+++ b/Engine/Source/Runtime/Core/Public/Delegates/DelegateBase.h
@@ -65,16 +65,6 @@ public:
 		return (InUserObject != nullptr) && (DelegateInstance != nullptr) && DelegateInstance->HasSameObject(InUserObject);
 	}
 
-	/**
-	 * Checks to see if the user object bound to this delegate will ever be valid again
-	 *
-	 * @return  True if the user object could be valid in the future or is valid now.
-	 */
-	inline bool IsCompactable( ) const
-	{
-		return (DelegateInstance != nullptr) && DelegateInstance->IsCompactable();
-	}
-
 	/**
 	 * Unbinds this delegate
 	 */
diff --git a/Engine/Source/Runtime/Core/Public/Delegates/DelegateSignatureImpl.inl b/Engine/Source/Runtime/Core/Public/Delegates/DelegateSignatureImpl.inl
--- a/Engine/Source/Runtime/Core/Public/Delegates/DelegateSignatureImpl.inl
+++ b/Engine/Source/Runtime/Core/Public/Delegates/DelegateSignatureImpl.inl
@@ -1601,8 +1601,7 @@ public:
 		// cannot be bound to the same single multi-cast delegate.  We could support this by changing IsSameFunction()
 		// to do comparison of payload data, but there are some trade-offs if we do that.
 
-		RemoveDelegateInstance(DelegateInstance);
-		CompactInvocationList();
+		RemoveDelegateInstance(&DelegateInstance);
 	}
 
 	/**
@@ -1624,8 +1623,7 @@ public:
 
 		if (DelegateInstance != nullptr)
 		{
-			RemoveDelegateInstance(*DelegateInstance);
-			CompactInvocationList();
+			RemoveDelegateInstance(DelegateInstance);
 		}		
 	}
 
@@ -1782,53 +1780,60 @@ protected:
 	{
 #if !UE_BUILD_SHIPPING && !UE_BUILD_TEST
 		// verify that the same function isn't already bound
-		for (IDelegateInstancePtr DelegateInstance : GetInvocationList())
+		for (IDelegateInstance* DelegateInstance : GetInvocationList())
 		{
 			// this down-cast is OK! allows for managing invocation list in the base class without requiring virtual functions
-			TDelegateInstanceInterface* DelegateInstanceInterface = (TDelegateInstanceInterface*)DelegateInstance.Get();
+			TDelegateInstanceInterface* DelegateInstanceInterface = (TDelegateInstanceInterface*)DelegateInstance;
 			check(!DelegateInstanceInterface->IsSameFunction(*InDelegateInstance));
 		}
 #endif
 
-		AddInternal(MakeShareable(InDelegateInstance));
+		AddInternal(InDelegateInstance);
 	}
 
 	/**
 	 * Broadcasts this delegate to all bound objects, except to those that may have expired.
 	 *
-	 * The constness of this method is a lie, because we are removing compactable delegate intsances.
+	 * The constness of this method is a lie, because we are removing compactable delegate instances.
 	 */
 	void Broadcast( FUNC_PARAM_LIST ) const
 	{
-		const TArray<IDelegateInstancePtr>& InvocationList = GetInvocationList();
+		const TArray<IDelegateInstance*>& InvocationList = GetInvocationList();
 
 		if (InvocationList.Num() == 0)
 		{
 			return;
 		}
 
-		// copy invocation list just in case the list is modified by one of the broadcast callbacks
-		typedef TArray<IDelegateInstancePtr, TInlineAllocator<4>> FInlineInvocationList;
-		FInlineInvocationList InvocationListCopy = FInlineInvocationList(InvocationList);
-
-		// Invoke each bound function
-		for (IDelegateInstancePtr DelegateInstance : InvocationListCopy)
+		// create deep copy of invocation list in case it is modified by one of the broadcast callbacks
+		TArray<TDelegateInstanceInterface*, TInlineAllocator<4>> InvocationListCopy;
+		InvocationListCopy.Reset(InvocationList.Num());
+			
+		for (IDelegateInstance* DelegateInstance : InvocationList)
+		{
+			InvocationListCopy.Add(((TDelegateInstanceInterface*)DelegateInstance)->CreateCopy());
+		}
+			
+		// invoke each bound function
+		for (int32 ListIndex = InvocationListCopy.Num() - 1; ListIndex >= 0 ; --ListIndex)
 		{
 			// this down-cast is OK! allows for managing invocation list in the base class without requiring virtual functions
-			TDelegateInstanceInterface* DelegateInstanceInterface = (TDelegateInstanceInterface*)DelegateInstance.Get();
+			TDelegateInstanceInterface* DelegateInstanceInterface = InvocationListCopy[ListIndex];
 
-			// execute the delegate instance...
-			if (!DelegateInstanceInterface->ExecuteIfSafe(FUNC_PARAM_PASSTHRU))
+			// execute the function and remove it from the local list if it was processed
+			if (DelegateInstanceInterface->ExecuteIfSafe(FUNC_PARAM_PASSTHRU) || !DelegateInstanceInterface->IsCompactable())
 			{
-				// ... or remove it if compactable
-				if (DelegateInstance->IsCompactable())
-				{
-					// search for the instance, because original list could have been modified during broadcast
-					// this is really shady, because it violates the constness of this method
-					const_cast<BASE_MULTICAST_DELEGATE_CLASS*>(this)->RemoveInternal(DelegateInstance);
-				}
+				InvocationListCopy.RemoveAtSwap(ListIndex);
 			}
 		}
+
+		// remove any bound functions that couldn't be processed
+		for (TDelegateInstanceInterface* DelegateInstance : InvocationListCopy)
+		{
+			// search for the instance, because original list could have been modified during broadcast
+			// this is really shady, because it violates the constness of this method
+			const_cast<BASE_MULTICAST_DELEGATE_CLASS*>(this)->RemoveDelegateInstance(DelegateInstance);
+		}
 	}
 
 	/**
@@ -1840,23 +1845,29 @@ protected:
 	 *
 	 * @param InDelegateInstance The delegate instance to remove.
 	 */
-	void RemoveDelegateInstance( const TDelegateInstanceInterface& InDelegateInstance )
+	void RemoveDelegateInstance( const TDelegateInstanceInterface* InDelegateInstance )
 	{
-		const TArray<IDelegateInstancePtr> InvocationList = GetInvocationList();
-
-		for (int32 InvocationListIndex = 0; InvocationListIndex < InvocationList.Num(); ++InvocationListIndex)
+		if (InDelegateInstance->IsCompactable())
 		{
-			// this down-cast is OK! allows for managing invocation list in the base class without requiring virtual functions
-			TDelegateInstanceInterface* DelegateInstance = (TDelegateInstanceInterface*)InvocationList[InvocationListIndex].Get();
+			CompactInvocationList();
+		}
+		else
+		{
+			const TArray<IDelegateInstance*> InvocationList = GetInvocationList();
 
-			// NOTE: We must do a deep compare here, not just compare delegate pointers, because multiple
-			//       delegate pointers can refer to the exact same object and method
-			if (DelegateInstance->IsSameFunction(InDelegateInstance))
+			for (int32 InvocationListIndex = 0; InvocationListIndex < InvocationList.Num(); ++InvocationListIndex)
 			{
-				RemoveByIndex(InvocationListIndex);
+				// this down-cast is OK! allows for managing invocation list in the base class without requiring virtual functions
+				TDelegateInstanceInterface* DelegateInstance = (TDelegateInstanceInterface*)InvocationList[InvocationListIndex];
 
-				// no need to continue, as we never allow the same delegate to be bound twice
-				break;
+				// NOTE: We must do a deep compare here, not just compare delegate pointers, because multiple
+				//       delegate pointers can refer to the exact same object and method
+				if (DelegateInstance->IsSameFunction(*InDelegateInstance))
+				{
+					RemoveByIndex(InvocationListIndex);
+				
+					break;	// no need to continue, as we never allow the same delegate to be bound twice
+				}
 			}
 		}
 	}
diff --git a/Engine/Source/Runtime/Core/Public/Delegates/IDelegateInstance.h b/Engine/Source/Runtime/Core/Public/Delegates/IDelegateInstance.h
--- a/Engine/Source/Runtime/Core/Public/Delegates/IDelegateInstance.h
+++ b/Engine/Source/Runtime/Core/Public/Delegates/IDelegateInstance.h
@@ -35,17 +35,6 @@ namespace EDelegateInstanceType
 }
 
 
-/**
- * Type definition for shared pointers to instances of IDelegateInstance.
- */
-typedef TSharedPtr<class IDelegateInstance> IDelegateInstancePtr;
-
-/**
- * Type definition for shared references to instances of IDelegateInstance.
- */
-typedef TSharedRef<class IDelegateInstance> IDelegateInstanceRef;
-
-
 /**
  * Interface for delegate instances.
  */
diff --git a/Engine/Source/Runtime/Core/Public/Delegates/MulticastDelegateBase.h b/Engine/Source/Runtime/Core/Public/Delegates/MulticastDelegateBase.h
--- a/Engine/Source/Runtime/Core/Public/Delegates/MulticastDelegateBase.h
+++ b/Engine/Source/Runtime/Core/Public/Delegates/MulticastDelegateBase.h
@@ -20,6 +20,11 @@ public:
 	 */
 	void Clear( )
 	{
+		for (IDelegateInstance* DelegateInstance : InvocationList)
+		{
+			delete DelegateInstance;
+		}
+
 		InvocationList.Empty();
 		ExpirableObjectCount = 0;
 	}
@@ -43,7 +48,7 @@ public:
 	 */
 	inline bool IsBoundToObject( void const* InUserObject ) const
 	{
-		for (IDelegateInstancePtr DelegateInstance : InvocationList)
+		for (IDelegateInstance* DelegateInstance : InvocationList)
 		{
 			if (DelegateInstance->HasSameObject(InUserObject))
 			{
@@ -82,7 +87,7 @@ protected:
 	 *
 	 * @param DelegateInstance The delegate instance to add.
 	 */
-	void AddInternal( const IDelegateInstanceRef& DelegateInstance )
+	void AddInternal( IDelegateInstance* DelegateInstance )
 	{
 		InvocationList.Add(DelegateInstance);
 
@@ -111,11 +116,12 @@ protected:
 		// remove expired delegates
 		for (int32 InvocationListIndex = InvocationList.Num() - 1; InvocationListIndex >= 0; --InvocationListIndex)
 		{
-			IDelegateInstancePtr DelegateInstance = InvocationList[InvocationListIndex];
+			IDelegateInstance* DelegateInstance = InvocationList[InvocationListIndex];
 
 			if (DelegateInstance->IsCompactable())
 			{
 				InvocationList.RemoveAtSwap(InvocationListIndex);
+				delete DelegateInstance;
 
 				checkSlow(ExpirableObjectCount > 0);
 				--ExpirableObjectCount;
@@ -128,7 +134,7 @@ protected:
 	 *
 	 * @return The invocation list.
 	 */
-	const TArray<IDelegateInstancePtr>& GetInvocationList( ) const
+	const TArray<IDelegateInstance*>& GetInvocationList( ) const
 	{
 		return InvocationList;
 	}
@@ -144,7 +150,7 @@ protected:
 	{
 		for (int32 InvocationListIndex = InvocationList.Num() - 1; InvocationListIndex >= 0; --InvocationListIndex)
 		{
-			IDelegateInstancePtr DelegateInstance = InvocationList[InvocationListIndex];
+			IDelegateInstance* DelegateInstance = InvocationList[InvocationListIndex];
 
 			if (DelegateInstance->HasSameObject(InUserObject))
 			{
@@ -160,7 +166,7 @@ protected:
 	 */
 	void RemoveByIndex( int32 InvocationListIndex )
 	{
-		IDelegateInstancePtr DelegateInstance = InvocationList[InvocationListIndex];
+		IDelegateInstance* DelegateInstance = InvocationList[InvocationListIndex];
 
 		// keep track of whether we have objects bound that may expire out from underneath us
 		const EDelegateInstanceType::Type DelegateType = DelegateInstance->GetType();
@@ -174,21 +180,7 @@ protected:
 		}
 
 		InvocationList.RemoveAtSwap(InvocationListIndex);
-	}
-
-	/**
-	 * Removes the given delegate instance from the invocation list.
-	 *
-	 * @param DelegateInstance The delegate instance to remove.
-	 */
-	void RemoveInternal( const IDelegateInstancePtr& DelegateInstance )
-	{
-		int32 DelegateIndex;
-		
-		if (InvocationList.Find(DelegateInstance, DelegateIndex))
-		{
-			RemoveByIndex(DelegateIndex);
-		}
+		delete DelegateInstance;
 	}
 
 private:
@@ -198,5 +190,5 @@ private:
 
 	// Holds the collection of delegate instances to invoke.
 	// Mutable so that we can do housekeeping during const broadcasts.
-	mutable TArray<IDelegateInstancePtr> InvocationList;
+	mutable TArray<IDelegateInstance*> InvocationList;
 };
