Added a recursion guard to Internationalization's initialization. Extra safety added in PopulateAllCultures to empty out AllCultures from previous attempts to populate.

diff --git a/Engine/Source/Runtime/Core/Private/Internationalization/Internationalization.cpp b/Engine/Source/Runtime/Core/Private/Internationalization/Internationalization.cpp
--- a/Engine/Source/Runtime/Core/Private/Internationalization/Internationalization.cpp
+++ b/Engine/Source/Runtime/Core/Private/Internationalization/Internationalization.cpp
@@ -179,10 +179,17 @@ namespace
 
 void FInternationalization::Initialize()
 {
-	if(IsInitialized())
+	static bool IsInitializing = false;
+
+	if(IsInitialized() || IsInitializing)
 	{
 		return;
 	}
+	struct FInitializingGuard
+	{
+		FInitializingGuard()	{IsInitializing = true;}
+		~FInitializingGuard()	{IsInitializing = false;}
+	} InitializingGuard;
 
 #if UE_ENABLE_ICU
 	UErrorCode ICUStatus = U_ZERO_ERROR;
@@ -404,7 +411,7 @@ void FInternationalization::PopulateAllCultures(void)
 	const icu::Locale* const AvailableLocales = icu::Locale::getAvailableLocales(LocaleCount);
 
 	// Get available locales and create cultures.
-	AllCultures.Reserve(LocaleCount);
+	AllCultures.Reset(LocaleCount);
 	for(int32_t i = 0; i < LocaleCount; ++i)
 	{
 		AllCultures.Add( MakeShareable( new FCulture( AvailableLocales[i].getName() ) ) );
