Fixes for how the editor handles localization when the OS is set to a language we don't have a translation for

TTP# 333731 - EDITOR: The editor region doesn't always seem to get set correctly on non-English systems

If you'd previously not set a language in the editor, the "Region & Language" settings would always show you the option for "en_US_POSIX" (and claim that you needed to restart the editor) because the saved culture name ("") didn't match the FInternationalization culture name (eg, "en_GB"), causing it to call GetCulture with an empty string (which seems to cause ICU to return "en_US_POSIX"). It now just uses the current FInternationalization culture name if the saved culture is empty.

Additionally, the editor/game now verifies that it has an available translation for the culture that FInternationalization is using (the one set during startup). If it doesn't (eg, your PC is using "fr_BR") then it will fallback to using generic English ("en") to ensure that the editor/game always has a valid translation to use.

diff --git a/Engine/Source/Developer/InternationalizationSettings/Private/InternationalizationSettingsModelDetails.cpp b/Engine/Source/Developer/InternationalizationSettings/Private/InternationalizationSettingsModelDetails.cpp
--- a/Engine/Source/Developer/InternationalizationSettings/Private/InternationalizationSettingsModelDetails.cpp
+++ b/Engine/Source/Developer/InternationalizationSettings/Private/InternationalizationSettingsModelDetails.cpp
@@ -72,7 +72,7 @@ void FInternationalizationSettingsModelDetails::OnSettingsChanged()
 	FInternationalization& I18N = FInternationalization::Get();
 
 	FString SavedCultureName = Model->GetCultureName();
-	if ( SavedCultureName != I18N.GetCurrentCulture()->GetName() )
+	if ( !SavedCultureName.IsEmpty() && SavedCultureName != I18N.GetCurrentCulture()->GetName() )
 	{
 		SelectedCulture = I18N.GetCulture(SavedCultureName);
 		RequiresRestart = true;
@@ -106,10 +106,10 @@ void FInternationalizationSettingsModelDetails::CustomizeDetails( IDetailLayoutB
 
 	// If the saved culture is not the same as the actual current culture, a restart is needed to sync them fully and properly.
 	FString SavedCultureName = Model->GetCultureName();
-	if ( SavedCultureName != I18N.GetCurrentCulture()->GetName() )
+	if ( !SavedCultureName.IsEmpty() && SavedCultureName != I18N.GetCurrentCulture()->GetName() )
 	{
-			SelectedCulture = I18N.GetCulture(SavedCultureName);
-			RequiresRestart = true;
+		SelectedCulture = I18N.GetCulture(SavedCultureName);
+		RequiresRestart = true;
 	}
 	else
 	{
@@ -242,64 +242,7 @@ void FInternationalizationSettingsModelDetails::CustomizeDetails( IDetailLayoutB
 void FInternationalizationSettingsModelDetails::RefreshAvailableCultures()
 {
 	AvailableCultures.Empty();
-
-	IPlatformFile& PlatformFile = IPlatformFile::GetPlatformPhysical();
-	TArray<FString> AllCultureNames;
-	FInternationalization::Get().GetCultureNames(AllCultureNames);
-	const TArray<FString> LocalizationPaths = FPaths::GetEditorLocalizationPaths();
-	for(const auto& LocalizationPath : LocalizationPaths)
-	{
-		/* Visitor class used to enumerate directories of culture */
-		class FCultureEnumeratorVistor : public IPlatformFile::FDirectoryVisitor
-		{
-		public:
-			FCultureEnumeratorVistor( const TArray<FString>& InAllCultureNames, TArray< TSharedPtr<FCulture> >& InAvailableCultures )
-				: AllCultureNames(InAllCultureNames)
-				, AvailableCultures(InAvailableCultures)
-			{
-			}
-
-			virtual bool Visit(const TCHAR* FilenameOrDirectory, bool bIsDirectory) OVERRIDE
-			{
-				if(bIsDirectory)
-				{
-					for( const auto& CultureName : AllCultureNames )
-					{
-						TSharedPtr<FCulture> Culture = FInternationalization::Get().GetCulture(CultureName);
-						if(Culture.IsValid() && !AvailableCultures.Contains(Culture))
-						{
-							// UE localization resource folders use "en-US" style while ICU uses "en_US" style so we replace underscores with dashes here.
-							const FString UnrealCultureName = FString(TEXT("/")) + CultureName.Replace(TEXT("_"), TEXT("-"));
-							if(FString(FilenameOrDirectory).EndsWith(UnrealCultureName))
-							{
-								AvailableCultures.Add(Culture);
-							}
-							else
-							{
-								// If the full name doesn't match, see if the base language is present
-								const FString CultureLanguageName = FString(TEXT("/")) + Culture->GetTwoLetterISOLanguageName();
-								if(FString(FilenameOrDirectory).EndsWith(CultureLanguageName))
-								{
-									AvailableCultures.Add(Culture);
-								}
-							}
-						}
-					}
-				}
-
-				return true;
-			}
-
-			/** Array of all culture names we can use */
-			const TArray<FString>& AllCultureNames;
-
-			/** Array of cultures that are available */
-			TArray< TSharedPtr<FCulture> >& AvailableCultures;
-		};
-
-		FCultureEnumeratorVistor CultureEnumeratorVistor(AllCultureNames, AvailableCultures);
-		PlatformFile.IterateDirectory(*LocalizationPath, CultureEnumeratorVistor);	
-	}
+	FInternationalization::Get().GetCulturesWithAvailableLocalization(FPaths::GetEditorLocalizationPaths(), AvailableCultures);
 
 	// Update our selected culture based on the available choices
 	if ( !AvailableCultures.Contains( SelectedCulture ) )
diff --git a/Engine/Source/Runtime/Core/Private/Internationalization/Internationalization.cpp b/Engine/Source/Runtime/Core/Private/Internationalization/Internationalization.cpp
--- a/Engine/Source/Runtime/Core/Private/Internationalization/Internationalization.cpp
+++ b/Engine/Source/Runtime/Core/Private/Internationalization/Internationalization.cpp
@@ -23,6 +23,27 @@
 
 #define LOCTEXT_NAMESPACE "Internationalization"
 
+namespace
+{
+
+FString GetCanonicalName(const FString& Name)
+{
+#if UE_ENABLE_ICU
+	static const int32 MaximumNameLength = 64;
+	const int32 NameLength = Name.Len();
+	check(NameLength < MaximumNameLength);
+	char CanonicalName[MaximumNameLength];
+
+	UErrorCode ICUStatus = U_ZERO_ERROR;
+	uloc_canonicalize(TCHAR_TO_ANSI( *Name ), CanonicalName, MaximumNameLength, &ICUStatus);
+	return CanonicalName;
+#else
+	return Name;
+#endif
+}
+
+}
+
 FInternationalization* FInternationalization::Instance;
 
 FInternationalization& FInternationalization::Get()
@@ -46,7 +67,7 @@ void FInternationalization::TearDown()
 	}
 }
 
-void FInternationalization::GetTimeZonesIDs(TArray<FString>& TimeZonesIDs)
+void FInternationalization::GetTimeZonesIDs(TArray<FString>& TimeZonesIDs) const
 {
 	TimeZonesIDs.Empty();
 #if UE_ENABLE_ICU
@@ -86,38 +107,26 @@ void FInternationalization::SetCurrentCulture(const FString& Name)
 	}
 }
 
-TSharedRef< FCulture > FInternationalization::GetCurrentCulture()
+TSharedRef< FCulture > FInternationalization::GetCurrentCulture() const
 {
 	return AllCultures[CurrentCultureIndex];
 }
 
-TSharedPtr< FCulture > FInternationalization::GetCulture(const FString& Name)
+TSharedPtr< FCulture > FInternationalization::GetCulture(const FString& Name) const
 {
 	int32 CultureIndex = GetCultureIndex(Name);
 	return CultureIndex != -1 ? AllCultures[CultureIndex] : TSharedPtr<FCulture>();
 }
 
-int32 FInternationalization::GetCultureIndex(const FString& Name)
+int32 FInternationalization::GetCultureIndex(const FString& Name) const
 {
-#if UE_ENABLE_ICU
-	static const int32 MaximumNameLength = 64;
-	const int32 NameLength = Name.Len();
-	check(NameLength < MaximumNameLength);
-	char CanonicalName[MaximumNameLength];
-
-	UErrorCode ICUStatus = U_ZERO_ERROR;
-	uloc_canonicalize(TCHAR_TO_ANSI( *Name ), CanonicalName, MaximumNameLength, &ICUStatus);
-#endif	
+	const FString CanonicalName = GetCanonicalName(Name);
 
 	const int32 CultureCount = AllCultures.Num();
 	int32 i;
 	for (i = 0; i < CultureCount; ++i)
 	{
-#if UE_ENABLE_ICU
 		if( AllCultures[i]->GetName() == CanonicalName )
-#else
-		if( AllCultures[i]->GetName() == Name )
-#endif
 		{
 			break;
 		}
@@ -376,25 +385,65 @@ FString& FInternationalization::Leetify(FString& SourceString)
 }
 #endif
 
-void FInternationalization::GetCultureNames(TArray<FString>& CultureNames)
+void FInternationalization::GetCultureNames(TArray<FString>& CultureNames) const
 {
-	CultureNames.Empty();
-#if UE_ENABLE_ICU
-	int32_t LocaleCount;
-	const icu::Locale* const AvailableLocales = icu::Locale::getAvailableLocales(LocaleCount);
+	CultureNames.Reset();
+	for(const auto& Culture : AllCultures)
+	{
+		CultureNames.Add(Culture->GetName());
+	}
+}
 
-	for(int32_t i = 0; i < LocaleCount; ++i)
+void FInternationalization::GetCulturesWithAvailableLocalization(const TArray<FString>& InLocalizationPaths, TArray< TSharedPtr<FCulture> >& OutAvailableCultures) const
+{
+	OutAvailableCultures.Reset();
+
+	TArray<FString> AllLocalizationFolders;
+	IPlatformFile& PlatformFile = IPlatformFile::GetPlatformPhysical();
+	for(const auto& LocalizationPath : InLocalizationPaths)
 	{
-		const char* const LocaleName = AvailableLocales[i].getName();
+		/* Visitor class used to enumerate directories of culture */
+		class FCultureEnumeratorVistor : public IPlatformFile::FDirectoryVisitor
+		{
+		public:
+			FCultureEnumeratorVistor( TArray<FString>& OutLocalizationFolders )
+				: LocalizationFolders(OutLocalizationFolders)
+			{
+			}
+
+			virtual bool Visit(const TCHAR* FilenameOrDirectory, bool bIsDirectory) OVERRIDE
+			{
+				if(bIsDirectory)
+				{
+					// UE localization resource folders use "en-US" style while ICU uses "en_US"
+					const FString LocalizationFolder = FPaths::GetCleanFilename(FilenameOrDirectory);
+					const FString CanonicalName = GetCanonicalName(LocalizationFolder);
+					LocalizationFolders.AddUnique(CanonicalName);
+				}
+
+				return true;
+			}
+
+			/** Array to fill with the names of the UE localization folders available at the given path */
+			TArray<FString>& LocalizationFolders;
+		};
 
-		CultureNames.Add(FString(LocaleName));
+		FCultureEnumeratorVistor CultureEnumeratorVistor(AllLocalizationFolders);
+		PlatformFile.IterateDirectory(*LocalizationPath, CultureEnumeratorVistor);	
 	}
-#else
-	for(int32 i = 0; i < AllCultures.Num(); ++i)
+
+	// Find any cultures that are a complete or partial match for the languages we have translations for
+	for(const auto& Culture : AllCultures)
 	{
-		CultureNames.Add(AllCultures[i]->GetName());
+		for(const FString& LocalizationFolder : AllLocalizationFolders)
+		{
+			// Check the full name, but if it doesn't match, see if the base language is present
+			if(LocalizationFolder == Culture->GetName() || LocalizationFolder == Culture->GetTwoLetterISOLanguageName())
+			{
+				OutAvailableCultures.AddUnique(Culture);
+			}
+		}
 	}
-#endif
 }
 
 FInternationalization::FInternationalization()
diff --git a/Engine/Source/Runtime/Core/Private/Internationalization/TextLocalizationManager.cpp b/Engine/Source/Runtime/Core/Private/Internationalization/TextLocalizationManager.cpp
--- a/Engine/Source/Runtime/Core/Private/Internationalization/TextLocalizationManager.cpp
+++ b/Engine/Source/Runtime/Core/Private/Internationalization/TextLocalizationManager.cpp
@@ -13,12 +13,96 @@ static FString AccessedStringBeforeLocLoadedErrorMsg = TEXT("Can't access string
 
 void BeginInitTextLocalization()
 {
+	// Initialize FInternationalization before we bind to OnCultureChanged otherwise we can accidentally initialize
+	// twice since FInternationalization::Initialize sets the culture
+	FInternationalization::Get();
+
 	FCoreDelegates::OnCultureChanged.AddRaw( &(FTextLocalizationManager::Get()), &FTextLocalizationManager::OnCultureChanged );
 }
 
 void EndInitTextLocalization()
 {
-	FTextLocalizationManager::Get().LoadResources(WITH_EDITOR, !GIsEditor);
+	const bool ShouldLoadEditor = WITH_EDITOR;
+	const bool ShouldLoadGame = FApp::IsGame();
+
+	FInternationalization& I18N = FInternationalization::Get();
+
+	// Set culture according to configuration now that configs are available.
+#if ENABLE_LOC_TESTING
+	if( FCommandLine::IsInitialized() && FParse::Param(FCommandLine::Get(), TEXT("LEET")) )
+	{
+		I18N.SetCurrentCulture(TEXT("LEET"));
+	}
+	else
+#endif
+	{
+		FString CultureName;
+		if (FParse::Value(FCommandLine::Get(), TEXT("CULTUREFORCOOKING="), CultureName))
+		{
+			// Write the culture passed in if first install...
+			if (FParse::Param(FCommandLine::Get(), TEXT("firstinstall")))
+			{
+				GConfig->SetString(TEXT("Internationalization"), TEXT("Culture"), *CultureName, GEngineIni);
+			}
+		}
+		else
+#if !UE_BUILD_SHIPPING
+		// Use culture override specified on commandline.
+		if (FParse::Value(FCommandLine::Get(), TEXT("CULTURE="), CultureName))
+		{
+			//UE_LOG(LogInit, Log, TEXT("Overriding culture %s w/ command-line option".), *CultureName);
+		}
+		else
+#endif // !UE_BUILD_SHIPPING
+#if WITH_EDITOR
+		// See if we've been provided a culture override in the editor
+		if(GConfig->GetString( TEXT("Internationalization"), TEXT("Culture"), CultureName, GEditorGameAgnosticIni ))
+		{
+			//UE_LOG(LogInit, Log, TEXT("Overriding culture %s w/ editor configuration."), *CultureName);
+		}
+		else
+#endif // WITH_EDITOR
+		// Use culture specified in engine configuration.
+		if(GConfig->GetString( TEXT("Internationalization"), TEXT("Culture"), CultureName, GEngineIni ))
+		{
+			//UE_LOG(LogInit, Log, TEXT("Overriding culture %s w/ engine configuration."), *CultureName);
+		}
+		else
+		{
+			CultureName = I18N.GetDefaultCulture()->GetName();
+		}
+
+		{
+			TArray<FString> LocalizationPaths;
+
+			if(ShouldLoadEditor)
+			{
+				LocalizationPaths += FPaths::GetEditorLocalizationPaths();
+			}
+
+			if(ShouldLoadGame)
+			{
+				LocalizationPaths += FPaths::GetGameLocalizationPaths();
+			}
+
+			LocalizationPaths += FPaths::GetEngineLocalizationPaths();
+
+			TArray< TSharedPtr<FCulture> > AvailableCultures;
+			I18N.GetCulturesWithAvailableLocalization(LocalizationPaths, AvailableCultures);
+
+			// If we don't have an available translation for the current culture, fallback to generic English
+			TSharedPtr<FCulture> TargetCulture = I18N.GetCulture(CultureName);
+			if(!TargetCulture.IsValid() || !AvailableCultures.Contains(TargetCulture))
+			{
+				UE_LOG(LogTextLocalizationManager, Warning, TEXT("The selected culture '%s' is not available; falling back to 'en'"), *CultureName);
+				CultureName = "en";
+			}
+		}
+
+		I18N.SetCurrentCulture(CultureName);
+	}
+
+	FTextLocalizationManager::Get().LoadResources(ShouldLoadEditor, ShouldLoadGame);
 	FTextLocalizationManager::Get().bIsInitialized = true;
 }
 
@@ -332,8 +416,8 @@ void FTextLocalizationManager::UpdateLiveTable(const TArray<FLocalizationEntryTr
 
 void FTextLocalizationManager::OnCultureChanged()
 {
-	const bool ShouldLoadEditor = bIsInitialized ? WITH_EDITOR : false;
-	const bool ShouldLoadGame = bIsInitialized ? !GIsEditor : false;
+	const bool ShouldLoadEditor = bIsInitialized && WITH_EDITOR;
+	const bool ShouldLoadGame = bIsInitialized && FApp::IsGame();
 	LoadResources(ShouldLoadEditor, ShouldLoadGame);
 }
 
diff --git a/Engine/Source/Runtime/Core/Public/Internationalization/Internationalization.h b/Engine/Source/Runtime/Core/Public/Internationalization/Internationalization.h
--- a/Engine/Source/Runtime/Core/Public/Internationalization/Internationalization.h
+++ b/Engine/Source/Runtime/Core/Public/Internationalization/Internationalization.h
@@ -36,37 +36,40 @@ public:
 		return FText( InTextLiteral, Namespace, Key );
 	}
 
-	CORE_API void GetTimeZonesIDs(TArray<FString>& TimeZonesIDs);
+	CORE_API void GetTimeZonesIDs(TArray<FString>& TimeZonesIDs) const;
 
 	//Set the current culture by name
 	CORE_API void SetCurrentCulture(const FString& Name);
 
 	//@return the current culture
-	CORE_API TSharedRef< FCulture > GetCurrentCulture();
+	CORE_API TSharedRef< FCulture > GetCurrentCulture() const;
 
 	//@return culture object by given name, or NULL if not found
-	CORE_API TSharedPtr< FCulture > GetCulture(const FString& Name);
+	CORE_API TSharedPtr< FCulture > GetCulture(const FString& Name) const;
 
 	//@return the default culture
-	CORE_API TSharedRef< FCulture > GetDefaultCulture()
+	CORE_API TSharedRef< FCulture > GetDefaultCulture() const
 	{
 		return DefaultCulture.ToSharedRef();
 	}
 
 	//@return the invariant culture
-	CORE_API TSharedRef< FCulture > GetInvariantCulture()
+	CORE_API TSharedRef< FCulture > GetInvariantCulture() const
 	{
 		return InvariantCulture.ToSharedRef();
 	}
 
 
-	CORE_API bool IsInitialized() {return bIsInitialized;}
+	CORE_API bool IsInitialized() const {return bIsInitialized;}
 
 #if ENABLE_LOC_TESTING
 	static CORE_API FString& Leetify(FString& SourceString);
 #endif
 
-	CORE_API void GetCultureNames(TArray<FString>& CultureNames);
+	CORE_API void GetCultureNames(TArray<FString>& CultureNames) const;
+
+	// Given some paths to look at, populate a list of cultures that we have available localization information for
+	CORE_API void GetCulturesWithAvailableLocalization(const TArray<FString>& InLocalizationPaths, TArray< TSharedPtr<FCulture> >& OutAvailableCultures) const;
 
 private:
 	FInternationalization();
@@ -75,7 +78,7 @@ private:
 	CORE_API void PopulateAllCultures(void);
 
 	//@return index of a culture by given name, or -1 if not found
-	CORE_API int32 GetCultureIndex(const FString& Name);
+	CORE_API int32 GetCultureIndex(const FString& Name) const;
 
 	CORE_API void Initialize();
 	CORE_API void Terminate();
diff --git a/Engine/Source/Runtime/Launch/Private/LaunchEngineLoop.cpp b/Engine/Source/Runtime/Launch/Private/LaunchEngineLoop.cpp
--- a/Engine/Source/Runtime/Launch/Private/LaunchEngineLoop.cpp
+++ b/Engine/Source/Runtime/Launch/Private/LaunchEngineLoop.cpp
@@ -2316,47 +2316,6 @@ void FEngineLoop::AppInit( )
 	// after the above has run we now have the REQUIRED set of engine .INIs  (all of the other .INIs)
 	// that are gotten from .h files' config() are not requires and are dynamically loaded when the .u files are loaded
 
-	FInternationalization& I18N = FInternationalization::Get();
-
-	// Set culture according to configuration now that configs are available.
-#if ENABLE_LOC_TESTING
-	if( FCommandLine::IsInitialized() && FParse::Param(FCommandLine::Get(), TEXT("LEET")) )
-	{
-		I18N.SetCurrentCulture(TEXT("LEET"));
-	}
-	else
-#endif
-	{
-		FString CultureName;
-#if !UE_BUILD_SHIPPING
-		// Use culture override specified on commandline.
-		if (FParse::Value(FCommandLine::Get(), TEXT("CULTURE="), CultureName))
-		{
-			//UE_LOG(LogInit, Log, TEXT("Overriding culture %s w/ command-line option".), *CultureName);
-		}
-		else
-#endif // !UE_BUILD_SHIPPING
-#if WITH_EDITOR
-		// See if we've been provided a culture override in the editor
-		if(GConfig->GetString( TEXT("Internationalization"), TEXT("Culture"), CultureName, GEditorGameAgnosticIni ))
-		{
-			//UE_LOG(LogInit, Log, TEXT("Overriding culture %s w/ editor configuration."), *CultureName);
-		}
-		else
-#endif // WITH_EDITOR
-		// Use culture specified in engine configuration.
-		if(GConfig->GetString( TEXT("Internationalization"), TEXT("Culture"), CultureName, GEngineIni ))
-		{
-			//UE_LOG(LogInit, Log, TEXT("Overriding culture %s w/ engine configuration."), *CultureName);
-		}
-		else
-		{
-			CultureName = I18N.GetDefaultCulture()->GetName();
-		}
-
-		I18N.SetCurrentCulture(CultureName);
-	}
-
 #if !UE_BUILD_SHIPPING
 	// Prompt the user for remote debugging?
 	bool bPromptForRemoteDebug = false;
@@ -2497,20 +2456,6 @@ void FEngineLoop::AppInit( )
 	check(sizeof(double) == 8);
 	check(GEngineNetVersion == 0 || GEngineNetVersion >= GEngineMinNetVersion || GEngineVersion.IsLicenseeVersion());
 
-	// Culture.
-	TCHAR CookerCulture[8];
-	
-	if (FParse::Value(FCommandLine::Get(), TEXT("CULTUREFORCOOKING="), CookerCulture, ARRAY_COUNT(CookerCulture)))
-	{
-		FInternationalization::Get().SetCurrentCulture( CookerCulture );
-
-		// Write the culture passed in if first install...
-		if (FParse::Param(FCommandLine::Get(), TEXT("firstinstall")))
-		{
-			GConfig->SetString(TEXT("Internationalization"), TEXT("Culture"), CookerCulture, GEngineIni);
-		}
-	}
-
 	// Init list of common colors.
 	GColorList.CreateColorMap();
 
