missing files for my last submission
#ue4
#ios

diff --git a/Engine/Source/Runtime/Core/Private/IOS/IOSAppDelegate.cpp b/Engine/Source/Runtime/Core/Private/IOS/IOSAppDelegate.cpp
--- a/Engine/Source/Runtime/Core/Private/IOS/IOSAppDelegate.cpp
+++ b/Engine/Source/Runtime/Core/Private/IOS/IOSAppDelegate.cpp
@@ -201,9 +201,9 @@ void InstallSignalHandlers()
 			UE_LOG(LogIOSAudioSession, Error, TEXT("Failed to set audio session category to AVAudioSessionCategoryAmbient!"));
 		}
 		ActiveError = nil;
-	}
+ 	}
+    self.bAudioActive = true;
 
-	[[AVAudioSession sharedInstance] setDelegate:self];
 	/* TODO::JTM - Jan 16, 2013 06:22PM - Music player support */
 }
 
@@ -211,6 +211,8 @@ void InstallSignalHandlers()
 {
 	if (bActive)
 	{
+        if (!self.bAudioActive)
+        {
 		bool bWasUsingBackgroundMusic = self.bUsingBackgroundMusic;
 		self.bUsingBackgroundMusic = [self IsBackgroundAudioPlaying];
 
@@ -266,8 +268,9 @@ void InstallSignalHandlers()
 			}
 			ActiveError = nil;
 		}
+        }
 	}
-	else if (!self.bUsingBackgroundMusic)
+	else if (self.bAudioActive && !self.bUsingBackgroundMusic)
 	{
 		NSError* ActiveError = nil;
 		[[AVAudioSession sharedInstance] setActive:NO error:&ActiveError];
@@ -276,7 +279,7 @@ void InstallSignalHandlers()
 			UE_LOG(LogIOSAudioSession, Error, TEXT("Failed to set audio session as inactive!"));
 		}
 		ActiveError = nil;
-
+        
 		// Necessary to prevent audio from getting killing when setup for background iPod audio playback
 		[[AVAudioSession sharedInstance] setCategory:AVAudioSessionCategoryAmbient error:&ActiveError];
 		if (ActiveError)
@@ -284,7 +287,8 @@ void InstallSignalHandlers()
 			UE_LOG(LogIOSAudioSession, Error, TEXT("Failed to set audio session category to AVAudioSessionCategoryAmbient!"));
 		}
 		ActiveError = nil;
-	}
+ 	}
+    self.bAudioActive = bActive;
 }
 
 - (bool)IsBackgroundAudioPlaying
@@ -388,12 +392,16 @@ void InstallSignalHandlers()
 
 - (void)beginInterruption
 {
+    FAppEntry::Suspend();
+
 	[self ToggleAudioSession:false];
-}
+ }
 
 - (void)endInterruption
 {
 	[self ToggleAudioSession:true];
+
+    FAppEntry::Resume();
 }
 
 - (void)applicationWillResignActive:(UIApplication *)application
diff --git a/Engine/Source/Runtime/Core/Public/IOS/IOSAppDelegate.h b/Engine/Source/Runtime/Core/Public/IOS/IOSAppDelegate.h
--- a/Engine/Source/Runtime/Core/Public/IOS/IOSAppDelegate.h
+++ b/Engine/Source/Runtime/Core/Public/IOS/IOSAppDelegate.h
@@ -21,6 +21,8 @@ namespace FAppEntry
 	void Tick();
     void SuspendTick();
 	void Shutdown();
+    void Suspend();
+    void Resume();
 }
 
 @interface IOSAppDelegate : UIResponder <UIApplicationDelegate,
@@ -84,6 +86,8 @@ UITextFieldDelegate, AVAudioSessionDelegate>
 - (void)ToggleAudioSession:(bool)bActive;
 - (bool)IsBackgroundAudioPlaying;
 
+@property (atomic) bool bAudioActive;
+
 @property (atomic) bool bIsSuspended;
 @property (atomic) bool bHasSuspended;
 @property (atomic) bool bHasStarted;
diff --git a/Engine/Source/Runtime/Engine/Public/AudioDevice.h b/Engine/Source/Runtime/Engine/Public/AudioDevice.h
--- a/Engine/Source/Runtime/Engine/Public/AudioDevice.h
+++ b/Engine/Source/Runtime/Engine/Public/AudioDevice.h
@@ -492,6 +492,12 @@ public:
 	/* When the set of Reverb volumes have changed invalidate the cached values of active sounds */
 	void InvalidateCachedInteriorVolumes() const;
 
+    /** Suspend any context related objects */
+    virtual void SuspendContext() {}
+    
+    /** Resume any context related objects */
+    virtual void ResumeContext() {}
+    
 protected:
 	friend class FSoundSource;
 
diff --git a/Engine/Source/Runtime/IOS/IOSAudio/Private/IOSAudioDevice.cpp b/Engine/Source/Runtime/IOS/IOSAudio/Private/IOSAudioDevice.cpp
--- a/Engine/Source/Runtime/IOS/IOSAudio/Private/IOSAudioDevice.cpp
+++ b/Engine/Source/Runtime/IOS/IOSAudio/Private/IOSAudioDevice.cpp
@@ -35,12 +35,16 @@ void FIOSAudioDevice::ResumeContext()
 {
 	int32& SuspendCounter = GetSuspendCounter();
 	FPlatformAtomics::InterlockedDecrement(&SuspendCounter);
+
+    AUGraphStart(AudioUnitGraph);
 }
 
 void FIOSAudioDevice::SuspendContext()
 {
 	int32& SuspendCounter = GetSuspendCounter();
 	FPlatformAtomics::InterlockedIncrement(&SuspendCounter);
+
+    AUGraphStop(AudioUnitGraph);
 }
 
 int32& FIOSAudioDevice::GetSuspendCounter()
diff --git a/Engine/Source/Runtime/Launch/Private/IOS/LaunchIOS.cpp b/Engine/Source/Runtime/Launch/Private/IOS/LaunchIOS.cpp
--- a/Engine/Source/Runtime/Launch/Private/IOS/LaunchIOS.cpp
+++ b/Engine/Source/Runtime/Launch/Private/IOS/LaunchIOS.cpp
@@ -7,10 +7,21 @@
 #include "EAGLView.h"
 #include "IOSCommandLineHelper.h"
 #include "GameLaunchDaemonMessageHandler.h"
+#include "AudioDevice.h"
 
 FEngineLoop GEngineLoop;
 FGameLaunchDaemonMessageHandler GCommandSystem;
 
+void FAppEntry::Suspend()
+{
+    GEngine->AudioDevice->SuspendContext();
+}
+
+void FAppEntry::Resume()
+{
+    GEngine->AudioDevice->ResumeContext();
+}
+
 void FAppEntry::PreInit(IOSAppDelegate* AppDelegate, UIApplication* Application)
 {
 	// make a controller object
