Better conversion of SDL keys to ASCII chars.

- Still awful, but at least handles all the (US keyboard) combinations and fixes TTP #334756.

diff --git a/Engine/Source/Runtime/Core/Private/Linux/LinuxApplication.cpp b/Engine/Source/Runtime/Core/Private/Linux/LinuxApplication.cpp
--- a/Engine/Source/Runtime/Core/Private/Linux/LinuxApplication.cpp
+++ b/Engine/Source/Runtime/Core/Private/Linux/LinuxApplication.cpp
@@ -706,22 +706,102 @@ TCHAR FLinuxApplication::ConvertChar( SDL_Keysym Keysym )
 		return 0;
 	}
 
-	if( Keysym.sym >= 97 && Keysym.sym <= 122 && (Keysym.mod & (KMOD_LSHIFT | KMOD_RSHIFT)) )
-	{
-		return Keysym.sym - 32; // Convert to uppercase
-	}
-
-	if( Keysym.sym >= 91 && Keysym.sym <= 93 && (Keysym.mod & (KMOD_LSHIFT | KMOD_RSHIFT)) )
-	{
-		return Keysym.sym + 32; // [ \ ] -> { | }
-	}
-
-	if( Keysym.sym == 96 && (Keysym.mod & (KMOD_LSHIFT | KMOD_RSHIFT)) )
-	{
-		return 126; // ` -> ~
-	}
-
-	return Keysym.sym;
+    TCHAR Char = Keysym.sym;
+
+    if (Keysym.mod & (KMOD_LSHIFT | KMOD_RSHIFT))
+    {
+        // Convert to uppercase (FIXME: what about CAPS?)
+        if( Keysym.sym >= 97 && Keysym.sym <= 122)
+        {
+            return Keysym.sym - 32;
+        }
+        else if( Keysym.sym >= 91 && Keysym.sym <= 93)
+        {
+            return Keysym.sym + 32; // [ \ ] -> { | }
+        }
+        else
+        {
+            switch(Keysym.sym)
+            {
+                case '`': // ` -> ~
+                    Char = TEXT('`');
+                    break;
+
+                case '-': // - -> _
+                    Char = TEXT('_');
+                    break;
+
+                case '=': // - -> _
+                    Char = TEXT('+');
+                    break;
+
+                case ',':
+                    Char = TEXT('<');
+                    break;
+
+                case '.':
+                    Char = TEXT('>');
+                    break;
+
+                case ';':
+                    Char = TEXT(':');
+                    break;
+
+                case '\'':
+                    Char = TEXT('\"');
+                    break;
+
+                case '/':
+                    Char = TEXT('?');
+                    break;
+
+                case '0':
+                    Char = TEXT(')');
+                    break;
+
+                case '9':
+                    Char = TEXT('(');
+                    break;
+
+                case '8':
+                    Char = TEXT('*');
+                    break;
+
+                case '7':
+                    Char = TEXT('&');
+                    break;
+
+                case '6':
+                    Char = TEXT('^');
+                    break;
+
+                case '5':
+                    Char = TEXT('%');
+                    break;
+
+                case '4':
+                    Char = TEXT('$');
+                    break;
+
+                case '3':
+                    Char = TEXT('#');
+                    break;
+
+                case '2':
+                    Char = TEXT('@');
+                    break;
+
+                case '1':
+                    Char = TEXT('!');
+                    break;
+
+                default:
+                    break;
+            }
+        }
+    }
+
+    return Char;
 }
 
 TSharedPtr< FLinuxWindow > FLinuxApplication::FindEventWindow( SDL_Event* Event )
