minor documentation updates

diff --git a/Engine/Source/Runtime/Core/Public/Delegates/DelegateSignatureImpl.inl b/Engine/Source/Runtime/Core/Public/Delegates/DelegateSignatureImpl.inl
--- a/Engine/Source/Runtime/Core/Public/Delegates/DelegateSignatureImpl.inl
+++ b/Engine/Source/Runtime/Core/Public/Delegates/DelegateSignatureImpl.inl
@@ -1113,12 +1113,13 @@ private:
 
 /**
  * Multicast delegate base class.
-	 *
+ *
  * This class implements the functionality of multicast delegates. It is templated to the function signature
  * that it is compatible with. Use the various DECLARE_MULTICAST_DELEGATE and DECLARE_EVENT macros to create
  * actual delegate types.
-	 *
- * Multicast delegates only support bindings with no return value.
+ *
+ * Multicast delegates offer no guarantees for the calling order of bound functions. As bindings get added
+ * and removed over time, the calling order may change. Only bindings without return values are supported.
  */
 template< FUNC_TEMPLATE_DECL_TYPENAME >
 class BASE_MULTICAST_DELEGATE_CLASS
@@ -1772,7 +1773,9 @@ protected:
 protected:
 
 	/**
-	 * Adds a function delegate to this multi-cast delegate's invocation list
+	 * Adds a function delegate to this multi-cast delegate's invocation list.
+	 *
+	 * This method will assert if the same function has already been bound.
 	 *
 	 * @param	InDelegate	Delegate to add
 	 */
@@ -1874,11 +1877,11 @@ protected:
 };
 
 
-	/**
+/**
  * Implements a multicast delegate.
  *
  * This class should not be instantiated directly. Use the DECLARE_MULTICAST_DELEGATE macros instead.
-	 */
+ */
 template< FUNC_TEMPLATE_DECL_TYPENAME >
 class MULTICAST_DELEGATE_CLASS : public BASE_MULTICAST_DELEGATE_CLASS< FUNC_TEMPLATE_DECL >
 {
