Delegates: Performance optimization for multicast delegates

The main idea is to defer the compaction of the invocation list in order to avoid expensive allocations and deep copies in Broadcast(). Instead of removing bound functions immediately they are merely freed and marked as deleted. The Broadcast() function already iterates over all functions anyway, so the initial check for whether compaction is needed is basically free. Only in the rare cases where compaction is needed, an additional O(n) cost applies. Callees are still able to add new functions, because they will be added to the end of the invocation list and hence ignored by Broadcast() since it is iterating in reverse order. Special guards were added to account for the possibility of re-entrant calls into Broadcast().

diff --git a/Engine/Source/Runtime/Core/Public/Delegates/DelegateSignatureImpl.inl b/Engine/Source/Runtime/Core/Public/Delegates/DelegateSignatureImpl.inl
--- a/Engine/Source/Runtime/Core/Public/Delegates/DelegateSignatureImpl.inl
+++ b/Engine/Source/Runtime/Core/Public/Delegates/DelegateSignatureImpl.inl
@@ -1147,7 +1147,6 @@ public:
 		if (DelegateInstance != nullptr)
 		{
 			AddDelegateInstance(DelegateInstance);
-			CompactInvocationList();
 		}
 	}
 
@@ -1165,7 +1164,6 @@ public:
 		if (DelegateInstance != nullptr)
 		{
 			AddDelegateInstance(DelegateInstance->CreateCopy());
-			CompactInvocationList();
 		}
 	}
 
@@ -1785,9 +1783,12 @@ protected:
 		// verify that the same function isn't already bound
 		for (IDelegateInstance* DelegateInstance : GetInvocationList())
 		{
-			// this down-cast is OK! allows for managing invocation list in the base class without requiring virtual functions
-			TDelegateInstanceInterface* DelegateInstanceInterface = (TDelegateInstanceInterface*)DelegateInstance;
-			check(!DelegateInstanceInterface->IsSameFunction(*InDelegateInstance));
+			if (DelegateInstance != nullptr)
+			{
+				// this down-cast is OK! allows for managing invocation list in the base class without requiring virtual functions
+				TDelegateInstanceInterface* DelegateInstanceInterface = (TDelegateInstanceInterface*)DelegateInstance;
+				check(!DelegateInstanceInterface->IsSameFunction(*InDelegateInstance));
+			}
 		}
 #endif
 
@@ -1797,80 +1798,63 @@ protected:
 	/**
 	 * Broadcasts this delegate to all bound objects, except to those that may have expired.
 	 *
-	 * The constness of this method is a lie, because we are removing compactable delegate instances.
+	 * The constness of this method is a lie, but it allows for broadcasting from const functions.
 	 */
 	void Broadcast( FUNC_PARAM_LIST ) const
 	{
 		const TArray<IDelegateInstance*>& InvocationList = GetInvocationList();
 
-		if (InvocationList.Num() == 0)
-		{
-			return;
-		}
-
-		// create deep copy of invocation list in case it is modified by one of the broadcast callbacks
-		TArray<TDelegateInstanceInterface*, TInlineAllocator<4>> InvocationListCopy;
-		InvocationListCopy.Reset(InvocationList.Num());
-			
-		for (IDelegateInstance* DelegateInstance : InvocationList)
+		LockInvocationList();
 		{
-			InvocationListCopy.Add(((TDelegateInstanceInterface*)DelegateInstance)->CreateCopy());
-		}
-			
-		// invoke each bound function
-		for (int32 ListIndex = InvocationListCopy.Num() - 1; ListIndex >= 0 ; --ListIndex)
-		{
-			// this down-cast is OK! allows for managing invocation list in the base class without requiring virtual functions
-			TDelegateInstanceInterface* DelegateInstanceInterface = InvocationListCopy[ListIndex];
-
-			// execute the function and remove it from the local list if it was processed
-			if (DelegateInstanceInterface->ExecuteIfSafe(FUNC_PARAM_PASSTHRU) || !DelegateInstanceInterface->IsCompactable())
+			// call bound functions in reverse order, so we ignore any instances that may be added by callees
+			for (int32 InvocationListIndex = InvocationList.Num() - 1; InvocationListIndex >= 0; --InvocationListIndex)
 			{
-				InvocationListCopy.RemoveAtSwap(ListIndex);
+				// this down-cast is OK! allows for managing invocation list in the base class without requiring virtual functions
+				TDelegateInstanceInterface* DelegateInstanceInterface = (TDelegateInstanceInterface*)InvocationList[InvocationListIndex];
+
+				if ((DelegateInstanceInterface == nullptr) || !DelegateInstanceInterface->ExecuteIfSafe(FUNC_PARAM_PASSTHRU))
+				{
+					const_cast<BASE_MULTICAST_DELEGATE_CLASS*>(this)->RequestCompaction();
+				}
 			}
 		}
+		UnlockInvocationList();
 
-		// remove any bound functions that couldn't be processed
-		for (TDelegateInstanceInterface* DelegateInstance : InvocationListCopy)
-		{
-			// search for the instance, because original list could have been modified during broadcast
-			// this is really shady, because it violates the constness of this method
-			const_cast<BASE_MULTICAST_DELEGATE_CLASS*>(this)->RemoveDelegateInstance(DelegateInstance);
-		}
+		// NOTE: The implementation of multicast delegates relies on the assumption
+		// that this is the only place where invocation list compaction takes place.
+
+		const_cast<BASE_MULTICAST_DELEGATE_CLASS*>(this)->CompactInvocationList();
 	}
 
 	/**
 	 * Removes a function from this multi-cast delegate's invocation list (performance is O(N)).
 	 *
-	 * Note that the order of the delegates may not be preserved! This method will not assert if
-	 * the given delegate instance was not found in the invocation list, because some instances
-	 * hold on to weak pointers to objects that may become invalid.
+	 * The function is not actually removed, but deleted and marked as removed.
+	 * It will be removed next time the invocation list is compacted within Broadcast().
 	 *
 	 * @param InDelegateInstance The delegate instance to remove.
 	 */
 	void RemoveDelegateInstance( const TDelegateInstanceInterface* InDelegateInstance )
 	{
-		if (InDelegateInstance->IsCompactable())
-		{
-			CompactInvocationList();
-		}
-		else
+		const TArray<IDelegateInstance*>& InvocationList = GetInvocationList();
+
+		// NOTE: We assume that this method is never called with a nullptr object, in which case the
+		//       the following algorithm would break down (it would remove the first found instance
+		//       of a matching function binding, which is not necessarily the instance we wish to remove).
+
+		for (int32 InvocationListIndex = 0; InvocationListIndex < InvocationList.Num(); ++InvocationListIndex)
 		{
-			const TArray<IDelegateInstance*> InvocationList = GetInvocationList();
+			// this down-cast is OK! allows for managing invocation list in the base class without requiring virtual functions
+			TDelegateInstanceInterface*& DelegateInstanceRef = (TDelegateInstanceInterface*&)InvocationList[InvocationListIndex];
 
-			for (int32 InvocationListIndex = 0; InvocationListIndex < InvocationList.Num(); ++InvocationListIndex)
+			// NOTE: We must do a deep compare here, not just compare delegate pointers, because multiple
+			//       delegate pointers can refer to the exact same object and method
+			if ((DelegateInstanceRef != nullptr) && DelegateInstanceRef->IsSameFunction(*InDelegateInstance))
 			{
-				// this down-cast is OK! allows for managing invocation list in the base class without requiring virtual functions
-				TDelegateInstanceInterface* DelegateInstance = (TDelegateInstanceInterface*)InvocationList[InvocationListIndex];
-
-				// NOTE: We must do a deep compare here, not just compare delegate pointers, because multiple
-				//       delegate pointers can refer to the exact same object and method
-				if (DelegateInstance->IsSameFunction(*InDelegateInstance))
-				{
-					RemoveByIndex(InvocationListIndex);
+				delete DelegateInstanceRef;
+				DelegateInstanceRef = nullptr;
 				
-					break;	// no need to continue, as we never allow the same delegate to be bound twice
-				}
+				break;	// no need to continue, as we never allow the same delegate to be bound twice
 			}
 		}
 	}
diff --git a/Engine/Source/Runtime/Core/Public/Delegates/MulticastDelegateBase.h b/Engine/Source/Runtime/Core/Public/Delegates/MulticastDelegateBase.h
--- a/Engine/Source/Runtime/Core/Public/Delegates/MulticastDelegateBase.h
+++ b/Engine/Source/Runtime/Core/Public/Delegates/MulticastDelegateBase.h
@@ -20,20 +20,16 @@ public:
 	 */
 	void Clear( )
 	{
-		for (IDelegateInstance* DelegateInstance : InvocationList)
+		for (IDelegateInstance*& DelegateInstanceRef : InvocationList)
 		{
-			delete DelegateInstance;
+			delete DelegateInstanceRef;
+			DelegateInstanceRef = nullptr;
 		}
-
-		InvocationList.Empty();
-		ExpirableObjectCount = 0;
 	}
 
 	/**
 	 * Checks to see if any functions are bound to this multi-cast delegate.
 	 *
-	 * This method is protected to avoid assumptions about only one delegate existing.
-	 *
 	 * @return true if any functions are bound, false otherwise.
 	 */
 	inline bool IsBound( ) const
@@ -50,7 +46,7 @@ public:
 	{
 		for (IDelegateInstance* DelegateInstance : InvocationList)
 		{
-			if (DelegateInstance->HasSameObject(InUserObject))
+			if ((DelegateInstance != nullptr) && DelegateInstance->HasSameObject(InUserObject))
 			{
 				return true;
 			}
@@ -67,8 +63,16 @@ public:
 	 */
 	void RemoveAll( const void* InUserObject )
 	{
-		RemoveByObject(InUserObject);
-		CompactInvocationList();
+		for (int32 InvocationListIndex = InvocationList.Num() - 1; InvocationListIndex >= 0; --InvocationListIndex)
+		{
+			IDelegateInstance*& DelegateInstanceRef = InvocationList[InvocationListIndex];
+
+			if ((DelegateInstanceRef != nullptr) && DelegateInstanceRef->HasSameObject(InUserObject))
+			{
+				delete DelegateInstanceRef;
+				DelegateInstanceRef = nullptr;
+			}
+		}
 	}
 
 protected:
@@ -77,7 +81,8 @@ protected:
 	 * Hidden default constructor.
 	 */
 	inline FMulticastDelegateBase( )
-		: ExpirableObjectCount(0)
+		: InvocationListLockCount(0)
+		, NeedsCompaction(false)
 	{ }
 
 protected:
@@ -87,46 +92,39 @@ protected:
 	 *
 	 * @param DelegateInstance The delegate instance to add.
 	 */
-	void AddInternal( IDelegateInstance* DelegateInstance )
+	inline void AddInternal( IDelegateInstance* DelegateInstance )
 	{
 		InvocationList.Add(DelegateInstance);
-
-		// Keep track of whether we have objects bound that may expire out from underneath us.
-		const EDelegateInstanceType::Type DelegateType = DelegateInstance->GetType();
-
-		if ((DelegateType == EDelegateInstanceType::SharedPointerMethod) ||
-			(DelegateType == EDelegateInstanceType::ThreadSafeSharedPointerMethod) ||
-			(DelegateType == EDelegateInstanceType::UObjectMethod))
-		{
-			++ExpirableObjectCount;
-		}
 	}
 
 	/**
-	 * Cleans up any delegates in our invocation list that have expired (performance is O(N))
+	 * Removes any expired or deleted functions from the invocation list.
+	 *
+	 * @see RequestCompaction
 	 */
-	void CompactInvocationList( ) const
+	void CompactInvocationList( )
 	{
-		// We only need to compact if any objects were added whose lifetime is known to us
-		if (ExpirableObjectCount == 0)
+		if (!NeedsCompaction || (InvocationListLockCount != 0))
 		{
 			return;
 		}
 
-		// remove expired delegates
 		for (int32 InvocationListIndex = InvocationList.Num() - 1; InvocationListIndex >= 0; --InvocationListIndex)
 		{
 			IDelegateInstance* DelegateInstance = InvocationList[InvocationListIndex];
 
-			if (DelegateInstance->IsCompactable())
+			if ((DelegateInstance == nullptr) || DelegateInstance->IsCompactable())
 			{
 				InvocationList.RemoveAtSwap(InvocationListIndex);
-				delete DelegateInstance;
 
-				checkSlow(ExpirableObjectCount > 0);
-				--ExpirableObjectCount;
+				if (DelegateInstance != nullptr)
+				{
+					delete DelegateInstance;
+				}
 			}
 		}
+
+		NeedsCompaction = false;
 	}
 
 	/**
@@ -134,61 +132,45 @@ protected:
 	 *
 	 * @return The invocation list.
 	 */
-	const TArray<IDelegateInstance*>& GetInvocationList( ) const
+	inline const TArray<IDelegateInstance*>& GetInvocationList( ) const
 	{
 		return InvocationList;
 	}
 
 	/**
-	 * Removes all functions from this multi-cast delegate's invocation list that are bound to the specified UserObject.
-	 *
-	 * Note that the order of the delegates may not be preserved!
+	 * Requests that the invocation list be compacted.
 	 *
-	 * @param InUserObject The object to remove the delegates for.
+	 * @see CompactInvocationList
 	 */
-	void RemoveByObject( const void* InUserObject )
+	inline void RequestCompaction( )
 	{
-		for (int32 InvocationListIndex = InvocationList.Num() - 1; InvocationListIndex >= 0; --InvocationListIndex)
-		{
-			IDelegateInstance* DelegateInstance = InvocationList[InvocationListIndex];
-
-			if (DelegateInstance->HasSameObject(InUserObject))
-			{
-				RemoveByIndex(InvocationListIndex);
-			}
-		}
+		NeedsCompaction = true;
 	}
 
 	/**
-	 * Removes a delegate instance given its index in the invocation list.
-	 *
-	 * @param InvocationListIndex The index of the delegate instance to remove.
+	 * Increments the lock counter for the invocation list.
 	 */
-	void RemoveByIndex( int32 InvocationListIndex )
+	inline void LockInvocationList( ) const
 	{
-		IDelegateInstance* DelegateInstance = InvocationList[InvocationListIndex];
-
-		// keep track of whether we have objects bound that may expire out from underneath us
-		const EDelegateInstanceType::Type DelegateType = DelegateInstance->GetType();
-
-		if ((DelegateType == EDelegateInstanceType::SharedPointerMethod) || 
-			(DelegateType == EDelegateInstanceType::ThreadSafeSharedPointerMethod) ||
-			(DelegateType == EDelegateInstanceType::UFunction) ||
-			(DelegateType == EDelegateInstanceType::UObjectMethod))
-		{
-			--ExpirableObjectCount;
-		}
+		++InvocationListLockCount;
+	}
 
-		InvocationList.RemoveAtSwap(InvocationListIndex);
-		delete DelegateInstance;
+	/**
+	 * Decrements the lock counter for the invocation list.
+	 */
+	inline void UnlockInvocationList( ) const
+	{
+		--InvocationListLockCount;
 	}
 
 private:
 
-	// Mutable so that we can housekeep list even with const broadcasts
-	mutable int32 ExpirableObjectCount;
-
 	// Holds the collection of delegate instances to invoke.
-	// Mutable so that we can do housekeeping during const broadcasts.
-	mutable TArray<IDelegateInstance*> InvocationList;
+	TArray<IDelegateInstance*> InvocationList;
+
+	// Holds a lock counter for the invocation list.
+	mutable int32 InvocationListLockCount;
+
+	// Holds a flag indicating whether the invocation list needs to be compacted.
+	bool NeedsCompaction;
 };
