#ttp 331740 UE4: Core: Runnable thread clean-up
#proj core
#branch UE4
#summary Remove all bAutoDeleteSelf and bAutoDeleteRunnable usage from the runnable.

diff --git a/Engine/Plugins/Messaging/UdpMessaging/Source/UdpMessaging/Private/Transport/UdpMessageBeacon.cpp b/Engine/Plugins/Messaging/UdpMessaging/Source/UdpMessaging/Private/Transport/UdpMessageBeacon.cpp
--- a/Engine/Plugins/Messaging/UdpMessaging/Source/UdpMessaging/Private/Transport/UdpMessageBeacon.cpp
+++ b/Engine/Plugins/Messaging/UdpMessaging/Source/UdpMessaging/Private/Transport/UdpMessageBeacon.cpp
@@ -28,7 +28,7 @@ FUdpMessageBeacon::FUdpMessageBeacon(FSocket* InSocket, const FGuid& InSocketId,
 	EndpointLeftEvent = FPlatformProcess::CreateSynchEvent();
 	MulticastAddress = InMulticastEndpoint.ToInternetAddr();
 
-	Thread = FRunnableThread::Create(this, TEXT("FUdpMessageBeacon"), false, false, 128 * 1024, TPri_AboveNormal);
+	Thread = FRunnableThread::Create(this, TEXT("FUdpMessageBeacon"), 128 * 1024, TPri_AboveNormal);
 }
 
 
diff --git a/Engine/Plugins/Messaging/UdpMessaging/Source/UdpMessaging/Private/Transport/UdpMessageProcessor.cpp b/Engine/Plugins/Messaging/UdpMessaging/Source/UdpMessaging/Private/Transport/UdpMessageProcessor.cpp
--- a/Engine/Plugins/Messaging/UdpMessaging/Source/UdpMessaging/Private/Transport/UdpMessageProcessor.cpp
+++ b/Engine/Plugins/Messaging/UdpMessaging/Source/UdpMessaging/Private/Transport/UdpMessageProcessor.cpp
@@ -26,7 +26,7 @@ FUdpMessageProcessor::FUdpMessageProcessor( FSocket* InSocket, const FGuid& InNo
 	, Stopping(false)
 {
 	WorkEvent = FPlatformProcess::CreateSynchEvent();
-	Thread = FRunnableThread::Create(this, TEXT("FUdpMessageProcessor"), false, false, 128 * 1024, TPri_AboveNormal);
+	Thread = FRunnableThread::Create(this, TEXT("FUdpMessageProcessor"), 128 * 1024, TPri_AboveNormal);
 
 	const UUdpMessagingSettings& Settings = *GetDefault<UUdpMessagingSettings>();
 
diff --git a/Engine/Plugins/Messaging/UdpMessaging/Source/UdpMessaging/Private/Tunnel/UdpMessageTunnel.h b/Engine/Plugins/Messaging/UdpMessaging/Source/UdpMessaging/Private/Tunnel/UdpMessageTunnel.h
--- a/Engine/Plugins/Messaging/UdpMessaging/Source/UdpMessaging/Private/Tunnel/UdpMessageTunnel.h
+++ b/Engine/Plugins/Messaging/UdpMessaging/Source/UdpMessaging/Private/Tunnel/UdpMessageTunnel.h
@@ -60,7 +60,7 @@ public:
 		MulticastSocket->SetReceiveBufferSize(2 * 1024 * 1024, NewSize);
 		UnicastSocket->SetReceiveBufferSize(2 * 1024 * 1024, NewSize);
 
-		Thread = FRunnableThread::Create(this, TEXT("FUdpMessageTunnel"), false, false, 128 * 1024, TPri_AboveNormal);
+		Thread = FRunnableThread::Create(this, TEXT("FUdpMessageTunnel"), 128 * 1024, TPri_AboveNormal);
 	}
 
 	/**
diff --git a/Engine/Plugins/Messaging/UdpMessaging/Source/UdpMessaging/Private/Tunnel/UdpMessageTunnelConnection.h b/Engine/Plugins/Messaging/UdpMessaging/Source/UdpMessaging/Private/Tunnel/UdpMessageTunnelConnection.h
--- a/Engine/Plugins/Messaging/UdpMessaging/Source/UdpMessaging/Private/Tunnel/UdpMessageTunnelConnection.h
+++ b/Engine/Plugins/Messaging/UdpMessaging/Source/UdpMessaging/Private/Tunnel/UdpMessageTunnelConnection.h
@@ -59,7 +59,7 @@ public:
 		int32 NewSize = 0;
 		Socket->SetReceiveBufferSize(2*1024*1024, NewSize);
 		Socket->SetSendBufferSize(2*1024*1024, NewSize);
-		Thread = FRunnableThread::Create(this, TEXT("FUdpMessageTunnelConnection"), false, false, 128 * 1024, TPri_Normal);
+		Thread = FRunnableThread::Create(this, TEXT("FUdpMessageTunnelConnection"), 128 * 1024, TPri_Normal);
 	}
 
 	/**
diff --git a/Engine/Source/Developer/Android/AndroidDeviceDetection/Private/AndroidDeviceDetectionModule.cpp b/Engine/Source/Developer/Android/AndroidDeviceDetection/Private/AndroidDeviceDetectionModule.cpp
--- a/Engine/Source/Developer/Android/AndroidDeviceDetection/Private/AndroidDeviceDetectionModule.cpp
+++ b/Engine/Source/Developer/Android/AndroidDeviceDetection/Private/AndroidDeviceDetectionModule.cpp
@@ -231,7 +231,7 @@ public:
 
 		// create and fire off our device detection thread
 		DetectionThreadRunnable = new FAndroidDeviceDetectionRunnable(ADBPath, DeviceMap, &DeviceMapLock);
-		DetectionThread = FRunnableThread::Create(DetectionThreadRunnable, TEXT("FAndroidDeviceDetectionRunnable"), false);
+		DetectionThread = FRunnableThread::Create(DetectionThreadRunnable, TEXT("FAndroidDeviceDetectionRunnable"));
 	}
 
 	virtual ~FAndroidDeviceDetection()
diff --git a/Engine/Source/Developer/DerivedDataCache/Private/DDCCleanup.cpp b/Engine/Source/Developer/DerivedDataCache/Private/DDCCleanup.cpp
--- a/Engine/Source/Developer/DerivedDataCache/Private/DDCCleanup.cpp
+++ b/Engine/Source/Developer/DerivedDataCache/Private/DDCCleanup.cpp
@@ -55,9 +55,7 @@ FDDCCleanup::FDDCCleanup()
 	, StopTaskCounter(0)
 {
 	// Don't delete the runnable automatically. It's going to be manually deleted in FDDCCleanup::Shutdown.
-	const bool bAutoDeleteSelf = false;
-	const bool bAutoDeleteRunnable = false;
-	Thread = FRunnableThread::Create( this, TEXT("FDDCCleanup"), bAutoDeleteSelf, bAutoDeleteRunnable, 0, TPri_BelowNormal );
+	Thread = FRunnableThread::Create(this, TEXT("FDDCCleanup"), 0, TPri_BelowNormal);
 }
 
 FDDCCleanup::~FDDCCleanup()
diff --git a/Engine/Source/Developer/LauncherServices/Private/Launcher/Launcher.cpp b/Engine/Source/Developer/LauncherServices/Private/Launcher/Launcher.cpp
--- a/Engine/Source/Developer/LauncherServices/Private/Launcher/Launcher.cpp
+++ b/Engine/Source/Developer/LauncherServices/Private/Launcher/Launcher.cpp
@@ -16,7 +16,7 @@ ILauncherWorkerPtr FLauncher::Launch( const ITargetDeviceProxyManagerRef& Device
 	{
 		FLauncherWorker* LauncherWorker = new FLauncherWorker(DeviceProxyManager, Profile);
 
-		if ((LauncherWorker != NULL) && (FRunnableThread::Create(LauncherWorker, TEXT("LauncherWorker"), false, false, 0, TPri_Normal) != NULL))
+		if ((LauncherWorker != NULL) && (FRunnableThread::Create(LauncherWorker, TEXT("LauncherWorker")) != NULL))
 		{
 			return MakeShareable(LauncherWorker);
 		}			
diff --git a/Engine/Source/Developer/LauncherServices/Private/Launcher/LauncherTask.h b/Engine/Source/Developer/LauncherServices/Private/Launcher/LauncherTask.h
--- a/Engine/Source/Developer/LauncherServices/Private/Launcher/LauncherTask.h
+++ b/Engine/Source/Developer/LauncherServices/Private/Launcher/LauncherTask.h
@@ -59,7 +59,7 @@ public:
 
 		LocalChainState = ChainState;
 
-		Thread = FRunnableThread::Create(this, TEXT("FLauncherTask"), false, false, 0, TPri_Normal);
+		Thread = FRunnableThread::Create(this, TEXT("FLauncherTask"));
 
 		if (Thread == NULL)
 		{
diff --git a/Engine/Source/Developer/ProfilerService/Private/ProfilerServiceManager.cpp b/Engine/Source/Developer/ProfilerService/Private/ProfilerServiceManager.cpp
--- a/Engine/Source/Developer/ProfilerService/Private/ProfilerServiceManager.cpp
+++ b/Engine/Source/Developer/ProfilerService/Private/ProfilerServiceManager.cpp
@@ -32,7 +32,7 @@ public:
 		, MessageEndpoint( InMessageEndpoint )
 		, StopTaskCounter( 0 )
 	{
-		Runnable = FRunnableThread::Create( this, TEXT("FFileTransferRunnable"), false, false, 128 * 1024, TPri_BelowNormal );
+		Runnable = FRunnableThread::Create(this, TEXT("FFileTransferRunnable"), 128 * 1024, TPri_BelowNormal);
 	}
 
 	/** Destructor. */
diff --git a/Engine/Source/Developer/iOS/IOSTargetPlatform/Private/Windows/IOSDeviceHelperWindows.cpp b/Engine/Source/Developer/iOS/IOSTargetPlatform/Private/Windows/IOSDeviceHelperWindows.cpp
--- a/Engine/Source/Developer/iOS/IOSTargetPlatform/Private/Windows/IOSDeviceHelperWindows.cpp
+++ b/Engine/Source/Developer/iOS/IOSTargetPlatform/Private/Windows/IOSDeviceHelperWindows.cpp
@@ -258,7 +258,7 @@ void FIOSDeviceHelper::Initialize()
 	// kick off a thread to query for connected devices
 	QueryTask = new FDeviceQueryTask();
 	QueryTask->OnDeviceNotification().AddStatic(FIOSDeviceHelper::DeviceCallback);
-	QueryThread = FRunnableThread::Create(QueryTask, TEXT("FIOSDeviceHelper.QueryTask"), false, false, 128 * 1024, TPri_Normal);
+	QueryThread = FRunnableThread::Create(QueryTask, TEXT("FIOSDeviceHelper.QueryTask"), 128 * 1024, TPri_Normal);
 #endif
 }
 
diff --git a/Engine/Source/Editor/TaskBrowser/Private/TaskDatabase.cpp b/Engine/Source/Editor/TaskBrowser/Private/TaskDatabase.cpp
--- a/Engine/Source/Editor/TaskBrowser/Private/TaskDatabase.cpp
+++ b/Engine/Source/Editor/TaskBrowser/Private/TaskDatabase.cpp
@@ -93,9 +93,7 @@ namespace TaskDatabaseSystem
 		// Launch the task database thread!
 		GTaskDatabaseThread = FRunnableThread::Create(
 			GTaskDatabaseThreadRunnable,		// Runnable module
-			TEXT( "TaskDatabaseThread" ),		// Thread name
-			false,								// Auto-delete self?
-			false,								// Auto-delete runnable?
+			TEXT("TaskDatabaseThread"),			// Thread name
 			0,									// Stack size
 			TaskDatabaseThreadPriority );		// Thread priority
 
diff --git a/Engine/Source/Editor/UnrealEd/Private/Dialogs/DlgSoundWaveOptions.cpp b/Engine/Source/Editor/UnrealEd/Private/Dialogs/DlgSoundWaveOptions.cpp
--- a/Engine/Source/Editor/UnrealEd/Private/Dialogs/DlgSoundWaveOptions.cpp
+++ b/Engine/Source/Editor/UnrealEd/Private/Dialogs/DlgSoundWaveOptions.cpp
@@ -388,8 +388,7 @@ void SSoundWaveCompressionOptions::CreateCompressedWaves( void )
 
 	SoundPreviewThreadRunnable = new FSoundPreviewThread( Count, SoundWave, PreviewInfo );
 
-	EThreadPriority SoundPreviewThreadPrio = TPri_Normal;
-	SoundPreviewThread = FRunnableThread::Create( SoundPreviewThreadRunnable, TEXT( "SoundPreviewThread" ), 0, 0, 0, SoundPreviewThreadPrio );
+	SoundPreviewThread = FRunnableThread::Create(SoundPreviewThreadRunnable, TEXT("SoundPreviewThread"));
 }
 
 TOptional<float> SSoundWaveCompressionOptions::GetProgress() const
diff --git a/Engine/Source/Runtime/AssetRegistry/Private/AssetDataGatherer.cpp b/Engine/Source/Runtime/AssetRegistry/Private/AssetDataGatherer.cpp
--- a/Engine/Source/Runtime/AssetRegistry/Private/AssetDataGatherer.cpp
+++ b/Engine/Source/Runtime/AssetRegistry/Private/AssetDataGatherer.cpp
@@ -31,7 +31,7 @@ FAssetDataGatherer::FAssetDataGatherer(const TArray<FString>& InPaths, bool bInI
 	}
 	else
 	{
-		Thread = FRunnableThread::Create( this, TEXT("FAssetDataGatherer"), false, false, 0, TPri_BelowNormal );
+		Thread = FRunnableThread::Create(this, TEXT("FAssetDataGatherer"), 0, TPri_BelowNormal);
 	}
 }
 
diff --git a/Engine/Source/Runtime/Core/Private/Async/TaskGraph.cpp b/Engine/Source/Runtime/Core/Private/Async/TaskGraph.cpp
--- a/Engine/Source/Runtime/Core/Private/Async/TaskGraph.cpp
+++ b/Engine/Source/Runtime/Core/Private/Async/TaskGraph.cpp
@@ -675,7 +675,7 @@ public:
 		{
 			FString Name = FString::Printf(TEXT("TaskGraphThread %d"), ThreadIndex - (LastExternalThread + 1));
 			uint32 StackSize = 256 * 1024;
-			WorkerThreads[ThreadIndex].RunnableThread = FRunnableThread::Create( &Thread(ThreadIndex), *Name, false, false, StackSize, TPri_Normal ); // these are below normal threads? so that they sleep when the named threads are active
+			WorkerThreads[ThreadIndex].RunnableThread = FRunnableThread::Create(&Thread(ThreadIndex), *Name, StackSize, TPri_Normal); // these are below normal threads? so that they sleep when the named threads are active
 			WorkerThreads[ThreadIndex].bAttached = true;
 		}
 	}
diff --git a/Engine/Source/Runtime/Core/Private/HAL/PThreadRunnableThread.h b/Engine/Source/Runtime/Core/Private/HAL/PThreadRunnableThread.h
--- a/Engine/Source/Runtime/Core/Private/HAL/PThreadRunnableThread.h
+++ b/Engine/Source/Runtime/Core/Private/HAL/PThreadRunnableThread.h
@@ -34,26 +34,6 @@ protected:
 	 */
 	FEvent* ThreadInitSyncEvent;
 
-	/** 
-	 * Sync event to make sure that CreateInteral() has been completed before allowing the thread to be auto-deleted
-	 */
-	FEvent* ThreadCreatedSyncEvent;
-
-	/** 
-	 * Flag used when the thread is waiting for the caller to finish setting it up before it can delete itself.
-	 */
-	FThreadSafeCounter WantsToDeleteSelf;
-
-	/**
-	 * Whether we should delete ourselves on thread exit
-	 */
-	bool bShouldDeleteSelf;
-
-	/**
-	 * Whether we should delete the runnable on thread exit
-	 */
-	bool bShouldDeleteRunnable;
-
 	/**
 	 * The priority to run the thread at
 	 */
@@ -226,16 +206,6 @@ protected:
 	 */
 	virtual void PostRun()
 	{
-		if (bShouldDeleteSelf == true)
-		{
-			// Make sure the caller knows we want to delete this thread if we're still int CreateInternal.
-			WantsToDeleteSelf.Increment();
-			// Wait until the caller has finished setting up this thread in case Runnable execution was very short.
-			ThreadCreatedSyncEvent->Wait();
-			// Now clean up the thread handle so we don't leak
-			Thread = PTHREAD_NULL;
-			delete this;
-		}
 	}
 
 	/**
@@ -265,12 +235,6 @@ protected:
 			ThreadInitSyncEvent->Trigger();
 		}
 
-		// Should we delete the runnable?
-		if (bShouldDeleteRunnable == true)
-		{
-			delete Runnable;
-			Runnable = NULL;
-		}
 		// Clean ourselves up without waiting
 		ThreadIsRunning = false;
 		return ExitCode;
@@ -281,10 +245,6 @@ public:
 		: Thread(PTHREAD_NULL)
 		, Runnable(NULL)
 		, ThreadInitSyncEvent(NULL)
-		, ThreadCreatedSyncEvent(NULL)
-		, WantsToDeleteSelf(0)
-		, bShouldDeleteSelf(false)
-		, bShouldDeleteRunnable(false)
 		, ThreadPriority(TPri_Normal)
 		, ThreadID(0)
 		, ThreadIsRunning(false)
@@ -300,7 +260,6 @@ public:
 		}
 		FRunnableThread::GetThreadRegistry().Remove(ThreadID);
 		ThreadID = 0;
-		delete ThreadCreatedSyncEvent;
 	}
 
 	virtual void SetThreadPriority(EThreadPriority NewPriority) OVERRIDE
@@ -342,18 +301,6 @@ public:
 
 		Thread = PTHREAD_NULL;
 
-		// delete the runnable if requested and we didn't shut down gracefully already.
-		if (Runnable && bShouldDeleteRunnable == true)
-		{
-			delete Runnable;
-			Runnable = NULL;
-		}
-		// Delete ourselves if requested and we didn't shut down gracefully already.
-		// This check prevents a double-delete of self when we shut down gracefully.
-		if (!bDidExitOK && bShouldDeleteSelf == true)
-		{
-			delete this;
-		}
 		return bDidExitOK;
 	}
 
@@ -383,18 +330,14 @@ public:
 protected:
 
 	virtual bool CreateInternal(FRunnable* InRunnable, const TCHAR* InThreadName,
-		bool bAutoDeleteSelf = 0,bool bAutoDeleteRunnable = 0,uint32 InStackSize = 0,
+		uint32 InStackSize = 0,
 		EThreadPriority InThreadPri = TPri_Normal, uint64 InThreadAffinityMask = 0) OVERRIDE
 	{
 		check(InRunnable);
 		Runnable = InRunnable;
-		bShouldDeleteSelf = bAutoDeleteSelf;
-		bShouldDeleteRunnable = bAutoDeleteRunnable;
 
 		// Create a sync event to guarantee the Init() function is called first
 		ThreadInitSyncEvent	= FPlatformProcess::CreateSynchEvent(true);
-		// Create a sync event to guarantee the thread will not delete itself until it has been fully set up.
-		ThreadCreatedSyncEvent = FPlatformProcess::CreateSynchEvent(true);
 		// A name for the thread in for debug purposes. _ThreadProc will set it.
 		ThreadName = InThreadName ? InThreadName : TEXT("Unnamed UE4");
 		// Create the new thread
@@ -415,10 +358,6 @@ protected:
 		}
 		else // If it fails, clear all the vars
 		{
-			if (bAutoDeleteRunnable == true)
-			{
-				delete InRunnable;
-			}
 			Runnable = NULL;
 		}
 
@@ -427,13 +366,5 @@ protected:
 		ThreadInitSyncEvent = NULL;
 		return Thread != PTHREAD_NULL;
 	}
-
-	virtual bool NotifyCreated() OVERRIDE
-	{
-		const bool bHasFinished = !!WantsToDeleteSelf.GetValue();
-		// It's ok to delete this thread if it wants to delete self.
-		ThreadCreatedSyncEvent->Trigger();
-		return bHasFinished;
-	}
 };
 
diff --git a/Engine/Source/Runtime/Core/Private/HAL/ThreadingBase.cpp b/Engine/Source/Runtime/Core/Private/HAL/ThreadingBase.cpp
--- a/Engine/Source/Runtime/Core/Private/HAL/ThreadingBase.cpp
+++ b/Engine/Source/Runtime/Core/Private/HAL/ThreadingBase.cpp
@@ -35,10 +35,6 @@ class FFakeThread : public FRunnableThread
 
 	/** Thread is suspended. */
 	bool bIsSuspended;
-	/** Thread should delete itself. */
-	bool bDeleteSelf;
-	/** Thread should automatically delete the runnable associated with it. */
-	bool bDeleteRunnable;
 	/** Name of this thread. */
 	FString Name;
 	/** Runnable object associated with this thread. */
@@ -53,8 +49,6 @@ public:
 	 */
 	FFakeThread()
 		: bIsSuspended(false)
-		, bDeleteSelf(false)
-		, bDeleteRunnable(false)
 		, Runnable(NULL)
 		, ThreadId(ThreadIdCounter++)
 	{
@@ -96,15 +90,6 @@ public:
 	virtual bool Kill(bool bShouldWait)
 	{
 		FSingleThreadManager::Get().RemoveThread(this);
-		if (bDeleteRunnable)
-		{
-			delete Runnable;
-			Runnable = NULL;
-		}
-		if (bDeleteSelf)
-		{
-			delete this;
-		}
 		return true;
 	}
 
@@ -124,7 +109,7 @@ public:
 	}
 
 	virtual bool CreateInternal(FRunnable* InRunnable, const TCHAR* ThreadName,
-		bool bAutoDeleteSelf,bool bAutoDeleteRunnable,uint32 InStackSize,
+		uint32 InStackSize,
 		EThreadPriority InThreadPri, uint64 InThreadAffinityMask)
 
 	{
@@ -136,11 +121,6 @@ public:
 		return Runnable != NULL;
 	}
 
-	virtual bool NotifyCreated() OVERRIDE
-	{
-		return false;
-	}
-
 	/**	 
 	 * Change the thread processor affinity
 	 *
@@ -222,8 +202,6 @@ FRunnableThread::~FRunnableThread()
 FRunnableThread* FRunnableThread::Create(
 	class FRunnable* InRunnable, 
 	const TCHAR* ThreadName, 
-	bool bAutoDeleteSelf,
-	bool bAutoDeleteRunnable, 
 	uint32 InStackSize,
 	EThreadPriority InThreadPri,
 	uint64 InThreadAffinityMask)
@@ -242,7 +220,7 @@ FRunnableThread* FRunnableThread::Create(
 			}
 
 			// Call the thread's create method
-			if (NewThread->CreateInternal(InRunnable,ThreadName,bAutoDeleteSelf,bAutoDeleteRunnable,InStackSize,InThreadPri,InThreadAffinityMask) == false)
+			if (NewThread->CreateInternal(InRunnable,ThreadName,InStackSize,InThreadPri,InThreadAffinityMask) == false)
 			{
 				// We failed to start the thread correctly so clean up
 				delete NewThread;
@@ -254,18 +232,13 @@ FRunnableThread* FRunnableThread::Create(
 	{
 		// Create a fake thread when multithreading is disabled.
 		NewThread = new FFakeThread();
-		if (NewThread->CreateInternal(InRunnable,ThreadName,bAutoDeleteSelf,bAutoDeleteRunnable,InStackSize,InThreadPri) == false)
+		if (NewThread->CreateInternal(InRunnable,ThreadName,InStackSize,InThreadPri) == false)
 		{
 			// We failed to start the thread correctly so clean up
 			delete NewThread;
 			NewThread = NULL;
 		}
 	}
-	if (NewThread && NewThread->NotifyCreated())
-	{
-		// The thread was marked to delete itself and it already has.
-		NewThread = NULL;
-	}
 
 	if( NewThread )
 	{
@@ -368,7 +341,7 @@ public:
 
 		OwningThreadPool = InPool;
 		DoWorkEvent = FPlatformProcess::CreateSynchEvent();
-		Thread = FRunnableThread::Create( this, *PoolThreadName, false, false, InStackSize, ThreadPriority );
+		Thread = FRunnableThread::Create(this, *PoolThreadName, InStackSize, ThreadPriority);
 		check(Thread);
 		return true;
 	}
diff --git a/Engine/Source/Runtime/Core/Private/Misc/MonitoredProcess.cpp b/Engine/Source/Runtime/Core/Private/Misc/MonitoredProcess.cpp
--- a/Engine/Source/Runtime/Core/Private/Misc/MonitoredProcess.cpp
+++ b/Engine/Source/Runtime/Core/Private/Misc/MonitoredProcess.cpp
@@ -68,7 +68,7 @@ bool FMonitoredProcess::Launch( )
 		return false;
 	}
 
-	Thread = FRunnableThread::Create(this, TEXT("FMonitoredProcess"), false, false, 128 * 1024, TPri_AboveNormal);
+	Thread = FRunnableThread::Create(this, TEXT("FMonitoredProcess"), 128 * 1024, TPri_AboveNormal);
 
 	return true;
 }
diff --git a/Engine/Source/Runtime/Core/Private/Serialization/AsyncIOSystemBase.cpp b/Engine/Source/Runtime/Core/Private/Serialization/AsyncIOSystemBase.cpp
--- a/Engine/Source/Runtime/Core/Private/Serialization/AsyncIOSystemBase.cpp
+++ b/Engine/Source/Runtime/Core/Private/Serialization/AsyncIOSystemBase.cpp
@@ -762,7 +762,7 @@ FIOSystem& FIOSystem::Get()
 			// the platform didn't have a specific need, so we just use the base class with the normal file system.
 			AsyncIOSystem = new FAsyncIOSystemBase(FPlatformFileManager::Get().GetPlatformFile());
 		}
-		AsyncIOThread = FRunnableThread::Create(AsyncIOSystem, TEXT("AsyncIOSystem"), 0, 0, 16384, TPri_AboveNormal );
+		AsyncIOThread = FRunnableThread::Create(AsyncIOSystem, TEXT("AsyncIOSystem"), 16384, TPri_AboveNormal);
 		check(AsyncIOThread);
 	}
 	check(AsyncIOSystem);
diff --git a/Engine/Source/Runtime/Core/Private/Stats/Stats2.cpp b/Engine/Source/Runtime/Core/Private/Stats/Stats2.cpp
--- a/Engine/Source/Runtime/Core/Private/Stats/Stats2.cpp
+++ b/Engine/Source/Runtime/Core/Private/Stats/Stats2.cpp
@@ -648,7 +648,7 @@ public:
 
 	void Start()
 	{
-		Thread = FRunnableThread::Create(this, TEXT("StatsThread"), 0, 0, 128 * 1024, TPri_BelowNormal);
+		Thread = FRunnableThread::Create(this, TEXT("StatsThread"), 128 * 1024, TPri_BelowNormal);
 		check(Thread != NULL);
 	}
 
diff --git a/Engine/Source/Runtime/Core/Private/WinRT/WinRTRunnableThread.h b/Engine/Source/Runtime/Core/Private/WinRT/WinRTRunnableThread.h
--- a/Engine/Source/Runtime/Core/Private/WinRT/WinRTRunnableThread.h
+++ b/Engine/Source/Runtime/Core/Private/WinRT/WinRTRunnableThread.h
@@ -23,26 +23,6 @@ class FRunnableThreadWinRT : public FRunnableThread
 	 */
 	FEvent* ThreadInitSyncEvent;
 
-	/** 
-	 * Sync event to make sure that CreateInteral() has been completed before allowing the thread to be auto-deleted
-	 */
-	FEvent* ThreadCreatedSyncEvent;
-
-	/** 
-	 * Flag used when the thread is waiting for the caller to finish setting it up before it can delete itself.
-	 */
-	FThreadSafeCounter WantsToDeleteSelf;
-
-	/**
-	 * Whether we should delete ourselves on thread exit
-	 */
-	bool bShouldDeleteSelf;
-
-	/**
-	 * Whether we should delete the runnable on thread exit
-	 */
-	bool bShouldDeleteRunnable;
-
 	/**
 	 * The priority to run the thread at
 	 */
@@ -141,24 +121,6 @@ class FRunnableThreadWinRT : public FRunnableThread
 			ThreadInitSyncEvent->Trigger();
 		}
 
-		// Should we delete the runnable?
-		if (bShouldDeleteRunnable == true)
-		{
-			delete Runnable;
-			Runnable = NULL;
-		}
-		// Clean ourselves up without waiting
-		if (bShouldDeleteSelf == true)
-		{
-			// Make sure the caller knows we want to delete this thread if we're still int CreateInternal.
-			WantsToDeleteSelf.Increment();
-			// Wait until the caller has finished setting up this thread in case Runnable execution was very short.
-			ThreadCreatedSyncEvent->Wait();
-			// Now clean up the thread handle so we don't leak
-			CloseHandle(Thread);
-			Thread = NULL;
-			delete this;
-		}
 		return ExitCode;
 	}
 
@@ -167,12 +129,8 @@ public:
 		: Thread(NULL)
 		, Runnable(NULL)
 		, ThreadInitSyncEvent(NULL)
-		, ThreadCreatedSyncEvent(NULL)
-		, WantsToDeleteSelf(0)
-		, bShouldDeleteSelf(false)
-		, bShouldDeleteRunnable(false)
 		, ThreadPriority(TPri_Normal)
-		, ThreadID(NULL)				
+		, ThreadID(NULL)
 	{
 	}
 
@@ -184,7 +142,6 @@ public:
 			Kill(true);
 		}
 		FRunnableThread::GetThreadRegistry().Remove(ThreadID);
-		delete ThreadCreatedSyncEvent;
 	}
 	
 	virtual void SetThreadPriority(EThreadPriority NewPriority) OVERRIDE
@@ -244,18 +201,7 @@ public:
 		// Now clean up the thread handle so we don't leak
 		CloseHandle(Thread);
 		Thread = NULL;
-		// delete the runnable if requested and we didn't shut down gracefully already.
-		if (Runnable && bShouldDeleteRunnable == true)
-		{
-			delete Runnable;
-			Runnable = NULL;
-		}
-		// Delete ourselves if requested and we didn't shut down gracefully already.
-		// This check prevents a double-delete of self when we shut down gracefully.
-		if (!bDidExitOK && bShouldDeleteSelf == true)
-		{
-			delete this;
-		}
+
 		return bDidExitOK;
 	}
 
@@ -278,29 +224,21 @@ public:
 
 protected:
 	virtual bool CreateInternal(FRunnable* InRunnable, const TCHAR* InThreadName,
-		bool bAutoDeleteSelf = 0,bool bAutoDeleteRunnable = 0,uint32 InStackSize = 0,
+		uint32 InStackSize = 0,
 		EThreadPriority InThreadPri = TPri_Normal, uint64 InThreadAffinityMask = 0) OVERRIDE
 	{
 		check(InRunnable);
 		Runnable = InRunnable;
-		bShouldDeleteSelf = bAutoDeleteSelf;
-		bShouldDeleteRunnable = bAutoDeleteRunnable;
 		ThreadAffintyMask = InThreadAffinityMask;
 
 		// Create a sync event to guarantee the Init() function is called first
 		ThreadInitSyncEvent	= FPlatformProcess::CreateSynchEvent(true);
-		// Create a sync event to guarantee the thread will not delete itself until it has been fully set up.
-		ThreadCreatedSyncEvent = FPlatformProcess::CreateSynchEvent(true);
 
 		// Create the new thread
 		Thread = ThreadEmulation::CreateThread(NULL,InStackSize,_ThreadProc,this,0,(::DWORD *)&ThreadID);
 		// If it fails, clear all the vars
 		if (Thread == NULL)
 		{
-			if (bAutoDeleteRunnable == true)
-			{
-				delete InRunnable;
-			}
 			Runnable = NULL;
 		}
 		else
@@ -319,14 +257,5 @@ protected:
 		ThreadInitSyncEvent = NULL;
 		return Thread != NULL;
 	}
-
-	virtual bool NotifyCreated() OVERRIDE
-	{
-		const bool bHasFinished = !!WantsToDeleteSelf.GetValue();
-		// It's ok to delete this thread if it wants to delete self.
-		ThreadCreatedSyncEvent->Trigger();
-		return bHasFinished;
-		return false;
-	}
 };
 
diff --git a/Engine/Source/Runtime/Core/Private/Windows/WindowsRunnableThread.cpp b/Engine/Source/Runtime/Core/Private/Windows/WindowsRunnableThread.cpp
--- a/Engine/Source/Runtime/Core/Private/Windows/WindowsRunnableThread.cpp
+++ b/Engine/Source/Runtime/Core/Private/Windows/WindowsRunnableThread.cpp
@@ -80,23 +80,5 @@ uint32 FRunnableThreadWin::Run()
 		ThreadInitSyncEvent->Trigger();
 	}
 
-	// Should we delete the runnable?
-	if (bShouldDeleteRunnable == true)
-	{
-		delete Runnable;
-		Runnable = NULL;
-	}
-	// Clean ourselves up without waiting
-	if (bShouldDeleteSelf == true)
-	{
-		// Make sure the caller knows we want to delete this thread if we're still int CreateInternal.
-		WantsToDeleteSelf.Increment();
-		// Wait until the caller has finished setting up this thread in case Runnable execution was very short.
-		ThreadCreatedSyncEvent->Wait();
-		// Now clean up the thread handle so we don't leak
-		CloseHandle(Thread);
-		Thread = NULL;
-		delete this;
-	}
 	return ExitCode;
 }
diff --git a/Engine/Source/Runtime/Core/Private/Windows/WindowsRunnableThread.h b/Engine/Source/Runtime/Core/Private/Windows/WindowsRunnableThread.h
--- a/Engine/Source/Runtime/Core/Private/Windows/WindowsRunnableThread.h
+++ b/Engine/Source/Runtime/Core/Private/Windows/WindowsRunnableThread.h
@@ -23,26 +23,6 @@ class FRunnableThreadWin : public FRunnableThread
 	 */
 	FEvent* ThreadInitSyncEvent;
 
-	/** 
-	 * Sync event to make sure that CreateInteral() has been completed before allowing the thread to be auto-deleted
-	 */
-	FEvent* ThreadCreatedSyncEvent;
-
-	/** 
-	 * Flag used when the thread is waiting for the caller to finish setting it up before it can delete itself.
-	 */
-	FThreadSafeCounter WantsToDeleteSelf;
-
-	/**
-	 * Whether we should delete ourselves on thread exit
-	 */
-	bool bShouldDeleteSelf;
-
-	/**
-	 * Whether we should delete the runnable on thread exit
-	 */
-	bool bShouldDeleteRunnable;
-
 	/**
 	 * The priority to run the thread at
 	 */
@@ -131,10 +111,6 @@ public:
 		: Thread(NULL)
 		, Runnable(NULL)
 		, ThreadInitSyncEvent(NULL)
-		, ThreadCreatedSyncEvent(NULL)
-		, WantsToDeleteSelf(0)
-		, bShouldDeleteSelf(false)
-		, bShouldDeleteRunnable(false)
 		, ThreadPriority(TPri_Normal)
 		, ThreadID(NULL)
 	{
@@ -149,7 +125,6 @@ public:
 			Kill(true);
 		}
 		FRunnableThread::GetThreadRegistry().Remove(ThreadID);
-		delete ThreadCreatedSyncEvent;
 	}
 	
 	virtual void SetThreadPriority(EThreadPriority NewPriority) OVERRIDE
@@ -208,18 +183,7 @@ public:
 		// Now clean up the thread handle so we don't leak
 		CloseHandle(Thread);
 		Thread = NULL;
-		// delete the runnable if requested and we didn't shut down gracefully already.
-		if (Runnable && bShouldDeleteRunnable == true)
-		{
-			delete Runnable;
-			Runnable = NULL;
-		}
-		// Delete ourselves if requested and we didn't shut down gracefully already.
-		// This check prevents a double-delete of self when we shut down gracefully.
-		if (!bDidExitOK && bShouldDeleteSelf == true)
-		{
-			delete this;
-		}
+
 		return bDidExitOK;
 	}
 
@@ -241,29 +205,21 @@ public:
 
 protected:
 	virtual bool CreateInternal(FRunnable* InRunnable, const TCHAR* InThreadName,
-		bool bAutoDeleteSelf = false, bool bAutoDeleteRunnable = false, uint32 InStackSize = 0,
+		uint32 InStackSize = 0,
 		EThreadPriority InThreadPri = TPri_Normal, uint64 InThreadAffinityMask = 0) OVERRIDE
 	{
 		check(InRunnable);
 		Runnable = InRunnable;
-		bShouldDeleteSelf = bAutoDeleteSelf;
-		bShouldDeleteRunnable = bAutoDeleteRunnable;
 		ThreadAffintyMask = InThreadAffinityMask;
 
 		// Create a sync event to guarantee the Init() function is called first
 		ThreadInitSyncEvent	= FPlatformProcess::CreateSynchEvent(true);
-		// Create a sync event to guarantee the thread will not delete itself until it has been fully set up.
-		ThreadCreatedSyncEvent = FPlatformProcess::CreateSynchEvent(true);
 
 		// Create the new thread
 		Thread = CreateThread(NULL,InStackSize,_ThreadProc,this,STACK_SIZE_PARAM_IS_A_RESERVATION,(::DWORD *)&ThreadID);
 		// If it fails, clear all the vars
 		if (Thread == NULL)
 		{
-			if (bAutoDeleteRunnable == true)
-			{
-				delete InRunnable;
-			}
 			Runnable = NULL;
 		}
 		else
@@ -280,13 +236,5 @@ protected:
 		ThreadInitSyncEvent = NULL;
 		return Thread != NULL;
 	}
-
-	virtual bool NotifyCreated() OVERRIDE
-	{
-		const bool bHasFinished = !!WantsToDeleteSelf.GetValue();
-		// It's ok to delete this thread if it wants to delete self.
-		ThreadCreatedSyncEvent->Trigger();
-		return bHasFinished;
-	}
 };
 
diff --git a/Engine/Source/Runtime/Core/Public/HAL/ThreadingBase.h b/Engine/Source/Runtime/Core/Public/HAL/ThreadingBase.h
--- a/Engine/Source/Runtime/Core/Public/HAL/ThreadingBase.h
+++ b/Engine/Source/Runtime/Core/Public/HAL/ThreadingBase.h
@@ -253,8 +253,6 @@ public:
 	 *
 	 * @param InRunnable The runnable object to execute
 	 * @param ThreadName Name of the thread
-	 * @param bAutoDeleteSelf Whether to delete this object on exit
-	 * @param bAutoDeleteRunnable Whether to delete the runnable object on exit
 	 * @param InStackSize The size of the stack to create. 0 means use the
 	 * current thread's stack size
 	 * @param InThreadPri Tells the thread whether it needs to adjust its
@@ -262,11 +260,9 @@ public:
 	 *
 	 * @return The newly created thread or NULL if it failed
 	 */
-	CORE_API static FRunnableThread* Create (
-		class FRunnable* InRunnable, 
-		const TCHAR* ThreadName, 
-		bool bAutoDeleteSelf = false,
-		bool bAutoDeleteRunnable = false, 
+	CORE_API static FRunnableThread* Create(
+		class FRunnable* InRunnable,
+		const TCHAR* ThreadName,
 		uint32 InStackSize = 0,
 		EThreadPriority InThreadPri = TPri_Normal,
 		uint64 InThreadAffinityMask = 0);
@@ -412,8 +408,6 @@ protected:
 	 *
 	 * @param InRunnable The runnable object to execute
 	 * @param ThreadName Name of the thread
-	 * @param bAutoDeleteSelf Whether to delete this object on exit
-	 * @param bAutoDeleteRunnable Whether to delete the runnable object on exit
 	 * @param InStackSize The size of the stack to create. 0 means use the
 	 * current thread's stack size
 	 * @param InThreadPri Tells the thread whether it needs to adjust its
@@ -422,15 +416,8 @@ protected:
 	 * @return True if the thread and all of its initialization was successful, false otherwise
 	 */
 	virtual bool CreateInternal (FRunnable* InRunnable, const TCHAR* ThreadName,
-		bool bAutoDeleteSelf = 0,bool bAutoDeleteRunnable = 0,uint32 InStackSize = 0,
+		uint32 InStackSize = 0,
 		EThreadPriority InThreadPri = TPri_Normal, uint64 InThreadAffinityMask = 0) = 0;
-	/**
-	 * Lets this thread know it has been created it case it has already finished its execution
-	 * and wants to delete itself.
-	 *
-	 * @return true if the thread has already finished execution.
-	 */
-	virtual bool NotifyCreated() = 0;
 
 	/** Called when the this runnable has been destroyed, so we should clean-up memory allocated by misc classes. */
 	FSimpleMulticastDelegate ThreadDestroyedDelegate;
diff --git a/Engine/Source/Runtime/Core/Public/HTML5/HTML5PlatformRunnableThread.h b/Engine/Source/Runtime/Core/Public/HTML5/HTML5PlatformRunnableThread.h
--- a/Engine/Source/Runtime/Core/Public/HTML5/HTML5PlatformRunnableThread.h
+++ b/Engine/Source/Runtime/Core/Public/HTML5/HTML5PlatformRunnableThread.h
@@ -62,15 +62,10 @@ public:
 protected:
 
 	virtual bool CreateInternal (FRunnable* InRunnable, const TCHAR* ThreadName,
-		bool bAutoDeleteSelf = 0,bool bAutoDeleteRunnable = 0,uint32 InStackSize = 0,
+		uint32 InStackSize = 0,
 		EThreadPriority InThreadPri = TPri_Normal, uint64 InThreadAffinityMask = 0) 
 	{
 		return false;
 	}
-
-	virtual bool NotifyCreated() OVERRIDE
-	{
-		return false;
-	}
 };
 
diff --git a/Engine/Source/Runtime/Engine/Classes/Engine/Engine.h b/Engine/Source/Runtime/Engine/Classes/Engine/Engine.h
--- a/Engine/Source/Runtime/Engine/Classes/Engine/Engine.h
+++ b/Engine/Source/Runtime/Engine/Classes/Engine/Engine.h
@@ -7,6 +7,8 @@
 #pragma once
 #include "Engine.generated.h"
 
+class FScreenSaverInhibitor;
+
 UENUM()
 enum EFullyLoadPackageType
 {
@@ -2116,6 +2118,7 @@ private:
 
 	/** Thread preventing screen saver from kicking. Suspend most of the time. */
 	FRunnableThread*		ScreenSaverInhibitor;
+	FScreenSaverInhibitor*  ScreenSaverInhibitorRunnable;
 
 	/** If true, the engine tick function will poll FPlatformSurvey for results */
 	bool					bPendingHardwareSurveyResults;
diff --git a/Engine/Source/Runtime/Engine/Private/ShaderCompiler/ShaderCompiler.cpp b/Engine/Source/Runtime/Engine/Private/ShaderCompiler/ShaderCompiler.cpp
--- a/Engine/Source/Runtime/Engine/Private/ShaderCompiler/ShaderCompiler.cpp
+++ b/Engine/Source/Runtime/Engine/Private/ShaderCompiler/ShaderCompiler.cpp
@@ -440,7 +440,7 @@ FShaderCompileThreadRunnable::FShaderCompileThreadRunnable(FShaderCompilingManag
 
 	if (Manager->bAllowAsynchronousShaderCompiling && !FPlatformProperties::RequiresCookedData())
 	{
-		Thread = FRunnableThread::Create(this, TEXT("ShaderCompilingThread"), 0, 0, 0, TPri_Normal);
+		Thread = FRunnableThread::Create(this, TEXT("ShaderCompilingThread"));
 	}
 }
 
diff --git a/Engine/Source/Runtime/Engine/Private/UnrealEngine.cpp b/Engine/Source/Runtime/Engine/Private/UnrealEngine.cpp
--- a/Engine/Source/Runtime/Engine/Private/UnrealEngine.cpp
+++ b/Engine/Source/Runtime/Engine/Private/UnrealEngine.cpp
@@ -811,6 +811,14 @@ void UEngine::PreExit()
 #if WITH_EDITOR
 	UGameViewportClient::OnScreenshotCaptured().RemoveUObject(this, &UEngine::HandleScreenshotCaptured);
 #endif
+
+	if (ScreenSaverInhibitor)
+	{
+		ScreenSaverInhibitor->Kill();
+		delete ScreenSaverInhibitor;
+	}
+
+	delete ScreenSaverInhibitorRunnable;
 }
 
 void UEngine::TickDeferredCommands()
@@ -5546,8 +5554,8 @@ void UEngine::EnableScreenSaver( bool bEnable )
 				if( !ScreenSaverInhibitor )
 				{
 					// Create thread inhibiting screen saver while it is running.
-					FScreenSaverInhibitor* ScreenSaverInhibitorRunnable = new FScreenSaverInhibitor();
-					ScreenSaverInhibitor = FRunnableThread::Create( ScreenSaverInhibitorRunnable, TEXT("ScreenSaverInhibitor"), true, true, 16 * 1024 );
+					ScreenSaverInhibitorRunnable = new FScreenSaverInhibitor();
+					ScreenSaverInhibitor = FRunnableThread::Create(ScreenSaverInhibitorRunnable, TEXT("ScreenSaverInhibitor"), 16 * 1024);
 					// Only actually run when needed to not bypass group policies for screensaver, etc.
 					ScreenSaverInhibitor->Suspend( true );
 					ScreenSaverInhibitorSemaphore = 0;
diff --git a/Engine/Source/Runtime/Messaging/Private/Bus/MessageBus.cpp b/Engine/Source/Runtime/Messaging/Private/Bus/MessageBus.cpp
--- a/Engine/Source/Runtime/Messaging/Private/Bus/MessageBus.cpp
+++ b/Engine/Source/Runtime/Messaging/Private/Bus/MessageBus.cpp
@@ -14,7 +14,7 @@ FMessageBus::FMessageBus( const IAuthorizeMessageRecipientsPtr& InRecipientAutho
 	: RecipientAuthorizer(InRecipientAuthorizer)
 {
 	Router = new FMessageRouter();
-	RouterThread = FRunnableThread::Create(Router, TEXT("FMessageBus.Router"), false, false, 128 * 1024, TPri_Normal);
+	RouterThread = FRunnableThread::Create(Router, TEXT("FMessageBus.Router"), 128 * 1024, TPri_Normal);
 
 	check(Router != nullptr);
 }
diff --git a/Engine/Source/Runtime/MoviePlayer/Private/MoviePlayerThreading.cpp b/Engine/Source/Runtime/MoviePlayer/Private/MoviePlayerThreading.cpp
--- a/Engine/Source/Runtime/MoviePlayer/Private/MoviePlayerThreading.cpp
+++ b/Engine/Source/Runtime/MoviePlayer/Private/MoviePlayerThreading.cpp
@@ -53,15 +53,7 @@ void FSlateLoadingSynchronizationMechanism::Initialize()
 	MainLoop.Lock();
 
 	SlateRunnableTask = new FSlateLoadingThreadTask( *this );
-	SlateLoadingThread = 
-		FRunnableThread::Create
-		(
-			SlateRunnableTask,
-			TEXT("SlateLoadingThread"),
-			false,
-			false
-		);
-
+	SlateLoadingThread = FRunnableThread::Create(SlateRunnableTask, TEXT("SlateLoadingThread"));
 }
 
 void FSlateLoadingSynchronizationMechanism::DestroySlateThread()
diff --git a/Engine/Source/Runtime/NetworkFileSystem/Private/NetworkFileServer.cpp b/Engine/Source/Runtime/NetworkFileSystem/Private/NetworkFileServer.cpp
--- a/Engine/Source/Runtime/NetworkFileSystem/Private/NetworkFileServer.cpp
+++ b/Engine/Source/Runtime/NetworkFileSystem/Private/NetworkFileServer.cpp
@@ -70,7 +70,7 @@ FNetworkFileServer::FNetworkFileServer( int32 InPort, const FFileRequestDelegate
 				ListenAddr->SetPort(port);
 
 				// now create a thread to accept connections
-				Thread = FRunnableThread::Create(this,TEXT("FNetworkFileServer"), false, false, 8 * 1024, TPri_AboveNormal);
+				Thread = FRunnableThread::Create(this, TEXT("FNetworkFileServer"), 8 * 1024, TPri_AboveNormal);
 
 				UE_LOG(LogFileServer, Display, TEXT("Unreal Network File Server is ready for client connections on %s!"), *ListenAddr->ToString(true));
 			}
diff --git a/Engine/Source/Runtime/NetworkFileSystem/Private/NetworkFileServerConnection.cpp b/Engine/Source/Runtime/NetworkFileSystem/Private/NetworkFileServerConnection.cpp
--- a/Engine/Source/Runtime/NetworkFileSystem/Private/NetworkFileServerConnection.cpp
+++ b/Engine/Source/Runtime/NetworkFileSystem/Private/NetworkFileServerConnection.cpp
@@ -37,7 +37,7 @@ FNetworkFileServerClientConnection::FNetworkFileServerClientConnection( FSocket*
 #else
 	const static uint32 NetworkFileServerThreadSize = 8 * 1024; 
 #endif
-	Thread = FRunnableThread::Create(this, TEXT("FNetworkFileServerClientConnection"), false, false, NetworkFileServerThreadSize, TPri_AboveNormal);
+	Thread = FRunnableThread::Create(this, TEXT("FNetworkFileServerClientConnection"), NetworkFileServerThreadSize, TPri_AboveNormal);
 }
 
 
diff --git a/Engine/Source/Runtime/Networking/Public/Common/TcpListener.h b/Engine/Source/Runtime/Networking/Public/Common/TcpListener.h
--- a/Engine/Source/Runtime/Networking/Public/Common/TcpListener.h
+++ b/Engine/Source/Runtime/Networking/Public/Common/TcpListener.h
@@ -38,7 +38,7 @@ public:
 		, Socket(NULL)
 		, Stopping(false)
 	{
-		Thread = FRunnableThread::Create(this, TEXT("FTcpListener"), false, false, 8 * 1024, TPri_Normal);
+		Thread = FRunnableThread::Create(this, TEXT("FTcpListener"), 8 * 1024, TPri_Normal);
 	}
 
 	/**
@@ -57,7 +57,7 @@ public:
 		Socket->GetAddress(*LocalAddress);
 		Endpoint = FIPv4Endpoint(LocalAddress);
 
-		Thread = FRunnableThread::Create(this, TEXT("FTcpListener"), false, false, 8 * 1024, TPri_Normal);
+		Thread = FRunnableThread::Create(this, TEXT("FTcpListener"), 8 * 1024, TPri_Normal);
 	}
 
 	/**
diff --git a/Engine/Source/Runtime/Networking/Public/Common/UdpSocketReceiver.h b/Engine/Source/Runtime/Networking/Public/Common/UdpSocketReceiver.h
--- a/Engine/Source/Runtime/Networking/Public/Common/UdpSocketReceiver.h
+++ b/Engine/Source/Runtime/Networking/Public/Common/UdpSocketReceiver.h
@@ -45,7 +45,7 @@ public:
 		check(Socket->GetSocketType() == SOCKTYPE_Datagram);
 
 		SocketSubsystem = ISocketSubsystem::Get(PLATFORM_SOCKETSUBSYSTEM);
-		Thread = FRunnableThread::Create(this, ThreadDescription, false, false, 128 * 1024, TPri_AboveNormal);
+		Thread = FRunnableThread::Create(this, ThreadDescription, 128 * 1024, TPri_AboveNormal);
 	}
 
 	/**
diff --git a/Engine/Source/Runtime/Networking/Public/Common/UdpSocketSender.h b/Engine/Source/Runtime/Networking/Public/Common/UdpSocketSender.h
--- a/Engine/Source/Runtime/Networking/Public/Common/UdpSocketSender.h
+++ b/Engine/Source/Runtime/Networking/Public/Common/UdpSocketSender.h
@@ -56,7 +56,7 @@ public:
 
 		WorkEvent = FPlatformProcess::CreateSynchEvent();
 
-		Thread = FRunnableThread::Create(this, ThreadDescription, false, false, 128 * 1024, TPri_AboveNormal);
+		Thread = FRunnableThread::Create(this, ThreadDescription, 128 * 1024, TPri_AboveNormal);
 	}
 
 	/**
diff --git a/Engine/Source/Runtime/Online/OnlineSubsystemNull/Private/OnlineSubsystemNull.cpp b/Engine/Source/Runtime/Online/OnlineSubsystemNull/Private/OnlineSubsystemNull.cpp
--- a/Engine/Source/Runtime/Online/OnlineSubsystemNull/Private/OnlineSubsystemNull.cpp
+++ b/Engine/Source/Runtime/Online/OnlineSubsystemNull/Private/OnlineSubsystemNull.cpp
@@ -132,7 +132,7 @@ bool FOnlineSubsystemNull::Init()
 		// Create the online async task thread
 		OnlineAsyncTaskThreadRunnable = new FOnlineAsyncTaskManagerNull(this);
 		check(OnlineAsyncTaskThreadRunnable);
-		OnlineAsyncTaskThread = FRunnableThread::Create(OnlineAsyncTaskThreadRunnable, TEXT("OnlineAsyncTaskThreadNull"), 0, 0, 128 * 1024, TPri_Normal);
+		OnlineAsyncTaskThread = FRunnableThread::Create(OnlineAsyncTaskThreadRunnable, TEXT("OnlineAsyncTaskThreadNull"), 128 * 1024, TPri_Normal);
 		check(OnlineAsyncTaskThread);
 		UE_LOG_ONLINE(Verbose, TEXT("Created thread (ID:%d)."), OnlineAsyncTaskThread->GetThreadID());
 
diff --git a/Engine/Source/Runtime/Online/OnlineSubsystemSteam/Private/OnlineSubsystemSteam.cpp b/Engine/Source/Runtime/Online/OnlineSubsystemSteam/Private/OnlineSubsystemSteam.cpp
--- a/Engine/Source/Runtime/Online/OnlineSubsystemSteam/Private/OnlineSubsystemSteam.cpp
+++ b/Engine/Source/Runtime/Online/OnlineSubsystemSteam/Private/OnlineSubsystemSteam.cpp
@@ -303,7 +303,7 @@ bool FOnlineSubsystemSteam::Init()
 		// Create the online async task thread
 		OnlineAsyncTaskThreadRunnable = new FOnlineAsyncTaskManagerSteam(this);
 		check(OnlineAsyncTaskThreadRunnable);		
-		OnlineAsyncTaskThread = FRunnableThread::Create(OnlineAsyncTaskThreadRunnable, TEXT("OnlineAsyncTaskThreadSteam"), 0, 0, 128 * 1024, TPri_Normal);
+		OnlineAsyncTaskThread = FRunnableThread::Create(OnlineAsyncTaskThreadRunnable, TEXT("OnlineAsyncTaskThreadSteam"), 128 * 1024, TPri_Normal);
 		check(OnlineAsyncTaskThread);
 		UE_LOG_ONLINE(Verbose, TEXT("Created thread (ID:%d)."), OnlineAsyncTaskThread->GetThreadID() );
 
diff --git a/Engine/Source/Runtime/PakFile/Private/SignedArchiveReader.cpp b/Engine/Source/Runtime/PakFile/Private/SignedArchiveReader.cpp
--- a/Engine/Source/Runtime/PakFile/Private/SignedArchiveReader.cpp
+++ b/Engine/Source/Runtime/PakFile/Private/SignedArchiveReader.cpp
@@ -25,7 +25,7 @@ FChunkCacheWorker::FChunkCacheWorker(FArchive* InReader)
 	check(!DecryptionKey.Exponent.IsZero() && !DecryptionKey.Modulus.IsZero());
 
 	QueuedRequestsEvent = FPlatformProcess::CreateSynchEvent();
-	Thread = FRunnableThread::Create(this, TEXT("FChunkCacheWorker"), false, false, 0, TPri_BelowNormal);
+	Thread = FRunnableThread::Create(this, TEXT("FChunkCacheWorker"), 0, TPri_BelowNormal);
 }
 
 FChunkCacheWorker::~FChunkCacheWorker()
diff --git a/Engine/Source/Runtime/RenderCore/Private/RenderingThread.cpp b/Engine/Source/Runtime/RenderCore/Private/RenderingThread.cpp
--- a/Engine/Source/Runtime/RenderCore/Private/RenderingThread.cpp
+++ b/Engine/Source/Runtime/RenderCore/Private/RenderingThread.cpp
@@ -457,9 +457,7 @@ void StartRenderingThread()
 	// Create the rendering thread.
 	GRenderingThreadRunnable = new FRenderingThread();
 
-	EThreadPriority RenderingThreadPrio = TPri_Normal;
-
-	GRenderingThread = FRunnableThread::Create( GRenderingThreadRunnable, *BuildRenderingThreadName( ThreadCount ), 0, 0, 0, RenderingThreadPrio );
+	GRenderingThread = FRunnableThread::Create( GRenderingThreadRunnable, *BuildRenderingThreadName( ThreadCount ));
 
 	// Wait for render thread to have taskgraph bound before we dispatch any tasks for it.
 	((FRenderingThread*)GRenderingThreadRunnable)->TaskGraphBoundSyncEvent->Wait();
@@ -476,7 +474,7 @@ void StartRenderingThread()
 	// Create the rendering thread heartbeat
 	GRenderingThreadRunnableHeartbeat = new FRenderingThreadTickHeartbeat();
 
-	GRenderingThreadHeartbeat = FRunnableThread::Create(GRenderingThreadRunnableHeartbeat, *FString::Printf(TEXT("RTHeartBeat %d"), ThreadCount), 0, 0, 16 * 1024, TPri_AboveNormal);
+	GRenderingThreadHeartbeat = FRunnableThread::Create(GRenderingThreadRunnableHeartbeat, *FString::Printf(TEXT("RTHeartBeat %d"), ThreadCount), 16 * 1024, TPri_AboveNormal);
 
 	ThreadCount++;
 }
diff --git a/Engine/Source/Runtime/Sockets/Public/MultichannelTcpReceiver.h b/Engine/Source/Runtime/Sockets/Public/MultichannelTcpReceiver.h
--- a/Engine/Source/Runtime/Sockets/Public/MultichannelTcpReceiver.h
+++ b/Engine/Source/Runtime/Sockets/Public/MultichannelTcpReceiver.h
@@ -35,7 +35,7 @@ public:
 		: Socket(InSocket)
 		, ReceiveDelegate(InReceiveDelegate)
 	{
-		Thread = FRunnableThread::Create(this,TEXT("FMultichannelTCPReceiver"), true, true, 8 * 1024, TPri_AboveNormal);
+		Thread = FRunnableThread::Create(this, TEXT("FMultichannelTCPReceiver"), 8 * 1024, TPri_AboveNormal);
 	}
 
 	/**
@@ -43,11 +43,11 @@ public:
 	 */
 	~FMultichannelTcpReceiver()
 	{
-		Thread->Kill(true);
-		delete Thread;
-
-		Socket = NULL;
-		Thread = NULL;
+		if (Thread)
+		{
+			Thread->Kill(true);
+			delete Thread;
+		}
 	}
 
 
diff --git a/Engine/Source/Runtime/Sockets/Public/MultichannelTcpSender.h b/Engine/Source/Runtime/Sockets/Public/MultichannelTcpSender.h
--- a/Engine/Source/Runtime/Sockets/Public/MultichannelTcpSender.h
+++ b/Engine/Source/Runtime/Sockets/Public/MultichannelTcpSender.h
@@ -43,7 +43,7 @@ public:
 		: Socket(InSocket)
 		, OkToSendDelegate(InOkToSendDelegate)
 	{
-		Thread = FRunnableThread::Create(this,TEXT("FMultichannelTCPSender"), true, true, 8 * 1024, TPri_AboveNormal);
+		Thread = FRunnableThread::Create(this, TEXT("FMultichannelTCPSender"), 8 * 1024, TPri_AboveNormal);
 	}
 
 	/**
@@ -51,11 +51,11 @@ public:
 	 */
 	~FMultichannelTcpSender(  )
 	{
-		Thread->Kill(true);
-		delete Thread;
-
-		Socket = NULL;
-		Thread = NULL;
+		if (Thread)
+		{
+			Thread->Kill(true);
+			delete Thread;
+		}
 	}
 
 
