CIS Fix for Android build.
Extra OVERRIDE keyword caused problems during Android compilation.

diff --git a/Engine/Source/Runtime/Core/Public/Android/AndroidAffinity.h b/Engine/Source/Runtime/Core/Public/Android/AndroidAffinity.h
--- a/Engine/Source/Runtime/Core/Public/Android/AndroidAffinity.h
+++ b/Engine/Source/Runtime/Core/Public/Android/AndroidAffinity.h
@@ -11,32 +11,32 @@ AndroidAffinity.h: Android affinity profile masks definitions.
 class FAndroidAffinity : public FGenericPlatformAffinity
 {
 public:
-	static const CORE_API uint64 GetMainGameMask() OVERRIDE
+	static const CORE_API uint64 GetMainGameMask()
 	{
 		return MAKEAFFINITYMASK1(0);
 	}
 
-	static const CORE_API uint64 GetRenderingThreadMask() OVERRIDE
+	static const CORE_API uint64 GetRenderingThreadMask()
 	{
 		return MAKEAFFINITYMASK1(1);
 	}
 
-	static const CORE_API uint64 GetRTHeartBeatMask() OVERRIDE
+	static const CORE_API uint64 GetRTHeartBeatMask()
 	{
 		return MAKEAFFINITYMASK1(4);
 	}
 
-	static const CORE_API uint64 GetPoolThreadMask() OVERRIDE
+	static const CORE_API uint64 GetPoolThreadMask()
 	{
 		return MAKEAFFINITYMASK1(4);
 	}
 
-	static const CORE_API uint64 GetTaskGraphThreadMask() OVERRIDE
+	static const CORE_API uint64 GetTaskGraphThreadMask()
 	{
 		return MAKEAFFINITYMASK3(0, 2, 3);
 	}
 
-	static const CORE_API uint64 GetStatsThreadMask() OVERRIDE
+	static const CORE_API uint64 GetStatsThreadMask()
 	{
 		return MAKEAFFINITYMASK1(5);
 	}
