Allowed dialog title to be specified in FMessageDialog

Since I need to change the callback delegate signature, I switched the parameter order to be consistent with the dialog functions.

diff --git a/Engine/Source/Editor/UnrealEd/Classes/Editor/EditorEngine.h b/Engine/Source/Editor/UnrealEd/Classes/Editor/EditorEngine.h
--- a/Engine/Source/Editor/UnrealEd/Classes/Editor/EditorEngine.h
+++ b/Engine/Source/Editor/UnrealEd/Classes/Editor/EditorEngine.h
@@ -1441,11 +1441,13 @@ public:
 
 	/** 
 	 * Displays a modal message dialog 
+	 * @param	InMessage	Type of the message box
 	 * @param	InText		Message to display
-	 * @param	InMessage		Type of the message box
+	 * @param	InTitle		Title for the message box
 	 * @return	Returns the result of the modal message box
 	 */
-	EAppReturnType::Type OnModalMessageDialog( const FText& InText, EAppMsgType::Type InMessage );
+	EAppReturnType::Type OnModalMessageDialog(EAppMsgType::Type InMessage, const FText& InText, const FText& InTitle);
+
 	/** 
 	 * Returns whether an object should replace an exisiting one or not 
 	 * @param	Filename		Filename of the package
diff --git a/Engine/Source/Editor/UnrealEd/Private/Editor.cpp b/Engine/Source/Editor/UnrealEd/Private/Editor.cpp
--- a/Engine/Source/Editor/UnrealEd/Private/Editor.cpp
+++ b/Engine/Source/Editor/UnrealEd/Private/Editor.cpp
@@ -4416,16 +4416,15 @@ void UEditorEngine::CloseEntryPopupWindow()
 	}
 }
 
-EAppReturnType::Type UEditorEngine::OnModalMessageDialog( const FText& InText, EAppMsgType::Type InMessage )
+EAppReturnType::Type UEditorEngine::OnModalMessageDialog(EAppMsgType::Type InMessage, const FText& InText, const FText& InTitle)
 {
-	const FText MessageTitle = NSLOCTEXT("UnrealEd", "GenericDialog_WindowTitle", "Message");
 	if( FSlateApplication::IsInitialized() && FSlateApplication::Get().CanAddModalWindow() )
 	{
-		return OpenMsgDlgInt(InMessage, InText, MessageTitle);
+		return OpenMsgDlgInt(InMessage, InText, InTitle);
 	}
 	else
 	{
-		return FPlatformMisc::MessageBoxExt( InMessage, *InText.ToString(), *MessageTitle.ToString() );
+		return FPlatformMisc::MessageBoxExt(InMessage, *InText.ToString(), *InTitle.ToString());
 	}
 }
 
diff --git a/Engine/Source/Runtime/Core/Private/Misc/OutputDevice.cpp b/Engine/Source/Runtime/Core/Private/Misc/OutputDevice.cpp
--- a/Engine/Source/Runtime/Core/Private/Misc/OutputDevice.cpp
+++ b/Engine/Source/Runtime/Core/Private/Misc/OutputDevice.cpp
@@ -10,6 +10,19 @@
 
 DEFINE_LOG_CATEGORY_STATIC(LogOutputDevice, Log, All);
 
+namespace
+{
+	/**
+	 * Singleton to only create error title text if needed (and after localization system is in place)
+	 */
+	const FText& GetDefaultMessageTitle()
+	{
+		// Will be initialised on first call
+		static FText DefaultMessageTitle(NSLOCTEXT("MessageDialog", "DefaultMessageTitle", "Message"));
+		return DefaultMessageTitle;
+	}
+}
+
 const TCHAR* FOutputDevice::VerbosityToString(ELogVerbosity::Type Verbosity)
 {
 	switch (Verbosity & ELogVerbosity::VerbosityMask)
@@ -336,7 +349,7 @@ bool VARARGS FDebug::EnsureNotFalseFormatted( bool bExpressionResult, const ANSI
 }
 #endif
 
-void FMessageDialog::Debugf( const FText& Message )
+void FMessageDialog::Debugf( const FText& Message, const FText* OptTitle )
 {
 	if( FApp::IsUnattended() == true )
 	{
@@ -344,9 +357,10 @@ void FMessageDialog::Debugf( const FText& Message )
 	}
 	else
 	{
+		FText Title = OptTitle ? *OptTitle : GetDefaultMessageTitle();
 		if ( GIsEditor && FCoreDelegates::ModalErrorMessage.IsBound() )
 		{
-			FCoreDelegates::ModalErrorMessage.Execute( Message, EAppMsgType::Ok );
+			FCoreDelegates::ModalErrorMessage.Execute(EAppMsgType::Ok, Message, Title);
 		}
 		else
 		{
@@ -372,7 +386,7 @@ void FMessageDialog::ShowLastError()
 	}
 }
 
-EAppReturnType::Type FMessageDialog::Open( EAppMsgType::Type MessageType, const FText& Message )
+EAppReturnType::Type FMessageDialog::Open( EAppMsgType::Type MessageType, const FText& Message, const FText* OptTitle )
 {
 	if( FApp::IsUnattended() == true )
 	{
@@ -403,13 +417,14 @@ EAppReturnType::Type FMessageDialog::Open( EAppMsgType::Type MessageType, const
 	}
 	else
 	{
+		FText Title = OptTitle ? *OptTitle : GetDefaultMessageTitle();
 		if ( GIsEditor && !IsRunningCommandlet() && FCoreDelegates::ModalErrorMessage.IsBound() )
 		{
-			return FCoreDelegates::ModalErrorMessage.Execute( Message, MessageType );
+			return FCoreDelegates::ModalErrorMessage.Execute( MessageType, Message, Title );
 		}
 		else
 		{
-			return FPlatformMisc::MessageBoxExt( MessageType, *Message.ToString(), *NSLOCTEXT("MessageDialog", "DefaultMessageTitle", "Message").ToString() );
+			return FPlatformMisc::MessageBoxExt( MessageType, *Message.ToString(), *Title.ToString() );
 		}
 	}
 }
diff --git a/Engine/Source/Runtime/Core/Public/Misc/CallbackDevice.h b/Engine/Source/Runtime/Core/Public/Misc/CallbackDevice.h
--- a/Engine/Source/Runtime/Core/Public/Misc/CallbackDevice.h
+++ b/Engine/Source/Runtime/Core/Public/Misc/CallbackDevice.h
@@ -73,10 +73,10 @@ public:
 	// Delegate type for redirector followed events ( Params: const FString& PackageName, UObject* Redirector )
 	DECLARE_MULTICAST_DELEGATE_TwoParams(FOnRedirectorFollowed, const FString&, UObject*);
 
-	/** delegate type for opening a modal message box ( Params: const FString& Text, EAppMsgType::Type MessageType ) */
-	DECLARE_DELEGATE_RetVal_TwoParams(EAppReturnType::Type, FOnModalMessageBox, const FText&, EAppMsgType::Type);
+	/** delegate type for opening a modal message box ( Params: EAppMsgType::Type MessageType, const FText& Text, const FText& Title ) */
+	DECLARE_DELEGATE_RetVal_ThreeParams(EAppReturnType::Type, FOnModalMessageBox, EAppMsgType::Type, const FText&, const FText&);
 
-	/** delegate type for querrying whether a loaded object should replace an already exisiting one */
+	/** delegate type for querying whether a loaded object should replace an already existing one */
 	DECLARE_DELEGATE_RetVal_OneParam(bool, FOnLoadObjectsOnTop, const FString&);
 
 	// Callback for handling an ensure
diff --git a/Engine/Source/Runtime/Core/Public/Misc/OutputDevice.h b/Engine/Source/Runtime/Core/Public/Misc/OutputDevice.h
--- a/Engine/Source/Runtime/Core/Public/Misc/OutputDevice.h
+++ b/Engine/Source/Runtime/Core/Public/Misc/OutputDevice.h
@@ -409,13 +409,22 @@ struct CORE_API FDebug
  **/
 struct CORE_API FMessageDialog
 {
-	/** Pops up a message dialog box containing the input string. */
-	static void Debugf( const FText& Message );
+	/** Pops up a message dialog box containing the input string.
+	 * @param Message Text of message to show
+	 * @param OptTitle Optional title to use (defaults to "Message")
+	*/
+	static void Debugf( const FText& Message, const FText* OptTitle = nullptr );
 
 	/** Pops up a message dialog box containing the last system error code in string form. */
 	static void ShowLastError();
 
-	static EAppReturnType::Type Open( EAppMsgType::Type MessageType, const FText& Message );
+	/**
+	 * Open a modal message box dialog
+	 * @param MessageType Controls buttons dialog should have
+	 * @param Message Text of message to show
+	 * @param OptTitle Optional title to use (defaults to "Message")
+	*/
+	static EAppReturnType::Type Open( EAppMsgType::Type MessageType, const FText& Message, const FText* OptTitle = nullptr );
 };
 
 
