A modal dialog is now shown if CrashReportClient is not available

TTP 322858: Make sure the crashing experience is sensible for Source versions that have no CrashReporter

diff --git a/Engine/Source/Runtime/Core/Private/Windows/WindowsPlatformExceptionHandling.cpp b/Engine/Source/Runtime/Core/Private/Windows/WindowsPlatformExceptionHandling.cpp
--- a/Engine/Source/Runtime/Core/Private/Windows/WindowsPlatformExceptionHandling.cpp
+++ b/Engine/Source/Runtime/Core/Private/Windows/WindowsPlatformExceptionHandling.cpp
@@ -321,12 +321,13 @@ int32 NewReportCrash( EXCEPTION_POINTERS* ExceptionInfo, const TCHAR* ErrorMessa
 		FString CrashReportUploaderArguments = FString::Printf( TEXT( "-LogFolder=%s" ), *FPaths::Combine(*FPaths::GameAgnosticSavedDir(), TEXT("CrashLogs")) );
 	
 		// Suppress the user input dialog if we're running in unattended mode
-		if( FApp::IsUnattended() || ReportUI == EErrorReportUI::ReportInUnattendedMode )
+		bool bNoDialog = FApp::IsUnattended() || ReportUI == EErrorReportUI::ReportInUnattendedMode;
+		if (bNoDialog)
 		{
 			CrashReportUploaderArguments += TEXT( " -Unattended" );
 		}
 
-		if ( FApp::IsInstalled() )
+		if (FApp::IsInstalled())
 		{
 			// Temporary workaround for CrashReportClient being built in Development, not Shipping (TTP328030). The
 			// following ensures that logs are saved to the user directory when UE4 is installed.
@@ -338,7 +339,19 @@ int32 NewReportCrash( EXCEPTION_POINTERS* ExceptionInfo, const TCHAR* ErrorMessa
 			TEXT("..\\..\\..\\Engine\\Binaries\\DotNET\\CrashReportUploader.exe");
 
 		// Passing !GUseCrashReportClient to 'hidden' and 'really-hidden' parameters
-		FPlatformProcess::CreateProc(CrashClientPath, *CrashReportUploaderArguments, true, !GUseCrashReportClient, !GUseCrashReportClient, NULL, 0, NULL, NULL);
+		bool bCrashReporterRan = FPlatformProcess::CreateProc(CrashClientPath, *CrashReportUploaderArguments, true, !GUseCrashReportClient, !GUseCrashReportClient, NULL, 0, NULL, NULL).IsValid();
+
+		if (!bCrashReporterRan && !bNoDialog)
+		{
+			UE_LOG(LogWindows, Log, TEXT("Could not start CrashReportClient"));
+			FPlatformMemory::DumpStats(*GWarn);
+			FText MessageTitle(FText::Format(
+				NSLOCTEXT("MessageDialog", "AppHasCrashed", "The {0} {1} has crashed and will close"),
+				FText::FromString(ReportInformation.wzApplicationName),
+				FText::FromString(GetEngineMode())
+			));
+			FMessageDialog::Open(EAppMsgType::Ok, FText::FromString(GErrorHist), &MessageTitle);
+		}
 	}
 
 	// Reset the setting after we've had our way with it
@@ -440,53 +453,54 @@ int32 ReportCrash( LPEXCEPTION_POINTERS ExceptionInfo )
 {
 	// Only create a minidump the first time this function is called.
 	// (Can be called the first time from the RenderThread, then a second time from the MainThread.)
-	if ( GAlreadyCreatedMinidump == false )
+	if (GAlreadyCreatedMinidump)
 	{
-		GAlreadyCreatedMinidump = true;
+		return EXCEPTION_EXECUTE_HANDLER;
+	}
+
+	GAlreadyCreatedMinidump = true;
 
-		FCoreDelegates::OnHandleSystemError.Broadcast();
+	FCoreDelegates::OnHandleSystemError.Broadcast();
+
+	// Note: NewReportCrash may read the callstack from GErrorHist
+	const SIZE_T StackTraceSize = 65535;
+	ANSICHAR* StackTrace = (ANSICHAR*) FMemory::SystemMalloc( StackTraceSize );
+	StackTrace[0] = 0;
+	// Walk the stack and dump it to the allocated memory.
+	FPlatformStackWalk::StackWalkAndDump( StackTrace, StackTraceSize, 0, ExceptionInfo->ContextRecord );
+	FCString::Strncat( GErrorHist, ANSI_TO_TCHAR(StackTrace), ARRAY_COUNT(GErrorHist) - 1 );
+	CreateExceptionInfoString(ExceptionInfo->ExceptionRecord);
+	FMemory::SystemFree( StackTrace );
 
 #if WINVER > 0x502	// Windows Error Reporting is not supported on Windows XP
-		// Loop in the new system while keeping the old until it's fully implemented
+	// 
+	if (GUseCrashReportClient)
+	{
 		NewReportCrash( ExceptionInfo, NULL, EErrorReportUI::ShowDialog );
+		return EXCEPTION_EXECUTE_HANDLER;
+	}
 #endif		// WINVER
 
-		if (!GUseCrashReportClient)
-		{
-			// Try to create file for minidump.
-			HANDLE FileHandle	= CreateFileW( MiniDumpFilenameW, GENERIC_WRITE, 0, NULL, CREATE_ALWAYS, FILE_ATTRIBUTE_NORMAL, NULL );
-			
-			// Write a minidump.
-			if( FileHandle != INVALID_HANDLE_VALUE )
-			{
-				MINIDUMP_EXCEPTION_INFORMATION DumpExceptionInfo;
-
-				DumpExceptionInfo.ThreadId			= GetCurrentThreadId();
-				DumpExceptionInfo.ExceptionPointers	= ExceptionInfo;
-				DumpExceptionInfo.ClientPointers	= false;
+	// Try to create file for minidump.
+	HANDLE FileHandle = CreateFileW( MiniDumpFilenameW, GENERIC_WRITE, 0, NULL, CREATE_ALWAYS, FILE_ATTRIBUTE_NORMAL, NULL );
+	
+	// Write a minidump.
+	if( FileHandle != INVALID_HANDLE_VALUE )
+	{
+		MINIDUMP_EXCEPTION_INFORMATION DumpExceptionInfo;
 
-				MiniDumpWriteDump( GetCurrentProcess(), GetCurrentProcessId(), FileHandle, MiniDumpNormal, &DumpExceptionInfo, NULL, NULL );
-				CloseHandle( FileHandle );
-			}
-		}
+		DumpExceptionInfo.ThreadId			= GetCurrentThreadId();
+		DumpExceptionInfo.ExceptionPointers	= ExceptionInfo;
+		DumpExceptionInfo.ClientPointers	= false;
 
-		const SIZE_T StackTraceSize = 65535;
-		ANSICHAR* StackTrace = (ANSICHAR*) FMemory::SystemMalloc( StackTraceSize );
-		StackTrace[0] = 0;
-		// Walk the stack and dump it to the allocated memory.
-		FPlatformStackWalk::StackWalkAndDump( StackTrace, StackTraceSize, 0, ExceptionInfo->ContextRecord );
-		FCString::Strncat( GErrorHist, ANSI_TO_TCHAR(StackTrace), ARRAY_COUNT(GErrorHist) - 1 );
-		CreateExceptionInfoString(ExceptionInfo->ExceptionRecord);
-		FMemory::SystemFree( StackTrace );
+		MiniDumpWriteDump( GetCurrentProcess(), GetCurrentProcessId(), FileHandle, MiniDumpNormal, &DumpExceptionInfo, NULL, NULL );
+		CloseHandle( FileHandle );
+	}
 
 #if UE_BUILD_SHIPPING && WITH_EDITOR
-		if (!GUseCrashReportClient)
-		{
-			uint32 dwOpt = 0;
-			EFaultRepRetVal repret = ReportFault( ExceptionInfo, dwOpt);
-		}
+	uint32 dwOpt = 0;
+	EFaultRepRetVal repret = ReportFault( ExceptionInfo, dwOpt);
 #endif
-	}
 
 	return EXCEPTION_EXECUTE_HANDLER;
 }
