multicast delegate bound check fix to see if there is a valid function in the invocation list

diff --git a/Engine/Source/Runtime/Core/Public/Delegates/MulticastDelegateBase.h b/Engine/Source/Runtime/Core/Public/Delegates/MulticastDelegateBase.h
--- a/Engine/Source/Runtime/Core/Public/Delegates/MulticastDelegateBase.h
+++ b/Engine/Source/Runtime/Core/Public/Delegates/MulticastDelegateBase.h
@@ -34,7 +34,14 @@ public:
 	 */
 	inline bool IsBound( ) const
 	{
-		return (InvocationList.Num() > 0);
+		for (IDelegateInstance* DelegateInstance : InvocationList)
+		{
+			if (DelegateInstance != nullptr)
+			{
+				return true;
+			}
+		}
+		return false;
 	}
 
 	/** 
