Core - Added GetCPUVendor, GetCPUBrand, GetPrimaryGPUBrand, GetOSVersion to the PlatformMisc (only for Windows), refactored a bit FWindowsPlatformSurvey::GetOSVersionLabels

diff --git a/Engine/Source/Runtime/Core/Private/GenericPlatform/GenericPlatformMisc.cpp b/Engine/Source/Runtime/Core/Private/GenericPlatform/GenericPlatformMisc.cpp
--- a/Engine/Source/Runtime/Core/Private/GenericPlatform/GenericPlatformMisc.cpp
+++ b/Engine/Source/Runtime/Core/Private/GenericPlatform/GenericPlatformMisc.cpp
@@ -195,6 +195,33 @@ void FGenericPlatformMisc::SubmitErrorReport( const TCHAR* InErrorHist, EErrorRe
 	UE_LOG(LogGenericPlatformMisc, Error, TEXT("This platform cannot submit a crash report. Report was:\n%s"), InErrorHist);
 }
 
+
+FString FGenericPlatformMisc::GetCPUVendor()
+{
+	// Not implemented cross-platform. Each platform may or may not choose to implement this.
+	return FString( TEXT( "GenericCPUVendor" ) );
+}
+
+FString FGenericPlatformMisc::GetCPUBrand()
+{
+	// Not implemented cross-platform. Each platform may or may not choose to implement this.
+	return FString( TEXT( "GenericCPUBrand" ) );
+}
+
+
+FString FGenericPlatformMisc::GetPrimaryGPUBrand()
+{
+	// Not implemented cross-platform. Each platform may or may not choose to implement this.
+	return FString( TEXT( "GenericGPUBrand" ) );
+}
+
+void FGenericPlatformMisc::GetOSVersions( FString& out_OSVersionLabel, FString& out_OSSubVersionLabel )
+{
+	// Not implemented cross-platform. Each platform may or may not choose to implement this.
+	out_OSVersionLabel = FString( TEXT( "GenericOSVersionLabel" ) );
+	out_OSSubVersionLabel = FString( TEXT( "GenericOSSubVersionLabel" ) );
+}
+
 void FGenericPlatformMisc::MemoryBarrier()
 {
 }
diff --git a/Engine/Source/Runtime/Core/Private/Windows/WindowsPlatformMisc.cpp b/Engine/Source/Runtime/Core/Private/Windows/WindowsPlatformMisc.cpp
--- a/Engine/Source/Runtime/Core/Private/Windows/WindowsPlatformMisc.cpp
+++ b/Engine/Source/Runtime/Core/Private/Windows/WindowsPlatformMisc.cpp
@@ -39,6 +39,235 @@
 #include <fcntl.h>
 #include <io.h>
 
+#include "AllowWindowsPlatformTypes.h"
+int32 FWindowsOSVersionHelper::GetOSVersions( FString& out_OSVersionLabel, FString& out_OSSubVersionLabel )
+{
+	int32 ErrorCode = (int32)SUCCEEDED;
+
+	// Get system info
+	SYSTEM_INFO SystemInfo;
+	if( FPlatformMisc::Is64bitOperatingSystem() )
+	{
+		GetNativeSystemInfo( &SystemInfo );
+	}
+	else
+	{
+		GetSystemInfo( &SystemInfo );
+	}
+
+	OSVERSIONINFOEX OsVersionInfo = {0};
+	OsVersionInfo.dwOSVersionInfoSize = sizeof( OSVERSIONINFOEX );
+	out_OSVersionLabel = TEXT( "Windows (unknown version)" );
+	out_OSSubVersionLabel = FString();
+#pragma warning(disable : 4996) // 'function' was declared deprecated
+	if( GetVersionEx( (LPOSVERSIONINFO)&OsVersionInfo ) )
+#pragma warning(default : 4996)
+	{
+		bool bIsInvalidVersion = false;
+
+		switch( OsVersionInfo.dwMajorVersion )
+		{
+			case 5:
+				switch( OsVersionInfo.dwMinorVersion )
+				{
+					case 0:
+						out_OSVersionLabel = TEXT( "Windows 2000" );
+						if( OsVersionInfo.wProductType == VER_NT_WORKSTATION )
+						{
+							out_OSSubVersionLabel = TEXT( "Professional" );
+						}
+						else
+						{
+							if( OsVersionInfo.wSuiteMask & VER_SUITE_DATACENTER )
+							{
+								out_OSSubVersionLabel = TEXT( "Datacenter Server" );
+							}
+							else if( OsVersionInfo.wSuiteMask & VER_SUITE_ENTERPRISE )
+							{
+								out_OSSubVersionLabel = TEXT( "Advanced Server" );
+							}
+							else
+							{
+								out_OSSubVersionLabel = TEXT( "Server" );
+							}
+						}
+						break;
+					case 1:
+						out_OSVersionLabel = TEXT( "Windows XP" );
+						if( OsVersionInfo.wSuiteMask & VER_SUITE_PERSONAL )
+						{
+							out_OSSubVersionLabel = TEXT( "Home Edition" );
+						}
+						else
+						{
+							out_OSSubVersionLabel = TEXT( "Professional" );
+						}
+						break;
+					case 2:
+						if( GetSystemMetrics( SM_SERVERR2 ) )
+						{
+							out_OSVersionLabel = TEXT( "Windows Server 2003 R2" );
+						}
+						else if( OsVersionInfo.wSuiteMask & VER_SUITE_STORAGE_SERVER )
+						{
+							out_OSVersionLabel = TEXT( "Windows Storage Server 2003" );
+						}
+						else if( OsVersionInfo.wSuiteMask & VER_SUITE_WH_SERVER )
+						{
+							out_OSVersionLabel = TEXT( "Windows Home Server" );
+						}
+						else if( OsVersionInfo.wProductType == VER_NT_WORKSTATION && SystemInfo.wProcessorArchitecture == PROCESSOR_ARCHITECTURE_AMD64 )
+						{
+							out_OSVersionLabel = TEXT( "Windows XP" );
+							out_OSSubVersionLabel = TEXT( "Professional x64 Edition" );
+						}
+						else
+						{
+							out_OSVersionLabel = TEXT( "Windows Server 2003" );
+						}
+						break;
+					default:
+						ErrorCode |= (int32)ERROR_UNKNOWNVERSION;
+				}
+				break;
+			case 6:
+				switch( OsVersionInfo.dwMinorVersion )
+				{
+					case 0:
+						if( OsVersionInfo.wProductType == VER_NT_WORKSTATION )
+						{
+							out_OSVersionLabel = TEXT( "Windows Vista" );
+						}
+						else
+						{
+							out_OSVersionLabel = TEXT( "Windows Server 2008" );
+						}
+						break;
+					case 1:
+						if( OsVersionInfo.wProductType == VER_NT_WORKSTATION )
+						{
+							out_OSVersionLabel = TEXT( "Windows 7" );
+						}
+						else
+						{
+							out_OSVersionLabel = TEXT( "Windows Server 2008 R2" );
+						}
+						break;
+					case 2:
+						if( OsVersionInfo.wProductType == VER_NT_WORKSTATION )
+						{
+							out_OSVersionLabel = TEXT( "Windows 8" );
+						}
+						else
+						{
+							out_OSVersionLabel = TEXT( "Windows Server 2012" );
+						}
+						break;
+					default:
+						ErrorCode |= (int32)ERROR_UNKNOWNVERSION;
+				}
+
+			{
+#pragma warning( push )
+#pragma warning( disable: 4191 )	// unsafe conversion from 'type of expression' to 'type required'
+				typedef BOOL( WINAPI *LPFN_GETPRODUCTINFO )(DWORD, DWORD, DWORD, DWORD, PDWORD);
+				LPFN_GETPRODUCTINFO fnGetProductInfo = (LPFN_GETPRODUCTINFO)GetProcAddress( GetModuleHandle( TEXT( "kernel32.dll" ) ), "GetProductInfo" );
+#pragma warning( pop )
+				if( fnGetProductInfo != NULL )
+				{
+					DWORD Type;
+					fnGetProductInfo( OsVersionInfo.dwMajorVersion, OsVersionInfo.dwMinorVersion, 0, 0, &Type );
+
+					switch( Type )
+					{
+						case PRODUCT_ULTIMATE:
+							out_OSSubVersionLabel = TEXT( "Ultimate Edition" );
+							break;
+						case PRODUCT_PROFESSIONAL:
+							out_OSSubVersionLabel = TEXT( "Professional" );
+							break;
+						case PRODUCT_HOME_PREMIUM:
+							out_OSSubVersionLabel = TEXT( "Home Premium Edition" );
+							break;
+						case PRODUCT_HOME_BASIC:
+							out_OSSubVersionLabel = TEXT( "Home Basic Edition" );
+							break;
+						case PRODUCT_ENTERPRISE:
+							out_OSSubVersionLabel = TEXT( "Enterprise Edition" );
+							break;
+						case PRODUCT_BUSINESS:
+							out_OSSubVersionLabel = TEXT( "Business Edition" );
+							break;
+						case PRODUCT_STARTER:
+							out_OSSubVersionLabel = TEXT( "Starter Edition" );
+							break;
+						case PRODUCT_CLUSTER_SERVER:
+							out_OSSubVersionLabel = TEXT( "Cluster Server Edition" );
+							break;
+						case PRODUCT_DATACENTER_SERVER:
+							out_OSSubVersionLabel = TEXT( "Datacenter Edition" );
+							break;
+						case PRODUCT_DATACENTER_SERVER_CORE:
+							out_OSSubVersionLabel = TEXT( "Datacenter Edition (core installation)" );
+							break;
+						case PRODUCT_ENTERPRISE_SERVER:
+							out_OSSubVersionLabel = TEXT( "Enterprise Edition" );
+							break;
+						case PRODUCT_ENTERPRISE_SERVER_CORE:
+							out_OSSubVersionLabel = TEXT( "Enterprise Edition (core installation)" );
+							break;
+						case PRODUCT_ENTERPRISE_SERVER_IA64:
+							out_OSSubVersionLabel = TEXT( "Enterprise Edition for Itanium-based Systems" );
+							break;
+						case PRODUCT_SMALLBUSINESS_SERVER:
+							out_OSSubVersionLabel = TEXT( "Small Business Server" );
+							break;
+						case PRODUCT_SMALLBUSINESS_SERVER_PREMIUM:
+							out_OSSubVersionLabel = TEXT( "Small Business Server Premium Edition" );
+							break;
+						case PRODUCT_STANDARD_SERVER:
+							out_OSSubVersionLabel = TEXT( "Standard Edition" );
+							break;
+						case PRODUCT_STANDARD_SERVER_CORE:
+							out_OSSubVersionLabel = TEXT( "Standard Edition (core installation)" );
+							break;
+						case PRODUCT_WEB_SERVER:
+							out_OSSubVersionLabel = TEXT( "Web Server Edition" );
+							break;
+					}
+				}
+				else
+				{
+					out_OSSubVersionLabel = TEXT( "(type unknown)" );
+					ErrorCode |= (int32)ERROR_GETPRODUCTINFO_FAILED;
+				}
+			}
+				break;
+			default:
+				ErrorCode |= ERROR_UNKNOWNVERSION;
+		}
+
+#if 0
+		// THIS BIT ADDS THE SERVICE PACK INFO TO THE EDITION STRING
+		// Append service pack info
+		if( OsVersionInfo.szCSDVersion[0] != 0 )
+		{
+			OSSubVersionLabel += FString::Printf( TEXT( " (%s)" ), OsVersionInfo.szCSDVersion );
+		}
+#else
+		// THIS BIT USES SERVICE PACK INFO ONLY
+		out_OSSubVersionLabel = OsVersionInfo.szCSDVersion;
+#endif
+	}
+	else
+	{
+		ErrorCode |= ERROR_GETVERSIONEX_FAILED;
+	}
+
+	return ErrorCode;
+}
+#include "HideWindowsPlatformTypes.h"
+
 /** 
  * Whether support for integrating into the firewall is there
  */
@@ -1625,6 +1854,7 @@ public:
 		if(bHasCPUIDInstruction)
 		{
 			Vendor = QueryCPUVendor();
+			Brand = QueryCPUBrand();
 			CPUInfo = QueryCPUInfo();
 			CacheLineSize = QueryCacheLineSize();
 		}
@@ -1635,28 +1865,50 @@ public:
 	 *
 	 * @returns True if this CPU supports __cpuid instruction. False otherwise.
 	 */
-	static bool HasCPUIDInstruction() { return CPUIDStaticCache.bHasCPUIDInstruction; }
+	static bool HasCPUIDInstruction()
+	{
+		return CPUIDStaticCache.bHasCPUIDInstruction;
+	}
 
 	/**
 	 * Gets pre-cached CPU vendor name.
 	 *
 	 * @returns CPU vendor name.
 	 */
-	static const FString& GetVendor() { return CPUIDStaticCache.Vendor; }
+	static const FString& GetVendor()
+	{
+		return CPUIDStaticCache.Vendor;
+	}
+
+	/**
+	* Gets pre-cached CPU brand string.
+	*
+	* @returns CPU brand string.
+	*/
+	static const FString& GetBrand()
+	{
+		return CPUIDStaticCache.Brand;
+	}
 
 	/**
 	 * Gets __cpuid CPU info.
 	 *
 	 * @returns CPU info unsigned int queried using __cpuid.
 	 */
-	static uint32 GetCPUInfo() { return CPUIDStaticCache.CPUInfo; }
+	static uint32 GetCPUInfo()
+	{
+		return CPUIDStaticCache.CPUInfo;
+	}
 
 	/** 
 	 * Gets cache line size.
 	 *
 	 * @returns Cache line size.
 	 */
-	static int32 GetCacheLineSize() { return CPUIDStaticCache.CacheLineSize; }
+	static int32 GetCacheLineSize()
+	{
+		return CPUIDStaticCache.CacheLineSize;
+	}
 
 private:
 	/**
@@ -1711,6 +1963,35 @@ private:
 		return ANSI_TO_TCHAR(VendorResult.Buffer);
 	}
 
+	/**
+	 * Queries brand string using __cpuid instruction.
+	 *
+	 * @returns CPU brand string.
+	 */
+	static FString QueryCPUBrand()
+	{
+		// @see for more information http://msdn.microsoft.com/en-us/library/vstudio/hskdteyh(v=vs.100).aspx
+		ANSICHAR BrandString[0x40] = {0};
+		int32 CPUInfo[4] = {-1};
+		const SIZE_T CPUInfoSize = sizeof( CPUInfo );
+
+		__cpuid( CPUInfo, 0x80000000 );
+		const uint32 MaxExtIDs = CPUInfo[0];
+
+		if( MaxExtIDs >= 0x80000004 )
+		{
+			const uint32 FirstBrandString = 0x80000002;
+			const uint32 NumBrandStrings = 3;
+			for( uint32 Index = 0; Index < NumBrandStrings; Index++ )
+			{
+				__cpuid( CPUInfo, FirstBrandString + Index );
+				FPlatformMemory::Memcpy( BrandString + CPUInfoSize * Index, CPUInfo, CPUInfoSize );
+			}
+		}
+
+		return ANSI_TO_TCHAR( BrandString );
+	}
+
 	/**
 	 * Queries CPU info using __cpuid instruction.
 	 *
@@ -1755,6 +2036,9 @@ private:
 	/** Vendor of the CPU. */
 	FString Vendor;
 
+	/** CPU brand. */
+	FString Brand;
+
 	/** CPU info from __cpuid. */
 	uint32 CPUInfo;
 
@@ -1775,6 +2059,47 @@ FString FWindowsPlatformMisc::GetCPUVendor()
 	return FCPUIDQueriedData::GetVendor();
 }
 
+FString FWindowsPlatformMisc::GetCPUBrand()
+{
+	return FCPUIDQueriedData::GetBrand();
+}
+
+#include "AllowWindowsPlatformTypes.h"
+FString FWindowsPlatformMisc::GetPrimaryGPUBrand()
+{
+	static FString PrimaryGPUBrand;
+	if( PrimaryGPUBrand.IsEmpty() )
+	{
+		// Find primary display adapter and get the device name.
+		PrimaryGPUBrand = FGenericPlatformMisc::GetPrimaryGPUBrand();
+
+		DISPLAY_DEVICE DisplayDevice;
+		DisplayDevice.cb = sizeof( DisplayDevice );
+		DWORD DeviceIndex = 0;
+
+		while( EnumDisplayDevices( 0, DeviceIndex, &DisplayDevice, 0 ) )
+		{
+			if( (DisplayDevice.StateFlags & (DISPLAY_DEVICE_ATTACHED_TO_DESKTOP | DISPLAY_DEVICE_PRIMARY_DEVICE)) > 0 )
+			{
+				PrimaryGPUBrand = DisplayDevice.DeviceString;
+				break;
+			}
+
+			FMemory::MemZero( DisplayDevice );
+			DisplayDevice.cb = sizeof( DisplayDevice );
+			DeviceIndex++;
+		}
+
+	}
+	return PrimaryGPUBrand;
+}
+#include "HideWindowsPlatformTypes.h"
+
+void FWindowsPlatformMisc::GetOSVersions( FString& out_OSVersionLabel, FString& out_OSSubVersionLabel )
+{
+	FWindowsOSVersionHelper::GetOSVersions( out_OSVersionLabel, out_OSSubVersionLabel );
+}
+
 uint32 FWindowsPlatformMisc::GetCPUInfo()
 {
 	return FCPUIDQueriedData::GetCPUInfo();
diff --git a/Engine/Source/Runtime/Core/Private/Windows/WindowsPlatformSurvey.cpp b/Engine/Source/Runtime/Core/Private/Windows/WindowsPlatformSurvey.cpp
--- a/Engine/Source/Runtime/Core/Private/Windows/WindowsPlatformSurvey.cpp
+++ b/Engine/Source/Runtime/Core/Private/Windows/WindowsPlatformSurvey.cpp
@@ -681,231 +681,42 @@ bool FWindowsPlatformSurvey::GetSubComponentIndex( IProvideWinSATResultsInfo* Wi
 
 void FWindowsPlatformSurvey::GetOSVersionLabels(const SYSTEM_INFO& SystemInfo, FHardwareSurveyResults& Results)
 {
-	OSVERSIONINFOEX OsVersionInfo = { 0 };
-	OsVersionInfo.dwOSVersionInfoSize = sizeof( OSVERSIONINFOEX );
-	FString OSVersionLabel = TEXT("Windows (unknown version)");
+	FString OSVersionLabel;
 	FString OSSubVersionLabel;
-#pragma warning(disable : 4996) // 'function' was declared deprecated
-	if (GetVersionEx( ( LPOSVERSIONINFO )&OsVersionInfo ))
-#pragma warning(default : 4996)
+	const int32 ErrorCode = FWindowsOSVersionHelper::GetOSVersions( OSVersionLabel, OSSubVersionLabel );
 
+	if( ErrorCode & FWindowsOSVersionHelper::ERROR_GETPRODUCTINFO_FAILED )
 	{
-		bool bIsInvalidVersion = false;
-
-		switch (OsVersionInfo.dwMajorVersion)
-		{
-		case 5:
-			switch (OsVersionInfo.dwMinorVersion)
-			{
-			case 0:
-				OSVersionLabel = TEXT("Windows 2000");
-				if (OsVersionInfo.wProductType == VER_NT_WORKSTATION)
-				{
-					OSSubVersionLabel = TEXT("Professional");
-				}
-				else
-				{
-					if ( OsVersionInfo.wSuiteMask & VER_SUITE_DATACENTER )
-					{
-						OSSubVersionLabel = TEXT("Datacenter Server");
-					}
-					else if ( OsVersionInfo.wSuiteMask & VER_SUITE_ENTERPRISE )
-					{
-						OSSubVersionLabel = TEXT("Advanced Server");
-					}
-					else
-					{
-						OSSubVersionLabel = TEXT("Server");
-					}
-				}
-				break;
-			case 1:
-				OSVersionLabel = TEXT("Windows XP");
-				if (OsVersionInfo.wSuiteMask & VER_SUITE_PERSONAL)
-				{
-					OSSubVersionLabel = TEXT("Home Edition");
-				}
-				else
-				{
-					OSSubVersionLabel = TEXT("Professional");
-				}
-				break;
-			case 2:
-				if (GetSystemMetrics(SM_SERVERR2))
-				{
-					OSVersionLabel = TEXT("Windows Server 2003 R2");
-				}
-				else if (OsVersionInfo.wSuiteMask & VER_SUITE_STORAGE_SERVER)
-				{
-					OSVersionLabel = TEXT("Windows Storage Server 2003");
-				}
-				else if (OsVersionInfo.wSuiteMask & VER_SUITE_WH_SERVER)
-				{
-					OSVersionLabel = TEXT("Windows Home Server");
-				}
-				else if (OsVersionInfo.wProductType == VER_NT_WORKSTATION && SystemInfo.wProcessorArchitecture == PROCESSOR_ARCHITECTURE_AMD64)
-				{
-					OSVersionLabel = TEXT("Windows XP");
-					OSSubVersionLabel = TEXT("Professional x64 Edition");
-				}
-				else
-				{
-					OSVersionLabel = TEXT("Windows Server 2003");
-				}
-				break;
-			default:
-				bIsInvalidVersion = true;
-			}
-			break;
-		case 6:
-			switch (OsVersionInfo.dwMinorVersion)
-			{
-			case 0:
-				if (OsVersionInfo.wProductType == VER_NT_WORKSTATION)
-				{
-					OSVersionLabel = TEXT("Windows Vista");
-				}
-				else
-				{
-					OSVersionLabel = TEXT("Windows Server 2008");
-				}
-				break;
-			case 1:
-				if (OsVersionInfo.wProductType == VER_NT_WORKSTATION)
-				{
-					OSVersionLabel = TEXT("Windows 7");
-				}
-				else
-				{
-					OSVersionLabel = TEXT("Windows Server 2008 R2");
-				}
-				break;
-			case 2:
-				if (OsVersionInfo.wProductType == VER_NT_WORKSTATION)
-				{
-					OSVersionLabel = TEXT("Windows 8");
-				}
-				else
-				{
-					OSVersionLabel = TEXT("Windows Server 2012");
-				}
-				break;
-			default:
-				bIsInvalidVersion = true;
-			}
-			{
-				#pragma warning( push )
-				#pragma warning( disable: 4191 )	// unsafe conversion from 'type of expression' to 'type required'
-				typedef BOOL (WINAPI *LPFN_GETPRODUCTINFO)(DWORD, DWORD, DWORD, DWORD, PDWORD);
-				LPFN_GETPRODUCTINFO fnGetProductInfo = (LPFN_GETPRODUCTINFO) GetProcAddress(GetModuleHandle(TEXT("kernel32.dll")), "GetProductInfo");
-				#pragma warning( pop )
-				if (fnGetProductInfo != NULL)
-				{
-					DWORD Type;
-					fnGetProductInfo(OsVersionInfo.dwMajorVersion, OsVersionInfo.dwMinorVersion, 0, 0, &Type);
-
-					switch( Type )
-					{
-					case PRODUCT_ULTIMATE:
-						OSSubVersionLabel = TEXT("Ultimate Edition" );
-						break;
-					case PRODUCT_PROFESSIONAL:
-						OSSubVersionLabel = TEXT("Professional" );
-						break;
-					case PRODUCT_HOME_PREMIUM:
-						OSSubVersionLabel = TEXT("Home Premium Edition" );
-						break;
-					case PRODUCT_HOME_BASIC:
-						OSSubVersionLabel = TEXT("Home Basic Edition" );
-						break;
-					case PRODUCT_ENTERPRISE:
-						OSSubVersionLabel = TEXT("Enterprise Edition" );
-						break;
-					case PRODUCT_BUSINESS:
-						OSSubVersionLabel = TEXT("Business Edition" );
-						break;
-					case PRODUCT_STARTER:
-						OSSubVersionLabel = TEXT("Starter Edition" );
-						break;
-					case PRODUCT_CLUSTER_SERVER:
-						OSSubVersionLabel = TEXT("Cluster Server Edition" );
-						break;
-					case PRODUCT_DATACENTER_SERVER:
-						OSSubVersionLabel = TEXT("Datacenter Edition" );
-						break;
-					case PRODUCT_DATACENTER_SERVER_CORE:
-						OSSubVersionLabel = TEXT("Datacenter Edition (core installation)" );
-						break;
-					case PRODUCT_ENTERPRISE_SERVER:
-						OSSubVersionLabel = TEXT("Enterprise Edition" );
-						break;
-					case PRODUCT_ENTERPRISE_SERVER_CORE:
-						OSSubVersionLabel = TEXT("Enterprise Edition (core installation)" );
-						break;
-					case PRODUCT_ENTERPRISE_SERVER_IA64:
-						OSSubVersionLabel = TEXT("Enterprise Edition for Itanium-based Systems" );
-						break;
-					case PRODUCT_SMALLBUSINESS_SERVER:
-						OSSubVersionLabel = TEXT("Small Business Server" );
-						break;
-					case PRODUCT_SMALLBUSINESS_SERVER_PREMIUM:
-						OSSubVersionLabel = TEXT("Small Business Server Premium Edition" );
-						break;
-					case PRODUCT_STANDARD_SERVER:
-						OSSubVersionLabel = TEXT("Standard Edition" );
-						break;
-					case PRODUCT_STANDARD_SERVER_CORE:
-						OSSubVersionLabel = TEXT("Standard Edition (core installation)" );
-						break;
-					case PRODUCT_WEB_SERVER:
-						OSSubVersionLabel = TEXT("Web Server Edition" );
-						break;
-					}
-				}
-				else
-				{
-					OSSubVersionLabel = TEXT("(type unknown)");
-					Results.ErrorCount++;
-					WriteFStringToResults(Results.LastSurveyError, TEXT("Failed to get GetProductInfo() function from GetProcAddress()."));
-					WriteFStringToResults(Results.LastSurveyErrorDetail, TEXT(""));
-				}
-			}
-			break;
-		default:
-			bIsInvalidVersion = true;
-		}
+		Results.ErrorCount++;
+		WriteFStringToResults( Results.LastSurveyError, TEXT( "Failed to get GetProductInfo() function from GetProcAddress()." ) );
+		WriteFStringToResults( Results.LastSurveyErrorDetail, TEXT( "" ) );
+	}
 
-#if 0
-		// THIS BIT ADDS THE SERVICE PACK INFO TO THE EDITION STRING
-		// Append service pack info
-		if (OsVersionInfo.szCSDVersion[0] != 0)
-		{
-			OSSubVersionLabel += FString::Printf(TEXT(" (%s)"), OsVersionInfo.szCSDVersion);
-		}
-#else
-		// THIS BIT USES SERVICE PACK INFO ONLY
-		OSSubVersionLabel = OsVersionInfo.szCSDVersion;
-#endif
+	if( ErrorCode & FWindowsOSVersionHelper::ERROR_UNKNOWNVERSION )
+	{
+		OSVERSIONINFOEX OsVersionInfo = {0};
+		OsVersionInfo.dwOSVersionInfoSize = sizeof( OSVERSIONINFOEX );
+#pragma warning(disable : 4996) // 'function' was declared deprecated
+		GetVersionEx( (LPOSVERSIONINFO)&OsVersionInfo );
+#pragma warning(default : 4996)
 
-		if (bIsInvalidVersion)
-		{
-			UE_LOG(LogWindows, Warning, TEXT("FWindowsPlatformSurvey::GetOSVersionLabel() unknown Windows version info from GetVersionEx()") );
-			Results.ErrorCount++;
-			WriteFStringToResults(Results.LastSurveyError, TEXT("GetVersionEx() returned unknown version"));
-			WriteFStringToResults(Results.LastSurveyErrorDetail, FString::Printf(TEXT("dwMajorVersion: %d  dwMinorVersion: %d"), OsVersionInfo.dwMajorVersion, OsVersionInfo.dwMinorVersion));
-		}
+		UE_LOG( LogWindows, Warning, TEXT( "FWindowsPlatformSurvey::GetOSVersionLabel() unknown Windows version info from GetVersionEx()" ) );
+		Results.ErrorCount++;
+		WriteFStringToResults( Results.LastSurveyError, TEXT( "GetVersionEx() returned unknown version" ) );
+		WriteFStringToResults( Results.LastSurveyErrorDetail, FString::Printf( TEXT( "dwMajorVersion: %d  dwMinorVersion: %d" ), OsVersionInfo.dwMajorVersion, OsVersionInfo.dwMinorVersion ) );
 	}
-	else
+
+	if( ErrorCode & FWindowsOSVersionHelper::ERROR_GETVERSIONEX_FAILED )
 	{
-		uint32 ErrorCode = FPlatformMisc::GetLastError();
-		UE_LOG(LogWindows, Warning, TEXT("FWindowsPlatformSurvey::GetOSVersionLabel() failed to get Windows version info from GetVersionEx()") );
+		const uint32 LastError = FPlatformMisc::GetLastError();
+		UE_LOG( LogWindows, Warning, TEXT( "FWindowsPlatformSurvey::GetOSVersionLabel() failed to get Windows version info from GetVersionEx()" ) );
 		Results.ErrorCount++;
-		WriteFStringToResults(Results.LastSurveyError, TEXT("GetVersionEx() failed"));
-		WriteFStringToResults(Results.LastSurveyErrorDetail, FString::Printf(TEXT("ErrorCode: 0x%0x"), ErrorCode));
+		WriteFStringToResults( Results.LastSurveyError, TEXT( "GetVersionEx() failed" ) );
+		WriteFStringToResults( Results.LastSurveyErrorDetail, FString::Printf( TEXT( "ErrorCode: 0x%0x" ), LastError ) );
 	}
 
-	WriteFStringToResults(Results.OSVersion, OSVersionLabel);
-	WriteFStringToResults(Results.OSSubVersion, OSSubVersionLabel);
+	WriteFStringToResults( Results.OSVersion, OSVersionLabel );
+	WriteFStringToResults( Results.OSSubVersion, OSSubVersionLabel );
 }
 
 bool FWindowsPlatformSurvey::GetLineFollowing(const FString& Token, const TArray<FString>& InLines, FString& OutString, int32 NthHit)
diff --git a/Engine/Source/Runtime/Core/Public/GenericPlatform/GenericPlatformMisc.h b/Engine/Source/Runtime/Core/Public/GenericPlatform/GenericPlatformMisc.h
--- a/Engine/Source/Runtime/Core/Public/GenericPlatform/GenericPlatformMisc.h
+++ b/Engine/Source/Runtime/Core/Public/GenericPlatform/GenericPlatformMisc.h
@@ -225,7 +225,7 @@ struct CORE_API FGenericPlatformMisc
 
 	/**
 	 * Installs handler for the unexpected (due to error) termination of the program,
-     * including, but not limited to, crashes.
+	 * including, but not limited to, crashes.
 	 *
 	 */
 	static void SetCrashHandler(void (* CrashHandler)(const FGenericCrashContext & Context))
@@ -294,6 +294,30 @@ struct CORE_API FGenericPlatformMisc
 		}
 	}
 
+	/**
+	* Uses cpuid instruction to get the vendor string
+	*
+	* @return	CPU vendor name
+	*/
+	static FString GetCPUVendor();
+
+	/** 
+	 * Uses cpuid instruction to get the CPU brand string
+	 *
+	 * @return	CPU brand string
+	 */
+	static FString GetCPUBrand();
+
+	/** 
+	 * @return primary GPU brand string
+	 */
+	static FString GetPrimaryGPUBrand();
+
+	/**
+	 * Gets the OS Version and OS Subversion.
+	 */
+	static void GetOSVersions( FString& out_OSVersionLabel, FString& out_OSSubVersionLabel );
+
 	static bool SupportsMessaging()
 	{
 		return true;
diff --git a/Engine/Source/Runtime/Core/Public/Windows/WindowsPlatformMisc.h b/Engine/Source/Runtime/Core/Public/Windows/WindowsPlatformMisc.h
--- a/Engine/Source/Runtime/Core/Public/Windows/WindowsPlatformMisc.h
+++ b/Engine/Source/Runtime/Core/Public/Windows/WindowsPlatformMisc.h
@@ -7,6 +7,20 @@
 
 #pragma once
 
+/** Helper struct used to get the string version of the Windows version. */
+struct FWindowsOSVersionHelper
+{
+	enum ErrorCodes
+	{
+		SUCCEEDED = 0,
+		ERROR_UNKNOWNVERSION = 1,
+		ERROR_GETPRODUCTINFO_FAILED = 2,
+		ERROR_GETVERSIONEX_FAILED = 4,
+	};
+
+	static int32 GetOSVersions( FString& out_OSVersion, FString& out_OSSubVersion );
+};
+
 /**
 * Windows implementation of the misc OS functions
 **/
@@ -136,12 +150,10 @@ struct CORE_API FWindowsPlatformMisc : public FGenericPlatformMisc
 	 */
 	static bool HasCPUIDInstruction();
 
-	/** 
-	 * Uses cpuid instruction to get the vendor string
-	 *
-	 * @return	CPU vendor name
-	 */
 	static FString GetCPUVendor();
+	static FString GetCPUBrand();
+	static FString GetPrimaryGPUBrand();
+	static void GetOSVersions( FString& out_OSVersionLabel, FString& out_OSSubVersionLabel );
 
 	/**
 	 * Uses cpuid instruction to get the vendor string
