Fixed recently added compiler issues for Windows XP.

diff --git a/Engine/Source/Programs/UnrealBuildTool/Windows/UEBuildWindows.cs b/Engine/Source/Programs/UnrealBuildTool/Windows/UEBuildWindows.cs
--- a/Engine/Source/Programs/UnrealBuildTool/Windows/UEBuildWindows.cs
+++ b/Engine/Source/Programs/UnrealBuildTool/Windows/UEBuildWindows.cs
@@ -338,10 +338,8 @@ namespace UnrealBuildTool
 			InBuildTarget.GlobalCompileEnvironment.Config.Definitions.Add("WIN32=1");
 
 			// Win32 XP is only supported at this time.
-			// The XP compilation is currently only supported with a VS2013 build.
 			SupportWindowsXP = false;
-				//&& InBuildTarget.Platform == UnrealTargetPlatform.Win32
-				//&& Compiler == WindowsCompiler.VisualStudio2013;
+				//&& InBuildTarget.Platform == UnrealTargetPlatform.Win32;
 
 			if( SupportWindowsXP )
 			{
diff --git a/Engine/Source/Runtime/Core/Private/Windows/WindowsPlatformProcess.cpp b/Engine/Source/Runtime/Core/Private/Windows/WindowsPlatformProcess.cpp
--- a/Engine/Source/Runtime/Core/Private/Windows/WindowsPlatformProcess.cpp
+++ b/Engine/Source/Runtime/Core/Private/Windows/WindowsPlatformProcess.cpp
@@ -13,6 +13,10 @@
 	#include <LM.h>
 	#include <tlhelp32.h>
 
+#if WINVER == 0x0502
+	#include <Psapi.h>
+#endif
+
 	namespace ProcessConstants
 	{
 		uint32 WIN_STD_INPUT_HANDLE = STD_INPUT_HANDLE;
@@ -24,6 +28,10 @@
 #include "HideWindowsPlatformTypes.h"
 #include "WindowsPlatformMisc.h"
 
+#if WINVER == 0x0502
+	#pragma comment(lib, "psapi.lib")
+#endif
+
 // static variables
 TArray<FString> FWindowsPlatformProcess::DllDirectoryStack;
 
@@ -444,7 +452,13 @@ FString FWindowsPlatformProcess::GetApplicationName( uint32 ProcessId )
 		
 		int32 InOutSize = ProcessNameBufferSize;
 		checkAtCompileTime(sizeof(::DWORD) == sizeof(int32), "DWORD size doesn't match int32.  Is it the future or the past?");
+
+#if WINVER == 0x0502
+		GetProcessImageFileName(ProcessHandle, ProcessNameBuffer, InOutSize);
+#else
 		QueryFullProcessImageName(ProcessHandle, 0, ProcessNameBuffer, (PDWORD)(&InOutSize));
+#endif
+
 		Output = ProcessNameBuffer;
 	}
 
diff --git a/Engine/Source/Runtime/Core/Private/Windows/WindowsRunnableThread.cpp b/Engine/Source/Runtime/Core/Private/Windows/WindowsRunnableThread.cpp
--- a/Engine/Source/Runtime/Core/Private/Windows/WindowsRunnableThread.cpp
+++ b/Engine/Source/Runtime/Core/Private/Windows/WindowsRunnableThread.cpp
@@ -10,7 +10,11 @@ uint32 FRunnableThreadWin::GuardedRun()
 	uint32 ExitCode = 1;
 
 	FPlatformProcess::SetThreadAffinityMask(ThreadAffintyMask);
+
+#if WINVER >= 0x0600
 	UE_LOG(LogThreadingWindows, Log, TEXT("Runnable thread %s is on Process %d."), *ThreadName  , static_cast<uint32>(::GetCurrentProcessorNumber()) );
+#endif
+
 
 #if !PLATFORM_SEH_EXCEPTIONS_DISABLED
 	if( !FPlatformMisc::IsDebuggerPresent() || GAlwaysReportCrash )
