Only need to forward declare FString once

diff --git a/Engine/Source/Runtime/Core/Public/Core.h b/Engine/Source/Runtime/Core/Public/Core.h
--- a/Engine/Source/Runtime/Core/Public/Core.h
+++ b/Engine/Source/Runtime/Core/Public/Core.h
@@ -99,10 +99,8 @@ class FPackageInfo;
 class ITransaction;
 class FUnknown;
 class FRepLink;
-class FString;
 class FText;
 class FMalloc;
-class FString;
 class FConfigCacheIni;
 class FConfigFile;
 class IFileManager;
