Send FRotator as compressed shorts rather than bytes

* The 8 bit precision is causing issues in certain places
* This is until we can specify how much precision we want more easily

diff --git a/Engine/Source/Runtime/Core/Private/Math/UnrealMath.cpp b/Engine/Source/Runtime/Core/Private/Math/UnrealMath.cpp
--- a/Engine/Source/Runtime/Core/Private/Math/UnrealMath.cpp
+++ b/Engine/Source/Runtime/Core/Private/Math/UnrealMath.cpp
@@ -54,7 +54,7 @@ bool FVector2D::NetSerialize(FArchive& Ar, class UPackageMap* Map, bool& bOutSuc
 
 bool FRotator::NetSerialize(FArchive& Ar, class UPackageMap* Map, bool& bOutSuccess)
 {
-	SerializeCompressed( Ar );
+	SerializeCompressedShort( Ar );
 	bOutSuccess = true;
 	return true;
 }
@@ -106,6 +106,53 @@ void FRotator::SerializeCompressed( FArchive& Ar )
 	}
 }
 
+void FRotator::SerializeCompressedShort( FArchive& Ar )
+{
+	uint16 ShortPitch = FRotator::CompressAxisToShort(Pitch);
+	uint16 ShortYaw = FRotator::CompressAxisToShort(Yaw);
+	uint16 ShortRoll = FRotator::CompressAxisToShort(Roll);
+
+	uint8 B = (ShortPitch!=0);
+	Ar.SerializeBits( &B, 1 );
+	if( B )
+	{
+		Ar << ShortPitch;
+	}
+	else
+	{
+		ShortPitch = 0;
+	}
+
+	B = (ShortYaw!=0);
+	Ar.SerializeBits( &B, 1 );
+	if( B )
+	{
+		Ar << ShortYaw;
+	}
+	else
+	{
+		ShortYaw = 0;
+	}
+
+	B = (ShortRoll!=0);
+	Ar.SerializeBits( &B, 1 );
+	if( B )
+	{
+		Ar << ShortRoll;
+	}
+	else
+	{
+		ShortRoll = 0;
+	}
+
+	if( Ar.IsLoading() )
+	{
+		Pitch = FRotator::DecompressAxisFromShort(ShortPitch);
+		Yaw	= FRotator::DecompressAxisFromShort(ShortYaw);
+		Roll = FRotator::DecompressAxisFromShort(ShortRoll);
+	}
+}
+
 CORE_API FRotator FVector::Rotation() const
 {
 	FRotator R;
diff --git a/Engine/Source/Runtime/Core/Public/Math/Rotator.h b/Engine/Source/Runtime/Core/Public/Math/Rotator.h
--- a/Engine/Source/Runtime/Core/Public/Math/Rotator.h
+++ b/Engine/Source/Runtime/Core/Public/Math/Rotator.h
@@ -32,6 +32,13 @@ public:
 	 */
 	CORE_API void SerializeCompressed( FArchive& Ar );
 
+	/**
+	 * Serializes the rotator compressed for e.g. network transmission (use shorts though).
+	 * 
+	 * @param	Ar	Archive to serialize to/ from
+	 */
+	CORE_API void SerializeCompressedShort( FArchive& Ar );
+
 	CORE_API bool NetSerialize(FArchive& Ar, class UPackageMap* Map, bool& bOutSuccess);
 
 	/** Constructor.*/
