Changes for packing OBB into APK file.
- Added functionality to GameActivity to query for the asset manager and the mode we are using for OBB data
- Added UI element to toggle OBB packing into APK in the project packaging config panel
- Added code to add a command line switch to turn on packing during building
- Changed build system so that it will write out a JavaBuildConfig file so we know at runtime where to find the data
- Change the build system to copy and rename the OBB so that it is in the APK's asset directory uncompressed
- Added IFileHandle implementation to handle reading from an Asset based OBB and changed the OBB locating code so deal with looking for the OBB in the APK
Submitted on behalf of Robert Jones.

diff --git a/Engine/Build/Android/Java/src/com/epicgames/ue4/GameActivity.java b/Engine/Build/Android/Java/src/com/epicgames/ue4/GameActivity.java
--- a/Engine/Build/Android/Java/src/com/epicgames/ue4/GameActivity.java
+++ b/Engine/Build/Android/Java/src/com/epicgames/ue4/GameActivity.java
@@ -14,6 +14,7 @@ import android.app.Dialog;
 import android.widget.EditText;
 import android.content.DialogInterface;
 import android.content.Intent;
+import android.content.res.AssetManager;
 import android.content.res.Configuration;
 import android.content.IntentSender.SendIntentException;
 
@@ -40,7 +41,7 @@ import com.google.android.gms.ads.AdSize;
 
 // TODO: use the resources from the UE4 lib project once we've got the packager up and running
 //import com.epicgames.ue4.R;
-
+import com.epicgames.ue4.JavaBuildSettings;
 
 //Extending NativeActivity so that this Java class is instantiated
 //from the beginning of the program.  This will allow the user
@@ -58,6 +59,9 @@ public class GameActivity extends NativeActivity implements GoogleApiClient.Conn
 	
 	GameActivity _activity;	
 	AlertDialog alert;
+
+	/** AssetManger reference - populated on start up and used when the OBB is packed into the APK */
+	private AssetManager			AssetManagerReference;
 	
 	private GoogleApiClient googleClient;
 	private boolean bResolvingGoogleServicesError = false;
@@ -172,6 +176,9 @@ public class GameActivity extends NativeActivity implements GoogleApiClient.Conn
 			_activity.setRequestedOrientation( android.content.pm.ActivityInfo.SCREEN_ORIENTATION_LANDSCAPE );
 		}
 		
+		// Grab a reference to the asset manager
+		AssetManagerReference = this.getAssets();
+
 		// tell the engine if this is a portrait app
 		nativeSetGlobalActivity();
 		nativeSetWindowInfo(getResources().getConfiguration().orientation == Configuration.ORIENTATION_PORTRAIT);
@@ -600,6 +607,23 @@ public class GameActivity extends NativeActivity implements GoogleApiClient.Conn
 		});
 	}
 
+	public AssetManager AndroidThunkJava_GetAssetManager()
+	{
+		if(AssetManagerReference == null)
+		{
+			Log.debug("No reference to asset manager found!");
+		}
+
+		return AssetManagerReference;
+	}
+
+	public static boolean isOBBInAPK()
+	{
+		return JavaBuildSettings.PackageType.AMAZON == JavaBuildSettings.PACKAGING;
+	}
+
+
+
 	public native boolean nativeIsShippingBuild();
 	public native void nativeSetGlobalActivity();
 	public native void nativeSetWindowInfo(boolean bIsPortrait);
diff --git a/Engine/Source/Editor/MainFrame/Private/Frame/MainFrameActions.cpp b/Engine/Source/Editor/MainFrame/Private/Frame/MainFrameActions.cpp
--- a/Engine/Source/Editor/MainFrame/Private/Frame/MainFrameActions.cpp
+++ b/Engine/Source/Editor/MainFrame/Private/Frame/MainFrameActions.cpp
@@ -531,6 +531,11 @@ void FMainFrameActionCallbacks::PackageProject( const FString InPlatformName, co
 	if (PackagingSettings->UsePakFile)
 	{
 		OptionalParams += TEXT(" -pak");
+
+		if (PackagingSettings->UseOBB_InAPK)
+		{
+			OptionalParams += TEXT(" -obbinapk");
+		}
 	}
 
 	if (PackagingSettings->ForDistribution)
diff --git a/Engine/Source/Editor/UnrealEd/Classes/Settings/ProjectPackagingSettings.h b/Engine/Source/Editor/UnrealEd/Classes/Settings/ProjectPackagingSettings.h
--- a/Engine/Source/Editor/UnrealEd/Classes/Settings/ProjectPackagingSettings.h
+++ b/Engine/Source/Editor/UnrealEd/Classes/Settings/ProjectPackagingSettings.h
@@ -104,4 +104,10 @@ public:
 	 */
 	UPROPERTY(config, EditAnywhere, Category=Packaging)
 	bool UsePakFile;
+
+	/**
+	* If enabled, on Android platforms, .pak files are placed inside the APK
+	*/
+	UPROPERTY(config, EditAnywhere, Category = Packaging)
+	bool UseOBB_InAPK;
 };
diff --git a/Engine/Source/Editor/UnrealEd/Private/Settings/SettingsClasses.cpp b/Engine/Source/Editor/UnrealEd/Private/Settings/SettingsClasses.cpp
--- a/Engine/Source/Editor/UnrealEd/Private/Settings/SettingsClasses.cpp
+++ b/Engine/Source/Editor/UnrealEd/Private/Settings/SettingsClasses.cpp
@@ -276,6 +276,7 @@ UProjectPackagingSettings::UProjectPackagingSettings( const class FPostConstruct
 	BuildConfiguration = PPBC_Development;
 	FullRebuild = true;
 	UsePakFile = true;
+	UseOBB_InAPK = false;
 }
 
 
diff --git a/Engine/Source/Programs/AutomationTool/Android/AndroidPlatform.Automation.cs b/Engine/Source/Programs/AutomationTool/Android/AndroidPlatform.Automation.cs
--- a/Engine/Source/Programs/AutomationTool/Android/AndroidPlatform.Automation.cs
+++ b/Engine/Source/Programs/AutomationTool/Android/AndroidPlatform.Automation.cs
@@ -101,9 +101,10 @@ public class AndroidPlatform : Platform
 		string BatchName = GetFinalBatchName(ApkName, Params);
 
 		// packaging just takes a pak file and makes it the .obb
-		UEBuildDeploy Deploy = UEBuildDeploy.GetBuildDeploy(UnrealTargetPlatform.Android);
+        UEBuildConfiguration.bOBBinAPK = Params.OBBinAPK; // Make sure this setting is sync'd pre-build
+        UEBuildDeploy Deploy = UEBuildDeploy.GetBuildDeploy(UnrealTargetPlatform.Android);
 		Deploy.PrepForUATPackageOrDeploy(Params.ShortProjectName, SC.ProjectRoot, SOName, SC.LocalRoot + "/Engine", Params.Distribution);
-		
+        
 		// first, look for a .pak file in the staged directory
 		string[] PakFiles = Directory.GetFiles(SC.StageDirectory, "*.pak", SearchOption.AllDirectories);
 
@@ -116,10 +117,19 @@ public class AndroidPlatform : Platform
 		string LocalObbName = GetFinalObbName(ApkName);
 		string DeviceObbName = GetDeviceObbName(ApkName);
 
-		Log("Creating {0} from {1}", LocalObbName, PakFiles[0]);
-		File.Delete(LocalObbName);
-		File.Copy(PakFiles[0], LocalObbName);
+        // Always delete the target OBB file if it exists
+        if (File.Exists(LocalObbName))
+        {
+            File.Delete(LocalObbName);
+        }
+
+        if (!Params.OBBinAPK)
+        {
+            Log("Creating {0} from {1}", LocalObbName, PakFiles[0]);
+            File.Copy(PakFiles[0], LocalObbName);
+        }
 
+        Log("Writing bat for install with {0}", Params.OBBinAPK ? "OBB in APK" : "OBB separate");
 		string PackageName = GetPackageInfo(ApkName, false);
 		// make a batch file that can be used to install the .apk and .obb files
 		string[] BatchLines = new string[] {
@@ -130,7 +140,7 @@ public class AndroidPlatform : Platform
 			"%ADB% shell rm -r /mnt/sdcard/" + Params.ShortProjectName,
 			SC.IsCodeBasedProject ? "" : "%ADB% shell rm /mnt/sdcard/UE4Game/UE4CommandLine.txt", // we need to delete the commandline in UE4Game or it will mess up loading
 			"%ADB% shell rm -r /mnt/sdcard/obb/" + PackageName,
-			"%ADB% push " + Path.GetFileName(LocalObbName) + " " + DeviceObbName,
+			Params.OBBinAPK ? "" : "%ADB% push " + Path.GetFileName(LocalObbName) + " " + DeviceObbName,
 		};
 		File.WriteAllLines(BatchName, BatchLines);
 
@@ -154,13 +164,17 @@ public class AndroidPlatform : Platform
 		{
 			throw new AutomationException("ARCHIVE FAILED - {0} was not found", ApkName);
 		}
-		if (!FileExists(ObbName))
+		if (!Params.OBBinAPK && !FileExists(ObbName))
 		{
 			throw new AutomationException("ARCHIVE FAILED - {0} was not found", ObbName);
 		}
 
 		SC.ArchiveFiles(Path.GetDirectoryName(ApkName), Path.GetFileName(ApkName));
-		SC.ArchiveFiles(Path.GetDirectoryName(ObbName), Path.GetFileName(ObbName));
+        if(!Params.OBBinAPK)
+		{
+            SC.ArchiveFiles(Path.GetDirectoryName(ObbName), Path.GetFileName(ObbName));
+        }
+
 		SC.ArchiveFiles(Path.GetDirectoryName(BatchName), Path.GetFileName(BatchName));
 	}
 
diff --git a/Engine/Source/Programs/AutomationTool/ProjectParams.cs b/Engine/Source/Programs/AutomationTool/ProjectParams.cs
--- a/Engine/Source/Programs/AutomationTool/ProjectParams.cs
+++ b/Engine/Source/Programs/AutomationTool/ProjectParams.cs
@@ -248,6 +248,7 @@ namespace AutomationTool
 			this.Archive = InParams.Archive;
 			this.ArchiveDirectoryParam = InParams.ArchiveDirectoryParam;
 			this.Distribution = InParams.Distribution;
+            this.OBBinAPK = InParams.OBBinAPK;
 		}
 
 		/// <summary>
@@ -321,7 +322,8 @@ namespace AutomationTool
 			bool? Archive = null,
 			string ArchiveDirectoryParam = null,
 			ParamList<string> ProgramTargets = null,
-			bool? Distribution = null
+			bool? Distribution = null,
+            bool? OBBinAPK = null
 			)
 		{
 			//
@@ -404,6 +406,7 @@ namespace AutomationTool
 			this.Archive = GetParamValueIfNotSpecified(Command, Archive, this.Archive, "archive");
 			this.ArchiveDirectoryParam = ParseParamValueIfNotSpecified(Command, ArchiveDirectoryParam, "archivedirectory", String.Empty);
 			this.Distribution = GetParamValueIfNotSpecified(Command, Distribution, this.Distribution, "distribution");
+            this.OBBinAPK = GetParamValueIfNotSpecified(Command, OBBinAPK, this.OBBinAPK, "obbinapk");
 			this.NoDebugInfo = GetParamValueIfNotSpecified(Command, NoDebugInfo, this.NoDebugInfo, "nodebuginfo");
 			this.NoCleanStage = GetParamValueIfNotSpecified(Command, NoCleanStage, this.NoCleanStage, "nocleanstage");
 			this.MapToRun = ParseParamValueIfNotSpecified(Command, MapToRun, "map", String.Empty);
@@ -984,6 +987,9 @@ namespace AutomationTool
 		[Help("distribution", "package for distribution the project")]
 		public bool Distribution { get; set; }
 
+        [Help("obbinapk", "package with OBB data in APK assets directory")]
+        public bool OBBinAPK {get; set; }
+
 		#endregion
 
 		#region Deploy
@@ -1514,6 +1520,7 @@ namespace AutomationTool
 				CommandUtils.Log("ProjectGameExeFilename={0}", ProjectGameExeFilename);
 				CommandUtils.Log("ProjectGameExePath={0}", ProjectGameExePath);
 				CommandUtils.Log("Distribution={0}", Distribution);
+                CommandUtils.Log("OBBinAPK={0}", OBBinAPK);
 				CommandUtils.Log("RawProjectPath={0}", RawProjectPath);
 				CommandUtils.Log("Rocket={0}", Rocket);
 				CommandUtils.Log("Run={0}", Run);
diff --git a/Engine/Source/Programs/UnrealBuildTool/Android/UEDeployAndroid.cs b/Engine/Source/Programs/UnrealBuildTool/Android/UEDeployAndroid.cs
--- a/Engine/Source/Programs/UnrealBuildTool/Android/UEDeployAndroid.cs
+++ b/Engine/Source/Programs/UnrealBuildTool/Android/UEDeployAndroid.cs
@@ -221,6 +221,48 @@ namespace UnrealBuildTool.Android
 			}
 		}
 
+        public string GetUE4BuildFilePath(String EngineDirectory)
+        {
+            return Path.GetFullPath(Path.Combine(EngineDirectory, "Build/Android/Java"));
+        }
+
+        public string GetUE4JavaFilePath(String EngineDirectory)
+        {
+            return Path.GetFullPath(Path.Combine(GetUE4BuildFilePath(EngineDirectory), "src/com/epicgames/ue4"));
+        }
+
+        public string GetUE4JavaBuildSettingsFileName(String EngineDirectory)
+        {
+            return Path.Combine(GetUE4JavaFilePath(EngineDirectory), "JavaBuildSettings.java");
+        }
+
+        public void WriteJavaBuildSettingsFile(string FileName, bool OBBinAPK)
+        {
+             // (!UEBuildConfiguration.bOBBinAPK ? "PackageType.AMAZON" : /*bPackageForGoogle ? "PackageType.GOOGLE" :*/ "PackageType.DEVELOPMENT") + ";\n");
+            string Setting = OBBinAPK ? "AMAZON" : "DEVELOPMENT";
+            if (!File.Exists(FileName) || ShouldWriteJavaBuildSettingsFile(FileName, Setting))
+            {
+                StringBuilder BuildSettings = new StringBuilder("package com.epicgames.ue4;\npublic class JavaBuildSettings\n{\n");
+                BuildSettings.Append("\tpublic enum PackageType {AMAZON, GOOGLE, DEVELOPMENT};\n");
+                BuildSettings.Append("\tpublic static final PackageType PACKAGING = PackageType." + Setting + ";\n");
+                BuildSettings.Append("}\n");
+                File.WriteAllText(FileName, BuildSettings.ToString());
+            }
+            else
+            {
+                Console.WriteLine("::Didn't write config file; contains same data as before.");
+            }
+        }
+
+        public bool ShouldWriteJavaBuildSettingsFile(string FileName, string setting)
+        {
+            var fileContent = File.ReadAllLines(FileName);
+            var packageLine = fileContent[4]; // We know this to be true... because we write it below...
+            int location = packageLine.IndexOf("PACKAGING") + 12 + 12; // + (PACKAGING = ) + ("PackageType.")
+            return String.Compare(setting, packageLine.Substring(location)) != 0;
+        }
+
+
 		private void MakeAPK(string ProjectName, string ProjectDirectory, string OutputPath, string EngineDirectory, bool bForDistribution)
 		{
 			// cache some build product paths
@@ -238,9 +280,13 @@ namespace UnrealBuildTool.Android
 			// set up some directory info
 			string IntermediateAndroidPath = Path.Combine(ProjectDirectory, "Intermediate/Android/");
 			string UE4BuildPath = IntermediateAndroidPath + "APK";
-			string UE4BuildFilesPath = Path.GetFullPath(Path.Combine(EngineDirectory, "Build/Android/Java"));
+            string UE4BuildFilesPath = GetUE4BuildFilePath(EngineDirectory);
 			string GameBuildFilesPath = Path.Combine(ProjectDirectory, "Build/Android");
 
+            // See if we need to create a 'default' Java Build settings file if one doesn't exist (if it does exist we have to assume it has been setup correctly)
+            string UE4JavaBuildSettingsFileName = GetUE4JavaBuildSettingsFileName(EngineDirectory);
+            WriteJavaBuildSettingsFile(UE4JavaBuildSettingsFileName, UEBuildConfiguration.bOBBinAPK);
+
 			// check to see if it's out of date before trying the slow make apk process (look at .so and all Engine and Project build files to be safe)
 			List<String> InputFiles = new List<string>();
 			InputFiles.Add(SourceSOName);
@@ -274,6 +320,39 @@ namespace UnrealBuildTool.Android
 			//Wipe the Intermediate/Build/APK directory first
 			DeleteDirectory(UE4BuildPath);
 
+            // If we are packaging for Amazon then we need to copy the PAK files to the correct location
+            // Currently we'll just support 1 of 'em
+            if (UEBuildConfiguration.bOBBinAPK)
+            {
+                string PAKFileLocation = ProjectDirectory + "/Saved/StagedBuilds/Android/" + ProjectName + "/Content/Paks";
+                Console.WriteLine("Pak location {0}", PAKFileLocation);
+                string PAKFileDestination = UE4BuildPath + "/assets";
+                Console.WriteLine("Pak destination location {0}", PAKFileDestination);
+                if (Directory.Exists(PAKFileLocation))
+                {
+                    Directory.CreateDirectory(UE4BuildPath);
+                    Directory.CreateDirectory(PAKFileDestination);
+                    Console.WriteLine("PAK file exists...");
+                    var pakFiles = Directory.EnumerateFiles(PAKFileLocation, "*.pak", SearchOption.TopDirectoryOnly);
+                    foreach (var s in pakFiles)
+                    {
+                        Console.WriteLine("Found file {0}", s);
+                    }
+
+                    if (pakFiles.Count() > 0)
+                    {
+                        var destFileName =  Path.Combine(PAKFileDestination, Path.GetFileName(pakFiles.ElementAt(0)) + ".png"); // Need a rename to turn off compression
+                        var srcFileName = pakFiles.ElementAt(0);
+                        if(!File.Exists(destFileName) || File.GetLastWriteTimeUtc(destFileName) < File.GetLastWriteTimeUtc(srcFileName))
+                        {
+                            Console.WriteLine("Copying {0} to {1}", srcFileName, destFileName);
+                            File.Copy(srcFileName,destFileName);
+                        }
+                    }
+                }
+                // Do we want to kill the OBB here or not???
+            }
+
 			Dictionary<string, string> Replacements = new Dictionary<string, string>();
 			Replacements.Add("${EXECUTABLE_NAME}", ProjectName);
 
diff --git a/Engine/Source/Programs/UnrealBuildTool/Configuration/UEBuildConfiguration.cs b/Engine/Source/Programs/UnrealBuildTool/Configuration/UEBuildConfiguration.cs
--- a/Engine/Source/Programs/UnrealBuildTool/Configuration/UEBuildConfiguration.cs
+++ b/Engine/Source/Programs/UnrealBuildTool/Configuration/UEBuildConfiguration.cs
@@ -109,6 +109,9 @@ namespace UnrealBuildTool
 		/** True if plugins should be excluded when building target. */
 		public static bool bExcludePlugins;
 
+        /** True if we need to package up Android with the OBB in the APK file */
+        public static bool bOBBinAPK;
+
 		/** Sets the configuration back to defaults. */
 		public static void Reset()
 		{
@@ -138,6 +141,7 @@ namespace UnrealBuildTool
             bUseLoggingInShipping = false;
 			bCompileSteamOSS = true;
 			bCompileMcpOSS = true;
+            bOBBinAPK = false;
 
 			XmlConfigLoader.Load(typeof(UEBuildConfiguration));
 
diff --git a/Engine/Source/Programs/UnrealBuildTool/System/UnrealBuildTool.cs b/Engine/Source/Programs/UnrealBuildTool/System/UnrealBuildTool.cs
--- a/Engine/Source/Programs/UnrealBuildTool/System/UnrealBuildTool.cs
+++ b/Engine/Source/Programs/UnrealBuildTool/System/UnrealBuildTool.cs
@@ -718,6 +718,10 @@ namespace UnrealBuildTool
 						{
 							//ConfigName = Arg;
 						}
+                        else if (LowercaseArg == "-obbinapk")
+                        {
+                            UEBuildConfiguration.bOBBinAPK = true;
+                        }
 						else if (LowercaseArg == "-modulewithsuffix")
 						{
 							bSpecificModulesOnly = true;
diff --git a/Engine/Source/Runtime/Core/Private/Android/AndroidFile.cpp b/Engine/Source/Runtime/Core/Private/Android/AndroidFile.cpp
--- a/Engine/Source/Runtime/Core/Private/Android/AndroidFile.cpp
+++ b/Engine/Source/Runtime/Core/Private/Android/AndroidFile.cpp
@@ -8,6 +8,8 @@
 
 #include <dirent.h>
 #include <jni.h>
+#include <Android/asset_manager.h>
+#include <Android/asset_manager_jni.h>
 
 // make an FTimeSpan object that represents the "epoch" for time_t (from a stat struct)
 const FDateTime AndroidEpoch(1970, 1, 1);
@@ -20,6 +22,11 @@ static int32 GPackagePatchVersion = 0;
 // External File Path base - setup during load
 FString GFilePathBase;
 
+// Is the OBB in an APK file or not
+bool GOBBinAPK;
+
+extern AAssetManager * AndroidThunkCpp_GetAssetManager();
+
 //This function is declared in the Java-defined class, GameActivity.java: "public native void nativeSetObbInfo(String PackageName, int Version, int PatchVersion);"
 extern "C" void Java_com_epicgames_ue4_GameActivity_nativeSetObbInfo(JNIEnv* jenv, jobject thiz, jstring PackageName, jint Version, jint PatchVersion)
 {
@@ -114,6 +121,104 @@ public:
 	}
 };
 
+/**
+* Android file handle implementation for files in Pak/OBB data
+**/
+class CORE_API FFileHandlePakAndroid : public IFileHandle
+{
+	enum { READWRITE_SIZE = 1024 * 1024 };
+	AAsset * File;
+	int32 FileHandle;
+	off64_t Start;
+	off64_t Length;
+
+	FORCEINLINE bool IsValid()
+	{
+		return FileHandle != -1;
+	}
+
+public:
+	FFileHandlePakAndroid(AAsset * file)
+		: File(file), FileHandle(0), Start(0), Length(0)
+	{
+		FileHandle = AAsset_openFileDescriptor64(file, &Start, &Length);
+	}
+
+	virtual ~FFileHandlePakAndroid()
+	{
+		close(FileHandle);
+		FileHandle = -1;
+		AAsset_close(File);
+	}
+
+	virtual int64 Tell() OVERRIDE
+	{
+		check(IsValid());
+		int64 pos = lseek64(FileHandle, 0, SEEK_CUR);
+		return pos - Start; // We are treating 'tell' as a virtual location from file Start
+	}
+
+	virtual bool Seek(int64 NewPosition) OVERRIDE
+	{
+		check(IsValid());
+		// we need to offset all positions by the Start offset
+		NewPosition += Start;
+		check(NewPosition >= 0);
+		
+		return lseek64(FileHandle, NewPosition, SEEK_SET) != -1;
+	}
+
+    virtual bool SeekFromEnd(int64 NewPositionRelativeToEnd = 0) OVERRIDE
+	{
+		check(IsValid());
+		check(NewPositionRelativeToEnd <= 0);
+		// We need to convert this to a virtual offset inside the file we are interested in
+		int64 position = Start + (Length - NewPositionRelativeToEnd);
+		return lseek64(FileHandle, position, SEEK_SET) != -1;
+	}
+
+	virtual bool Read(uint8* Destination, int64 BytesToRead) OVERRIDE
+	{
+		check(IsValid());
+		while (BytesToRead)
+		{
+			check(BytesToRead >= 0);
+			int64 ThisSize = FMath::Min<int64>(READWRITE_SIZE, BytesToRead);
+			check(Destination);
+			if (read(FileHandle, Destination, ThisSize) != ThisSize)
+			{
+				return false;
+			}
+			Destination += ThisSize;
+			BytesToRead -= ThisSize;
+		}
+		return true;
+	}
+
+	virtual bool Write(const uint8* Source, int64 BytesToWrite) OVERRIDE
+	{
+		check(IsValid());
+		while (BytesToWrite)
+		{
+			check(BytesToWrite >= 0);
+			int64 ThisSize = FMath::Min<int64>(READWRITE_SIZE, BytesToWrite);
+			check(Source);
+			if (write(FileHandle, Source, ThisSize) != ThisSize)
+			{
+				return false;
+			}
+			Source += ThisSize;
+			BytesToWrite -= ThisSize;
+		}
+		return true;
+	}
+
+	virtual int64 Size()
+	{
+		return Length;
+	}
+};
+
 /**
  * Android File I/O implementation
 **/
@@ -141,8 +246,10 @@ protected:
 public:
 
 	FAndroidPlatformFile()
-		: bLookedForObbs(false)
+		: AssetMgr(nullptr)
+		, bLookedForObbs(false)
 	{
+		AssetMgr = AndroidThunkCpp_GetAssetManager();
 	}
 
 	virtual bool FileExists(const TCHAR* Filename) OVERRIDE
@@ -150,10 +257,12 @@ public:
 		struct stat FileInfo;
 		FString NormalizedFilename = NormalizeFilename(Filename);
 		// check the read path
-		if (stat(TCHAR_TO_UTF8(*ConvertToAndroidPath(NormalizedFilename, false)), &FileInfo) == -1)
+		FString convertedPath = ConvertToAndroidPath(NormalizedFilename, false);
+		if (stat(TCHAR_TO_UTF8(*convertedPath), &FileInfo) == -1)
 		{
 			// if not in read path, check the write path
-			if (stat(TCHAR_TO_UTF8(*ConvertToAndroidPath(NormalizedFilename, true)), &FileInfo) == -1)
+			convertedPath = ConvertToAndroidPath(NormalizedFilename, true);
+			if (stat(TCHAR_TO_UTF8(*convertedPath), &FileInfo) == -1)
 			{
 				return false;
 			}
@@ -168,16 +277,28 @@ public:
 		FileInfo.st_size = -1;
 		FString NormalizedFilename = NormalizeFilename(Filename);
 		// check the read path
-		if(stat(TCHAR_TO_UTF8(*ConvertToAndroidPath(NormalizedFilename, false)), &FileInfo) == -1)
+		FString convertedPath = ConvertToAndroidPath(NormalizedFilename, false);
+		// See if we need to route via the asset path
+		if (convertedPath.Contains(TEXT(".pak")) && GOBBinAPK)
 		{
-			// if not in read path, check the write path
-			stat(TCHAR_TO_UTF8(*ConvertToAndroidPath(NormalizedFilename, true)), &FileInfo);
+			AAsset * file = AAssetManager_open(AssetMgr, TCHAR_TO_UTF8(*convertedPath), AASSET_MODE_RANDOM); // seems reasonable?
+			FileInfo.st_size = AAsset_getLength64(file);
+			AAsset_close(file);
 		}
-
-		// make sure to return -1 for directories
-		if (S_ISDIR(FileInfo.st_mode))
+		else
 		{
-			FileInfo.st_size = -1;
+			if (stat(TCHAR_TO_UTF8(*convertedPath), &FileInfo) == -1)
+			{
+				// if not in read path, check the write path
+				convertedPath = ConvertToAndroidPath(NormalizedFilename, true);
+				stat(TCHAR_TO_UTF8(*convertedPath), &FileInfo);
+			}
+
+			// make sure to return -1 for directories
+			if (S_ISDIR(FileInfo.st_mode))
+			{
+				FileInfo.st_size = -1;
+			}
 		}
 		return FileInfo.st_size;
 	}
@@ -186,6 +307,10 @@ public:
 	{
 		// only delete from write path
 		FString AndroidFilename = ConvertToAndroidPath(NormalizeFilename(Filename), true);
+		// lets not try and delete things if we are pulling files from the APK
+		if (AndroidFilename.Contains(TEXT(".pak")) && GOBBinAPK)
+			return false;
+
 		return unlink(TCHAR_TO_UTF8(*AndroidFilename)) == 0;
 	}
 
@@ -193,6 +318,12 @@ public:
 	{
 		FString NormalizedFilename = NormalizeFilename(Filename);
 		FString Filepath = ConvertToAndroidPath(NormalizedFilename, false);
+		
+		// If we are pulling the pak files from the APK they are ALWAYS read only
+		if (Filepath.Contains(TEXT(".pak")) && GOBBinAPK)
+		{
+			return true;
+		}
 		// check read path
 		if (access(TCHAR_TO_UTF8(*Filepath), F_OK) == -1)
 		{
@@ -215,6 +346,10 @@ public:
 	{
 		// move to the write path
 		FString ToAndroidFilename = ConvertToAndroidPath(NormalizeFilename(To), true);
+		// Can't move pak files in the APK
+		if (ToAndroidFilename.Contains(TEXT(".pak")) && GOBBinAPK)
+			return false;
+
 		FString FromAndroidFilename = ConvertToAndroidPath(NormalizeFilename(From), false);
 		return rename(TCHAR_TO_UTF8(*FromAndroidFilename), TCHAR_TO_UTF8(*ToAndroidFilename)) != -1;
 	}
@@ -223,6 +358,11 @@ public:
 	{
 		struct stat FileInfo;
 		FString AndroidFilename = ConvertToAndroidPath(NormalizeFilename(Filename), false);
+		
+		// We can't change the state of pak in APK files - so indicate failure
+		if (AndroidFilename.Contains(TEXT(".pak")) && GOBBinAPK)
+			return false;
+
 		if (stat(TCHAR_TO_UTF8(*AndroidFilename), &FileInfo) != -1)
 		{
 			if (bNewReadOnlyValue)
@@ -245,6 +385,11 @@ public:
 		FString NormalizedFilename = NormalizeFilename(Filename);
 		// check the read path
 		FString AndroidFilename = ConvertToAndroidPath(NormalizedFilename, false);
+
+		// No TimeStamp for pak files in APK, so just return a default timespan for now
+		if (AndroidFilename.Contains(TEXT(".pak")) && GOBBinAPK)
+			return AndroidEpoch;
+
 		if(stat(TCHAR_TO_UTF8(*AndroidFilename), &FileInfo) == -1)
 		{
 			// if not in the read path, check the write path
@@ -265,6 +410,10 @@ public:
 		// get file times
 		struct stat FileInfo;
 		FString AndroidFilename = ConvertToAndroidPath(NormalizeFilename(Filename), true);
+		
+		// Can't set time stamp on pak files in APK
+		if (AndroidFilename.Contains(TEXT(".pak")) && GOBBinAPK)
+			return ;
 		if(stat(TCHAR_TO_UTF8(*AndroidFilename), &FileInfo) == -1)
 		{
 			return;
@@ -283,10 +432,16 @@ public:
 		struct stat FileInfo;
 		FString NormalizedFilename = NormalizeFilename(Filename);
 		// check the read path
-		if(stat(TCHAR_TO_UTF8(*ConvertToAndroidPath(NormalizedFilename, false)), &FileInfo) == -1)
+		FString convertedPath = ConvertToAndroidPath(NormalizedFilename, false);
+		// No TimeStamp for pak files in APK, so just return a default timespan for now
+		if (convertedPath.Contains(TEXT(".pak")) && GOBBinAPK)
+			return AndroidEpoch;
+
+		if(stat(TCHAR_TO_UTF8(*convertedPath), &FileInfo) == -1)
 		{
 			// if not in the read path, check the write path
-			if(stat(TCHAR_TO_UTF8(*ConvertToAndroidPath(NormalizedFilename, true)), &FileInfo) == -1)
+			convertedPath = ConvertToAndroidPath(NormalizedFilename, true);
+			if (stat(TCHAR_TO_UTF8(*convertedPath), &FileInfo) == -1)
 			{
 				return FDateTime::MinValue();
 			}
@@ -302,23 +457,38 @@ public:
 		FString NormalizedFilename = NormalizeFilename(Filename);
 		// check the read path
 		FString AndroidFilename = ConvertToAndroidPath(NormalizedFilename, false);
-		int32 Handle = open(TCHAR_TO_UTF8(*AndroidFilename), O_RDONLY);
-		if(Handle == -1)
+		// No TimeStamp for pak files in APK, so just return a default timespan for now
+		if (AndroidFilename.Contains(TEXT(".pak")) && GOBBinAPK)
 		{
-			// if not in the read path, check the write path
-			FString AndroidFilename = ConvertToAndroidPath(NormalizedFilename, true);
-			Handle = open(TCHAR_TO_UTF8(*AndroidFilename), O_RDONLY);
+			AAsset * file = AAssetManager_open(AssetMgr, TCHAR_TO_UTF8(*AndroidFilename), AASSET_MODE_RANDOM); // seems reasonable?
+			return new FFileHandlePakAndroid(file);
 		}
-
-		if (Handle != -1)
+		else
 		{
-			return new FFileHandleAndroid(Handle);
+			int32 Handle = open(TCHAR_TO_UTF8(*AndroidFilename), O_RDONLY);
+			if (Handle == -1)
+			{
+				// if not in the read path, check the write path
+				AndroidFilename = ConvertToAndroidPath(NormalizedFilename, true);
+				Handle = open(TCHAR_TO_UTF8(*AndroidFilename), O_RDONLY);
+			}
+
+			if (Handle != -1)
+			{
+				return new FFileHandleAndroid(Handle);
+			}
 		}
 		return NULL;
 	}
 
 	virtual IFileHandle* OpenWrite(const TCHAR* Filename, bool bAppend, bool bAllowRead) OVERRIDE
 	{
+		FString AndroidFilename = ConvertToAndroidPath(NormalizeFilename(Filename), true);
+		if (AndroidFilename.Contains(TEXT(".pak")) && GOBBinAPK)
+		{
+			return nullptr; // Can't open pak files in the APK for writing
+		}
+
 		int Flags = O_CREAT;
 		if (bAppend)
 		{
@@ -336,7 +506,7 @@ public:
 		{
 			Flags |= O_WRONLY;
 		}
-		FString AndroidFilename = ConvertToAndroidPath(NormalizeFilename(Filename), true);
+		
 		int32 Handle = open(TCHAR_TO_UTF8(*AndroidFilename), Flags, S_IRUSR | S_IWUSR);
 		if (Handle != -1)
 		{
@@ -357,7 +527,7 @@ public:
 		FString AndroidDirectory = ConvertToAndroidPath(NormalizedDirectory, false);
 		if (stat(TCHAR_TO_UTF8(*AndroidDirectory), &FileInfo)== -1)
 		{
-			FString AndroidDirectory = ConvertToAndroidPath(NormalizeFilename(Directory), true);
+			AndroidDirectory = ConvertToAndroidPath(NormalizeFilename(Directory), true);
 			if (stat(TCHAR_TO_UTF8(*AndroidDirectory), &FileInfo)== -1)
 			{
 				return false;
@@ -385,37 +555,72 @@ public:
 		{
 			// never look again after this time
 			bLookedForObbs = true;
-
-			// look in the downloader path first, then manual location as a fallback
-			for (int32 PassIndex = 0; PassIndex < 2 && ObbHijackedPakDirectory == TEXT(""); PassIndex++)
+			// If we aren't looking in the APK for the OBB then go looking in the expansion directory...
+			FPlatformMisc::LowLevelOutputDebugStringf(TEXT("Obb %s.\n"), GOBBinAPK ? TEXT("in APK") : TEXT("in external dir."));
+			if (!GOBBinAPK)
 			{
-				FString ObbPath = GFilePathBase + ((PassIndex == 0) ? FString(TEXT("/Android/obb/")) : FString(TEXT("/obb/"))) + GPackageName;
+				FPlatformMisc::LowLevelOutputDebugStringf(TEXT("Looking for Obb in external dir.\n"));
+				// look in the downloader path first, then manual location as a fallback
+				for (int32 PassIndex = 0; PassIndex < 2 && ObbHijackedPakDirectory == TEXT(""); PassIndex++)
+				{
+					FString ObbPath = GFilePathBase + ((PassIndex == 0) ? FString(TEXT("/Android/obb/")) : FString(TEXT("/obb/"))) + GPackageName;
 
-				DIR* Handle = opendir(TCHAR_TO_UTF8(*ObbPath));
-				if (Handle)
+					DIR* Handle = opendir(TCHAR_TO_UTF8(*ObbPath));
+					if (Handle)
+					{
+						struct dirent *Entry;
+						while ((Entry = readdir(Handle)) != NULL)
+						{
+							if (FCStringAnsi::Strstr(Entry->d_name, ".obb") != NULL)
+							{
+								// convert to a ".pak" filename
+								FString ObbFileName = ObbPath / FString(UTF8_TO_TCHAR(Entry->d_name));
+								FString PakFileName = FString(Directory) / FString(UTF8_TO_TCHAR(Entry->d_name)) + TEXT(".pak");
+								ObbRemaps.Add(PakFileName, ObbFileName);
+
+								FPlatformMisc::LowLevelOutputDebugStringf(TEXT("Obb found at '%s', pretending to be '%s'\n"), *ObbFileName, *PakFileName);
+
+								// remember which directory we are hijacking only if we found a .obb file
+								ObbHijackedPakDirectory = Directory;
+							}
+						}
+						closedir(Handle);
+					}
+					else
+					{
+						FPlatformMisc::LowLevelOutputDebugStringf(TEXT("Obb directory not found\n"));
+					}
+				}
+			}
+			else // otherwise we need to open the assets directory to find the target file + deal with renaming
+			{
+				FPlatformMisc::LowLevelOutputDebugStringf(TEXT("Looking for Obb in APK file.\n"));
+				AAssetDir * topDir = AAssetManager_openDir(AssetMgr, ""); // grab the top level dir
+				if (topDir)
 				{
-					bool Result = true;
-					struct dirent *Entry;
-					while ((Entry = readdir(Handle)) != NULL && Result == true)
+					const char * fileName = NULL;
+					while ((fileName = AAssetDir_getNextFileName(topDir)) != NULL)
 					{
-						if (FCStringAnsi::Strstr(Entry->d_name, ".obb") != NULL)
+						const char * pakPosition = NULL;
+						if ((pakPosition = strstr(fileName, ".pak")) != 0)
 						{
-							// convert to a ".pak" filename
-							FString ObbFilename = ObbPath / FString(UTF8_TO_TCHAR(Entry->d_name));
-							FString PakFilename = FString(Directory) / FString(UTF8_TO_TCHAR(Entry->d_name)) + TEXT(".pak");
-							ObbRemaps.Add(PakFilename, ObbFilename);
+							FString ObbFileName = FString(UTF8_TO_TCHAR(fileName)); // The real asset file name
+							FString PakFileName = FString(Directory) / FString(UTF8_TO_TCHAR(fileName));
+							PakFileName.RemoveFromEnd(TEXT(".png")); // strip the '.png' from the end
+							ObbRemaps.Add(PakFileName, ObbFileName);
 
-							FPlatformMisc::LowLevelOutputDebugStringf(TEXT("Obb found at '%s', pretending to be '%s'\n"), *ObbFilename, *PakFilename);
+							FPlatformMisc::LowLevelOutputDebugStringf(TEXT("Obb found at '%s', pretending to be '%s'\n"), *ObbFileName, *PakFileName);
 
-							// remember which directory we are hijacking only if we found a .obb file
+							// remember which directory we are hijacking only if we found a .pak file
 							ObbHijackedPakDirectory = Directory;
 						}
 					}
-					closedir(Handle);
+					AAssetDir_close(topDir);
 				}
-				else
+
+				if (ObbHijackedPakDirectory == TEXT(""))
 				{
-					FPlatformMisc::LowLevelOutputDebugStringf(TEXT("Obb directory not found\n"));
+					FPlatformMisc::LowLevelOutputDebugStringf(TEXT("Failed to find PAK in APK file.\n"));
 				}
 			}
 		}
@@ -473,6 +678,8 @@ public:
 	FString ReadDirectory;
 	FString WriteDirectory;
 	
+	AAssetManager* AssetMgr;
+
 	/** 
 	 * Because the .obb file is simply a .pak file, we redirect the engine looking for .pak files to actually find
 	 * .obb files instead. These variables track that remapping, as it happens only the first time a query into a Paks
diff --git a/Engine/Source/Runtime/Launch/Private/Android/AndroidJNI.cpp b/Engine/Source/Runtime/Launch/Private/Android/AndroidJNI.cpp
--- a/Engine/Source/Runtime/Launch/Private/Android/AndroidJNI.cpp
+++ b/Engine/Source/Runtime/Launch/Private/Android/AndroidJNI.cpp
@@ -4,6 +4,8 @@
 #include "ExceptionHandling.h"
 
 #include "AndroidJNI.h"
+#include <Android/asset_manager.h>
+#include <Android/asset_manager_jni.h>
 
 #define JNI_CURRENT_VERSION JNI_VERSION_1_6
 
@@ -11,6 +13,7 @@ JavaVM* GJavaVM;
 jobject GJavaGlobalThis = NULL;
 
 extern FString GFilePathBase;
+extern bool GOBBinAPK;
 
 //////////////////////////////////////////////////////////////////////////
 // FJNIHelper
@@ -102,6 +105,7 @@ jmethodID JDef_GameActivity::AndroidThunkJava_WriteAchievement;
 jmethodID JDef_GameActivity::AndroidThunkJava_ShowAdBanner;
 jmethodID JDef_GameActivity::AndroidThunkJava_HideAdBanner;
 jmethodID JDef_GameActivity::AndroidThunkJava_CloseAdBanner;
+jmethodID JDef_GameActivity::AndroidThunkJava_GetAssetManager;
 
 DEFINE_LOG_CATEGORY_STATIC(LogEngine, Log, All);
 
@@ -242,6 +246,28 @@ void AndroidThunkCpp_CloseAdBanner()
 }
 
 
+namespace
+{
+	jobject GJavaAssetManager = NULL;
+	AAssetManager* GAssetManagerRef = NULL;
+}
+
+AAssetManager * AndroidThunkCpp_GetAssetManager()
+{
+	if (!GAssetManagerRef)
+	{
+		if (JNIEnv* Env = GetJavaEnv())
+		{
+			GJavaAssetManager = Env->CallObjectMethod(GJavaGlobalThis, JDef_GameActivity::AndroidThunkJava_GetAssetManager);
+			Env->NewGlobalRef(GJavaAssetManager);
+			GAssetManagerRef = AAssetManager_fromJava(Env, GJavaAssetManager);
+
+		}
+	}
+
+	return GAssetManagerRef;
+}
+
 //The JNI_OnLoad function is triggered by loading the game library from 
 //the Java source file.
 //	static
@@ -273,6 +299,7 @@ JNIEXPORT jint JNI_OnLoad(JavaVM* InJavaVM, void* InReserved)
 	JDef_GameActivity::AndroidThunkJava_ShowAdBanner = env->GetMethodID(JDef_GameActivity::ClassID, "AndroidThunkJava_ShowAdBanner", "(Ljava/lang/String;Z)V");
 	JDef_GameActivity::AndroidThunkJava_HideAdBanner = env->GetMethodID(JDef_GameActivity::ClassID, "AndroidThunkJava_HideAdBanner", "()V");
 	JDef_GameActivity::AndroidThunkJava_CloseAdBanner = env->GetMethodID(JDef_GameActivity::ClassID, "AndroidThunkJava_CloseAdBanner", "()V");
+	JDef_GameActivity::AndroidThunkJava_GetAssetManager = env->GetMethodID(JDef_GameActivity::ClassID, "AndroidThunkJava_GetAssetManager", "()Landroid/content/res/AssetManager;");
 
 	// hook signals
 #if UE_BUILD_DEBUG
@@ -298,6 +325,9 @@ JNIEXPORT jint JNI_OnLoad(JavaVM* InJavaVM, void* InReserved)
 	env->ReleaseStringUTFChars(pathString, nativePathString);
 	FPlatformMisc::LowLevelOutputDebugStringf(TEXT("Path found as '%s'\n"), *GFilePathBase);
 
+	// Next we check to see if the OBB file is in the APK
+	jmethodID isOBBInAPKMethod = env->GetStaticMethodID(JDef_GameActivity::ClassID, "isOBBInAPK", "()Z");
+	GOBBinAPK = (bool)env->CallStaticBooleanMethod(JDef_GameActivity::ClassID, isOBBInAPKMethod, nullptr);
 
 	// Wire up to core delegates, so core code can call out to Java
 	DECLARE_DELEGATE_OneParam(FAndroidLaunchURLDelegate, const FString&);
diff --git a/Engine/Source/Runtime/Launch/Private/Android/AndroidJNI.h b/Engine/Source/Runtime/Launch/Private/Android/AndroidJNI.h
--- a/Engine/Source/Runtime/Launch/Private/Android/AndroidJNI.h
+++ b/Engine/Source/Runtime/Launch/Private/Android/AndroidJNI.h
@@ -29,7 +29,9 @@ public:
 	static jmethodID AndroidThunkJava_ShowAdBanner;
 	static jmethodID AndroidThunkJava_HideAdBanner;
 	static jmethodID AndroidThunkJava_CloseAdBanner;
+
+	static jmethodID AndroidThunkJava_GetAssetManager;
 };
 
 // Returns the java environment
-JNIEnv* GetJavaEnv(bool bRequireGlobalThis = true);
\ No newline at end of file
+JNIEnv* GetJavaEnv(bool bRequireGlobalThis = true);
