Slate: Added ability to specify window title bar content alignment on a per-platform basis

diff --git a/Engine/Source/Runtime/Core/Private/Windows/WindowsApplication.cpp b/Engine/Source/Runtime/Core/Private/Windows/WindowsApplication.cpp
--- a/Engine/Source/Runtime/Core/Private/Windows/WindowsApplication.cpp
+++ b/Engine/Source/Runtime/Core/Private/Windows/WindowsApplication.cpp
@@ -432,6 +432,27 @@ void FWindowsApplication::GetInitialDisplayMetrics( FDisplayMetrics& OutDisplayM
 	OutDisplayMetrics = InitialDisplayMetrics;
 }
 
+EWindowTitleAlignment::Type FWindowsApplication::GetWindowTitleAlignment() const
+{
+	OSVERSIONINFOEX VersionInfo;
+	FMemory::MemZero(VersionInfo);
+	VersionInfo.dwMajorVersion = 6;
+	VersionInfo.dwMinorVersion = 2;
+	VersionInfo.dwOSVersionInfoSize = sizeof(VersionInfo);
+
+	DWORDLONG LongConditionMask = 0;
+	int ConditionMask = VER_GREATER_EQUAL;
+	VER_SET_CONDITION(LongConditionMask, VER_MAJORVERSION, ConditionMask);
+	VER_SET_CONDITION(LongConditionMask, VER_MINORVERSION, ConditionMask);
+
+	if (::VerifyVersionInfo(&VersionInfo, VER_MAJORVERSION | VER_MINORVERSION, LongConditionMask) != 0)
+	{
+		return EWindowTitleAlignment::Center;
+	}		
+
+	return EWindowTitleAlignment::Left;
+}
+
 void FWindowsApplication::DestroyApplication()
 {
 
diff --git a/Engine/Source/Runtime/Core/Public/GenericPlatform/GenericApplication.h b/Engine/Source/Runtime/Core/Public/GenericPlatform/GenericApplication.h
--- a/Engine/Source/Runtime/Core/Public/GenericPlatform/GenericApplication.h
+++ b/Engine/Source/Runtime/Core/Public/GenericPlatform/GenericApplication.h
@@ -213,6 +213,24 @@ struct FDisplayMetrics
 	FVector2D ActionSafePaddingSize;
 };
 
+
+/**
+ * Enumerates available horizontal alignments for window title bars.
+ */
+namespace EWindowTitleAlignment
+{
+	enum Type
+	{
+		/** The window title should be aligned to the left. */
+		Left,
+		/** The window title should be centered within the title bar. */
+		Center,
+		/** The window title should be aligned to the right. */
+		Right
+	};
+}
+
+
 /**
  * Generic platform application interface
  */
@@ -272,6 +290,12 @@ public:
 
 	virtual void GetInitialDisplayMetrics( FDisplayMetrics& OutDisplayMetrics ) const { GetDisplayMetrics(OutDisplayMetrics); }
 
+	/** Gets the horizontal alignment of the window title bar's title text. */
+	virtual EWindowTitleAlignment::Type GetWindowTitleAlignment() const
+	{
+		return EWindowTitleAlignment::Left;
+	}
+
 	virtual void DestroyApplication() { }
 
 	/** Function to return the current implementation of the ForceFeedback system */
diff --git a/Engine/Source/Runtime/Core/Public/Mac/MacApplication.h b/Engine/Source/Runtime/Core/Public/Mac/MacApplication.h
--- a/Engine/Source/Runtime/Core/Public/Mac/MacApplication.h
+++ b/Engine/Source/Runtime/Core/Public/Mac/MacApplication.h
@@ -53,6 +53,11 @@ public:
 
 	virtual void GetDisplayMetrics( FDisplayMetrics& OutDisplayMetrics ) const OVERRIDE;
 
+	virtual EWindowTitleAlignment::Type GetWindowTitleAlignment() const OVERRIDE
+	{
+		return EWindowTitleAlignment::Center;
+	}
+
 	virtual ITextInputMethodSystem *GetTextInputMethodSystem() OVERRIDE
 	{
 		return TextInputMethodSystem.Get();
diff --git a/Engine/Source/Runtime/Core/Public/Windows/WindowsApplication.h b/Engine/Source/Runtime/Core/Public/Windows/WindowsApplication.h
--- a/Engine/Source/Runtime/Core/Public/Windows/WindowsApplication.h
+++ b/Engine/Source/Runtime/Core/Public/Windows/WindowsApplication.h
@@ -191,6 +191,8 @@ public:
 
 	virtual void GetInitialDisplayMetrics( FDisplayMetrics& OutDisplayMetrics ) const OVERRIDE;
 
+	virtual EWindowTitleAlignment::Type GetWindowTitleAlignment() const OVERRIDE;
+
 	virtual IForceFeedbackSystem *GetForceFeedbackSystem() OVERRIDE 
 	{
 		return this; 
diff --git a/Engine/Source/Runtime/Slate/Private/Framework/Application/SlateApplication.cpp b/Engine/Source/Runtime/Slate/Private/Framework/Application/SlateApplication.cpp
--- a/Engine/Source/Runtime/Slate/Private/Framework/Application/SlateApplication.cpp
+++ b/Engine/Source/Runtime/Slate/Private/Framework/Application/SlateApplication.cpp
@@ -407,7 +407,6 @@ void FSlateApplication::Create()
 }
 
 TSharedPtr<FSlateApplication> FSlateApplication::CurrentApplication = NULL;
-TSharedPtr< GenericApplication > FSlateApplication::PlatformApplication = NULL;
 
 FSlateApplication::FSlateApplication()
 	: bAppIsActive(true)
diff --git a/Engine/Source/Runtime/Slate/Public/Framework/Application/SlateApplication.h b/Engine/Source/Runtime/Slate/Public/Framework/Application/SlateApplication.h
--- a/Engine/Source/Runtime/Slate/Public/Framework/Application/SlateApplication.h
+++ b/Engine/Source/Runtime/Slate/Public/Framework/Application/SlateApplication.h
@@ -520,8 +520,6 @@ public:
 
 	virtual void GetInitialDisplayMetrics( FDisplayMetrics& OutDisplayMetrics ) const { PlatformApplication->GetInitialDisplayMetrics( OutDisplayMetrics ); }
 
-	virtual const TSharedPtr<GenericApplication> GetPlatformApplication ( ) const { return PlatformApplication; }
-
 	/** Are we drag-dropping right now? */
 	bool IsDragDropping() const;
 
@@ -1130,8 +1128,6 @@ private:
 	/** Application singleton */
 	static TSharedPtr< FSlateApplication > CurrentApplication;
 
-	static TSharedPtr< class GenericApplication > PlatformApplication;
-	
 	TSet<FKey> PressedMouseButtons;
 
 	/** true when the slate app is active; i.e. the current foreground window is from our Slate app*/
diff --git a/Engine/Source/Runtime/SlateCore/Private/Application/SlateApplicationBase.cpp b/Engine/Source/Runtime/SlateCore/Private/Application/SlateApplicationBase.cpp
--- a/Engine/Source/Runtime/SlateCore/Private/Application/SlateApplicationBase.cpp
+++ b/Engine/Source/Runtime/SlateCore/Private/Application/SlateApplicationBase.cpp
@@ -11,3 +11,4 @@
  *****************************************************************************/
 
 TSharedPtr<FSlateApplicationBase> FSlateApplicationBase::CurrentBaseApplication = nullptr;
+TSharedPtr<GenericApplication> FSlateApplicationBase::PlatformApplication = nullptr;
diff --git a/Engine/Source/Runtime/SlateCore/Private/Widgets/SWindow.cpp b/Engine/Source/Runtime/SlateCore/Private/Widgets/SWindow.cpp
--- a/Engine/Source/Runtime/SlateCore/Private/Widgets/SWindow.cpp
+++ b/Engine/Source/Runtime/SlateCore/Private/Widgets/SWindow.cpp
@@ -354,10 +354,26 @@ void SWindow::ConstructWindowInternals( const bool bCreateTitleBar )
 		// @todo mainframe: Should be measured from actual title bar content widgets.  Don't use a hard-coded size!
 		TitleBarSize = SWindowDefs::DefaultTitleBarSize;
 
+		EWindowTitleAlignment::Type TitleAlignment = FSlateApplicationBase::Get().GetPlatformApplication()->GetWindowTitleAlignment();
+		EHorizontalAlignment TitleContentAlignment;
+
+		if (TitleAlignment == EWindowTitleAlignment::Left)
+		{
+			TitleContentAlignment = HAlign_Left;
+		}
+		else if (TitleAlignment == EWindowTitleAlignment::Center)
+		{
+			TitleContentAlignment = HAlign_Center;
+		}
+		else
+		{
+			TitleContentAlignment = HAlign_Right;
+		}
+
 		MainWindowArea->AddSlot()
 			.AutoHeight()
 			[
-				FSlateApplicationBase::Get().MakeWindowTitleBar(SharedThis(this), nullptr, HAlign_Center, TitleBar)
+				FSlateApplicationBase::Get().MakeWindowTitleBar(SharedThis(this), nullptr, TitleContentAlignment, TitleBar)
 			];
 	}
 	else
diff --git a/Engine/Source/Runtime/SlateCore/Public/Application/SlateApplicationBase.h b/Engine/Source/Runtime/SlateCore/Public/Application/SlateApplicationBase.h
--- a/Engine/Source/Runtime/SlateCore/Public/Application/SlateApplicationBase.h
+++ b/Engine/Source/Runtime/SlateCore/Public/Application/SlateApplicationBase.h
@@ -129,12 +129,23 @@ public:
 	 */
 	virtual TSharedPtr< SWidget > GetKeyboardFocusedWidget( ) const = 0;
 
-	/** Gets the Widget that currently captures the mouse.
+	/**
+	 * Gets the Widget that currently captures the mouse.
 	 *
 	 * @return The captor widget, or nullptr if no widget captured the mouse.
 	 */
 	virtual TSharedPtr< SWidget > GetMouseCaptor( ) const = 0;
 
+	/**
+	 * Gets the platform application.
+	 *
+	 * @return Platform application.
+	 */
+	virtual const TSharedPtr<GenericApplication> GetPlatformApplication ( ) const
+	{
+		return PlatformApplication;
+	}
+
 	/**
 	 * Gets the rectangle of the current preferred work area.
 	 *
@@ -266,4 +277,7 @@ protected:
 
 	// Holds a pointer to the current application.
 	static TSharedPtr<FSlateApplicationBase> CurrentBaseApplication;
+
+	// Holds a pointer to the platform application.
+	static TSharedPtr<class GenericApplication> PlatformApplication;
 };
