Linux: mimic Windows behavior regarding OnKeyChar() events.

Particularly, do not trigger this for arrow keys, since they do not generate WM_CHAR on Windows.
Otherwise console will be broken in subtle ways (TTP #331984).

diff --git a/Engine/Source/Runtime/Core/Private/Linux/LinuxApplication.cpp b/Engine/Source/Runtime/Core/Private/Linux/LinuxApplication.cpp
--- a/Engine/Source/Runtime/Core/Private/Linux/LinuxApplication.cpp
+++ b/Engine/Source/Runtime/Core/Private/Linux/LinuxApplication.cpp
@@ -156,6 +156,15 @@ void FLinuxApplication::AddPendingEvent( SDL_Event SDLEvent )
 	}
 }
 
+bool FLinuxApplication::GeneratesKeyCharMessage(const SDL_KeyboardEvent & KeyDownEvent)
+{
+	bool bCmdKeyPressed = (KeyDownEvent.keysym.mod & (KMOD_LCTRL | KMOD_RCTRL)) != 0;
+	const SDL_Keycode Sym = KeyDownEvent.keysym.sym;
+
+	// filter out command keys, non-ASCI and arrow keycodes that don't generate WM_CHAR under Windows (TODO: find a table?)
+	return !bCmdKeyPressed && Sym < 128 &&
+		(Sym != SDLK_DOWN && Sym != SDLK_LEFT && Sym != SDLK_RIGHT && Sym != SDLK_UP);
+}
 
 void FLinuxApplication::ProcessDeferredMessage( SDL_Event Event )
 {
@@ -177,24 +186,21 @@ void FLinuxApplication::ProcessDeferredMessage( SDL_Event Event )
 	{
 		return;
 	}
-
 	switch( Event.type )
 	{
 	case SDL_KEYDOWN:
 		{
-			SDL_KeyboardEvent keyEvent = Event.key;
-			const SDL_Keycode KeyCode = keyEvent.keysym.scancode;
-			const TCHAR Character = ConvertChar( keyEvent.keysym );
-			const bool IsRepeat = keyEvent.repeat != 0;
+			SDL_KeyboardEvent KeyEvent = Event.key;
+			const SDL_Keycode KeyCode = KeyEvent.keysym.scancode;
+			const bool bIsRepeated = KeyEvent.repeat != 0;
 				
-			MessageHandler->OnKeyDown( KeyCode, keyEvent.keysym.sym, IsRepeat );
-
 			// First KeyDown, then KeyChar. This is important, as in-game console ignores first character otherwise
+			MessageHandler->OnKeyDown(KeyCode, KeyEvent.keysym.sym, bIsRepeated);
 
-			bool bCmdKeyPressed = keyEvent.keysym.mod & (KMOD_LCTRL | KMOD_RCTRL);
-			if ( !bCmdKeyPressed && KeyCode < 128 )
+			if (GeneratesKeyCharMessage(KeyEvent))
 			{
-				MessageHandler->OnKeyChar( Character, IsRepeat );
+				const TCHAR Character = ConvertChar(KeyEvent.keysym);
+				MessageHandler->OnKeyChar(Character, bIsRepeated);
 			}
 		}
 		break;
diff --git a/Engine/Source/Runtime/Core/Public/Linux/LinuxApplication.h b/Engine/Source/Runtime/Core/Public/Linux/LinuxApplication.h
--- a/Engine/Source/Runtime/Core/Public/Linux/LinuxApplication.h
+++ b/Engine/Source/Runtime/Core/Public/Linux/LinuxApplication.h
@@ -68,6 +68,15 @@ private:
 
 	void ProcessDeferredMessage( SDL_Event Event );
 
+	/** 
+	 * Determines whether this particular SDL_KEYDOWN event should also be routed to OnKeyChar()
+	 *
+	 * @param KeyDownEvent event in question
+	 *
+	 * @return true if character needs to be passed to OnKeyChar
+	 */
+	static bool GeneratesKeyCharMessage(const SDL_KeyboardEvent & KeyDownEvent);
+
 private:
 
 	struct SDLControllerState
