Linux: fix mouse look in games.

Integrating MichaelL's fix (somewhat changed) which avoids grabbing the mouse if cursor is hidden.

diff --git a/Engine/Source/Runtime/Core/Private/Linux/LinuxApplication.cpp b/Engine/Source/Runtime/Core/Private/Linux/LinuxApplication.cpp
--- a/Engine/Source/Runtime/Core/Private/Linux/LinuxApplication.cpp
+++ b/Engine/Source/Runtime/Core/Private/Linux/LinuxApplication.cpp
@@ -884,6 +884,7 @@ void* FLinuxApplication::GetCapture( void ) const
 void FLinuxApplication::UpdateMouseCaptureWindow( SDL_HWindow TargetWindow )
 {
 	const bool bEnable = bIsMouseCaptureEnabled || bIsMouseCursorLocked;
+	FLinuxCursor *LinuxCursor = static_cast< FLinuxCursor* >(Cursor.Get());
 
 	if( bEnable )
 	{
@@ -891,7 +892,7 @@ void FLinuxApplication::UpdateMouseCaptureWindow( SDL_HWindow TargetWindow )
 		{
 			MouseCaptureWindow = TargetWindow;
 		}
-		if( MouseCaptureWindow )
+		if (MouseCaptureWindow && !LinuxCursor->IsHidden())
 		{
 			//SDL_SetWindowGrab(TargetWindow, SDL_TRUE);
 			DSEXT_SetMouseGrab(TargetWindow, SDL_TRUE);
@@ -901,8 +902,11 @@ void FLinuxApplication::UpdateMouseCaptureWindow( SDL_HWindow TargetWindow )
 	{
 		if( MouseCaptureWindow )
 		{
-			//SDL_SetWindowGrab(TargetWindow, SDL_FALSE);
-			DSEXT_SetMouseGrab(TargetWindow, SDL_FALSE);
+			if (!LinuxCursor->IsHidden())
+			{
+				//SDL_SetWindowGrab(TargetWindow, SDL_FALSE);
+				DSEXT_SetMouseGrab(TargetWindow, SDL_FALSE);
+			}
 			MouseCaptureWindow = NULL;
 		}
 	}
