Add static ForwardVector to FVector

diff --git a/Engine/Source/Runtime/Core/Private/Math/UnrealMath.cpp b/Engine/Source/Runtime/Core/Private/Math/UnrealMath.cpp
--- a/Engine/Source/Runtime/Core/Private/Math/UnrealMath.cpp
+++ b/Engine/Source/Runtime/Core/Private/Math/UnrealMath.cpp
@@ -13,6 +13,7 @@ DEFINE_LOG_CATEGORY(LogUnrealMath);
 
 CORE_API const FVector FVector::ZeroVector(0.0f, 0.0f, 0.0f);
 CORE_API const FVector FVector::UpVector(0.0f, 0.0f, 1.0f);
+CORE_API const FVector FVector::ForwardVector(1.0f, 0.0f, 0.0f);
 CORE_API const FVector2D FVector2D::ZeroVector(0.0f, 0.0f);
 CORE_API const FVector2D FVector2D::UnitVector(1.0f, 1.0f);
 CORE_API const FRotator FRotator::ZeroRotator(0.f,0.f,0.f);
diff --git a/Engine/Source/Runtime/Core/Public/Math/Vector.h b/Engine/Source/Runtime/Core/Public/Math/Vector.h
--- a/Engine/Source/Runtime/Core/Public/Math/Vector.h
+++ b/Engine/Source/Runtime/Core/Public/Math/Vector.h
@@ -860,8 +860,12 @@ public:
 
 	CORE_API bool NetSerialize(FArchive& Ar, class UPackageMap* Map, bool& bOutSuccess);
 
+	/** A zero vector (0,0,0) */
 	static CORE_API const FVector ZeroVector;
+	/** World up vector (0,0,1) */
 	static CORE_API const FVector UpVector;
+	/** Local Unreal forward vector (1,0,0) */
+	static CORE_API const FVector ForwardVector;
 };
 
 /**
