Fix flickering when using windows that resize themselves via disclosure arrows on OS X. When setting the window frame tell Cocoa to send a displayIfNeeded message so that the OpenGL context updates & redraws appropriately. This stops us relying on the result of a race between Cocoa in the main thread & OpenGL in the render thread as to who draws first.

diff --git a/Engine/Source/Runtime/Core/Private/Mac/MacWindow.cpp b/Engine/Source/Runtime/Core/Private/Mac/MacWindow.cpp
--- a/Engine/Source/Runtime/Core/Private/Mac/MacWindow.cpp
+++ b/Engine/Source/Runtime/Core/Private/Mac/MacWindow.cpp
@@ -877,21 +877,23 @@ void FMacWindow::ReshapeWindow( int32 X, int32 Y, int32 Width, int32 Height )
 	const TSharedRef<FGenericApplicationMessageHandler> MessageHandler = OwningApplication->MessageHandler;
 	MessageHandler->BeginReshapingWindow( SharedThis( this ) );
 
+	BOOL DisplayIfNeeded = (WindowMode == EWindowMode::Windowed);
+	
 	const int32 InvertedY = FPlatformMisc::ConvertSlateYPositionToCocoa(Y) - Height + 1;
 	if (Definition->HasOSWindowBorder)
 	{
-		[WindowHandle setFrame: [WindowHandle frameRectForContentRect: NSMakeRect(X, InvertedY, FMath::Max(Width, 1), FMath::Max(Height, 1))] display:NO];
+		[WindowHandle setFrame: [WindowHandle frameRectForContentRect: NSMakeRect(X, InvertedY, FMath::Max(Width, 1), FMath::Max(Height, 1))] display:DisplayIfNeeded];
 	}
 	else
 	{
-		[WindowHandle setFrame: NSMakeRect(X, InvertedY, FMath::Max(Width, 1), FMath::Max(Height, 1)) display:NO];
+		[WindowHandle setFrame: NSMakeRect(X, InvertedY, FMath::Max(Width, 1), FMath::Max(Height, 1)) display:DisplayIfNeeded];
 	}
 	
 	// Force resize back to screen size in fullscreen - not ideally pretty but means we don't
 	// have to subvert the OS X or UE fullscreen handling events elsewhere.
 	if(WindowMode != EWindowMode::Windowed)
 	{
-		[WindowHandle setFrame: [WindowHandle screen].frame display:NO];
+		[WindowHandle setFrame: [WindowHandle screen].frame display:YES];
 	}
 	
 	WindowHandle->bZoomed = [WindowHandle isZoomed];
