Typo in FEngineVersion::ToString().

diff --git a/Engine/Source/Runtime/Core/Private/Misc/EngineVersion.cpp b/Engine/Source/Runtime/Core/Private/Misc/EngineVersion.cpp
--- a/Engine/Source/Runtime/Core/Private/Misc/EngineVersion.cpp
+++ b/Engine/Source/Runtime/Core/Private/Misc/EngineVersion.cpp
@@ -83,7 +83,7 @@ bool FEngineVersion::IsCompatibleWith(const FEngineVersion &Other) const
 FString FEngineVersion::ToString(EVersionComponent::Type LastComponent) const
 {
 	FString Result = FString::Printf(TEXT("%u"), Major);
-	if(LastComponent >= EVersionComponent::Major)
+	if(LastComponent >= EVersionComponent::Minor)
 	{
 		Result += FString::Printf(TEXT(".%u"), Minor);
 		if(LastComponent >= EVersionComponent::Patch)
