Linux: Fix incorrect error message for missing #define.

diff --git a/Engine/Source/Runtime/Core/Public/Linux/LinuxPlatformIncludes.h b/Engine/Source/Runtime/Core/Public/Linux/LinuxPlatformIncludes.h
--- a/Engine/Source/Runtime/Core/Public/Linux/LinuxPlatformIncludes.h
+++ b/Engine/Source/Runtime/Core/Public/Linux/LinuxPlatformIncludes.h
@@ -79,7 +79,7 @@ typedef FGenericPlatformAffinity FPlatformAffinity;
 #endif
 
 #ifndef UE_SERVER
-#error "WITH_EDITORONLY_DATA must be defined"
+#error "UE_SERVER must be defined"
 #endif
 
 typedef FLinuxPlatformProperties<WITH_EDITORONLY_DATA, UE_SERVER, !WITH_SERVER_CODE> FPlatformProperties;
@@ -87,4 +87,4 @@ typedef FLinuxPlatformProperties<WITH_EDITORONLY_DATA, UE_SERVER, !WITH_SERVER_C
 // In glibc 2.17, __secure_getenv was renamed to secure_getenv.
 #if defined(__GLIBC__) && (__GLIBC__ == 2) && (__GLIBC_MINOR__ < 17)
 #  define secure_getenv __secure_getenv
-#endif
\ No newline at end of file
+#endif
