Fixed unnecessary runnable thread reporting to log.

diff --git a/Engine/Source/Runtime/Core/Private/Windows/WindowsRunnableThread.cpp b/Engine/Source/Runtime/Core/Private/Windows/WindowsRunnableThread.cpp
--- a/Engine/Source/Runtime/Core/Private/Windows/WindowsRunnableThread.cpp
+++ b/Engine/Source/Runtime/Core/Private/Windows/WindowsRunnableThread.cpp
@@ -11,7 +11,7 @@ uint32 FRunnableThreadWin::GuardedRun()
 
 	FPlatformProcess::SetThreadAffinityMask(ThreadAffintyMask);
 
-#if WINVER >= 0x0600
+#if PLATFORM_XBOXONE
 	UE_LOG(LogThreadingWindows, Log, TEXT("Runnable thread %s is on Process %d."), *ThreadName  , static_cast<uint32>(::GetCurrentProcessorNumber()) );
 #endif
 
