Fix for TextInputSystem shared ptr crash if the text input system failed to initialize

diff --git a/Engine/Source/Runtime/Core/Private/Windows/WindowsApplication.cpp b/Engine/Source/Runtime/Core/Private/Windows/WindowsApplication.cpp
--- a/Engine/Source/Runtime/Core/Private/Windows/WindowsApplication.cpp
+++ b/Engine/Source/Runtime/Core/Private/Windows/WindowsApplication.cpp
@@ -103,7 +103,10 @@ bool FWindowsApplication::RegisterClass( const HINSTANCE HInstance, const HICON
 
 FWindowsApplication::~FWindowsApplication()
 {
-	TextInputMethodSystem->Terminate();
+	if (TextInputMethodSystem.IsValid())
+	{
+		TextInputMethodSystem->Terminate();
+	}
 
 	::CoUninitialize();
 	OleUninitialize();
