Add FBox2D type
https://github.com/EpicGames/UnrealEngine/pull/129 (EverNewJoy)

diff --git a/Engine/Source/Runtime/Core/Private/Math/Box2D.cpp b/Engine/Source/Runtime/Core/Private/Math/Box2D.cpp
new file mode 100644
--- /dev/null
+++ b/Engine/Source/Runtime/Core/Private/Math/Box2D.cpp
@@ -0,0 +1,33 @@
+// Copyright 1998-2014 Epic Games, Inc. All Rights Reserved.
+
+/*=============================================================================
+	Box2D.cpp: Implements the FBox2D class.
+=============================================================================*/
+
+/* FBox2D structors
+ *****************************************************************************/
+
+#include "CorePrivate.h"
+
+FBox2D::FBox2D(const FVector2D* Points, const int32 Count)
+	: Min(0.f, 0.f)
+	, Max(0.f, 0.f)
+	, bIsValid(false)
+{
+	for (int32 PointItr = 0; PointItr < Count; PointItr++)
+	{
+		*this += Points[PointItr];
+	}
+}
+
+
+FBox2D::FBox2D(const TArray<FVector2D>& Points)
+	: Min(0.f, 0.f)
+	, Max(0.f, 0.f)
+	, bIsValid(false)
+{
+	for(const FVector2D& EachPoint : Points)
+	{
+		*this += EachPoint;
+	}
+}
\ No newline at end of file
diff --git a/Engine/Source/Runtime/Core/Public/Math/Box2D.h b/Engine/Source/Runtime/Core/Public/Math/Box2D.h
new file mode 100644
--- /dev/null
+++ b/Engine/Source/Runtime/Core/Public/Math/Box2D.h
@@ -0,0 +1,460 @@
+// Copyright 1998-2014 Epic Games, Inc. All Rights Reserved.
+
+/*=============================================================================
+	Box2D.h: Declares the FBox2D class.
+=============================================================================*/
+
+#pragma once
+
+
+/**
+ * Implements a rectangular 2D Box.
+ */
+class FBox2D
+{
+public:
+
+	/** Holds the box's minimum point. */
+	FVector2D Min;
+
+	/** Holds the box's maximum point. */
+	FVector2D Max;
+
+	/** Holds a flag indicating whether this box is valid. */
+	bool bIsValid;
+
+
+public:
+
+	/** Default constructor. */
+	FBox2D( ) { }
+
+	/** Creates and initializes a new box. */
+	FBox2D( int32 )
+	{
+		Init();
+	}
+
+	/**
+	 * Creates and initializes a new box.
+	 *
+	 * @param EForceInit - Force Init Enum.
+	 */
+	explicit FBox2D( EForceInit )
+	{
+		Init();
+	}
+
+	/**
+	 * Creates and initializes a new box from the specified parameters.
+	 *
+	 * @param InMin - The box's minimum point.
+	 * @param InMax - The box's maximum point.
+	 */
+	FBox2D( const FVector2D& InMin, const FVector2D& InMax )
+		: Min(InMin)
+		, Max(InMax)
+		, bIsValid(true)
+	{ }
+
+	
+	 
+	/**
+	 * Creates and initializes a new box from the given set of points.
+	 *
+	 * @param Points - Array of Points to create for the bounding volume.
+	 * @param Count - The number of points.
+	 */
+	CORE_API FBox2D( const FVector2D* Points, const int32 Count );
+
+	/**
+	 * Creates and initializes a new box from an array of points.
+	 *
+	 * @param Points - Array of Points to create for the bounding volume.
+	 */
+	CORE_API FBox2D( const TArray<FVector2D>& Points );
+
+
+public:
+
+	/**
+	 * Adds to this bounding box to include a given point.
+	 *
+	 * @param Other the point to increase the bounding volume to.
+	 *
+	 * @return Reference to this bounding box after resizing to include the other point.
+	 */
+	FORCEINLINE FBox2D& operator+=( const FVector2D &Other );
+
+	/**
+	 * Gets the result of addition to this bounding volume.
+	 *
+	 * @param Other The other point to add to this.
+	 *
+	 * @return A new bounding volume.
+	 */
+	FBox2D operator+( const FVector2D& Other ) const
+	{
+		return FBox2D(*this) += Other;
+	}
+
+	/**
+	 * Adds to this bounding box to include a new bounding volume.
+	 *
+	 * @param Other the bounding volume to increase the bounding volume to.
+	 *
+	 * @return Reference to this bounding volume after resizing to include the other bounding volume.
+	 */
+	FORCEINLINE FBox2D& operator+=( const FBox2D& Other );
+
+	/**
+	 * Gets the result of addition to this bounding volume.
+	 *
+	 * @param Other The other volume to add to this.
+	 *
+	 * @return A new bounding volume.
+	 */
+	FBox2D operator+( const FBox2D& Other ) const
+	{
+		return FBox2D(*this) += Other;
+	}
+
+	/**
+	 * Gets reference to the min or max of this bounding volume.
+	 *
+	 * @param Index the index into points of the bounding volume.
+	 *
+	 * @return a reference to a point of the bounding volume.
+	 */
+   FVector2D& operator[]( int32 Index )
+	{
+		check((Index >= 0) && (Index < 2));
+
+		if (Index == 0)
+		{
+			return Min;
+		}
+
+		return Max;
+	}
+	bool operator==( const FBox2D& Other ) const
+	{
+		return ( Min == Other.Min ) && ( Max == Other.Max );
+	}
+
+
+public:
+
+
+
+	/** 
+	 * Calculates the distance of a point to this box.
+	 *
+	 * @param Point - The point.
+	 *
+	 * @return The distance.
+	 */
+	FORCEINLINE float ComputeSquaredDistanceToPoint( const FVector2D& Point ) const
+	{
+		// Accumulates the distance as we iterate axis
+		float DistSquared = 0.f;
+		
+		if (Point.X < Min.X)
+		{
+			DistSquared += FMath::Square(Point.X - Min.X);
+		}
+		else if (Point.X > Max.X)
+		{
+			DistSquared += FMath::Square(Point.X - Max.X);
+		}
+		
+		if (Point.Y < Min.Y)
+		{
+			DistSquared += FMath::Square(Point.Y - Min.Y);
+		}
+		else if (Point.Y > Max.Y)
+		{
+			DistSquared += FMath::Square(Point.Y - Max.Y);
+		}
+		
+		return DistSquared;
+	}
+
+	/** 
+	 * Increase the bounding box volume
+	 *
+	 * @param	W size to increase volume by
+	 * @return	new bounding box increased in size
+	 */
+	FBox2D ExpandBy( const float W ) const
+	{
+		return FBox2D(Min - FVector2D(W, W), Max + FVector2D(W, W));
+	}
+
+	/** 
+	 * Shift bounding box position
+	 *
+	 * @param	Offset vector to shift by
+	 * @return	new shifted bounding box 
+	 */
+	FBox2D ShiftBy( const FVector2D& Offset ) const
+	{
+		return FBox2D(Min + Offset, Max + Offset);
+	}
+
+
+	/**
+	 * Gets the box's center point.
+	 *
+	 * @return Center point.
+	 */
+	FVector2D GetCenter( ) const
+	{
+		return FVector2D((Min + Max) * 0.5f);
+	}
+
+	/**
+	 * Get the center and extents
+	 *
+	 * @param center[out] reference to center point
+	 * @param Extents[out] reference to the extent around the center
+	 */
+	void GetCenterAndExtents( FVector2D & center, FVector2D & Extents ) const
+	{
+		Extents = GetExtent();
+		center = Min + Extents;
+	}
+
+	/**
+	 * Calculates the closest point on or inside the box to a given point in space.
+	 *
+	 * @param Point - The point in space.
+	 *
+	 * @return The closest point on or inside the box.
+	 */
+	FORCEINLINE FVector2D GetClosestPointTo( const FVector2D& Point ) const;
+
+	/**
+	 * Gets the box extents around the center.
+	 *
+	 * @return Box extents.
+	 */
+	FVector2D GetExtent( ) const
+	{
+		return 0.5f * (Max - Min);
+	}
+
+
+	/**
+	 * Gets the box size.
+	 *
+	 * @return Box size.
+	 */
+	FVector2D GetSize( ) const
+	{
+		return (Max - Min);
+	}
+
+	/**
+	 * Gets the box area.
+	 *
+	 * @return Box area.
+	 */
+	float GetArea( ) const
+	{
+		return (Max.X - Min.X) * (Max.Y - Min.Y);
+	}
+
+	/**
+	 * Set the initial values of the bounding box to Zero.
+	 */
+	void Init( )
+	{
+		Min = Max = FVector2D::ZeroVector;
+		bIsValid = false;
+	}
+
+	
+	/**
+	 * Test for intersection of bounding box and this bounding box
+	 *
+	 * @param other bounding box to test intersection
+	 * @return true if the other box intersects this box, otherwise false
+	 */
+	FORCEINLINE bool Intersect( const FBox2D & other ) const;
+
+	/**
+	 * Test for intersection of bounding box and this bounding box in the X plane
+	 *
+	 * @param other bounding box to test intersection
+	 * @return true if the other box intersects this box in the X Plane, otherwise false
+	 */
+	FORCEINLINE bool IntersectX( const FBox2D& other ) const;
+	
+	/**
+	 * Test for intersection of bounding box and this bounding box in the Y plane
+	 *
+	 * @param other bounding box to test intersection
+	 * @return true if the other box intersects this box in the Y Plane, otherwise false
+	 */
+	FORCEINLINE bool IntersectY( const FBox2D& other ) const;
+	
+	/**
+	 * Test if point is inside this box
+	 *
+	 * @param Point to test
+	 * @return true if the point is inside this box, otherwise false
+	 */
+	bool IsInside( const FVector2D & TestPoint ) const
+	{
+		return ((TestPoint.X > Min.X) && (TestPoint.X < Max.X) && (TestPoint.Y > Min.Y) && (TestPoint.Y < Max.Y));
+	}
+
+
+	/** 
+	 * Checks to see if given box is fully encapsulated by this box
+	 * 
+	 * @param Other box to test for encapsulation within the bounding volume
+	 * @return true if box is inside this volume
+	 */
+	bool IsInside( const FBox2D& Other ) const
+	{
+		return (IsInside(Other.Min) && IsInside(Other.Max));
+	}
+
+
+	/**
+	 * Get a textual representation of this box.
+	 *
+	 * @return A string describing the box.
+	 */
+	FString ToString( ) const;
+
+public:
+
+	/**
+	 * Serializes the bounding box.
+	 *
+	 * @param Ar - The archive to serialize into.
+	 * @param Box - The box to serialize.
+	 *
+	 * @return Reference to the Archive after serialization.
+	 */
+	friend FArchive& operator<<( FArchive& Ar, FBox2D& Box )
+	{
+		return Ar << Box.Min << Box.Max << Box.bIsValid;
+	}
+};
+
+
+
+/* FBox2D inline functions
+ *****************************************************************************/
+
+FORCEINLINE FBox2D& FBox2D::operator+=( const FVector2D &Other )
+{
+	if (bIsValid)
+	{
+		Min.X = FMath::Min(Min.X, Other.X);
+		Min.Y = FMath::Min(Min.Y, Other.Y);
+	
+		Max.X = FMath::Max(Max.X, Other.X);
+		Max.Y = FMath::Max(Max.Y, Other.Y);
+		
+	}
+	else
+	{
+		Min = Max = Other;
+		bIsValid = true;
+	}
+
+	return *this;
+}
+
+
+FORCEINLINE FBox2D& FBox2D::operator+=( const FBox2D& Other )
+{
+	if (bIsValid && Other.bIsValid)
+	{
+		Min.X = FMath::Min(Min.X, Other.Min.X);
+		Min.Y = FMath::Min(Min.Y, Other.Min.Y);
+
+		Max.X = FMath::Max(Max.X, Other.Max.X);
+		Max.Y = FMath::Max(Max.Y, Other.Max.Y);
+	}
+	else if (Other.bIsValid)
+	{
+		*this = Other;
+	}
+
+	return *this;
+}
+
+
+FORCEINLINE FVector2D FBox2D::GetClosestPointTo( const FVector2D& Point ) const
+{
+	// start by considering the point inside the box
+	FVector2D ClosestPoint = Point;
+
+	// now clamp to inside box if it's outside
+	if (Point.X < Min.X)
+	{
+		ClosestPoint.X = Min.X;
+	}
+	else if (Point.X > Max.X)
+	{
+		ClosestPoint.X = Max.X;
+	}
+
+	// now clamp to inside box if it's outside
+	if (Point.Y < Min.Y)
+	{
+		ClosestPoint.Y = Min.Y;
+	}
+	else if (Point.Y > Max.Y)
+	{
+		ClosestPoint.Y = Max.Y;
+	}
+
+	return ClosestPoint;
+}
+
+FORCEINLINE bool FBox2D::Intersect( const FBox2D & Other ) const
+{
+	if ((Min.X > Other.Max.X) || (Other.Min.X > Max.X))
+	{
+		return false;
+	}
+
+	if ((Min.Y > Other.Max.Y) || (Other.Min.Y > Max.Y))
+	{
+		return false;
+	}
+
+	return true;
+}
+
+
+FORCEINLINE bool FBox2D::IntersectX( const FBox2D& Other ) const
+{
+	if ((Min.X > Other.Max.X) || (Other.Min.X > Max.X))
+	{
+		return false;
+	}
+
+	return true;
+}
+FORCEINLINE bool FBox2D::IntersectY( const FBox2D& Other ) const
+{
+	if ((Min.Y > Other.Max.Y) || (Other.Min.Y > Max.Y))
+	{
+		return false;
+	}
+
+	return true;
+}
+
+FORCEINLINE FString FBox2D::ToString( ) const
+{
+	return FString::Printf(TEXT("bIsValid=%s, Min=(%s), Max=(%s)"), bIsValid ? TEXT("true") : TEXT("false"), *Min.ToString(), *Max.ToString());
+}
\ No newline at end of file
diff --git a/Engine/Source/Runtime/Core/Public/Math/UnrealMath.h b/Engine/Source/Runtime/Core/Public/Math/UnrealMath.h
--- a/Engine/Source/Runtime/Core/Public/Math/UnrealMath.h
+++ b/Engine/Source/Runtime/Core/Public/Math/UnrealMath.h
@@ -23,6 +23,7 @@
 #include "RangeSet.h"
 #include "Interval.h"
 #include "Box.h"
+#include "Box2D.h"
 #include "BoxSphereBounds.h"
 #include "OrientedBox.h"
 #include "Axis.h"
