Optimize FBox::TransformBy() by using the branchless vectorized implementation for matrices
https://github.com/EpicGames/UnrealEngine/pull/124 (pzurita)

diff --git a/Engine/Source/Runtime/Core/Private/Math/Box.cpp b/Engine/Source/Runtime/Core/Private/Math/Box.cpp
--- a/Engine/Source/Runtime/Core/Private/Math/Box.cpp
+++ b/Engine/Source/Runtime/Core/Private/Math/Box.cpp
@@ -125,27 +125,7 @@ FBox FBox::TransformBy(const FMatrix& M) const
 
 FBox FBox::TransformBy(const FTransform & M) const
 {
-	FVector Vertices[8] = 
-	{
-		FVector(Min),
-		FVector(Min.X, Min.Y, Max.Z),
-		FVector(Min.X, Max.Y, Min.Z),
-		FVector(Max.X, Min.Y, Min.Z),
-		FVector(Max.X, Max.Y, Min.Z),
-		FVector(Max.X, Min.Y, Max.Z),
-		FVector(Min.X, Max.Y, Max.Z),
-		FVector(Max)
-	};
-
-	FBox NewBox(0);
-
-	for (int32 VertexIndex = 0; VertexIndex < ARRAY_COUNT(Vertices); VertexIndex++)
-	{
-		FVector4 ProjectedVertex = M.TransformPosition(Vertices[VertexIndex]);
-		NewBox += ProjectedVertex;
-	}
-
-	return NewBox;
+	return TransformBy(M.ToMatrixWithScale());
 }
 
 FBox FBox::InverseTransformBy(const FTransform & M) const
