TaskGraph - Removed duplicated code

diff --git a/Engine/Source/Runtime/Core/Private/Async/TaskGraph.cpp b/Engine/Source/Runtime/Core/Private/Async/TaskGraph.cpp
--- a/Engine/Source/Runtime/Core/Private/Async/TaskGraph.cpp
+++ b/Engine/Source/Runtime/Core/Private/Async/TaskGraph.cpp
@@ -123,7 +123,7 @@ private:
 class FTaskThread : public FRunnable, FSingleThreadRunnable
 {
 public:
-
+	/** Array of tasks for this task thread. */
 	TArray<FBaseGraphTask*> NewTasks;
 
 	// Calls meant to be called from a "main" or supervisor thread.
@@ -1057,25 +1057,9 @@ FGraphEventRef FGraphEvent::CreateGraphEvent()
 
 void FGraphEvent::DispatchSubsequents(ENamedThreads::Type CurrentThreadIfKnown)
 {
-	if (EventsToWaitFor.Num())
-	{
-		// need to save this first and empty the actual tail of the task might be recycled faster than it is cleared.
-		FGraphEventArray TempEventsToWaitFor;
-		Exchange(EventsToWaitFor, TempEventsToWaitFor);
-		// create the Gather...this uses a special version of private CreateTask that "assumes" the subsequent list (which other threads might still be adding too).
-		TGraphTask<FNullGraphTask>::CreateTask(TRefCountPtr<FGraphEvent>(this), &TempEventsToWaitFor, CurrentThreadIfKnown).ConstructAndDispatchWhenReady(TEXT("DontCompleteUntil"));
-		return;
-	}
-	//@todo need to get rid of these arrays, or at least use inline allocators
 	TArray<FBaseGraphTask*> NewTasks;
-	NewTasks.Reset(128);
-	SubsequentList.PopAllAndClose(NewTasks);
-	for (int32 Index = NewTasks.Num() - 1; Index >= 0 ; Index--) // reverse the order since PopAll is implicitly backwards
-	{
-		FBaseGraphTask* NewTask = NewTasks[Index];
-		checkThreadGraph(NewTask);
-		NewTask->ConditionalQueueTask(CurrentThreadIfKnown);
-	}
+	NewTasks.Reset( 128 );
+	DispatchSubsequents( NewTasks, CurrentThreadIfKnown );
 }
 
 void FGraphEvent::DispatchSubsequents(TArray<FBaseGraphTask*>& NewTasks, ENamedThreads::Type CurrentThreadIfKnown)
@@ -1086,7 +1070,7 @@ void FGraphEvent::DispatchSubsequents(TArray<FBaseGraphTask*>& NewTasks, ENamedT
 		FGraphEventArray TempEventsToWaitFor;
 		Exchange(EventsToWaitFor, TempEventsToWaitFor);
 		// create the Gather...this uses a special version of private CreateTask that "assumes" the subsequent list (which other threads might still be adding too).
-		TGraphTask<FNullGraphTask>::CreateTask(TRefCountPtr<FGraphEvent>(this), &TempEventsToWaitFor, CurrentThreadIfKnown).ConstructAndDispatchWhenReady(TEXT("DontCompleteUntil"));
+		TGraphTask<FNullGraphTask>::CreateTask( FGraphEventRef( this ), &TempEventsToWaitFor, CurrentThreadIfKnown ).ConstructAndDispatchWhenReady( TEXT( "DontCompleteUntil" ) );
 		return;
 	}
 
diff --git a/Engine/Source/Runtime/Core/Public/Async/TaskGraphInterfaces.h b/Engine/Source/Runtime/Core/Public/Async/TaskGraphInterfaces.h
--- a/Engine/Source/Runtime/Core/Public/Async/TaskGraphInterfaces.h
+++ b/Engine/Source/Runtime/Core/Public/Async/TaskGraphInterfaces.h
@@ -347,7 +347,7 @@ public:
 
 	/**
 	 *	Attempts to a new subsequent task. If this event has already fired, false is returned and action must be taken to ensure that the task will still fire even though this event cannot be a prerequisite (because it is already finished).
-	 *	@return true if the task was sucessfully set up as a subsequent. false if the event has already fired.
+	 *	@return true if the task was successfully set up as a subsequent. false if the event has already fired.
 	**/
 	bool AddSubsequent(class FBaseGraphTask* Task)
 	{
