Core - Added checking whether the processor can execute 128-bit CAS operation, required for lockfreelist128 (PS4,XB1,Windows only)

diff --git a/Engine/Source/Runtime/Core/Public/GenericPlatform/GenericPlatformAtomics.h b/Engine/Source/Runtime/Core/Public/GenericPlatform/GenericPlatformAtomics.h
--- a/Engine/Source/Runtime/Core/Public/GenericPlatform/GenericPlatformAtomics.h
+++ b/Engine/Source/Runtime/Core/Public/GenericPlatform/GenericPlatformAtomics.h
@@ -14,14 +14,8 @@ struct MS_ALIGN(16) FInt128
 	int64 Low;
 	/** The high part of the 128-bit integer. */
 	int64 High;
-
-	/** Initialization constructor. */
-	FInt128( int64 InHigh, int64 InLow )
-		: Low( InLow )
-		, High( InHigh )
-	{}
-
-} GCC_ALIGN(16);
+} 
+GCC_ALIGN(16);
 
 /**
 * Generic implementation...you are required to implement at least FPlatformAtomics::InterlockedCompareExchange but they are listed here 
@@ -207,6 +201,16 @@ struct FGenericPlatformAtomics
 
 #endif // 0
 
+
+	/**
+	 * @return true, if the processor we are running on can execute compare and exchange 128-bit operation.
+	 * @see cmpxchg16b, early AMD64 processors don't support this operation.
+	 */
+	static FORCEINLINE bool CanUseCompareExchange128()
+	{
+		return false;
+	}
+
 protected:
 	/**
 	 * Checks if a pointer is aligned and can be used with atomic functions.
diff --git a/Engine/Source/Runtime/Core/Public/Windows/WindowsPlatformAtomics.h b/Engine/Source/Runtime/Core/Public/Windows/WindowsPlatformAtomics.h
--- a/Engine/Source/Runtime/Core/Public/Windows/WindowsPlatformAtomics.h
+++ b/Engine/Source/Runtime/Core/Public/Windows/WindowsPlatformAtomics.h
@@ -108,7 +108,7 @@ struct CORE_API FWindowsPlatformAtomics
 	 *	To determine whether the processor supports this operation, call the IsProcessorFeaturePresent function with PF_COMPARE64_EXCHANGE128.
 	 */
 #if	PLATFORM_HAS_128BIT_ATOMICS
-	static FORCEINLINE int64 InterlockedCompareExchange128( volatile FInt128* Dest, const FInt128& Exchange, FInt128* Comparand )
+	static FORCEINLINE bool InterlockedCompareExchange128( volatile FInt128* Dest, const FInt128& Exchange, FInt128* Comparand )
 	{
 #if !(UE_BUILD_SHIPPING || UE_BUILD_TEST)
 		if (IsAligned(Dest,16) == false)
@@ -137,6 +137,15 @@ struct CORE_API FWindowsPlatformAtomics
 		return ::InterlockedCompareExchangePointer(Dest, Exchange, Comparand);
 	}
 
+	/**
+	* @return true, if the processor we are running on can execute compare and exchange 128-bit operation.
+	* @see cmpxchg16b, early AMD64 processors don't support this operation.
+	*/
+	static FORCEINLINE bool CanUseCompareExchange128()
+	{
+		return !!IsProcessorFeaturePresent( PF_COMPARE_EXCHANGE128 );
+	}
+
 protected:
 	/**
 	 * Handles atomics function failure.
