Unmark warning 4427 - 'enum' is a native enum: add an access specifier (private/public) to declare a managed enum as error. (warning off by default)

diff --git a/Engine/Source/Runtime/Core/Public/WinRT/WinRTCompilerSetup.h b/Engine/Source/Runtime/Core/Public/WinRT/WinRTCompilerSetup.h
--- a/Engine/Source/Runtime/Core/Public/WinRT/WinRTCompilerSetup.h
+++ b/Engine/Source/Runtime/Core/Public/WinRT/WinRTCompilerSetup.h
@@ -14,7 +14,7 @@
 #pragma warning (error: 4100 4101 4102 4103 4104 4105 4106 4107 4108 4109 4110 4111 4112 4113 4114 4115 4116 4117 4118 4119 4120 4121 4122 4123 4124 4125 4126 4127 4128 4129 4130 4131 4132 4133 4134 4135 4136 4137 4138 4139 4140 4141 4142 4143 4144 4145 4146 4147 4148 4149 4150 4151 4152 4153 4154 4155 4156 4157 4158 4159 4160 4161 4162 4163 4164      4166 4167 4168 4169 4170 4171 4172 4173 4174 4175 4176 4177 4178 4179 4180 4181 4182 4183 4184 4185 4186 4187 4188 4189 4190 4191 4192 4193 4194 4195 4196 4197 4198 4199)
 #pragma warning (error: 4200 4201 4202 4203 4204 4205 4206 4207 4208 4209 4210 4211 4212 4213 4214 4215 4216 4217 4218 4219 4220 4221 4222 4223 4224 4225 4226 4227 4228 4229 4230 4231 4232 4233 4234 4235 4236 4237 4238 4239 4240 4241      4243 4244 4245 4246 4247 4248 4249 4250 4251 4252 4253      4255 4256 4257 4258 4259 4260 4261 4262 4263 4264           4267 4268 4269 4270 4271 4272 4273 4274 4275 4276 4277 4278 4279 4280 4281 4282 4283 4284 4285 4286 4287 4288 4289 4290 4291 4292 4293 4294 4295      4297 4298 4299)
 #pragma warning (error: 4300 4301 4302 4303 4304      4306 4307 4308 4309 4310 4311      4313 4314 4315 4316 4317 4318 4319 4320 4321 4322 4323 4324 4325 4326 4327 4328 4329 4330 4331 4332 4333 4334 4335 4336 4337 4338      4340 4341 4342 4343 4344 4345 4346 4347 4348 4349      4351 4352 4353 4354 4355 4356 4357 4358 4359 4360 4361 4362 4363 4364      4366 4367 4368 4369 4370      4372 4373 4374 4375 4376 4377 4378 4379 4380 4381 4382 4383 4384 4385 4386 4387      4389 4390 4391 4392 4393 4394 4395 4396 4397 4398 4399)
-#pragma warning (error: 4400 4401 4402 4403 4404 4405 4406 4407 4408 4409 4410 4411 4412 4413 4414 4415 4416 4417 4418 4419 4420 4421 4422 4423 4424 4425 4426 4427 4428 4429 4430 4431 4432 4433 4434      4436 4437 4438 4439 4440 4441 4442 4443 4444 4445 4446 4447 4448 4449 4450 4451 4452 4453 4454 4455 4456 4457 4458 4459 4460 4461 4462 4463 4464 4465 4466 4467 4468 4469 4470      4472 4473 4474 4475 4476 4477 4478 4479 4480      4482 4483 4484 4485 4486 4487 4488 4489 4490 4491 4492 4493 4494 4495 4496 4497 4498 4499)
+#pragma warning (error: 4400 4401 4402 4403 4404 4405 4406 4407 4408 4409 4410 4411 4412 4413 4414 4415 4416 4417 4418 4419 4420 4421 4422 4423 4424 4425 4426      4428 4429 4430 4431 4432 4433 4434      4436 4437 4438 4439 4440 4441 4442 4443 4444 4445 4446 4447 4448 4449 4450 4451 4452 4453 4454 4455 4456 4457 4458 4459 4460 4461 4462 4463 4464 4465 4466 4467 4468 4469 4470      4472 4473 4474 4475 4476 4477 4478 4479 4480      4482 4483 4484 4485 4486 4487 4488 4489 4490 4491 4492 4493 4494 4495 4496 4497 4498 4499)
 #pragma warning (error: 4500 4501 4502 4503 4504 4505 4506 4507 4508 4509 4510 4511 4512 4513 4514 4515 4516 4517 4518 4519 4520 4521 4522 4523 4524 4525 4526 4527 4528 4529 4530 4531 4532 4533 4534 4535 4536 4537 4538 4539 4540 4541 4542 4543 4544 4545 4546                4550 4551 4552 4553 4554      4556 4557 4558 4559 4560 4561 4562 4563 4564 4565 4566 4567 4568 4569 4570      4572 4573      4575 4576 4577 4578 4579 4580 4581 4582 4583 4584 4585 4586 4587 4588 4589 4590 4591 4592 4593 4594 4595 4596 4597 4598 4599)
 #pragma warning (error: 4600 4601 4602 4603 4604 4605 4606 4607 4608 4609 4610 4611 4612 4613 4614 4615 4616 4617 4618      4620 4621 4622 4623 4624 4625 4626 4627 4628 4629 4630 4631 4632 4633 4634 4635 4636 4637 4638 4639 4640 4641 4642 4643 4644 4645 4646 4647 4648 4649 4650      4652 4653 4654 4655 4656 4657 4658 4659 4660 4661 4662 4663 4664 4665 4666 4667 4668 4669 4670 4671 4672 4673 4674 4675 4676 4677 4678 4679 4680 4681 4682 4683 4684 4685 4686 4687 4688 4689 4690 4691 4692 4693 4694 4695 4696 4697 4698 4699)
 #pragma warning (error: 4700 4701 4702 4703 4704 4705 4706 4707 4708 4709 4710 4711 4712 4713 4714 4715 4716 4717 4718 4719 4720 4721 4722 4723 4724 4725 4726 4727 4728 4729      4731 4732 4733 4734 4735 4736 4737 4738 4739 4740 4741 4742 4743 4744 4745 4746 4747 4748 4749 4750 4751 4752 4753 4754 4755 4756 4757 4758 4759 4760 4761 4762 4763 4764 4765 4766 4767 4768 4769 4770 4771 4772 4773 4774 4775 4776 4777 4778 4779 4780 4781 4782 4783 4784 4785 4786 4787 4788 4789 4790 4791 4792 4793 4794 4795 4796 4797 4798 4799)
@@ -41,6 +41,7 @@
 // 4365 - 'action' : conversion from 'type_1' to 'type_2', signed/unsigned mismatch												http://msdn.microsoft.com/en-us/library/ms173683.aspx
 // 4371 - layout of class may have changed from a previous version of the compiler due to better packing of member
 // 4388 - '==' : signed/unsigned mismatch																						http://msdn.microsoft.com/en-us/library/jj155806.aspx
+// 4427 - 'enum' is a native enum: add an access specifier (private/public) to declare a managed enum
 // 4435 - 'class1' : Object layout under /vd2 will change due to virtual base 'class2'
 // 4471 - a forward declaration of an unscoped enumeration must have an underlying type (int assumed)
 // 4481 - nonstandard extension used: override specifier 'keyword'																http://msdn.microsoft.com/en-us/library/ms173703.aspx
diff --git a/Engine/Source/Runtime/Core/Public/Windows/WindowsPlatformCompilerSetup.h b/Engine/Source/Runtime/Core/Public/Windows/WindowsPlatformCompilerSetup.h
--- a/Engine/Source/Runtime/Core/Public/Windows/WindowsPlatformCompilerSetup.h
+++ b/Engine/Source/Runtime/Core/Public/Windows/WindowsPlatformCompilerSetup.h
@@ -28,7 +28,7 @@ checkAtCompileTime(_MSC_VER >= 1700, Visual_Studio_2012_Required);
 #pragma warning (error: 4100 4101 4102 4103 4104 4105 4106 4107 4108 4109 4110 4111 4112 4113 4114 4115 4116 4117 4118 4119 4120 4121 4122 4123 4124 4125 4126 4127 4128 4129 4130 4131 4132 4133 4134 4135 4136 4137 4138 4139 4140 4141 4142 4143 4144 4145 4146 4147 4148 4149 4150 4151 4152 4153 4154 4155 4156 4157 4158 4159 4160 4161 4162 4163 4164      4166 4167 4168 4169 4170 4171 4172 4173 4174 4175 4176 4177 4178 4179 4180 4181 4182 4183 4184 4185 4186 4187 4188 4189 4190 4191 4192 4193 4194 4195 4196 4197 4198 4199)
 #pragma warning (error: 4200 4201 4202 4203 4204 4205 4206 4207 4208 4209 4210 4211 4212 4213 4214 4215 4216 4217 4218 4219 4220 4221 4222 4223 4224 4225 4226 4227 4228 4229 4230 4231 4232 4233 4234 4235 4236 4237 4238 4239 4240 4241      4243 4244 4245 4246 4247 4248 4249 4250 4251 4252 4253      4255 4256 4257 4258 4259 4260 4261 4262 4263 4264           4267 4268 4269 4270 4271 4272 4273 4274 4275 4276 4277 4278 4279 4280 4281 4282 4283 4284 4285 4286 4287 4288 4289 4290 4291 4292 4293 4294 4295      4297 4298 4299)
 #pragma warning (error: 4300 4301 4302 4303 4304      4306 4307 4308 4309 4310 4311      4313 4314 4315 4316 4317 4318 4319 4320 4321 4322 4323 4324 4325 4326 4327 4328 4329 4330 4331 4332 4333 4334 4335 4336 4337 4338      4340 4341 4342 4343 4344 4345 4346 4347 4348 4349      4351 4352 4353 4354 4355 4356 4357 4358 4359 4360 4361 4362 4363 4364      4366 4367 4368 4369 4370      4372 4373 4374 4375 4376 4377 4378 4379 4380 4381 4382 4383 4384 4385 4386 4387      4389 4390 4391 4392 4393 4394 4395 4396 4397 4398 4399)
-#pragma warning (error: 4400 4401 4402 4403 4404 4405 4406 4407 4408 4409 4410 4411 4412 4413 4414 4415 4416 4417 4418 4419 4420 4421 4422 4423 4424 4425 4426 4427 4428 4429 4430 4431 4432 4433 4434      4436 4437 4438 4439 4440 4441 4442 4443 4444 4445 4446 4447 4448 4449 4450 4451 4452 4453 4454 4455 4456 4457 4458 4459 4460 4461 4462 4463 4464 4465 4466 4467 4468 4469 4470      4472 4473 4474 4475 4476 4477 4478 4479 4480      4482 4483 4484 4485 4486 4487 4488 4489 4490 4491 4492 4493 4494 4495 4496 4497 4498 4499)
+#pragma warning (error: 4400 4401 4402 4403 4404 4405 4406 4407 4408 4409 4410 4411 4412 4413 4414 4415 4416 4417 4418 4419 4420 4421 4422 4423 4424 4425 4426      4428 4429 4430 4431 4432 4433 4434      4436 4437 4438 4439 4440 4441 4442 4443 4444 4445 4446 4447 4448 4449 4450 4451 4452 4453 4454 4455 4456 4457 4458 4459 4460 4461 4462 4463 4464 4465 4466 4467 4468 4469 4470      4472 4473 4474 4475 4476 4477 4478 4479 4480      4482 4483 4484 4485 4486 4487 4488 4489 4490 4491 4492 4493 4494 4495 4496 4497 4498 4499)
 #pragma warning (error: 4500 4501 4502 4503 4504 4505 4506 4507 4508 4509 4510 4511 4512 4513 4514 4515 4516 4517 4518 4519 4520 4521 4522 4523 4524 4525 4526 4527 4528 4529 4530 4531 4532 4533 4534 4535 4536 4537 4538 4539 4540 4541 4542 4543 4544 4545 4546                4550 4551 4552 4553 4554      4556 4557 4558 4559 4560 4561 4562 4563 4564 4565 4566 4567 4568 4569 4570      4572 4573      4575 4576 4577 4578 4579 4580 4581 4582 4583 4584 4585 4586 4587 4588 4589 4590 4591 4592 4593 4594 4595 4596 4597 4598 4599)
 #pragma warning (error: 4600 4601 4602 4603 4604 4605 4606 4607 4608 4609 4610 4611 4612 4613 4614 4615 4616 4617 4618      4620 4621 4622 4623 4624 4625 4626 4627 4628 4629 4630 4631 4632 4633 4634 4635 4636 4637 4638 4639 4640 4641 4642 4643 4644 4645 4646 4647 4648 4649 4650      4652 4653 4654 4655 4656 4657 4658 4659 4660 4661 4662 4663 4664 4665 4666 4667 4668 4669 4670 4671 4672 4673 4674 4675 4676 4677 4678 4679 4680 4681 4682 4683 4684 4685 4686 4687 4688 4689 4690 4691 4692 4693 4694 4695 4696 4697 4698 4699)
 #pragma warning (error: 4700 4701 4702 4703 4704 4705 4706 4707 4708 4709 4710 4711 4712 4713 4714 4715 4716 4717 4718 4719 4720 4721 4722 4723 4724 4725 4726 4727 4728 4729      4731 4732 4733 4734 4735 4736 4737 4738 4739 4740 4741 4742 4743 4744 4745 4746 4747 4748 4749 4750 4751 4752 4753 4754 4755 4756 4757 4758 4759 4760 4761 4762 4763 4764 4765 4766 4767 4768 4769 4770 4771 4772 4773 4774 4775 4776 4777 4778 4779 4780 4781 4782 4783 4784 4785 4786 4787 4788 4789 4790 4791 4792 4793 4794 4795 4796 4797 4798 4799)
@@ -55,6 +55,7 @@ checkAtCompileTime(_MSC_VER >= 1700, Visual_Studio_2012_Required);
 // 4365 - 'action' : conversion from 'type_1' to 'type_2', signed/unsigned mismatch												http://msdn.microsoft.com/en-us/library/ms173683.aspx
 // 4371 - layout of class may have changed from a previous version of the compiler due to better packing of member
 // 4388 - '==' : signed/unsigned mismatch																						http://msdn.microsoft.com/en-us/library/jj155806.aspx
+// 4427 - 'enum' is a native enum: add an access specifier (private/public) to declare a managed enum
 // 4435 - 'class1' : Object layout under /vd2 will change due to virtual base 'class2'
 // 4471 - a forward declaration of an unscoped enumeration must have an underlying type (int assumed)
 // 4481 - nonstandard extension used: override specifier 'keyword'																http://msdn.microsoft.com/en-us/library/ms173703.aspx
