CrashReportClient now checks app name before uploading crash

Addresses TTP 332274 (mainly CrashReportClientApp.cpp changes) and
- Fixes hang on Linux (uploader wasn't honoring always-unattended setting)
- Fixed erroneous check warnings
- Removed code duplication between CrashReportClient platforms
- Error message now passed on for ensures as well as asserts
- Windows crash reporting multiple instance bail-out is now thread-safe

diff --git a/Engine/Source/Programs/CrashReporter/CrashReportClient/Private/CrashReportClient.cpp b/Engine/Source/Programs/CrashReporter/CrashReportClient/Private/CrashReportClient.cpp
--- a/Engine/Source/Programs/CrashReporter/CrashReportClient/Private/CrashReportClient.cpp
+++ b/Engine/Source/Programs/CrashReporter/CrashReportClient/Private/CrashReportClient.cpp
@@ -14,17 +14,15 @@ FString GCrashUserId;
 
 #if !CRASH_REPORT_UNATTENDED_ONLY
 
-FCrashReportClient::FCrashReportClient(const FString& Directory)
-	: ReportDirectory(Directory)
-	, AppState(EApplicationState::Ready)
+FCrashReportClient::FCrashReportClient(const FPlatformErrorReport& ErrorReport, const FString& AppName)
+	: AppState(EApplicationState::Ready)
 	, SubmittedCountdown(-1)
 	, bDiagnosticFileSent(false)
-	, ErrorReportFiles(Directory)
+	, ErrorReportFiles(ErrorReport)
 	, Uploader(GServerIP)
+	, CrashedAppName(AppName)
 	, CancelButtonText(LOCTEXT("Cancel", "Don't Send"))
 {
-	CrashedAppName = ErrorReportFiles.FindCrashedAppName();
-
 	// Set global user name ID: will be added to the report
 	if (FRocketSupport::IsRocket())
 	{
@@ -135,10 +133,9 @@ void FCrashReportClient::StartUIWillCloseTicker()
 
 void FCrashReportClient::StoreCommentAndUpload()
 {
+	// Call upload even if the report is empty: pending reports will be sent if any
 	ErrorReportFiles.SetUserComment(UserComment.IsEmpty() ? LOCTEXT("NoComment", "No comment provided") : UserComment);
-
-	// Using GetCleanFilename to actually get directory leaf name
-	Uploader.BeginUpload(FPaths::GetCleanFilename(ReportDirectory), ErrorReportFiles);
+	Uploader.BeginUpload(ErrorReportFiles);
 
 	SubmittedCountdown = 5;
 	AppState = EApplicationState::CountingDown;
@@ -160,7 +157,7 @@ bool FCrashReportClient::UIWillCloseTick(float UnusedDeltaTime)
 
 	if (!bDiagnosticFileSent && DiagnoseReportTask.IsWorkDone())
 	{
-		auto DiagnosticsFilePath = ReportDirectory / GDiagnosticsFilename;
+		auto DiagnosticsFilePath = ErrorReportFiles.GetReportDirectory() / GDiagnosticsFilename;
 
 		Uploader.LocalDiagnosisComplete(FPaths::FileExists(DiagnosticsFilePath) ? DiagnosticsFilePath : TEXT(""));
 		bDiagnosticFileSent = true;	
@@ -180,9 +177,8 @@ bool FCrashReportClient::UIWillCloseTick(float UnusedDeltaTime)
 		}
 	}
 
-	// IsWorkDone will always return true here (since uploader can't finish until
-	// the diagnosis has been sent), but it has the side effect of joining the
-	// worker thread.
+	// IsWorkDone will always return true here (since uploader can't finish until the diagnosis has been sent), but it
+	//  has the side effect of joining the worker thread.
 	if (!Uploader.IsFinished() || !DiagnoseReportTask.IsDone())
 	{
 		// More ticks, please
diff --git a/Engine/Source/Programs/CrashReporter/CrashReportClient/Private/CrashReportClient.h b/Engine/Source/Programs/CrashReporter/CrashReportClient/Private/CrashReportClient.h
--- a/Engine/Source/Programs/CrashReporter/CrashReportClient/Private/CrashReportClient.h
+++ b/Engine/Source/Programs/CrashReporter/CrashReportClient/Private/CrashReportClient.h
@@ -39,10 +39,10 @@ class FCrashReportClient : public TSharedFromThis<FCrashReportClient>
 {
 public:
 	/**
-	 * Find crashed app name and set up background tasks
-	 * @param ReportDirectory Full path to report to upload
+	 * Constructor: sets up background diagnosis
+	 * @param ErrorReport Error report to upload
 	 */
-	explicit FCrashReportClient(const FString& ReportDirectory);
+	explicit FCrashReportClient(const FPlatformErrorReport& ErrorReport, const FString& AppName);
 
 	/**
 	 * Respond to the user pressing Submit
diff --git a/Engine/Source/Programs/CrashReporter/CrashReportClient/Private/CrashReportClientApp.cpp b/Engine/Source/Programs/CrashReporter/CrashReportClient/Private/CrashReportClientApp.cpp
--- a/Engine/Source/Programs/CrashReporter/CrashReportClient/Private/CrashReportClientApp.cpp
+++ b/Engine/Source/Programs/CrashReporter/CrashReportClient/Private/CrashReportClientApp.cpp
@@ -26,12 +26,20 @@ IMPLEMENT_APPLICATION(CrashReportClient, "CrashReportClient");
 
 DEFINE_LOG_CATEGORY(CrashReportClientLog);
 
+namespace
+{
+	/** Directory containing the report */
+	FString ReportDirectoryAbsolutePath;
+
+	/** Name of crashed app to display */
+	FString CrashedAppName;
+}
 /**
  * Upload the crash report with no user interaction
  */
-void RunCrashReportClientUnattended(FMainLoopTiming& MainLoop, const FString& ReportDirectory)
+void RunCrashReportClientUnattended(FMainLoopTiming& MainLoop, const FPlatformErrorReport& ErrorReport)
 {
-	FCrashReportClientUnattended CrashReportClient(ReportDirectory);
+	FCrashReportClientUnattended CrashReportClient(ErrorReport);
 
 	// loop until the app is ready to quit
 	while (!GIsRequestingExit)
@@ -43,19 +51,54 @@ void RunCrashReportClientUnattended(FMainLoopTiming& MainLoop, const FString& Re
 /**
  * Look for the report to upload, either in the command line or in the platform's report queue
  */
-FString FindCrashReport(const TCHAR* CommandLine)
+void ParseCommandLine(const TCHAR* CommandLine)
 {
 	const TCHAR* CommandLineAfterExe = FCommandLine::RemoveExeName(CommandLine);
 
 	// Use the first argument if present and it's not a flag
-	if (*CommandLineAfterExe && *CommandLineAfterExe != '-')
+	if (*CommandLineAfterExe)
 	{
-		return FParse::Token(CommandLineAfterExe, true /* handle escaped quotes */);
+		if (*CommandLineAfterExe != '-')
+		{
+			ReportDirectoryAbsolutePath = FParse::Token(CommandLineAfterExe, true /* handle escaped quotes */);
+		}
+		FParse::Value(CommandLineAfterExe, TEXT("AppName="), CrashedAppName);
 	}
-	else
+
+	if (ReportDirectoryAbsolutePath.IsEmpty())
+	{
+		ReportDirectoryAbsolutePath = FPlatformErrorReport::FindMostRecentErrorReport();
+	}
+}
+
+/**
+ * Find the error report folder and check it matches the app name if provided
+ */
+FPlatformErrorReport LoadErrorReport()
+{
+	if (ReportDirectoryAbsolutePath.IsEmpty())
+	{
+		UE_LOG(CrashReportClientLog, Warning, TEXT("No error report found"));
+		return FPlatformErrorReport();
+	}
+
+	FPlatformErrorReport ErrorReport(ReportDirectoryAbsolutePath);
+#if CRASH_REPORT_UNATTENDED_ONLY
+	return ErrorReport;
+#else
+	auto AppNameInReport = ErrorReport.FindCrashedAppName();
+	if (CrashedAppName.IsEmpty())
+	{
+		// If no app name was provided on the command line, assume the report is the correct one
+		CrashedAppName = AppNameInReport;
+	}
+	else if (CrashedAppName != AppNameInReport)
 	{
-		return FPlatformErrorReport::FindMostRecentErrorReport();
+		// Don't display or upload anything if it's not the report we expected
+		ErrorReport = FPlatformErrorReport();
 	}
+	return ErrorReport;
+#endif
 }
 
 void RunCrashReportClient(const TCHAR* CommandLine)
@@ -74,19 +117,14 @@ void RunCrashReportClient(const TCHAR* CommandLine)
 	FMainLoopTiming MainLoop(IdealTickRate, bUnattended ? EMainLoopOptions::CoreTickerOnly : EMainLoopOptions::UsingSlate);
 
 	// Find the report to upload in the command line arguments
-	auto CommandLineAfterExe = FCommandLine::RemoveExeName(CommandLine);
-	FString ReportDirectory = FindCrashReport(CommandLine);
-
-	if (ReportDirectory.IsEmpty())
-	{
-		UE_LOG(CrashReportClientLog, Warning, TEXT("No error report found"));
-	}
-
-	FParse::Token(CommandLineAfterExe, true /* handle escaped quotes */);
+	ParseCommandLine(CommandLine);
 
+	FPlatformErrorReport::Init();
+	auto ErrorReport = LoadErrorReport();
+	
 	if (bUnattended)
 	{
-		RunCrashReportClientUnattended(MainLoop, ReportDirectory);
+		RunCrashReportClientUnattended(MainLoop, ErrorReport);
 	}
 	else
 	{
@@ -98,7 +136,7 @@ void RunCrashReportClient(const TCHAR* CommandLine)
 		FCrashReportClientStyle::Initialize();
 
 		// Create the main implementation object
-		TSharedRef<FCrashReportClient> CrashReportClient = MakeShareable(new FCrashReportClient(ReportDirectory));
+		TSharedRef<FCrashReportClient> CrashReportClient = MakeShareable(new FCrashReportClient(ErrorReport, CrashedAppName));
 
 		// open up the app window	
 		TSharedRef<SCrashReportClient> ClientControl = SNew(SCrashReportClient, CrashReportClient);
@@ -147,10 +185,15 @@ void RunCrashReportClient(const TCHAR* CommandLine)
 	}
 
 	GEngineLoop.AppExit();
+
+	FPlatformErrorReport::ShutDown();
 }
 
 void CrashReportClientCheck(bool bCondition, const TCHAR* Location)
 {
-	UE_LOG(CrashReportClientLog, Warning, TEXT("CHECK FAILED at %s"), Location);
+	if (!bCondition)
+	{
+		UE_LOG(CrashReportClientLog, Warning, TEXT("CHECK FAILED at %s"), Location);
+	}
 }
 
diff --git a/Engine/Source/Programs/CrashReporter/CrashReportClient/Private/CrashReportClientUnattended.cpp b/Engine/Source/Programs/CrashReporter/CrashReportClient/Private/CrashReportClientUnattended.cpp
--- a/Engine/Source/Programs/CrashReporter/CrashReportClient/Private/CrashReportClientUnattended.cpp
+++ b/Engine/Source/Programs/CrashReporter/CrashReportClient/Private/CrashReportClientUnattended.cpp
@@ -4,12 +4,13 @@
 
 #include "CrashReportClientUnattended.h"
 
-FCrashReportClientUnattended::FCrashReportClientUnattended(const FString& Directory)
-	: ReportDirectory(Directory)
-	, Uploader(GServerIP)
+FCrashReportClientUnattended::FCrashReportClientUnattended(const FPlatformErrorReport& ErrorReport)
+	: Uploader(GServerIP)
 {
-	// Using GetCleanFilename to actually get directory leaf name
-	Uploader.BeginUpload(FPaths::GetCleanFilename(ReportDirectory), FPlatformErrorReport(ReportDirectory));
+	Uploader.BeginUpload(ErrorReport);
+
+	// Prevent uploader waiting for a diagnosis file
+	Uploader.LocalDiagnosisSkipped();
 
 	StartTicker();
 }
diff --git a/Engine/Source/Programs/CrashReporter/CrashReportClient/Private/CrashReportClientUnattended.h b/Engine/Source/Programs/CrashReporter/CrashReportClient/Private/CrashReportClientUnattended.h
--- a/Engine/Source/Programs/CrashReporter/CrashReportClient/Private/CrashReportClientUnattended.h
+++ b/Engine/Source/Programs/CrashReporter/CrashReportClient/Private/CrashReportClientUnattended.h
@@ -12,9 +12,9 @@ class FCrashReportClientUnattended
 public:
 	/**
 	 * Set up uploader object
-	 * @param ReportDirectory Full path to report to upload
+	 * @param ErrorReport Error report to upload
 	 */
-	explicit FCrashReportClientUnattended(const FString& ReportDirectory);
+	explicit FCrashReportClientUnattended(const FPlatformErrorReport& ErrorReport);
 
 private:
 	/**
@@ -29,9 +29,6 @@ private:
 	 */
 	void StartTicker();
 
-	/** Complete path of report directory - must be first (see constructor) */
-	FString ReportDirectory;
-
 	/** Object that uploads report files to the server */
 	FCrashUpload Uploader;
 };
diff --git a/Engine/Source/Programs/CrashReporter/CrashReportClient/Private/CrashUpload.cpp b/Engine/Source/Programs/CrashReporter/CrashReportClient/Private/CrashUpload.cpp
--- a/Engine/Source/Programs/CrashReporter/CrashReportClient/Private/CrashUpload.cpp
+++ b/Engine/Source/Programs/CrashReporter/CrashReportClient/Private/CrashUpload.cpp
@@ -48,15 +48,14 @@ bool FCrashUpload::PingTimeout(float DeltaTime)
 	return false;
 }
 
-void FCrashUpload::BeginUpload(const FString& ReportDirectoryNameParam, const FPlatformErrorReport& PlatformErrorReport)
+void FCrashUpload::BeginUpload(const FPlatformErrorReport& PlatformErrorReport)
 {
-	ReportDirectoryName = ReportDirectoryNameParam;
 	ErrorReport = PlatformErrorReport;
 	PendingFiles = ErrorReport.GetFilesToUpload();
-	UE_LOG(CrashReportClientLog, Log, TEXT("Got %d pending files to upload from '%s'"), PendingFiles.Num(), *ReportDirectoryName);
+	UE_LOG(CrashReportClientLog, Log, TEXT("Got %d pending files to upload from '%s'"), PendingFiles.Num(), *ErrorReport.GetReportDirectoryLeafName());
 
 	// Pause before posting completed message, to allow for additional files to be uploaded
-	PauseState = FApp::IsUnattended() ? EUploadState::Finished : EUploadState::PostingReportComplete;
+	PauseState = EUploadState::PostingReportComplete;
 	if (State == EUploadState::Ready)
 	{
 		BeginUploadImpl();
@@ -178,7 +177,7 @@ void FCrashUpload::UploadNextFile()
 		Request->SetHeader(TEXT("Content-Type"), TEXT("application/octet-stream"));
 		Request->SetURL(UrlPrefix / TEXT("UploadReportFile"));
 		Request->SetContent(PostData);
-		Request->SetHeader(TEXT("DirectoryName"), ReportDirectoryName);
+		Request->SetHeader(TEXT("DirectoryName"), *ErrorReport.GetReportDirectoryLeafName());
 		Request->SetHeader(TEXT("FileName"), Filename);
 		Request->SetHeader(TEXT("FileLength"), FString::FromInt(PostData.Num()));
 
@@ -194,6 +193,7 @@ void FCrashUpload::UploadNextFile()
 
 void FCrashUpload::AssignReportIdToPostDataBuffer()
 {
+	FString ReportDirectoryName = *ErrorReport.GetReportDirectoryLeafName();
 	const int32 DirectoryNameLength = ReportDirectoryName.Len();
 	PostData.SetNum(DirectoryNameLength);
 	for (int32 Index = 0; Index != DirectoryNameLength; ++Index)
@@ -246,7 +246,7 @@ void FCrashUpload::OnProcessRequestComplete(FHttpRequestPtr HttpRequest, FHttpRe
 		}
 		else
 		{
-			SetCurrentState(EUploadState::ServerNotAvailable);
+			PingTimeout(0);
 		}
 		break;
 
@@ -339,10 +339,7 @@ void FCrashUpload::CheckPendingReportsForFilesToUpload()
 			return;
 		}
 
-		auto ReportDirectory = PendingReportDirectories.Pop();
-		ReportDirectoryName = FPaths::GetCleanFilename(ReportDirectory);
-
-		ErrorReport = FPlatformErrorReport(ReportDirectory);
+		ErrorReport = FPlatformErrorReport(PendingReportDirectories.Pop());
 		PendingFiles = ErrorReport.GetFilesToUpload();
 	}
 	while (PendingFiles.Num() == 0 || !SendCheckReportRequest());
@@ -352,10 +349,10 @@ void FCrashUpload::BeginUploadImpl()
 {
 	FPendingReports PendingReports;
 	// Make sure we don't try to upload the current report twice
-	PendingReports.Forget(ReportDirectoryName);
+	PendingReports.Forget(*ErrorReport.GetReportDirectoryLeafName());
 	PendingReportDirectories = PendingReports.GetReportDirectories();
 
-	UE_LOG(CrashReportClientLog, Log, TEXT("Proceeding to upload %d pending files from '%s'"), PendingFiles.Num(), *ReportDirectoryName);
+	UE_LOG(CrashReportClientLog, Log, TEXT("Proceeding to upload %d pending files from '%s'"), PendingFiles.Num(), *ErrorReport.GetReportDirectoryLeafName());
 
 	// Not sure about this: probably want to removes files one by one as they're successfully uploaded
 	PendingReports.Clear();
diff --git a/Engine/Source/Programs/CrashReporter/CrashReportClient/Private/CrashUpload.h b/Engine/Source/Programs/CrashReporter/CrashReportClient/Private/CrashUpload.h
--- a/Engine/Source/Programs/CrashReporter/CrashReportClient/Private/CrashUpload.h
+++ b/Engine/Source/Programs/CrashReporter/CrashReportClient/Private/CrashUpload.h
@@ -24,10 +24,9 @@ public:
 
 	/**
 	 * Commence upload when ready
-	 * @param ReportDirectoryName Leaf name of the directory containing the report
 	 * @param PlatformErrorReport Error report to upload files from
 	 */
-	void BeginUpload(const FString& ReportDirectoryName, const FPlatformErrorReport& PlatformErrorReport);
+	void BeginUpload(const FPlatformErrorReport& PlatformErrorReport);
 
 	/**
 	 * Provide progress or error information for the UI
@@ -40,6 +39,14 @@ public:
 	 */
 	void LocalDiagnosisComplete(const FString& DiagnosticsFile);
 
+	/**
+	 * Inform uploader that no local diagnosis is happening
+	 */
+	void LocalDiagnosisSkipped()
+	{
+		LocalDiagnosisComplete("");
+	}
+
 	/**
 	 * Determine whether the upload has finished (successfully or otherwise)
 	 * @return Whether the upload has finished
@@ -179,9 +186,6 @@ private:
 	/** Flag to record if a diagnostics report text file has been sent */
 	bool bDiagnosticsFileSent;
 
-	/** Leaf name of the report directory */
-	FString ReportDirectoryName;
-
 	/** Full paths of files still to be uploaded */
 	TArray<FString> PendingFiles;
 
diff --git a/Engine/Source/Programs/CrashReporter/CrashReportClient/Private/GenericErrorReport.cpp b/Engine/Source/Programs/CrashReporter/CrashReportClient/Private/GenericErrorReport.cpp
--- a/Engine/Source/Programs/CrashReporter/CrashReportClient/Private/GenericErrorReport.cpp
+++ b/Engine/Source/Programs/CrashReporter/CrashReportClient/Private/GenericErrorReport.cpp
@@ -4,7 +4,6 @@
 
 #include "GenericErrorReport.h"
 #include "XmlFile.h"
-#include "CrashDebugHelperModule.h"
 #include "CrashReportUtil.h"
 
 // ----------------------------------------------------------------
@@ -39,7 +38,6 @@ FGenericErrorReport::FGenericErrorReport(const FString& Directory)
 		return true;
 	});
 	FPlatformFileManager::Get().GetPlatformFile().IterateDirectory(*ReportDirectory, FilenamesVisitor);
-
 }
 
 bool FGenericErrorReport::SetUserComment(const FText& UserComment)
diff --git a/Engine/Source/Programs/CrashReporter/CrashReportClient/Private/GenericErrorReport.h b/Engine/Source/Programs/CrashReporter/CrashReportClient/Private/GenericErrorReport.h
--- a/Engine/Source/Programs/CrashReporter/CrashReportClient/Private/GenericErrorReport.h
+++ b/Engine/Source/Programs/CrashReporter/CrashReportClient/Private/GenericErrorReport.h
@@ -2,6 +2,8 @@
 
 #pragma once
 
+class FCrashDebugHelperModule;
+
 /**
  * Helper that works with Windows Error Reports
  */
@@ -21,6 +23,20 @@ public:
 	 */
 	explicit FGenericErrorReport(const FString& Directory);
 
+	/**
+	 * One-time initialisation: does nothing by default
+	 */
+	static void Init()
+	{
+	}
+
+	/**
+	 * One-time clean-up: does nothing by default
+	 */
+	static void ShutDown()
+	{
+	}
+
 	/**
 	 * Write the provided comment into the error report
 	 * @param UserComment Information provided by the user to add to the report
@@ -49,6 +65,31 @@ public:
 	 */
 	bool TryReadDiagnosticsFile(FText& OutReportDescription);
 
+	/**
+	 * Provide the full path of the error report directory
+	 */
+	FString GetReportDirectory() const
+	{
+		return ReportDirectory;
+	}
+
+	/**
+	 * Provide the name of the error report directory
+	 */
+	FString GetReportDirectoryLeafName() const
+	{
+		// Using GetCleanFilename to actually get directory name
+		return FPaths::GetCleanFilename(ReportDirectory);
+	}
+
+	/**
+	 * Is there anything to upload?
+	 */
+	bool HasFilesToUpload() const
+	{
+		return ReportFilenames.Num() != 0;
+	}
+
 protected:
 	/**
 	 * Look throught the list of report files to find one with the given extension
@@ -59,6 +100,7 @@ protected:
 	/** Full path to the directory the report files are in */
 	FString ReportDirectory;
 
+private:
 	/** List of leaf filenames of all the files in the report folder */
 	TArray<FString> ReportFilenames;
 };
diff --git a/Engine/Source/Programs/CrashReporter/CrashReportClient/Private/Linux/LinuxErrorReport.cpp b/Engine/Source/Programs/CrashReporter/CrashReportClient/Private/Linux/LinuxErrorReport.cpp
deleted file mode 100644
diff --git a/Engine/Source/Programs/CrashReporter/CrashReportClient/Private/Linux/LinuxErrorReport.h b/Engine/Source/Programs/CrashReporter/CrashReportClient/Private/Linux/LinuxErrorReport.h
--- a/Engine/Source/Programs/CrashReporter/CrashReportClient/Private/Linux/LinuxErrorReport.h
+++ b/Engine/Source/Programs/CrashReporter/CrashReportClient/Private/Linux/LinuxErrorReport.h
@@ -2,10 +2,12 @@
 
 #pragma once
 
+#include "../GenericErrorReport.h"
+
 /**
  * Helper that works with Linux Error Reports
  */
-class FLinuxErrorReport
+class FLinuxErrorReport : public FGenericErrorReport
 {
 public:
 	/**
@@ -19,62 +21,27 @@ public:
 	 * Discover all files in the crash report directory
 	 * @param Directory Full path to directory containing the report
 	 */
-	explicit FLinuxErrorReport(const FString& Directory);
-
-	/**
-	 * Unload helper modules
-	 */
-	~FLinuxErrorReport();
-
-	/**
-	 * Write the provided comment into the error report
-	 * @param UserComment Information provided by the user to add to the report
-	 * @return Whether the comment was successfully written to the report
-	 */
-	bool SetUserComment(const FText& UserComment);
-
-	/**
-	 * Provide full paths to all the report files
-	 * @return List of paths
-	 */
-	TArray<FString> GetFilesToUpload() const;
-
-	/**
-	 * Provide the exception and a call-stack as plain text if possible
-	 * @note This can take quite a long time
-	 */
-	FText DiagnoseReport() const;
-
-	/**
-	 * Get the name of the crashed app from the report
-	 */
-	FString FindCrashedAppName() const;
-
-	/**
-	 * Load the WER XML file for this report
-	 * @note This is Windows specific and so shouldn't really be part of the public interface, but currently the server
-	 * is Windows-specific in its checking of reports, so this is needed.
-	 * @param OutBuffer Buffer to load the file into
-	 * @return Whether finding and loading the file succeeded
-	 */
-	bool LoadWindowsReportXmlFile(TArray<uint8>& OutBuffer) const;
+	explicit FLinuxErrorReport(const FString& Directory)
+		: FGenericErrorReport(Directory)
+	{
+	}
 
 	/**
-	 * Look for the most recent Linux Error Report
-	 * @return Full path to the most recent report, or an empty string if none found
+	 * Do nothing - shouldn't be called on Linux
+	 * @return Dummy text
 	 */
-	static FString FindMostRecentErrorReport();
+	FText DiagnoseReport() const	
+	{
+		return FText::FromString("No local diagnosis on Linux");
+	}
 
-private:
 	/**
-	 * Look throught the list of report files to find one with the given extension
-	 * @return Whether a file with the extension was found
+	 * Do nothing - shouldn't be called on Linux
+	 * @return Empty string
 	 */
-	bool FindFirstReportFileWithExtension(FString& OutFilename, const TCHAR* Extension) const;
-
-	/** Full path to the directory the report files are in */
-	FString ReportDirectory;
-
-	/** List of leaf filenames of all the files in the report folder */
-	TArray<FString> ReportFilenames;
+	static FString FindMostRecentErrorReport()
+	{
+		// The report folder is currently always sent on the command-line on Linux
+		return TEXT("");
+	}
 };
diff --git a/Engine/Source/Programs/CrashReporter/CrashReportClient/Private/Mac/MacErrorReport.h b/Engine/Source/Programs/CrashReporter/CrashReportClient/Private/Mac/MacErrorReport.h
--- a/Engine/Source/Programs/CrashReporter/CrashReportClient/Private/Mac/MacErrorReport.h
+++ b/Engine/Source/Programs/CrashReporter/CrashReportClient/Private/Mac/MacErrorReport.h
@@ -32,9 +32,8 @@ public:
 		return FText::FromString("No local diagnosis on Mac");
 	}
 
-
 	/**
-	 * Get the name of the crashed app from the report
+	 * Get the name of the crashed app from the report (hides implementation in FGenericErrorReport)
 	 */
 	FString FindCrashedAppName() const;
 
diff --git a/Engine/Source/Programs/CrashReporter/CrashReportClient/Private/Windows/WindowsErrorReport.cpp b/Engine/Source/Programs/CrashReporter/CrashReportClient/Private/Windows/WindowsErrorReport.cpp
--- a/Engine/Source/Programs/CrashReporter/CrashReportClient/Private/Windows/WindowsErrorReport.cpp
+++ b/Engine/Source/Programs/CrashReporter/CrashReportClient/Private/Windows/WindowsErrorReport.cpp
@@ -60,11 +60,15 @@ struct FWindowsReportParser
 
 FWindowsErrorReport::FWindowsErrorReport(const FString& Directory)
 	: FGenericErrorReport(Directory)
+{
+}
+
+void FWindowsErrorReport::Init()
 {
 	CrashHelperModule = &FModuleManager::LoadModuleChecked<FCrashDebugHelperModule>(FName("CrashDebugHelper"));
 }
 
-FWindowsErrorReport::~FWindowsErrorReport()
+void FWindowsErrorReport::ShutDown()
 {
 	CrashHelperModule->ShutdownModule();
 }
@@ -72,7 +76,7 @@ FWindowsErrorReport::~FWindowsErrorReport()
 FText FWindowsErrorReport::DiagnoseReport() const
 {
 	// Should check if there are local PDBs before doing anything
-	auto CrashDebugHelper = CrashHelperModule->Get();
+	auto CrashDebugHelper = CrashHelperModule ? CrashHelperModule->Get() : nullptr;
 	if (!CrashDebugHelper)
 	{
 		// Not localized: should never be seen
@@ -115,7 +119,7 @@ FText FWindowsErrorReport::DiagnoseReport() const
 
 FString FWindowsErrorReport::FindCrashedAppName() const
 {
-	const FString CrashedAppName = FWindowsReportParser::Find( ReportDirectory, TEXT( "AppPath=" ) );
+	const FString CrashedAppName = FWindowsReportParser::Find( ReportDirectory, TEXT( "AppName=" ) );
 	return CrashedAppName;
 }
 
diff --git a/Engine/Source/Programs/CrashReporter/CrashReportClient/Private/Windows/WindowsErrorReport.h b/Engine/Source/Programs/CrashReporter/CrashReportClient/Private/Windows/WindowsErrorReport.h
--- a/Engine/Source/Programs/CrashReporter/CrashReportClient/Private/Windows/WindowsErrorReport.h
+++ b/Engine/Source/Programs/CrashReporter/CrashReportClient/Private/Windows/WindowsErrorReport.h
@@ -18,15 +18,20 @@ public:
 	}
 
 	/**
-	 * Discover all files in the crash report directory
-	 * @param Directory Full path to directory containing the report
+	 * Load helper modules
 	 */
-	explicit FWindowsErrorReport(const FString& Directory);
+	static void Init();
 
 	/**
 	 * Unload helper modules
 	 */
-	~FWindowsErrorReport();
+	static void ShutDown();
+
+	/**
+	 * Discover all files in the crash report directory
+	 * @param Directory Full path to directory containing the report
+	 */
+	explicit FWindowsErrorReport(const FString& Directory);
 
 	/**
 	 * Provide the exception and a call-stack as plain text if possible
diff --git a/Engine/Source/Runtime/Core/Private/Windows/WindowsPlatformExceptionHandling.cpp b/Engine/Source/Runtime/Core/Private/Windows/WindowsPlatformExceptionHandling.cpp
--- a/Engine/Source/Runtime/Core/Private/Windows/WindowsPlatformExceptionHandling.cpp
+++ b/Engine/Source/Runtime/Core/Private/Windows/WindowsPlatformExceptionHandling.cpp
@@ -21,27 +21,49 @@
 #pragma comment( lib, "version.lib" )
 #pragma comment( lib, "Shlwapi.lib" )
 
-static bool GAlreadyCreatedMinidump = false;
+namespace
+{
+static const int LocalBufferSize = 1024;
+volatile LONG ReportCrashCallCount = 0;
+const TCHAR* AssertionMessageStart = nullptr;	// pointer into GErrorHist, if non-null
+int AssertionMessageLength;						// Number of characters to copy after AssertionMessageStart
+
+/**
+ * Write a Windows minidump to disk
+ * @param Path Full path of file to write (normally a .dmp file)
+ * @param ExceptionInfo Pointer to structure containing the exception information
+ * @return Success or failure
+ */
+bool WriteMinidump(const TCHAR* Path, LPEXCEPTION_POINTERS ExceptionInfo)
+{
+	// Try to create file for minidump.
+	HANDLE FileHandle = CreateFileW(Path, GENERIC_WRITE, 0, NULL, CREATE_ALWAYS, FILE_ATTRIBUTE_NORMAL, NULL);
+	
+	if (FileHandle == INVALID_HANDLE_VALUE)
+		return false;
+
+	// Initialise structure required by MiniDumpWriteDumps
+	MINIDUMP_EXCEPTION_INFORMATION DumpExceptionInfo = {};
+
+	DumpExceptionInfo.ThreadId			= GetCurrentThreadId();
+	DumpExceptionInfo.ExceptionPointers	= ExceptionInfo;
+	DumpExceptionInfo.ClientPointers	= FALSE;
+
+	BOOL result = MiniDumpWriteDump(GetCurrentProcess(), GetCurrentProcessId(), FileHandle, MiniDumpNormal, &DumpExceptionInfo, NULL, NULL);
+	CloseHandle(FileHandle);
+
+	return result == TRUE;
+}
 
 /** 
  * Get a string description of the mode the engine was running in when it crashed
  */
-static const TCHAR* GetEngineMode()
+const TCHAR* GetEngineMode()
 {
-	if( IsRunningCommandlet() )
-	{
-		return TEXT( "Commandlet" );
-	}
-	else if( GIsEditor )
-	{
-		return TEXT( "Editor" );
-	}
-	else if( IsRunningDedicatedServer() )
-	{
-		return TEXT( "Server" );
-	}
-
-	return TEXT( "Game" );
+	return	IsRunningCommandlet()?	 	TEXT("Commandlet") :
+			GIsEditor?				 	TEXT("Editor") :
+			IsRunningDedicatedServer()?	TEXT("Server") :
+										TEXT("Game");
 }
 
 #if WINVER > 0x502	// Windows Error Reporting is not supported on Windows XP
@@ -49,32 +71,22 @@ static const TCHAR* GetEngineMode()
 /** 
  * Get one line of text to describe the crash
  */
-static void GetCrashDescription( WER_REPORT_INFORMATION& ReportInformation, const TCHAR* ErrorMessage )
+void GetCrashDescription(WER_REPORT_INFORMATION& ReportInformation, const TCHAR* ErrorMessage)
 {
-	if( ErrorMessage != NULL )
+	if (ErrorMessage)
 	{
-		StringCchCat( ReportInformation.wzDescription, ARRAYSIZE( ReportInformation.wzDescription ), ErrorMessage );
+		StringCchCat(ReportInformation.wzDescription, ARRAYSIZE(ReportInformation.wzDescription), ErrorMessage);
 	}
 	else
 	{
-		StringCchCopy( ReportInformation.wzDescription, ARRAYSIZE( ReportInformation.wzDescription ), TEXT( "The application crashed while running " ) );
+		StringCchCopy(ReportInformation.wzDescription, ARRAYSIZE(ReportInformation.wzDescription), TEXT("The application crashed while running "));
 
-		if( IsRunningCommandlet() )
-		{
-			StringCchCat( ReportInformation.wzDescription, ARRAYSIZE( ReportInformation.wzDescription ), TEXT( "a commandlet" ) );
-		}
-		else if( GIsEditor )
-		{
-			StringCchCat( ReportInformation.wzDescription, ARRAYSIZE( ReportInformation.wzDescription ), TEXT( "the editor" ) );
-		}
-		else if( GIsServer && !GIsClient )
-		{
-			StringCchCat( ReportInformation.wzDescription, ARRAYSIZE( ReportInformation.wzDescription ), TEXT( "a server" ) );
-		}
-		else
-		{
-			StringCchCat( ReportInformation.wzDescription, ARRAYSIZE( ReportInformation.wzDescription ), TEXT( "the game" ) );
-		}
+		const TCHAR* Description =	IsRunningCommandlet()?	 	TEXT("a commandlet") :
+									GIsEditor?				 	TEXT("the editor") :
+									IsRunningDedicatedServer()?	TEXT("a server") :
+																TEXT("the game");
+
+		StringCchCat(ReportInformation.wzDescription, ARRAYSIZE(ReportInformation.wzDescription), Description);
 	}
 }
 
@@ -83,7 +95,7 @@ static void GetCrashDescription( WER_REPORT_INFORMATION& ReportInformation, cons
 /**
  * Get the 4 element version number of the module
  */
-static void GetModuleVersion( TCHAR* ModuleName, TCHAR* StringBuffer, DWORD MaxSize )
+void GetModuleVersion( TCHAR* ModuleName, TCHAR* StringBuffer, DWORD MaxSize )
 {
 	StringCchCopy( StringBuffer, MaxSize, TEXT( "0.0.0.0" ) );
 	
@@ -117,79 +129,74 @@ static void GetModuleVersion( TCHAR* ModuleName, TCHAR* StringBuffer, DWORD MaxS
  * Parameters 0 through 7 are predefined for Windows
  * Parameters 8 and 9 are user defined
  */
-static void SetReportParameters( HREPORT ReportHandle, EXCEPTION_POINTERS* ExceptionInfo )
+void SetReportParameters( HREPORT ReportHandle, EXCEPTION_POINTERS* ExceptionInfo )
 {
 	HRESULT Result;
-	enum 
-	{
-		BUFFER_SIZE = 1024,
-	};
-	TCHAR StringBuffer[BUFFER_SIZE] = {0};
-	TCHAR LocalBuffer[BUFFER_SIZE] = {0};
+	TCHAR StringBuffer[LocalBufferSize] = {0};
+	TCHAR LocalBuffer[LocalBufferSize] = {0};
 
 	// Set the parameters for the standard problem signature
 	HMODULE ModuleHandle = GetModuleHandle( NULL );
 
-	StringCchPrintf( StringBuffer, BUFFER_SIZE, TEXT( "UE4-%s" ), FApp::GetGameName() );
+	StringCchPrintf( StringBuffer, LocalBufferSize, TEXT( "UE4-%s" ), FApp::GetGameName() );
 	Result = WerReportSetParameter( ReportHandle, WER_P0, TEXT( "Application Name" ), StringBuffer );
 
-	GetModuleFileName( ModuleHandle, LocalBuffer, BUFFER_SIZE );
+	GetModuleFileName( ModuleHandle, LocalBuffer, LocalBufferSize );
 	PathStripPath( LocalBuffer );
-	GetModuleVersion( LocalBuffer, StringBuffer, BUFFER_SIZE );
+	GetModuleVersion( LocalBuffer, StringBuffer, LocalBufferSize );
 	Result = WerReportSetParameter( ReportHandle, WER_P1, TEXT( "Application Version" ), StringBuffer );
 
-	StringCchPrintf( StringBuffer, BUFFER_SIZE, TEXT( "%08x" ), GetTimestampForLoadedLibrary( ModuleHandle ) );
+	StringCchPrintf( StringBuffer, LocalBufferSize, TEXT( "%08x" ), GetTimestampForLoadedLibrary( ModuleHandle ) );
 	Result = WerReportSetParameter( ReportHandle, WER_P2, TEXT( "Application Timestamp" ), StringBuffer );
 
 	HMODULE FaultModuleHandle = NULL;
 	GetModuleHandleEx( GET_MODULE_HANDLE_EX_FLAG_FROM_ADDRESS | GET_MODULE_HANDLE_EX_FLAG_UNCHANGED_REFCOUNT, ( LPCTSTR )ExceptionInfo->ExceptionRecord->ExceptionAddress, &FaultModuleHandle );
 
-	GetModuleFileName( FaultModuleHandle, LocalBuffer, BUFFER_SIZE );
+	GetModuleFileName( FaultModuleHandle, LocalBuffer, LocalBufferSize );
 	PathStripPath( LocalBuffer );
 	Result = WerReportSetParameter( ReportHandle, WER_P3, TEXT( "Fault Module Name" ), LocalBuffer );
 
-	GetModuleVersion( LocalBuffer, StringBuffer, BUFFER_SIZE );
+	GetModuleVersion( LocalBuffer, StringBuffer, LocalBufferSize );
 	Result = WerReportSetParameter( ReportHandle, WER_P4, TEXT( "Fault Module Version" ), StringBuffer );
 
-	StringCchPrintf( StringBuffer, BUFFER_SIZE, TEXT( "%08x" ), GetTimestampForLoadedLibrary( FaultModuleHandle ) );
+	StringCchPrintf( StringBuffer, LocalBufferSize, TEXT( "%08x" ), GetTimestampForLoadedLibrary( FaultModuleHandle ) );
 	Result = WerReportSetParameter( ReportHandle, WER_P5, TEXT( "Fault Module Timestamp" ), StringBuffer );
 
-	StringCchPrintf( StringBuffer, BUFFER_SIZE, TEXT( "%08x" ), ExceptionInfo->ExceptionRecord->ExceptionCode );
+	StringCchPrintf( StringBuffer, LocalBufferSize, TEXT( "%08x" ), ExceptionInfo->ExceptionRecord->ExceptionCode );
 	Result = WerReportSetParameter( ReportHandle, WER_P6, TEXT( "Exception Code" ), StringBuffer );
 
 	INT_PTR ExceptionOffset = ( char* )( ExceptionInfo->ExceptionRecord->ExceptionAddress ) - ( char* )FaultModuleHandle;
-	StringCchPrintf( StringBuffer, BUFFER_SIZE, TEXT( "%p" ), ExceptionOffset );
+	StringCchPrintf( StringBuffer, LocalBufferSize, TEXT( "%p" ), ExceptionOffset );
 	Result = WerReportSetParameter( ReportHandle, WER_P7, TEXT( "Exception Offset" ), StringBuffer );
 
 	// Look for the assertion fail.
-	const TCHAR* AssertPos = FCString::Stristr( GErrorHist, TEXT( "Assertion failed: " ) );
-	const TCHAR* StackPos = FCString::Stristr( GErrorHist, TEXT( "\nStack:" ) );
-	if( AssertPos != nullptr && StackPos != nullptr && StackPos > AssertPos )
+	if (AssertionMessageStart)
 	{
 		// Use LocalBuffer to store the assertion fail.
-		FCString::Strncpy( LocalBuffer, AssertPos, FMath::Min<int32>( StackPos - AssertPos, BUFFER_SIZE ) );
-
-		const int32 BufferLen = FCString::Strlen( LocalBuffer );
+		FCString::Strncpy(LocalBuffer, AssertionMessageStart, AssertionMessageLength);
 
 		// Replace " with ' and replace \n with #
-		for( int32 ChIndex = 0; ChIndex < BufferLen; ++ChIndex )
+		for (TCHAR& Char: LocalBuffer)
 		{
-			if( LocalBuffer[ChIndex] == TEXT( '\"' ) )
+			if (Char == 0)
 			{
-				LocalBuffer[ChIndex] = TEXT( '\'' );
+				break;
 			}
 
-			if( LocalBuffer[ChIndex] == TEXT( '\n' ) )
+			switch (Char)
 			{
-				LocalBuffer[ChIndex] = TEXT( '#' );
+				default: break;
+				case '"':	Char = '\'';	break;
+				case '\r':	Char = '#';		break;
+				case '\n':	Char = '#';		break;
 			}
 		}
 	}
 
-	StringCchPrintf( StringBuffer, BUFFER_SIZE, TEXT( "!%s!AssertLog=\"%s\"" ), FCommandLine::Get(), LocalBuffer );
+	StringCchPrintf( StringBuffer, LocalBufferSize, TEXT( "!%s!AssertLog=\"%s\"" ), FCommandLine::Get(), LocalBuffer );
 	Result = WerReportSetParameter( ReportHandle, WER_P8, TEXT( "Commandline" ), StringBuffer );
 
-	StringCchPrintf( StringBuffer, BUFFER_SIZE, TEXT( "%s!%s!%s!%d" ), TEXT( BRANCH_NAME ), FPlatformProcess::BaseDir(), GetEngineMode(), BUILT_FROM_CHANGELIST );
+	StringCchPrintf( StringBuffer, LocalBufferSize, TEXT( "%s!%s!%s!%d" ), TEXT( BRANCH_NAME ), FPlatformProcess::BaseDir(), GetEngineMode(), BUILT_FROM_CHANGELIST );
 	Result = WerReportSetParameter( ReportHandle, WER_P9, TEXT( "BranchBaseDir" ), StringBuffer );
 }
 
@@ -198,26 +205,18 @@ static void SetReportParameters( HREPORT ReportHandle, EXCEPTION_POINTERS* Excep
  *
  * Note this has to be a minidump (and not a microdump) for the dbgeng functions to work correctly
  */
-static void AddMiniDump( HREPORT ReportHandle, EXCEPTION_POINTERS* ExceptionInfo )
+void AddMiniDump(HREPORT ReportHandle, EXCEPTION_POINTERS* ExceptionInfo)
 {
 #if 0
 	// This doesn't work in x64
 	WER_EXCEPTION_INFORMATION WerExceptionInfo = { ExceptionInfo, FALSE };
 	Result = WerReportAddDump( ReportHandle, GetCurrentProcess(), GetCurrentThread(), WerDumpTypeMiniDump, &WerExceptionInfo, NULL, 0 );
 #else
-	FString MinidumpFileName = FString::Printf( TEXT( "%sDump%d.dmp" ), *FPaths::GameLogDir(), FDateTime::UtcNow().GetTicks() );
-	HANDLE DumpHandle = CreateFileW( *MinidumpFileName, GENERIC_WRITE, 0, NULL, CREATE_ALWAYS, FILE_ATTRIBUTE_NORMAL, NULL );
-	if( DumpHandle != INVALID_HANDLE_VALUE )
+	FString MinidumpFileName = FString::Printf(TEXT("%sDump%d.dmp"), *FPaths::GameLogDir(), FDateTime::UtcNow().GetTicks());
+	
+	if (WriteMinidump(*MinidumpFileName, ExceptionInfo))
 	{
-		MINIDUMP_EXCEPTION_INFORMATION DumpExceptionInfo = { 0 };
-		DumpExceptionInfo.ThreadId = GetCurrentThreadId();
-		DumpExceptionInfo.ExceptionPointers = ExceptionInfo;
-		DumpExceptionInfo.ClientPointers = true;
-
-		MiniDumpWriteDump( GetCurrentProcess(), GetCurrentProcessId(), DumpHandle, MiniDumpNormal, &DumpExceptionInfo, NULL, NULL );
-		CloseHandle( DumpHandle );
-
-		WerReportAddFile( ReportHandle, *MinidumpFileName, WerFileTypeMinidump, WER_FILE_ANONYMOUS_DATA ); 
+		WerReportAddFile(ReportHandle, *MinidumpFileName, WerFileTypeMinidump, WER_FILE_ANONYMOUS_DATA); 
 	}
 #endif
 }
@@ -225,7 +224,7 @@ static void AddMiniDump( HREPORT ReportHandle, EXCEPTION_POINTERS* ExceptionInfo
 /** 
  * Add miscellaneous files to the report. Currently the log and the video file
  */
-static void AddMiscFiles( HREPORT ReportHandle )
+void AddMiscFiles( HREPORT ReportHandle )
 {
 	FString LogFileName = FPaths::GameLogDir() / FApp::GetGameName() + TEXT( ".log" );
 	WerReportAddFile( ReportHandle, *LogFileName, WerFileTypeOther, WER_FILE_ANONYMOUS_DATA ); 
@@ -239,7 +238,7 @@ static void AddMiscFiles( HREPORT ReportHandle )
  *
  * This is required so that the crash report doesn't immediately get sent to Microsoft, and we can't intercept it
  */
-static void ForceWERQueuing( bool bEnable )
+void ForceWERQueuing( bool bEnable )
 {
 	HKEY Key = 0;
 	HRESULT KeyResult = RegOpenKeyEx( HKEY_CURRENT_USER, TEXT( "Software\\Microsoft\\Windows\\Windows Error Reporting" ), 0, KEY_WRITE, &Key );
@@ -251,8 +250,6 @@ static void ForceWERQueuing( bool bEnable )
 }
 
 
-namespace
-{
 /**
  * Enum indicating whether to run the crash reporter UI
  */
@@ -265,28 +262,46 @@ enum class EErrorReportUI
 	ReportInUnattendedMode	
 };
 
+/**
+ * Scoped setter of WER flag, to ensure it always gets put back
+ */
+class FScopedWERQueuing
+{
+public:
+	/**	Default constructor - force WER to queue, not send, reports */
+	FScopedWERQueuing()
+	{
+		ForceWERQueuing(true);
+	}
+
+	/**	Destructor - Reset the setting after we've had our way with it */
+	~FScopedWERQueuing()
+	{
+		ForceWERQueuing(false);
+	}
+
+private:
+	// disallow copying
+	FScopedWERQueuing(const FScopedWERQueuing&);
+	FScopedWERQueuing& operator=(const FScopedWERQueuing&);
+};
+
+
 /** 
  * Create a Windows Error Report, add the user log and video, and add it to the WER queue
- * Launch CrashReportUploader.exe to intercept the report and upload to our local site
+ * Launch CrashReportClient.exe to intercept the report and upload to our local site
  */
-int32 NewReportCrash( EXCEPTION_POINTERS* ExceptionInfo, const TCHAR* ErrorMessage, EErrorReportUI ReportUI )
+int32 ReportCrashUsingCrashReportClient(EXCEPTION_POINTERS* ExceptionInfo, const TCHAR* ErrorMessage, EErrorReportUI ReportUI)
 {
-	HRESULT Result;
-	HREPORT ReportHandle = NULL;
-	WER_REPORT_INFORMATION ReportInformation = { sizeof( WER_REPORT_INFORMATION ), 0 };
-	WER_SUBMIT_RESULT SubmitResult;
-	
 	// Flush out the log
-#if 1
 	GLog->Flush();
-#else
-	GError->HandleError();
-#endif
 
 	// Set the report to force queue
-	ForceWERQueuing( true );
+	FScopedWERQueuing ScopedQueueForcer;
 
 	// Construct the report details
+	WER_REPORT_INFORMATION ReportInformation = { sizeof(WER_REPORT_INFORMATION) };
+
 	StringCchCopy( ReportInformation.wzConsentKey, ARRAYSIZE( ReportInformation.wzConsentKey ), TEXT( "" ) );
 	StringCchCopy( ReportInformation.wzApplicationName, ARRAYSIZE( ReportInformation.wzApplicationName ), FApp::GetGameName() );
 	StringCchCopy( ReportInformation.wzApplicationPath, ARRAYSIZE( ReportInformation.wzApplicationPath ), FPlatformProcess::BaseDir() );
@@ -296,69 +311,67 @@ int32 NewReportCrash( EXCEPTION_POINTERS* ExceptionInfo, const TCHAR* ErrorMessa
 	GetCrashDescription( ReportInformation, ErrorMessage );
 	
 	// Create a crash event report
-	if( WerReportCreate( APPCRASH_EVENT, WerReportApplicationCrash, &ReportInformation, &ReportHandle ) == S_OK )
+	HREPORT ReportHandle = NULL;
+	if (WerReportCreate(APPCRASH_EVENT, WerReportApplicationCrash, &ReportInformation, &ReportHandle) == S_OK)
 	{
 		// Set the standard set of a crash parameters
-		SetReportParameters( ReportHandle, ExceptionInfo );
+		SetReportParameters(ReportHandle, ExceptionInfo);
 
 		// Add a manually generated minidump
-		AddMiniDump( ReportHandle, ExceptionInfo );
+		AddMiniDump(ReportHandle, ExceptionInfo);
 
 		// Add the log and video
-		AddMiscFiles( ReportHandle );
+		AddMiscFiles(ReportHandle);
 
 		// Submit
-		Result = WerReportSubmit( ReportHandle, WerConsentAlwaysPrompt, WER_SUBMIT_QUEUE | WER_SUBMIT_BYPASS_DATA_THROTTLING, &SubmitResult );
+		WER_SUBMIT_RESULT SubmitResult;
+		WerReportSubmit(ReportHandle, WerConsentAlwaysPrompt, WER_SUBMIT_QUEUE | WER_SUBMIT_BYPASS_DATA_THROTTLING, &SubmitResult);
 
 		// Cleanup
-		Result = WerReportCloseHandle( ReportHandle );
+		WerReportCloseHandle(ReportHandle);
 	}
 
 	// Build machines do not upload these automatically since it is not okay to have lingering processes after the build completes.
-	if ( !GIsBuildMachine )
+	if (GIsBuildMachine)
+		return EXCEPTION_CONTINUE_EXECUTION;
+
+	FString CrashReportClientArguments;
+
+	// Suppress the user input dialog if we're running in unattended mode
+	bool bNoDialog = FApp::IsUnattended() || ReportUI == EErrorReportUI::ReportInUnattendedMode;
+	if (bNoDialog)
 	{
-		// When invoking CrashReportUploader, put the logs in our saved directory
-		FString CrashReportUploaderArguments = FString::Printf( TEXT( "-LogFolder=%s" ), *FPaths::Combine(*FPaths::GameAgnosticSavedDir(), TEXT("CrashLogs")) );
-	
-		// Suppress the user input dialog if we're running in unattended mode
-		bool bNoDialog = FApp::IsUnattended() || ReportUI == EErrorReportUI::ReportInUnattendedMode;
-		if (bNoDialog)
-		{
-			CrashReportUploaderArguments += TEXT( " -Unattended" );
-		}
+		CrashReportClientArguments += TEXT(" -Unattended");
+	}
 
-		if (FApp::IsInstalled())
-		{
-			// Temporary workaround for CrashReportClient being built in Development, not Shipping (TTP328030). The
-			// following ensures that logs are saved to the user directory when UE4 is installed.
-			CrashReportUploaderArguments += TEXT( " -Installed" );
-		}
+	if (FApp::IsInstalled())
+	{
+		// Temporary workaround for CrashReportClient being built in Development, not Shipping (TTP328030). The
+		// following ensures that logs are saved to the user directory when UE4 is installed.
+		CrashReportClientArguments += TEXT(" -Installed");
+	}
 
-		const TCHAR* CrashClientPath = GUseCrashReportClient ?
-			TEXT("CrashReportClient.exe") :
-			TEXT("..\\..\\..\\Engine\\Binaries\\DotNET\\CrashReportUploader.exe");
+	CrashReportClientArguments += FString(TEXT(" -AppName=")) + ReportInformation.wzApplicationName;
 
-		// Passing !GUseCrashReportClient to 'hidden' and 'really-hidden' parameters
-		bool bCrashReporterRan = FPlatformProcess::CreateProc(CrashClientPath, *CrashReportUploaderArguments, true, !GUseCrashReportClient, !GUseCrashReportClient, NULL, 0, NULL, NULL).IsValid();
+	static const TCHAR CrashReportClientExeName[] = TEXT("CrashReportClient.exe");
+	FString CrashClientPath = FString(TEXT("..\\..\\..\\Engine\\Binaries")) / FPlatformProcess::GetBinariesSubdirectory() / CrashReportClientExeName;
 
-		if (!bCrashReporterRan && !bNoDialog)
-		{
-			UE_LOG(LogWindows, Log, TEXT("Could not start CrashReportClient"));
-			FPlatformMemory::DumpStats(*GWarn);
-			FText MessageTitle(FText::Format(
-				NSLOCTEXT("MessageDialog", "AppHasCrashed", "The {0} {1} has crashed and will close"),
-				FText::FromString(ReportInformation.wzApplicationName),
-				FText::FromString(GetEngineMode())
-			));
-			FMessageDialog::Open(EAppMsgType::Ok, FText::FromString(GErrorHist), &MessageTitle);
-		}
-	}
+	bool bCrashReporterRan = FPlatformProcess::CreateProc(*CrashClientPath, *CrashReportClientArguments, true, false, false, NULL, 0, NULL, NULL).IsValid();
 
-	// Reset the setting after we've had our way with it
-	ForceWERQueuing( false );
+	if (!bCrashReporterRan && !bNoDialog)
+	{
+		UE_LOG(LogWindows, Log, TEXT("Could not start %s"), CrashReportClientExeName);
+		FPlatformMemory::DumpStats(*GWarn);
+		FText MessageTitle(FText::Format(
+			NSLOCTEXT("MessageDialog", "AppHasCrashed", "The {0} {1} has crashed and will close"),
+			FText::FromString(ReportInformation.wzApplicationName),
+			FText::FromString(GetEngineMode())
+		));
+		FMessageDialog::Open(EAppMsgType::Ok, FText::FromString(GErrorHist), &MessageTitle);
+	}
 
-	// Let the system take back over
-	return ( ErrorMessage != NULL ) ? EXCEPTION_CONTINUE_EXECUTION : EXCEPTION_EXECUTE_HANDLER;
+	// Let the system take back over (return value only used by NewReportEnsure)
+	return EXCEPTION_CONTINUE_EXECUTION;
 }
 
 } // end anonymous namespace
@@ -378,7 +391,7 @@ void NewReportEnsure( const TCHAR* ErrorMessage )
 		FPlatformMisc::RaiseException( 1 );
 	}
 #if !PLATFORM_SEH_EXCEPTIONS_DISABLED
-	__except( NewReportCrash( GetExceptionInformation(), ErrorMessage, EErrorReportUI::ReportInUnattendedMode ) )
+	__except(ReportCrashUsingCrashReportClient(GetExceptionInformation(), ErrorMessage, EErrorReportUI::ReportInUnattendedMode))
 	{
 	}
 #endif
@@ -453,54 +466,52 @@ int32 ReportCrash( LPEXCEPTION_POINTERS ExceptionInfo )
 {
 	// Only create a minidump the first time this function is called.
 	// (Can be called the first time from the RenderThread, then a second time from the MainThread.)
-	if (GAlreadyCreatedMinidump)
-	{
+	if (::InterlockedIncrement(&ReportCrashCallCount) != 1)
 		return EXCEPTION_EXECUTE_HANDLER;
-	}
-
-	GAlreadyCreatedMinidump = true;
 
 	FCoreDelegates::OnHandleSystemError.Broadcast();
 
-	// Note: NewReportCrash may read the callstack from GErrorHist
-	const SIZE_T StackTraceSize = 65535;
-	ANSICHAR* StackTrace = (ANSICHAR*) FMemory::SystemMalloc( StackTraceSize );
-	StackTrace[0] = 0;
-	// Walk the stack and dump it to the allocated memory.
-	FPlatformStackWalk::StackWalkAndDump( StackTrace, StackTraceSize, 0, ExceptionInfo->ContextRecord );
-	FCString::Strncat( GErrorHist, ANSI_TO_TCHAR(StackTrace), ARRAY_COUNT(GErrorHist) - 1 );
-	CreateExceptionInfoString(ExceptionInfo->ExceptionRecord);
-	FMemory::SystemFree( StackTrace );
+	const TCHAR* AssertPos = FCString::Stristr(GErrorHist, TEXT("Assertion failed:"));
+	const TCHAR* StackPos = FCString::Stristr(GErrorHist, TEXT("\nStack:"));
+	if (AssertPos && StackPos && StackPos > AssertPos)
+	{
+		// The assertion code has already written a callstack to the log
+		AssertionMessageStart = AssertPos;
+		AssertionMessageLength = FMath::Min(int(StackPos - AssertPos), LocalBufferSize);
+	}
+	else
+	{
+		AssertionMessageStart = GErrorHist;
+		AssertionMessageLength = FCString::Strlen(GErrorHist);
+
+		// Note: ReportCrashUsingCrashReportClient reads the callstack from GErrorHist
+		const SIZE_T StackTraceSize = 65535;
+		ANSICHAR* StackTrace = (ANSICHAR*) FMemory::SystemMalloc( StackTraceSize );
+		StackTrace[0] = 0;
+		// Walk the stack and dump it to the allocated memory.
+		FPlatformStackWalk::StackWalkAndDump( StackTrace, StackTraceSize, 0, ExceptionInfo->ContextRecord );
+		FCString::Strncat( GErrorHist, ANSI_TO_TCHAR(StackTrace), ARRAY_COUNT(GErrorHist) - 1 );
+		CreateExceptionInfoString(ExceptionInfo->ExceptionRecord);
+		FMemory::SystemFree( StackTrace );
+	}
+
 
 #if WINVER > 0x502	// Windows Error Reporting is not supported on Windows XP
 	// 
 	if (GUseCrashReportClient)
 	{
-		NewReportCrash( ExceptionInfo, NULL, EErrorReportUI::ShowDialog );
-		return EXCEPTION_EXECUTE_HANDLER;
+		ReportCrashUsingCrashReportClient(ExceptionInfo, nullptr, EErrorReportUI::ShowDialog);
 	}
+	else
 #endif		// WINVER
-
-	// Try to create file for minidump.
-	HANDLE FileHandle = CreateFileW( MiniDumpFilenameW, GENERIC_WRITE, 0, NULL, CREATE_ALWAYS, FILE_ATTRIBUTE_NORMAL, NULL );
-	
-	// Write a minidump.
-	if( FileHandle != INVALID_HANDLE_VALUE )
 	{
-		MINIDUMP_EXCEPTION_INFORMATION DumpExceptionInfo;
-
-		DumpExceptionInfo.ThreadId			= GetCurrentThreadId();
-		DumpExceptionInfo.ExceptionPointers	= ExceptionInfo;
-		DumpExceptionInfo.ClientPointers	= false;
-
-		MiniDumpWriteDump( GetCurrentProcess(), GetCurrentProcessId(), FileHandle, MiniDumpNormal, &DumpExceptionInfo, NULL, NULL );
-		CloseHandle( FileHandle );
-	}
+		WriteMinidump(MiniDumpFilenameW, ExceptionInfo);
 
 #if UE_BUILD_SHIPPING && WITH_EDITOR
-	uint32 dwOpt = 0;
-	EFaultRepRetVal repret = ReportFault( ExceptionInfo, dwOpt);
+		uint32 dwOpt = 0;
+		EFaultRepRetVal repret = ReportFault( ExceptionInfo, dwOpt);
 #endif
+	}
 
 	return EXCEPTION_EXECUTE_HANDLER;
 }
