Pull request https://github.com/EpicGames/UnrealEngine/pull/170/ authored by I-Could-Not-Think-Of-Something-Good

diff --git a/Engine/Source/Developer/MeshUtilities/Private/Windows/D3D9MeshUtils.cpp b/Engine/Source/Developer/MeshUtilities/Private/Windows/D3D9MeshUtils.cpp
--- a/Engine/Source/Developer/MeshUtilities/Private/Windows/D3D9MeshUtils.cpp
+++ b/Engine/Source/Developer/MeshUtilities/Private/Windows/D3D9MeshUtils.cpp
@@ -58,7 +58,7 @@ FD3D9MeshUtilities::FD3D9MeshUtilities()
 
 	RegisterClassExW(&wc);
 
-	DummyWindowHandle = CreateWindowExW(NULL, wc.lpszClassName, L"", WS_OVERLAPPEDWINDOW, 0, 0, 100, 100, NULL, NULL, wc.hInstance, NULL);
+	DummyWindowHandle = CreateWindowExW(0, wc.lpszClassName, L"", WS_OVERLAPPEDWINDOW, 0, 0, 100, 100, NULL, NULL, wc.hInstance, NULL);
 
 	D3DPRESENT_PARAMETERS d3dpp;
 
diff --git a/Engine/Source/Editor/UnrealEd/Private/ObjectTools.cpp b/Engine/Source/Editor/UnrealEd/Private/ObjectTools.cpp
--- a/Engine/Source/Editor/UnrealEd/Private/ObjectTools.cpp
+++ b/Engine/Source/Editor/UnrealEd/Private/ObjectTools.cpp
@@ -733,7 +733,7 @@ namespace ObjectTools
 
 			// See if this is a blueprint consolidate and replace instances of the generated class
 			UBlueprint* BlueprintToConsolidateTo = Cast<UBlueprint>(ObjectToConsolidateTo);
-			if ( BlueprintToConsolidateTo != NULL && ensure(BlueprintToConsolidateTo->GeneratedClass) != NULL )
+			if ( BlueprintToConsolidateTo != NULL && ensure(BlueprintToConsolidateTo->GeneratedClass) )
 			{
 				for ( TArray<UObject*>::TConstIterator ConsolIter( ReplaceInfo.ReplaceableObjects ); ConsolIter; ++ConsolIter )
 				{
diff --git a/Engine/Source/Runtime/Core/Private/Windows/WindowsPlatformProcess.cpp b/Engine/Source/Runtime/Core/Private/Windows/WindowsPlatformProcess.cpp
--- a/Engine/Source/Runtime/Core/Private/Windows/WindowsPlatformProcess.cpp
+++ b/Engine/Source/Runtime/Core/Private/Windows/WindowsPlatformProcess.cpp
@@ -277,7 +277,7 @@ FProcHandle FWindowsPlatformProcess::CreateProc( const TCHAR* URL, const TCHAR*
 	{
 		CreateFlags |= DETACHED_PROCESS;
 	}
-	uint32 dwFlags = NULL;
+	uint32 dwFlags = 0;
 	uint16 ShowWindowFlags = SW_HIDE;
 	if (bLaunchReallyHidden)
 	{
@@ -295,7 +295,7 @@ FProcHandle FWindowsPlatformProcess::CreateProc( const TCHAR* URL, const TCHAR*
 
 	STARTUPINFO StartupInfo = { sizeof(STARTUPINFO), NULL, NULL, NULL,
 		CW_USEDEFAULT, CW_USEDEFAULT, CW_USEDEFAULT, CW_USEDEFAULT,
-		NULL, NULL, NULL, dwFlags, ShowWindowFlags, NULL, NULL,
+		0, 0, 0, dwFlags, ShowWindowFlags, 0, NULL,
 		::GetStdHandle(ProcessConstants::WIN_STD_INPUT_HANDLE), HANDLE(PipeWrite), HANDLE(PipeWrite) };
 	if( !CreateProcess( NULL, CommandLine.GetCharArray().GetTypedData(), &Attr, &Attr, true, CreateFlags,
 		NULL, OptionalWorkingDirectory, &StartupInfo, &ProcInfo ) )
@@ -523,7 +523,7 @@ bool FWindowsPlatformProcess::ExecProcess( const TCHAR* URL, const TCHAR* Params
 	bool bSuccess = false;
 	STARTUPINFO StartupInfo = { sizeof(STARTUPINFO), NULL, NULL, NULL,
 		CW_USEDEFAULT, CW_USEDEFAULT, CW_USEDEFAULT, CW_USEDEFAULT,
-		NULL, NULL, NULL, dwFlags, ShowWindowFlags, NULL, NULL,
+		0, 0, 0, dwFlags, ShowWindowFlags, 0, NULL,
 		::GetStdHandle(ProcessConstants::WIN_STD_INPUT_HANDLE), WritablePipes[0], WritablePipes[1] };
 	if (CreateProcess(NULL, CommandLine.GetCharArray().GetTypedData(), &Attr, &Attr, true, CreateFlags,
 		NULL, NULL, &StartupInfo, &ProcInfo))
diff --git a/Engine/Source/Runtime/Core/Private/Windows/WindowsPlatformStackWalk.cpp b/Engine/Source/Runtime/Core/Private/Windows/WindowsPlatformStackWalk.cpp
--- a/Engine/Source/Runtime/Core/Private/Windows/WindowsPlatformStackWalk.cpp
+++ b/Engine/Source/Runtime/Core/Private/Windows/WindowsPlatformStackWalk.cpp
@@ -127,7 +127,7 @@ static int32 CaptureStackTraceHelper( uint64 *BackTrace, uint32 MaxDepth, CONTEX
 	// NULL out remaining entries.
 	for ( ; CurrentDepth<MaxDepth; CurrentDepth++ )
 	{
-		BackTrace[CurrentDepth] = NULL;
+		BackTrace[CurrentDepth] = 0;
 	}
 
 	return EXCEPTION_EXECUTE_HANDLER;
@@ -223,7 +223,7 @@ void FWindowsPlatformStackWalk::CaptureStackBackTrace( uint64* BackTrace, uint32
 			}
 			while ( NumFrames < MaxDepth )
 			{
-				BackTrace[ NumFrames++ ] = NULL;
+				BackTrace[ NumFrames++ ] = 0;
 			}
 		}
 #elif PLATFORM_64BITS
