Improved setting of bUsingTrackpad variable in MacApplication

diff --git a/Engine/Source/Runtime/Core/Private/Mac/MacApplication.cpp b/Engine/Source/Runtime/Core/Private/Mac/MacApplication.cpp
--- a/Engine/Source/Runtime/Core/Private/Mac/MacApplication.cpp
+++ b/Engine/Source/Runtime/Core/Private/Mac/MacApplication.cpp
@@ -262,6 +262,8 @@ void FMacApplication::ProcessEvent( NSEvent* Event )
 		case NSRightMouseDragged:
 		case NSOtherMouseDragged:
 		{
+			bUsingTrackpad = [Event subtype] == NSTouchEventSubtype;
+
 			if( CurrentEventWindow.IsValid() && CurrentEventWindow->IsRegularWindow() )
 			{
 				bool IsMouseOverTitleBar = false;
@@ -345,6 +347,8 @@ void FMacApplication::ProcessEvent( NSEvent* Event )
 		case NSRightMouseDown:
 		case NSOtherMouseDown:
 		{
+			bUsingTrackpad = [Event subtype] == NSTouchEventSubtype;
+
 			EMouseButtons::Type Button = [Event type] == NSLeftMouseDown ? EMouseButtons::Left : EMouseButtons::Right;
 			if ([Event type] == NSOtherMouseDown)
 			{
@@ -390,6 +394,8 @@ void FMacApplication::ProcessEvent( NSEvent* Event )
 		case NSRightMouseUp:
 		case NSOtherMouseUp:
 		{
+			bUsingTrackpad = false;
+
 			EMouseButtons::Type Button = [Event type] == NSLeftMouseUp ? EMouseButtons::Left : EMouseButtons::Right;
 			if ([Event type] == NSOtherMouseUp)
 			{
