Stop tooltips & other windows initially appearing stretched & distorted on OS X.
When deferring window frame change operations on OS X send the OnSizeChanged message to the game & ensure we return the deferred size from openGLFrame when required. This keeps Slate in-sync with the expected size of the window so that the render thread doesn't try to blit to a context with an incorrectly sized backing store.

diff --git a/Engine/Source/Runtime/Core/Private/Mac/MacWindow.cpp b/Engine/Source/Runtime/Core/Private/Mac/MacWindow.cpp
--- a/Engine/Source/Runtime/Core/Private/Mac/MacWindow.cpp
+++ b/Engine/Source/Runtime/Core/Private/Mac/MacWindow.cpp
@@ -43,11 +43,11 @@
 {
 	if([self styleMask] & (NSTexturedBackgroundWindowMask))
 	{
-		return [self frame];
+		return (!bDeferSetFrame ? [self frame] : DeferFrame);
 	}
 	else
 	{
-		return [[self contentView] frame];
+		return (!bDeferSetFrame ? [[self contentView] frame] : [self contentRectForFrameRect:DeferFrame]);
 	}
 }
 
@@ -96,7 +96,7 @@
 				DeferFrame.size = [self frame].size;
 			}
 			
-			[super setFrame:DeferFrame display:NO];
+			[super setFrame:DeferFrame display:YES];
 		};
 		
 		if([NSThread isMainThread])
@@ -330,6 +330,10 @@
 	{
 		bDeferSetFrame = true;
 		DeferFrame = FrameRect;
+		if(self.bForwardEvents)
+		{
+			MacApplication->OnWindowDidResize( self );
+		}
 	}
 }
 
