Improved error reporting when failing to create a collator or formatting object from ICU.

diff --git a/Engine/Source/Runtime/Core/Private/Internationalization/ICUCulture.cpp b/Engine/Source/Runtime/Core/Private/Internationalization/ICUCulture.cpp
--- a/Engine/Source/Runtime/Core/Private/Internationalization/ICUCulture.cpp
+++ b/Engine/Source/Runtime/Core/Private/Internationalization/ICUCulture.cpp
@@ -11,43 +11,57 @@ namespace
 	TSharedRef<const icu::Collator> CreateCollator( const icu::Locale& ICULocale )
 	{
 		UErrorCode ICUStatus = U_ZERO_ERROR;
-		return MakeShareable( icu::Collator::createInstance( ICULocale, ICUStatus ) );
+		TSharedPtr<const icu::Collator> Ptr;
+		checkf(Ptr.IsValid(), TEXT("Creating a collator object failed using locale %s. Perhaps this locale has no data."), StringCast<TCHAR>(ICULocale.getName()));
+		return Ptr.ToSharedRef();
 	}
 
 	TSharedRef<const icu::DecimalFormat> CreateDecimalFormat( const icu::Locale& ICULocale )
 	{
 		UErrorCode ICUStatus = U_ZERO_ERROR;
-		return MakeShareable( static_cast<icu::DecimalFormat*>(icu::NumberFormat::createInstance( ICULocale, ICUStatus )) );
+		TSharedPtr<const icu::DecimalFormat> Ptr = MakeShareable( static_cast<icu::DecimalFormat*>(icu::NumberFormat::createInstance( ICULocale, ICUStatus )) );
+		checkf(Ptr.IsValid(), TEXT("Creating a decimal format object failed using locale %s. Perhaps this locale has no data."), StringCast<TCHAR>(ICULocale.getName()));
+		return Ptr.ToSharedRef();
 	}
 
 	TSharedRef<const icu::DecimalFormat> CreateCurrencyFormat( const icu::Locale& ICULocale )
 	{
 		UErrorCode ICUStatus = U_ZERO_ERROR;
-		return MakeShareable( static_cast<icu::DecimalFormat*>(icu::NumberFormat::createCurrencyInstance( ICULocale, ICUStatus )) );
+		TSharedPtr<const icu::DecimalFormat> Ptr = MakeShareable( static_cast<icu::DecimalFormat*>(icu::NumberFormat::createCurrencyInstance( ICULocale, ICUStatus )) );
+		checkf(Ptr.IsValid(), TEXT("Creating a currency format object failed using locale %s. Perhaps this locale has no data."), StringCast<TCHAR>(ICULocale.getName()));
+		return Ptr.ToSharedRef();
 	}
 
 	TSharedRef<const icu::DecimalFormat> CreatePercentFormat( const icu::Locale& ICULocale )
 	{
 		UErrorCode ICUStatus = U_ZERO_ERROR;
-		return MakeShareable( static_cast<icu::DecimalFormat*>(icu::NumberFormat::createPercentInstance( ICULocale, ICUStatus )) );
+		TSharedPtr<const icu::DecimalFormat> Ptr = MakeShareable( static_cast<icu::DecimalFormat*>(icu::NumberFormat::createPercentInstance( ICULocale, ICUStatus )) );
+		checkf(Ptr.IsValid(), TEXT("Creating a percent format object failed using locale %s. Perhaps this locale has no data."), StringCast<TCHAR>(ICULocale.getName()));
+		return Ptr.ToSharedRef();
 	}
 
 	TSharedRef<const icu::DateFormat> CreateDateFormat( const icu::Locale& ICULocale )
 	{
 		UErrorCode ICUStatus = U_ZERO_ERROR;
-		return MakeShareable( icu::DateFormat::createDateInstance( icu::DateFormat::EStyle::kDefault, ICULocale ) );
+		TSharedPtr<const icu::DateFormat> Ptr = MakeShareable( icu::DateFormat::createDateInstance( icu::DateFormat::EStyle::kDefault, ICULocale ) );
+		checkf(Ptr.IsValid(), TEXT("Creating a date format object failed using locale %s. Perhaps this locale has no data."), StringCast<TCHAR>(ICULocale.getName()));
+		return Ptr.ToSharedRef();
 	}
 
 	TSharedRef<const icu::DateFormat> CreateTimeFormat( const icu::Locale& ICULocale )
 	{
 		UErrorCode ICUStatus = U_ZERO_ERROR;
-		return MakeShareable( icu::DateFormat::createTimeInstance( icu::DateFormat::EStyle::kDefault, ICULocale ) );
+		TSharedPtr<const icu::DateFormat> Ptr = MakeShareable( icu::DateFormat::createTimeInstance( icu::DateFormat::EStyle::kDefault, ICULocale ) );
+		checkf(Ptr.IsValid(), TEXT("Creating a time format object failed using locale %s. Perhaps this locale has no data."), StringCast<TCHAR>(ICULocale.getName()));
+		return Ptr.ToSharedRef();
 	}
 
 	TSharedRef<const icu::DateFormat> CreateDateTimeFormat( const icu::Locale& ICULocale )
 	{
 		UErrorCode ICUStatus = U_ZERO_ERROR;
-		return MakeShareable( icu::DateFormat::createDateTimeInstance( icu::DateFormat::EStyle::kDefault, icu::DateFormat::EStyle::kDefault, ICULocale ) );
+		TSharedPtr<const icu::DateFormat> Ptr = MakeShareable( icu::DateFormat::createDateTimeInstance( icu::DateFormat::EStyle::kDefault, icu::DateFormat::EStyle::kDefault, ICULocale ) );
+		checkf(Ptr.IsValid(), TEXT("Creating a date-time format object failed using locale %s. Perhaps this locale has no data."), StringCast<TCHAR>(ICULocale.getName()));
+		return Ptr.ToSharedRef();
 	}
 }
 
