Resolved accidental uninitialized pointer in culture initialization. Resolved compilation error in Mac builds due to non POD type passed in varargs.

diff --git a/Engine/Source/Runtime/Core/Private/Internationalization/ICUCulture.cpp b/Engine/Source/Runtime/Core/Private/Internationalization/ICUCulture.cpp
--- a/Engine/Source/Runtime/Core/Private/Internationalization/ICUCulture.cpp
+++ b/Engine/Source/Runtime/Core/Private/Internationalization/ICUCulture.cpp
@@ -11,8 +11,8 @@ namespace
 	TSharedRef<const icu::Collator> CreateCollator( const icu::Locale& ICULocale )
 	{
 		UErrorCode ICUStatus = U_ZERO_ERROR;
-		TSharedPtr<const icu::Collator> Ptr;
-		checkf(Ptr.IsValid(), TEXT("Creating a collator object failed using locale %s. Perhaps this locale has no data."), StringCast<TCHAR>(ICULocale.getName()));
+		TSharedPtr<const icu::Collator> Ptr = MakeShareable( icu::Collator::createInstance( ICULocale, ICUStatus ) );
+		checkf(Ptr.IsValid(), TEXT("Creating a collator object failed using locale %s. Perhaps this locale has no data."), StringCast<TCHAR>(ICULocale.getName()).Get());
 		return Ptr.ToSharedRef();
 	}
 
@@ -20,7 +20,7 @@ namespace
 	{
 		UErrorCode ICUStatus = U_ZERO_ERROR;
 		TSharedPtr<const icu::DecimalFormat> Ptr = MakeShareable( static_cast<icu::DecimalFormat*>(icu::NumberFormat::createInstance( ICULocale, ICUStatus )) );
-		checkf(Ptr.IsValid(), TEXT("Creating a decimal format object failed using locale %s. Perhaps this locale has no data."), StringCast<TCHAR>(ICULocale.getName()));
+		checkf(Ptr.IsValid(), TEXT("Creating a decimal format object failed using locale %s. Perhaps this locale has no data."), StringCast<TCHAR>(ICULocale.getName()).Get());
 		return Ptr.ToSharedRef();
 	}
 
@@ -28,7 +28,7 @@ namespace
 	{
 		UErrorCode ICUStatus = U_ZERO_ERROR;
 		TSharedPtr<const icu::DecimalFormat> Ptr = MakeShareable( static_cast<icu::DecimalFormat*>(icu::NumberFormat::createCurrencyInstance( ICULocale, ICUStatus )) );
-		checkf(Ptr.IsValid(), TEXT("Creating a currency format object failed using locale %s. Perhaps this locale has no data."), StringCast<TCHAR>(ICULocale.getName()));
+		checkf(Ptr.IsValid(), TEXT("Creating a currency format object failed using locale %s. Perhaps this locale has no data."), StringCast<TCHAR>(ICULocale.getName()).Get());
 		return Ptr.ToSharedRef();
 	}
 
@@ -36,7 +36,7 @@ namespace
 	{
 		UErrorCode ICUStatus = U_ZERO_ERROR;
 		TSharedPtr<const icu::DecimalFormat> Ptr = MakeShareable( static_cast<icu::DecimalFormat*>(icu::NumberFormat::createPercentInstance( ICULocale, ICUStatus )) );
-		checkf(Ptr.IsValid(), TEXT("Creating a percent format object failed using locale %s. Perhaps this locale has no data."), StringCast<TCHAR>(ICULocale.getName()));
+		checkf(Ptr.IsValid(), TEXT("Creating a percent format object failed using locale %s. Perhaps this locale has no data."), StringCast<TCHAR>(ICULocale.getName()).Get());
 		return Ptr.ToSharedRef();
 	}
 
@@ -44,7 +44,7 @@ namespace
 	{
 		UErrorCode ICUStatus = U_ZERO_ERROR;
 		TSharedPtr<const icu::DateFormat> Ptr = MakeShareable( icu::DateFormat::createDateInstance( icu::DateFormat::EStyle::kDefault, ICULocale ) );
-		checkf(Ptr.IsValid(), TEXT("Creating a date format object failed using locale %s. Perhaps this locale has no data."), StringCast<TCHAR>(ICULocale.getName()));
+		checkf(Ptr.IsValid(), TEXT("Creating a date format object failed using locale %s. Perhaps this locale has no data."), StringCast<TCHAR>(ICULocale.getName()).Get());
 		return Ptr.ToSharedRef();
 	}
 
@@ -52,7 +52,7 @@ namespace
 	{
 		UErrorCode ICUStatus = U_ZERO_ERROR;
 		TSharedPtr<const icu::DateFormat> Ptr = MakeShareable( icu::DateFormat::createTimeInstance( icu::DateFormat::EStyle::kDefault, ICULocale ) );
-		checkf(Ptr.IsValid(), TEXT("Creating a time format object failed using locale %s. Perhaps this locale has no data."), StringCast<TCHAR>(ICULocale.getName()));
+		checkf(Ptr.IsValid(), TEXT("Creating a time format object failed using locale %s. Perhaps this locale has no data."), StringCast<TCHAR>(ICULocale.getName()).Get());
 		return Ptr.ToSharedRef();
 	}
 
@@ -60,7 +60,7 @@ namespace
 	{
 		UErrorCode ICUStatus = U_ZERO_ERROR;
 		TSharedPtr<const icu::DateFormat> Ptr = MakeShareable( icu::DateFormat::createDateTimeInstance( icu::DateFormat::EStyle::kDefault, icu::DateFormat::EStyle::kDefault, ICULocale ) );
-		checkf(Ptr.IsValid(), TEXT("Creating a date-time format object failed using locale %s. Perhaps this locale has no data."), StringCast<TCHAR>(ICULocale.getName()));
+		checkf(Ptr.IsValid(), TEXT("Creating a date-time format object failed using locale %s. Perhaps this locale has no data."), StringCast<TCHAR>(ICULocale.getName()).Get());
 		return Ptr.ToSharedRef();
 	}
 }
