Disable reduced memory mode for virtual textures, as this is currently causing corruption.

diff --git a/Engine/Source/Runtime/Core/Private/HAL/ConsoleManager.cpp b/Engine/Source/Runtime/Core/Private/HAL/ConsoleManager.cpp
--- a/Engine/Source/Runtime/Core/Private/HAL/ConsoleManager.cpp
+++ b/Engine/Source/Runtime/Core/Private/HAL/ConsoleManager.cpp
@@ -2200,13 +2200,13 @@ static TAutoConsoleVariable<float> CVarSetMipMapLODBias(
 static TAutoConsoleVariable<int32> CVarVirtualTextureEnabled(
 	TEXT("r.VirtualTexture"),
 	1,
-	TEXT("If set to 1, textures will use virtual memory to optimize texture mip streaming"),
+	TEXT("If set to 1, textures will use virtual memory so they can be partially resident."),
 	ECVF_RenderThreadSafe);
 
 static TAutoConsoleVariable<int32> CVarVirtualTextureReducedMemoryEnabled(
 	TEXT("r.VirtualTextureReducedMemory"),
-	1,
-	TEXT("If set to 1, the cost of virtual memory will be reduced by transforming the texture between streamed & resident"),
+	0,
+	TEXT("If set to 1, the cost of virtual textures will be reduced by using a more packed layout."),
 	ECVF_RenderThreadSafe);
 
 static TAutoConsoleVariable<int32> CVarPrecomputedVisibilityWarning(
