Stats - Added macro GET_STATDESCRIPTION to get the stat's description

diff --git a/Engine/Source/Runtime/Core/Public/Stats/Stats.h b/Engine/Source/Runtime/Core/Public/Stats/Stats.h
--- a/Engine/Source/Runtime/Core/Public/Stats/Stats.h
+++ b/Engine/Source/Runtime/Core/Public/Stats/Stats.h
@@ -159,5 +159,6 @@ FORCEINLINE void StatsMasterEnableSubtract(int32 Value = 1)
 
 #define GET_STATID(Stat) (TStatId())
 #define GET_STATFNAME(Stat) (FName())
+#define GET_STATDESCRIPTION(Stat) (nullptr)
 
 #endif
\ No newline at end of file
diff --git a/Engine/Source/Runtime/Core/Public/Stats/Stats2.h b/Engine/Source/Runtime/Core/Public/Stats/Stats2.h
--- a/Engine/Source/Runtime/Core/Public/Stats/Stats2.h
+++ b/Engine/Source/Runtime/Core/Public/Stats/Stats2.h
@@ -1339,6 +1339,7 @@ struct FStat_##StatName\
 
 #define GET_STATID(Stat) (StatPtr_##Stat.GetStatId())
 #define GET_STATFNAME(Stat) (StatPtr_##Stat.GetStatFName())
+#define GET_STATDESCRIPTION(Stat) (FStat_##Stat::GetDescription())
 
 #define STAT_GROUP_TO_FStatGroup(Group) FStatGroup_##Group
 
