Core.MemoryProfiler - Fixed an issue with bad stats descriptions

diff --git a/Engine/Source/Runtime/Core/Private/GenericPlatform/GenericPlatformMemory.cpp b/Engine/Source/Runtime/Core/Private/GenericPlatform/GenericPlatformMemory.cpp
--- a/Engine/Source/Runtime/Core/Private/GenericPlatform/GenericPlatformMemory.cpp
+++ b/Engine/Source/Runtime/Core/Private/GenericPlatform/GenericPlatformMemory.cpp
@@ -14,17 +14,18 @@ DEFINE_STAT(MCR_Physical);
 DEFINE_STAT(MCR_GPU);
 DEFINE_STAT(MCR_TexturePool);
 
+// Must match values in the MemoryProfiler2.FMemoryAllocationStatsV4
 DECLARE_MEMORY_STAT(TEXT("Total Physical"),		STAT_TotalPhysical,STATGROUP_MemoryPlatform);
 DECLARE_MEMORY_STAT(TEXT("Total Virtual"),		STAT_TotalVirtual,STATGROUP_MemoryPlatform);
-DECLARE_MEMORY_STAT(TEXT("PageSize"),			STAT_PageSize,STATGROUP_MemoryPlatform);
+DECLARE_MEMORY_STAT(TEXT("Page Size"),			STAT_PageSize,STATGROUP_MemoryPlatform);
 DECLARE_MEMORY_STAT(TEXT("Total Physical GB"),	STAT_TotalPhysicalGB,STATGROUP_MemoryPlatform);
 
-DECLARE_MEMORY_STAT(TEXT("AvailablePhysical"),	STAT_AvailablePhysical,STATGROUP_MemoryPlatform);
-DECLARE_MEMORY_STAT(TEXT("AvailableVirtual"),	STAT_AvailableVirtual,STATGROUP_MemoryPlatform);
-DECLARE_MEMORY_STAT(TEXT("UsedPhysical"),		STAT_UsedPhysical,STATGROUP_MemoryPlatform);
-DECLARE_MEMORY_STAT(TEXT("PeakUsedPhysical"),	STAT_PeakUsedPhysical,STATGROUP_MemoryPlatform);
-DECLARE_MEMORY_STAT(TEXT("UsedVirtual"),		STAT_UsedVirtual,STATGROUP_MemoryPlatform);
-DECLARE_MEMORY_STAT(TEXT("PeakUsedVirtual"),	STAT_PeakUsedVirtual,STATGROUP_MemoryPlatform);
+DECLARE_MEMORY_STAT(TEXT("Available Physical"),	STAT_AvailablePhysical,STATGROUP_MemoryPlatform);
+DECLARE_MEMORY_STAT(TEXT("Available Virtual"),	STAT_AvailableVirtual,STATGROUP_MemoryPlatform);
+DECLARE_MEMORY_STAT(TEXT("Used Physical"),		STAT_UsedPhysical,STATGROUP_MemoryPlatform);
+DECLARE_MEMORY_STAT(TEXT("Peak Used Physical"),	STAT_PeakUsedPhysical,STATGROUP_MemoryPlatform);
+DECLARE_MEMORY_STAT(TEXT("Used Virtual"),		STAT_UsedVirtual,STATGROUP_MemoryPlatform);
+DECLARE_MEMORY_STAT(TEXT("Peak Used Virtual"),	STAT_PeakUsedVirtual,STATGROUP_MemoryPlatform);
 
 FGenericPlatformMemoryStats::FGenericPlatformMemoryStats()
 	: FGenericPlatformMemoryConstants( FPlatformMemory::GetConstants() )
@@ -71,16 +72,16 @@ void FGenericPlatformMemory::GetStatsForMallocProfiler( FGenericMemoryStats& out
 	FPlatformMemoryStats Stats = FPlatformMemory::GetStats();
 
 	// Base common stats for all platforms.
-	out_Stats.Add( GET_STATFNAME( STAT_TotalPhysical ), Stats.TotalPhysical );
-	out_Stats.Add( GET_STATFNAME( STAT_TotalVirtual ), Stats.TotalVirtual );
-	out_Stats.Add( GET_STATFNAME( STAT_PageSize ), Stats.PageSize );
-	out_Stats.Add( GET_STATFNAME( STAT_TotalPhysicalGB ), (SIZE_T)Stats.TotalPhysicalGB );
-	out_Stats.Add( GET_STATFNAME( STAT_AvailablePhysical ), Stats.AvailablePhysical );
-	out_Stats.Add( GET_STATFNAME( STAT_AvailableVirtual ), Stats.AvailableVirtual );
-	out_Stats.Add( GET_STATFNAME( STAT_UsedPhysical ), Stats.UsedPhysical );
-	out_Stats.Add( GET_STATFNAME( STAT_PeakUsedPhysical ), Stats.PeakUsedPhysical );
-	out_Stats.Add( GET_STATFNAME( STAT_UsedVirtual ), Stats.UsedVirtual );
-	out_Stats.Add( GET_STATFNAME( STAT_PeakUsedVirtual ), Stats.PeakUsedVirtual );
+	out_Stats.Add( GET_STATDESCRIPTION( STAT_TotalPhysical ), Stats.TotalPhysical );
+	out_Stats.Add( GET_STATDESCRIPTION( STAT_TotalVirtual ), Stats.TotalVirtual );
+	out_Stats.Add( GET_STATDESCRIPTION( STAT_PageSize ), Stats.PageSize );
+	out_Stats.Add( GET_STATDESCRIPTION( STAT_TotalPhysicalGB ), (SIZE_T)Stats.TotalPhysicalGB );
+	out_Stats.Add( GET_STATDESCRIPTION( STAT_AvailablePhysical ), Stats.AvailablePhysical );
+	out_Stats.Add( GET_STATDESCRIPTION( STAT_AvailableVirtual ), Stats.AvailableVirtual );
+	out_Stats.Add( GET_STATDESCRIPTION( STAT_UsedPhysical ), Stats.UsedPhysical );
+	out_Stats.Add( GET_STATDESCRIPTION( STAT_PeakUsedPhysical ), Stats.PeakUsedPhysical );
+	out_Stats.Add( GET_STATDESCRIPTION( STAT_UsedVirtual ), Stats.UsedVirtual );
+	out_Stats.Add( GET_STATDESCRIPTION( STAT_PeakUsedVirtual ), Stats.PeakUsedVirtual );
 #endif // STATS
 }
 
diff --git a/Engine/Source/Runtime/Core/Private/HAL/MallocBinned.cpp b/Engine/Source/Runtime/Core/Private/HAL/MallocBinned.cpp
--- a/Engine/Source/Runtime/Core/Private/HAL/MallocBinned.cpp
+++ b/Engine/Source/Runtime/Core/Private/HAL/MallocBinned.cpp
@@ -55,14 +55,14 @@ void FMallocBinned::GetAllocatorStats( FGenericMemoryStats& out_Stats )
 	}
 
 	// Malloc binned stats.
-	out_Stats.Add( GET_STATFNAME( STAT_Binned_OsCurrent ), LocalOsCurrent );
-	out_Stats.Add( GET_STATFNAME( STAT_Binned_OsPeak ), LocalOsPeak );
-	out_Stats.Add( GET_STATFNAME( STAT_Binned_WasteCurrent ), LocalWasteCurrent );
-	out_Stats.Add( GET_STATFNAME( STAT_Binned_WastePeak ), LocalWastePeak );
-	out_Stats.Add( GET_STATFNAME( STAT_Binned_UsedCurrent ), LocalUsedCurrent );
-	out_Stats.Add( GET_STATFNAME( STAT_Binned_UsedPeak ), LocalUsedPeak );
-	out_Stats.Add( GET_STATFNAME( STAT_Binned_CurrentAllocs ), LocalCurrentAllocs );
-	out_Stats.Add( GET_STATFNAME( STAT_Binned_TotalAllocs ), LocalTotalAllocs );
-	out_Stats.Add( GET_STATFNAME( STAT_Binned_SlackCurrent ), LocalSlackCurrent );
+	out_Stats.Add( GET_STATDESCRIPTION( STAT_Binned_OsCurrent ), LocalOsCurrent );
+	out_Stats.Add( GET_STATDESCRIPTION( STAT_Binned_OsPeak ), LocalOsPeak );
+	out_Stats.Add( GET_STATDESCRIPTION( STAT_Binned_WasteCurrent ), LocalWasteCurrent );
+	out_Stats.Add( GET_STATDESCRIPTION( STAT_Binned_WastePeak ), LocalWastePeak );
+	out_Stats.Add( GET_STATDESCRIPTION( STAT_Binned_UsedCurrent ), LocalUsedCurrent );
+	out_Stats.Add( GET_STATDESCRIPTION( STAT_Binned_UsedPeak ), LocalUsedPeak );
+	out_Stats.Add( GET_STATDESCRIPTION( STAT_Binned_CurrentAllocs ), LocalCurrentAllocs );
+	out_Stats.Add( GET_STATDESCRIPTION( STAT_Binned_TotalAllocs ), LocalTotalAllocs );
+	out_Stats.Add( GET_STATDESCRIPTION( STAT_Binned_SlackCurrent ), LocalSlackCurrent );
 #endif // STATS
 }
\ No newline at end of file
diff --git a/Engine/Source/Runtime/Core/Private/HAL/MemoryBase.cpp b/Engine/Source/Runtime/Core/Private/HAL/MemoryBase.cpp
--- a/Engine/Source/Runtime/Core/Private/HAL/MemoryBase.cpp
+++ b/Engine/Source/Runtime/Core/Private/HAL/MemoryBase.cpp
@@ -67,9 +67,9 @@ void FMalloc::UpdateStats()
 void FMalloc::GetAllocatorStats( FGenericMemoryStats& out_Stats )
 {
 #if	STATS
-	out_Stats.Add( GET_STATFNAME( STAT_MallocCalls ), GetCurrentFrameCalls().MallocCalls );
-	out_Stats.Add( GET_STATFNAME( STAT_ReallocCalls ), GetCurrentFrameCalls().ReallocCalls );
-	out_Stats.Add( GET_STATFNAME( STAT_FreeCalls ), GetCurrentFrameCalls().FreeCalls );
-	out_Stats.Add( GET_STATFNAME( STAT_TotalAllocatorCalls ), GetCurrentFrameCalls().AllocatorCalls );
+	out_Stats.Add( GET_STATDESCRIPTION( STAT_MallocCalls ), GetCurrentFrameCalls().MallocCalls );
+	out_Stats.Add( GET_STATDESCRIPTION( STAT_ReallocCalls ), GetCurrentFrameCalls().ReallocCalls );
+	out_Stats.Add( GET_STATDESCRIPTION( STAT_FreeCalls ), GetCurrentFrameCalls().FreeCalls );
+	out_Stats.Add( GET_STATDESCRIPTION( STAT_TotalAllocatorCalls ), GetCurrentFrameCalls().AllocatorCalls );
 #endif // STATS
 }
\ No newline at end of file
diff --git a/Engine/Source/Runtime/Core/Private/ProfilingDebugging/MallocProfiler.cpp b/Engine/Source/Runtime/Core/Private/ProfilingDebugging/MallocProfiler.cpp
--- a/Engine/Source/Runtime/Core/Private/ProfilingDebugging/MallocProfiler.cpp
+++ b/Engine/Source/Runtime/Core/Private/ProfilingDebugging/MallocProfiler.cpp
@@ -892,9 +892,9 @@ void FMallocProfiler::WriteMemoryAllocationStats()
 	FPlatformMemory::GetStatsForMallocProfiler( Stats ); 
 	UsedMalloc->GetAllocatorStats( Stats );
 
-	static const FName NAME_MemoryProfilingOverhead = TEXT("Memory Profiling Overhead");
+	const TCHAR* DESC_MemoryProfilingOverhead = TEXT("Memory Profiling Overhead");
 	int64 MemoryProfilingOverhead = CalculateMemoryProfilingOverhead();
-	Stats.Add( NAME_MemoryProfilingOverhead, MemoryProfilingOverhead );
+	Stats.Add( DESC_MemoryProfilingOverhead, MemoryProfilingOverhead );
 
 	check( Stats.Data.Num() < MAX_uint8 );
 	uint8 StatsCount = Stats.Data.Num();
@@ -902,7 +902,7 @@ void FMallocProfiler::WriteMemoryAllocationStats()
 
 	for( const auto& MapIt : Stats.Data )
 	{
-		FString StatName = MapIt.Key.ToString();
+		FString StatName = MapIt.Key;
 		int64 StatValue = (int64)MapIt.Value;
 		int32 StatNameIndex = GetNameTableIndex(StatName);
 		BufferedFileWriter << StatNameIndex << StatValue;
diff --git a/Engine/Source/Runtime/Core/Private/Windows/WindowsPlatformMemory.cpp b/Engine/Source/Runtime/Core/Private/Windows/WindowsPlatformMemory.cpp
--- a/Engine/Source/Runtime/Core/Private/Windows/WindowsPlatformMemory.cpp
+++ b/Engine/Source/Runtime/Core/Private/Windows/WindowsPlatformMemory.cpp
@@ -121,8 +121,7 @@ void FWindowsPlatformMemory::GetStatsForMallocProfiler( FGenericMemoryStats& out
 	FPlatformMemoryStats Stats = GetStats();
 
 	// Windows specific stats.
-	static const FName NAME_WindowsSpecificMemoryStat   = TEXT("Windows Specific Memory Stat");
-	out_Stats.Add( NAME_WindowsSpecificMemoryStat, Stats.WindowsSpecificMemoryStat );
+	out_Stats.Add( GET_STATDESCRIPTION( STAT_WindowsSpecificMemoryStat ), Stats.WindowsSpecificMemoryStat );
 #endif // STATS
 }
 
diff --git a/Engine/Source/Runtime/Core/Public/HAL/MemoryMisc.h b/Engine/Source/Runtime/Core/Public/HAL/MemoryMisc.h
--- a/Engine/Source/Runtime/Core/Public/HAL/MemoryMisc.h
+++ b/Engine/Source/Runtime/Core/Public/HAL/MemoryMisc.h
@@ -5,10 +5,10 @@
 /** Holds generic memory stats, internally implemented as a map. */
 struct FGenericMemoryStats
 {
-	void Add( const FName StatName, const SIZE_T StatValue )
+	void Add( const TCHAR* StatDescription, const SIZE_T StatValue )
 	{
-		Data.Add( StatName, StatValue );
+		Data.Add( FString(StatDescription), StatValue );
 	}
 
-	TMap<FName,SIZE_T> Data;
+	TMap<FString, SIZE_T> Data;
 };
\ No newline at end of file
diff --git a/Engine/Source/Runtime/Engine/Private/MallocProfilerEx.cpp b/Engine/Source/Runtime/Engine/Private/MallocProfilerEx.cpp
--- a/Engine/Source/Runtime/Engine/Private/MallocProfilerEx.cpp
+++ b/Engine/Source/Runtime/Engine/Private/MallocProfilerEx.cpp
@@ -83,7 +83,7 @@ void FMallocProfilerEx::GetTexturePoolSize( FGenericMemoryStats& out_Stats )
 		RHIGetTextureMemoryStats(Stats);
 	}
 
-	static const FName NAME_TextureAllocatedMemorySize = TEXT("Texture Allocated Memory Size");
+	const TCHAR* NAME_TextureAllocatedMemorySize = TEXT("Texture Allocated Memory Size");
 	out_Stats.Add( NAME_TextureAllocatedMemorySize, Stats.AllocatedMemorySize );
 }
 
