TPS: Made libdsext more resilient to various problems.

Takes care of some errors reported by valgrind and hopefully will make it easier to diagnose other problems like TTP #331994.

diff --git a/Engine/Source/Runtime/Core/Private/Linux/LinuxApplication.cpp b/Engine/Source/Runtime/Core/Private/Linux/LinuxApplication.cpp
--- a/Engine/Source/Runtime/Core/Private/Linux/LinuxApplication.cpp
+++ b/Engine/Source/Runtime/Core/Private/Linux/LinuxApplication.cpp
@@ -894,8 +894,10 @@ void FLinuxApplication::UpdateMouseCaptureWindow( SDL_HWindow TargetWindow )
 		}
 		if (MouseCaptureWindow && !LinuxCursor->IsHidden())
 		{
-			//SDL_SetWindowGrab(TargetWindow, SDL_TRUE);
-			DSEXT_SetMouseGrab(TargetWindow, SDL_TRUE);
+			if (EDSExtSuccess != DSEXT_SetMouseGrab(TargetWindow, SDL_TRUE))
+			{
+				UE_LOG(LogHAL, Log, TEXT("Could not grab cursor for SDL window %p"), TargetWindow);
+			}
 		}
 	}
 	else
@@ -904,8 +906,10 @@ void FLinuxApplication::UpdateMouseCaptureWindow( SDL_HWindow TargetWindow )
 		{
 			if (!LinuxCursor->IsHidden())
 			{
-				//SDL_SetWindowGrab(TargetWindow, SDL_FALSE);
-				DSEXT_SetMouseGrab(TargetWindow, SDL_FALSE);
+				if (EDSExtSuccess != DSEXT_SetMouseGrab(TargetWindow, SDL_FALSE))
+				{
+					UE_LOG(LogHAL, Log, TEXT("Could not ungrab cursor for SDL window %p"), TargetWindow);
+				}
 			}
 			MouseCaptureWindow = NULL;
 		}
diff --git a/Engine/Source/Runtime/Core/Private/Linux/LinuxCursor.cpp b/Engine/Source/Runtime/Core/Private/Linux/LinuxCursor.cpp
--- a/Engine/Source/Runtime/Core/Private/Linux/LinuxCursor.cpp
+++ b/Engine/Source/Runtime/Core/Private/Linux/LinuxCursor.cpp
@@ -150,7 +150,11 @@ FVector2D FLinuxCursor::GetPosition() const
 {
 	int CursorX, CursorY;
 
-	DSEXT_GetMousePosition( &CursorX, &CursorY );
+	if (EDSExtSuccess != DSEXT_GetMousePosition(&CursorX, &CursorY))
+	{
+		UE_LOG(LogHAL, Log, TEXT("Could not get absolute mouse position"));
+		CursorX = CursorY = 0;
+	}
 
 	return FVector2D( CursorX, CursorY );
 }
