UE4: Fixed a couple minor issues that showed up in static code analysis.

diff --git a/Engine/Source/Runtime/Core/Private/Async/TaskGraph.cpp b/Engine/Source/Runtime/Core/Private/Async/TaskGraph.cpp
--- a/Engine/Source/Runtime/Core/Private/Async/TaskGraph.cpp
+++ b/Engine/Source/Runtime/Core/Private/Async/TaskGraph.cpp
@@ -749,7 +749,7 @@ public:
 		}
 		int32 QueueToExecuteOn = ENamedThreads::GetQueueIndex(ThreadToExecuteOn);
 		ThreadToExecuteOn = ENamedThreads::GetThreadIndex(ThreadToExecuteOn);
-		FTaskThread* Target = Target = &Thread(ThreadToExecuteOn);
+		FTaskThread* Target = &Thread(ThreadToExecuteOn);
 		if (ThreadToExecuteOn == CurrentThreadIfKnown)
 		{
 			Target->EnqueueFromThisThread(QueueToExecuteOn, Task);
diff --git a/Engine/Source/Runtime/Core/Private/HAL/ThreadingBase.cpp b/Engine/Source/Runtime/Core/Private/HAL/ThreadingBase.cpp
--- a/Engine/Source/Runtime/Core/Private/HAL/ThreadingBase.cpp
+++ b/Engine/Source/Runtime/Core/Private/HAL/ThreadingBase.cpp
@@ -315,6 +315,7 @@ public:
 		, TimeToDie(0)
 		, QueuedWork(NULL)
 		, OwningThreadPool(NULL)
+		, Thread(NULL)
 	{
 	}
 
diff --git a/Engine/Source/Runtime/Core/Private/Misc/ConfigCacheIni.cpp b/Engine/Source/Runtime/Core/Private/Misc/ConfigCacheIni.cpp
--- a/Engine/Source/Runtime/Core/Private/Misc/ConfigCacheIni.cpp
+++ b/Engine/Source/Runtime/Core/Private/Misc/ConfigCacheIni.cpp
@@ -2335,8 +2335,7 @@ static void LoadAnIniFile(const FString& FilenameToLoad, FConfigFile& ConfigFile
  */
 static bool GenerateDestIniFile(FConfigFile& DestConfigFile, const FString& DestIniFilename, const TArray<FIniFilename>& SourceIniHierarchy, bool bAllowGeneratedINIs)
 {
-	bool bResult = true;
-	bResult = LoadIniFileHierarchy(SourceIniHierarchy, *DestConfigFile.SourceConfigFile);
+	bool bResult = LoadIniFileHierarchy(SourceIniHierarchy, *DestConfigFile.SourceConfigFile);
 	if( !bResult )
 	{
 		return false;
diff --git a/Engine/Source/Runtime/Core/Private/Stats/Stats2.cpp b/Engine/Source/Runtime/Core/Private/Stats/Stats2.cpp
--- a/Engine/Source/Runtime/Core/Private/Stats/Stats2.cpp
+++ b/Engine/Source/Runtime/Core/Private/Stats/Stats2.cpp
@@ -566,7 +566,8 @@ class FStatsThread : public FRunnable, FSingleThreadRunnable
 public:
 
 	FStatsThread()
-		: State(FStatsThreadState::GetLocalState())
+		: Thread(nullptr)
+		, State(FStatsThreadState::GetLocalState())
 		, bReadyToProcess(false)
 	{
 		check(IsInGameThread());
