PlatformMisc - Added method GetDiskTotalAndFreeSpace which retrieves information about the total number of bytes and number of free bytes for the specified disk path (only for Windows)

diff --git a/Engine/Source/Runtime/Core/Private/GenericPlatform/GenericPlatformMisc.cpp b/Engine/Source/Runtime/Core/Private/GenericPlatform/GenericPlatformMisc.cpp
--- a/Engine/Source/Runtime/Core/Private/GenericPlatform/GenericPlatformMisc.cpp
+++ b/Engine/Source/Runtime/Core/Private/GenericPlatform/GenericPlatformMisc.cpp
@@ -222,6 +222,16 @@ void FGenericPlatformMisc::GetOSVersions( FString& out_OSVersionLabel, FString&
 	out_OSSubVersionLabel = FString( TEXT( "GenericOSSubVersionLabel" ) );
 }
 
+
+bool FGenericPlatformMisc::GetDiskTotalAndFreeSpace( const FString& InPath, uint64& TotalNumberOfBytes, uint64& NumberOfFreeBytes )
+{
+	// Not implemented cross-platform. Each platform may or may not choose to implement this.
+	TotalNumberOfBytes = 0;
+	NumberOfFreeBytes = 0;
+	return false;
+}
+
+
 void FGenericPlatformMisc::MemoryBarrier()
 {
 }
diff --git a/Engine/Source/Runtime/Core/Private/Windows/WindowsPlatformMisc.cpp b/Engine/Source/Runtime/Core/Private/Windows/WindowsPlatformMisc.cpp
--- a/Engine/Source/Runtime/Core/Private/Windows/WindowsPlatformMisc.cpp
+++ b/Engine/Source/Runtime/Core/Private/Windows/WindowsPlatformMisc.cpp
@@ -2100,6 +2100,20 @@ void FWindowsPlatformMisc::GetOSVersions( FString& out_OSVersionLabel, FString&
 	FWindowsOSVersionHelper::GetOSVersions( out_OSVersionLabel, out_OSSubVersionLabel );
 }
 
+
+bool FWindowsPlatformMisc::GetDiskTotalAndFreeSpace( const FString& InPath, uint64& TotalNumberOfBytes, uint64& NumberOfFreeBytes )
+{
+	bool bSuccess = false;
+	// We need to convert the path to make sure it is formatted with windows style Drive e.g. "C:\"
+	const FString ValidatedPath = FPaths::ConvertRelativePathToFull( InPath ).Replace( TEXT( "/" ), TEXT( "\\" ) );
+	if( ValidatedPath.Len() >= 3 && ValidatedPath[1] == ':' && ValidatedPath[2] == '\\' )
+	{
+		bSuccess = !!::GetDiskFreeSpaceEx( *ValidatedPath, nullptr, reinterpret_cast<ULARGE_INTEGER*>(&TotalNumberOfBytes), reinterpret_cast<ULARGE_INTEGER*>(&NumberOfFreeBytes) );
+	}
+	return bSuccess;	
+}
+
+
 uint32 FWindowsPlatformMisc::GetCPUInfo()
 {
 	return FCPUIDQueriedData::GetCPUInfo();
diff --git a/Engine/Source/Runtime/Core/Public/GenericPlatform/GenericPlatformMisc.h b/Engine/Source/Runtime/Core/Public/GenericPlatform/GenericPlatformMisc.h
--- a/Engine/Source/Runtime/Core/Public/GenericPlatform/GenericPlatformMisc.h
+++ b/Engine/Source/Runtime/Core/Public/GenericPlatform/GenericPlatformMisc.h
@@ -318,6 +318,9 @@ struct CORE_API FGenericPlatformMisc
 	 */
 	static void GetOSVersions( FString& out_OSVersionLabel, FString& out_OSSubVersionLabel );
 
+	/** Retrieves information about the total number of bytes and number of free bytes for the specified disk path. */
+	static bool GetDiskTotalAndFreeSpace( const FString& InPath, uint64& TotalNumberOfBytes, uint64& NumberOfFreeBytes );
+
 	static bool SupportsMessaging()
 	{
 		return true;
diff --git a/Engine/Source/Runtime/Core/Public/Windows/WindowsPlatformMisc.h b/Engine/Source/Runtime/Core/Public/Windows/WindowsPlatformMisc.h
--- a/Engine/Source/Runtime/Core/Public/Windows/WindowsPlatformMisc.h
+++ b/Engine/Source/Runtime/Core/Public/Windows/WindowsPlatformMisc.h
@@ -154,6 +154,7 @@ struct CORE_API FWindowsPlatformMisc : public FGenericPlatformMisc
 	static FString GetCPUBrand();
 	static FString GetPrimaryGPUBrand();
 	static void GetOSVersions( FString& out_OSVersionLabel, FString& out_OSSubVersionLabel );
+	static bool GetDiskTotalAndFreeSpace( const FString& InPath, uint64& TotalNumberOfBytes, uint64& NumberOfFreeBytes );
 
 	/**
 	 * Uses cpuid instruction to get the vendor string
