Get multitouch info on Mac using private MultitouchDevices framework instead of trying to guess based on limited info provided by NSEvent

diff --git a/Engine/Source/Runtime/Core/Core.Build.cs b/Engine/Source/Runtime/Core/Core.Build.cs
--- a/Engine/Source/Runtime/Core/Core.Build.cs
+++ b/Engine/Source/Runtime/Core/Core.Build.cs
@@ -81,6 +81,11 @@ public class Core : ModuleRules
 				"OpenGL"
 				);
 			PublicFrameworks.AddRange(new string[] { "Cocoa", "Carbon", "IOKit" });
+			
+			if (UEBuildConfiguration.bBuildEditor == true)
+			{
+				PublicAdditionalLibraries.Add("/System/Library/PrivateFrameworks/MultitouchSupport.framework/Versions/Current/MultitouchSupport");
+			}
 		}
 		else if (Target.Platform == UnrealTargetPlatform.IOS)
 		{
diff --git a/Engine/Source/Runtime/Core/Private/Mac/MacApplication.cpp b/Engine/Source/Runtime/Core/Private/Mac/MacApplication.cpp
--- a/Engine/Source/Runtime/Core/Private/Mac/MacApplication.cpp
+++ b/Engine/Source/Runtime/Core/Private/Mac/MacApplication.cpp
@@ -14,6 +14,22 @@
 
 FMacApplication* MacApplication = NULL;
 
+#if WITH_EDITOR
+typedef int32 (*MTContactCallbackFunction)(int32, void*, int32, double, int32);
+extern "C" CFMutableArrayRef MTDeviceCreateList(void);
+extern "C" void MTRegisterContactFrameCallback(void*, MTContactCallbackFunction);
+extern "C" void MTDeviceStart(void*, int);
+
+static int MTContactCallback(int32 Device, void* Data, int32 NumFingers, double TimeStamp, int32 Frame)
+{
+	if (MacApplication)
+	{
+		MacApplication->SetIsUsingTrackpad(NumFingers > 0);
+	}
+	return 1;
+}
+#endif
+
 FMacApplication* FMacApplication::CreateMacApplication()
 {
 	MacApplication = new FMacApplication();
@@ -46,6 +62,15 @@ FMacApplication::FMacApplication()
 
 	[NSEvent addGlobalMonitorForEventsMatchingMask:NSMouseMovedMask handler:^(NSEvent* Event){AddPendingEvent(Event);}];
 
+#if WITH_EDITOR
+	NSMutableArray* MultiTouchDevices = (__bridge NSMutableArray*)MTDeviceCreateList();
+	for (id Device in MultiTouchDevices)
+	{
+		MTRegisterContactFrameCallback((void*)Device, MTContactCallback);
+		MTDeviceStart((void*)Device, 0);
+	}
+#endif
+
 	TextInputMethodSystem = MakeShareable( new FMacTextInputMethodSystem );
 	if(!TextInputMethodSystem->Initialize())
 	{
@@ -262,8 +287,6 @@ void FMacApplication::ProcessEvent( NSEvent* Event )
 		case NSRightMouseDragged:
 		case NSOtherMouseDragged:
 		{
-			bUsingTrackpad = [Event subtype] == NSTouchEventSubtype;
-
 			if( CurrentEventWindow.IsValid() && CurrentEventWindow->IsRegularWindow() )
 			{
 				bool IsMouseOverTitleBar = false;
@@ -347,8 +370,6 @@ void FMacApplication::ProcessEvent( NSEvent* Event )
 		case NSRightMouseDown:
 		case NSOtherMouseDown:
 		{
-			bUsingTrackpad = [Event subtype] == NSTouchEventSubtype;
-
 			EMouseButtons::Type Button = [Event type] == NSLeftMouseDown ? EMouseButtons::Left : EMouseButtons::Right;
 			if ([Event type] == NSOtherMouseDown)
 			{
@@ -394,8 +415,6 @@ void FMacApplication::ProcessEvent( NSEvent* Event )
 		case NSRightMouseUp:
 		case NSOtherMouseUp:
 		{
-			bUsingTrackpad = false;
-
 			EMouseButtons::Type Button = [Event type] == NSLeftMouseUp ? EMouseButtons::Left : EMouseButtons::Right;
 			if ([Event type] == NSOtherMouseUp)
 			{
@@ -436,15 +455,6 @@ void FMacApplication::ProcessEvent( NSEvent* Event )
 		
 				NSEventPhase Phase = [Event phase];
 				
-				if( Phase == NSEventPhaseNone || Phase == NSEventPhaseEnded || Phase == NSEventPhaseCancelled )
-				{
-					bUsingTrackpad = false;
-				}
-				else
-				{
-					bUsingTrackpad = true;
-				}
-				
 				if ([Event momentumPhase] != NSEventPhaseNone || [Event phase] != NSEventPhaseNone)
 				{
 					bool bInverted = [Event isDirectionInvertedFromDevice];
@@ -501,7 +511,6 @@ void FMacApplication::ProcessEvent( NSEvent* Event )
 	
 		case NSEventTypeBeginGesture:
 		{
-			bUsingTrackpad = true;
             if( CurrentEventWindow.IsValid() )
 			{
 				MessageHandler->OnBeginGesture();
@@ -515,7 +524,6 @@ void FMacApplication::ProcessEvent( NSEvent* Event )
 			{
 				MessageHandler->OnEndGesture();
             }
-			bUsingTrackpad = false;
 #if WITH_EDITOR
 			LastGestureUsed = EGestureEvent::None;
 #endif
diff --git a/Engine/Source/Runtime/Core/Public/Mac/MacApplication.h b/Engine/Source/Runtime/Core/Public/Mac/MacApplication.h
--- a/Engine/Source/Runtime/Core/Public/Mac/MacApplication.h
+++ b/Engine/Source/Runtime/Core/Public/Mac/MacApplication.h
@@ -79,6 +79,11 @@ public:
 
 #if WITH_EDITOR
     virtual void SendAnalytics(IAnalyticsProvider* Provider) OVERRIDE;
+
+	void SetIsUsingTrackpad(bool bInUsingTrackpad)
+	{
+		bUsingTrackpad = bInUsingTrackpad;
+	}
 #endif
 
 
