Linux: making QA-ClickHUD work.

- This change makes QA-ClickHUD map work (resolving 4.2 TTP #331994) while also not breaking SlateViewer and (hopefully) external efforts to make editor work. Needs to be reconsidered and redone properly, probably will not work with QA-ClickHUD via PIE (when we get to this).
- Also make DSExt more informative about its errors and more X11 specific ;/ Also to be redone after 4.2 TTPs are gone.

diff --git a/Engine/Source/Runtime/Core/Private/Linux/LinuxApplication.cpp b/Engine/Source/Runtime/Core/Private/Linux/LinuxApplication.cpp
--- a/Engine/Source/Runtime/Core/Private/Linux/LinuxApplication.cpp
+++ b/Engine/Source/Runtime/Core/Private/Linux/LinuxApplication.cpp
@@ -293,7 +293,7 @@ void FLinuxApplication::ProcessDeferredMessage( SDL_Event Event )
 			}
 			if(buttonEvent.type == SDL_MOUSEBUTTONUP)
 			{
-				MessageHandler->OnMouseUp( button );
+				MessageHandler->OnMouseUp(button);
 			}
 			// SDL 2.0.2+
 			//else if(buttonEvent.clicks > 1)
@@ -886,13 +886,16 @@ void FLinuxApplication::UpdateMouseCaptureWindow( SDL_HWindow TargetWindow )
 	const bool bEnable = bIsMouseCaptureEnabled || bIsMouseCursorLocked;
 	FLinuxCursor *LinuxCursor = static_cast< FLinuxCursor* >(Cursor.Get());
 
-	if( bEnable )
+	// this is a hacky heuristic which makes QA-ClickHUD work while not ruining SlateViewer...
+	bool bShouldGrab = (IS_PROGRAM != 0 || GIsEditor) && !LinuxCursor->IsHidden();
+
+	if (bEnable)
 	{
 		if( TargetWindow )
 		{
 			MouseCaptureWindow = TargetWindow;
 		}
-		if (MouseCaptureWindow && !LinuxCursor->IsHidden())
+		if (bShouldGrab && MouseCaptureWindow)
 		{
 			if (EDSExtSuccess != DSEXT_SetMouseGrab(TargetWindow, SDL_TRUE))
 			{
@@ -904,7 +907,7 @@ void FLinuxApplication::UpdateMouseCaptureWindow( SDL_HWindow TargetWindow )
 	{
 		if( MouseCaptureWindow )
 		{
-			if (!LinuxCursor->IsHidden())
+			if (bShouldGrab)
 			{
 				if (EDSExtSuccess != DSEXT_SetMouseGrab(TargetWindow, SDL_FALSE))
 				{
diff --git a/Engine/Source/Runtime/Core/Private/Linux/LinuxCursor.cpp b/Engine/Source/Runtime/Core/Private/Linux/LinuxCursor.cpp
--- a/Engine/Source/Runtime/Core/Private/Linux/LinuxCursor.cpp
+++ b/Engine/Source/Runtime/Core/Private/Linux/LinuxCursor.cpp
@@ -150,9 +150,10 @@ FVector2D FLinuxCursor::GetPosition() const
 {
 	int CursorX, CursorY;
 
-	if (EDSExtSuccess != DSEXT_GetMousePosition(&CursorX, &CursorY))
+	int DSRetCode = DSEXT_GetAbsoluteMousePosition(&CursorX, &CursorY);
+	if (EDSExtSuccess != DSRetCode)
 	{
-		UE_LOG(LogHAL, Log, TEXT("Could not get absolute mouse position"));
+		UE_LOG(LogHAL, Log, TEXT("Could not get absolute mouse position, DSExt returned %d"), DSRetCode);
 		CursorX = CursorY = 0;
 	}
 
diff --git a/Engine/Source/Runtime/OpenGLDrv/Private/Linux/OpenGLLinux.cpp b/Engine/Source/Runtime/OpenGLDrv/Private/Linux/OpenGLLinux.cpp
--- a/Engine/Source/Runtime/OpenGLDrv/Private/Linux/OpenGLLinux.cpp
+++ b/Engine/Source/Runtime/OpenGLDrv/Private/Linux/OpenGLLinux.cpp
@@ -5,6 +5,7 @@
 =============================================================================*/
 
 #include "OpenGLDrvPrivate.h"
+#include "ds_extensions.h"
 
 /*------------------------------------------------------------------------------
 	OpenGL function pointers.
@@ -107,6 +108,12 @@ static void _PlatformCreateDummyGLWindow( FPlatformOpenGLContext *OutContext )
 
 	OutContext->hWnd					= h_wnd;
 	OutContext->bReleaseWindowOnDestroy	= true;
+
+	int CacheRet = DSEXT_CacheX11Info(h_wnd);
+	if (EDSExtSuccess != CacheRet)
+	{
+		UE_LOG(LogRHI, Error, TEXT("Could not cache X11 info, DSExt error: %d"), CacheRet);
+	}
 }
 
 /**
