Allow starting CrashReportClient from a case-sensitive fs.

We used to lowercase everything, and some current builds still do this. Thus try lowercased path first, then switch to a proper one (assuming case-sensitive deployment).
Needed for TTP #309293.

diff --git a/Engine/Source/Runtime/Core/Private/Linux/LinuxPlatformMisc.cpp b/Engine/Source/Runtime/Core/Private/Linux/LinuxPlatformMisc.cpp
--- a/Engine/Source/Runtime/Core/Private/Linux/LinuxPlatformMisc.cpp
+++ b/Engine/Source/Runtime/Core/Private/Linux/LinuxPlatformMisc.cpp
@@ -970,7 +970,13 @@ void GenerateCrashInfoAndLaunchReporter(const FLinuxCrashContext & Context)
 
 		// try launching the tool and wait for its exit, if at all
 		const TCHAR * RelativePathToCrashReporter = TEXT("../../../engine/binaries/linux/crashreportclient");	// FIXME: painfully hard-coded
+		if (!FPaths::FileExists(RelativePathToCrashReporter))
+		{
+			RelativePathToCrashReporter = TEXT("../../../Engine/Binaries/Linux/CrashReportClient");	// FIXME: even more painfully hard-coded
+		}
 
+		// show on the console
+		printf("Starting %s\n", StringCast<ANSICHAR>(RelativePathToCrashReporter).Get());
 		FProcHandle RunningProc = FPlatformProcess::CreateProc(RelativePathToCrashReporter, *(CrashInfoAbsolute + TEXT("/")), true, false, false, NULL, 0, NULL, NULL);
 		if (FPlatformProcess::IsProcRunning(RunningProc))
 		{
@@ -988,6 +994,8 @@ void GenerateCrashInfoAndLaunchReporter(const FLinuxCrashContext & Context)
 				{
 					break;
 				}
+
+				FPlatformProcess::Sleep(1.0f);
 			};
 		}
 	}
