Core - SCOPE_LOG_TIME_FUNC should work on all platforms

diff --git a/Engine/Source/Runtime/Core/Private/Stats/StatsMisc.cpp b/Engine/Source/Runtime/Core/Private/Stats/StatsMisc.cpp
--- a/Engine/Source/Runtime/Core/Private/Stats/StatsMisc.cpp
+++ b/Engine/Source/Runtime/Core/Private/Stats/StatsMisc.cpp
@@ -2,10 +2,17 @@
 
 #include "CorePrivate.h"
 
-FScopeLogTime::FScopeLogTime( const TCHAR* InName, FTotalTimeAndCount* InGlobal /*= nullptr */, EScopeLogTimeUnits InUnits /*= ScopeLog_Milliseconds */ ) 
+FScopeLogTime::FScopeLogTime( const WIDECHAR* InName, FTotalTimeAndCount* InCumulative /*= nullptr */, EScopeLogTimeUnits InUnits /*= ScopeLog_Milliseconds */ )
 : StartTime( FPlatformTime::Seconds() )
 , Name( InName )
-, Cumulative( InGlobal )
+, Cumulative( InCumulative )
+, Units( InUnits )
+{}
+
+FScopeLogTime::FScopeLogTime( const ANSICHAR* InName, FTotalTimeAndCount* InCumulative /*= nullptr*/, EScopeLogTimeUnits InUnits /*= ScopeLog_Milliseconds */ )
+: StartTime( FPlatformTime::Seconds() )
+, Name( InName )
+, Cumulative( InCumulative )
 , Units( InUnits )
 {}
 
diff --git a/Engine/Source/Runtime/Core/Public/Stats/StatsMisc.h b/Engine/Source/Runtime/Core/Public/Stats/StatsMisc.h
--- a/Engine/Source/Runtime/Core/Public/Stats/StatsMisc.h
+++ b/Engine/Source/Runtime/Core/Public/Stats/StatsMisc.h
@@ -60,16 +60,17 @@ struct FScopeLogTime
 	 * @param InGlobal - Pointer to the variable that holds the cumulative stats
 	 *
 	 */
-	CORE_API FScopeLogTime( const TCHAR* InName, FTotalTimeAndCount* InCumulative = nullptr, EScopeLogTimeUnits InUnits = ScopeLog_Milliseconds );
+	CORE_API FScopeLogTime( const WIDECHAR* InName, FTotalTimeAndCount* InCumulative = nullptr, EScopeLogTimeUnits InUnits = ScopeLog_Milliseconds );
+	CORE_API FScopeLogTime( const ANSICHAR* InName, FTotalTimeAndCount* InCumulative = nullptr, EScopeLogTimeUnits InUnits = ScopeLog_Milliseconds );
 
 	/** Destructor. */
 	CORE_API ~FScopeLogTime();
 
 protected:
-	virtual double GetDisplayScopedTime(double InScopedTime) const;
-	virtual FString GetDisplayUnitsString() const;
+	double GetDisplayScopedTime(double InScopedTime) const;
+	FString GetDisplayUnitsString() const;
+
 
-protected:
 	const double StartTime;
 	const FString Name;
 	FTotalTimeAndCount* Cumulative;
@@ -83,7 +84,7 @@ protected:
 	FScopeLogTime ScopeLogTime(Name, CumulativePtr, FScopeLogTime::ScopeLog_Seconds);
 
 #define SCOPE_LOG_TIME_FUNC() \
-	FScopeLogTime ScopeLogTime(TEXT(__FUNCTION__));
+	FScopeLogTime ScopeLogTime(__FUNCTION__);
 
 #define SCOPE_LOG_TIME_FUNC_WITH_GLOBAL(CumulativePtr) \
-	FScopeLogTime ScopeLogTime(TEXT(__FUNCTION__),CumulativePtr);
+	FScopeLogTime ScopeLogTime(__FUNCTION__,CumulativePtr);
