Linux: fix Delete key (TTP #334760).

- Also minor formatting changes to make code more in line with our coding style.

diff --git a/Engine/Source/Runtime/Core/Private/Linux/LinuxApplication.cpp b/Engine/Source/Runtime/Core/Private/Linux/LinuxApplication.cpp
--- a/Engine/Source/Runtime/Core/Private/Linux/LinuxApplication.cpp
+++ b/Engine/Source/Runtime/Core/Private/Linux/LinuxApplication.cpp
@@ -163,7 +163,7 @@ bool FLinuxApplication::GeneratesKeyCharMessage(const SDL_KeyboardEvent & KeyDow
 
 	// filter out command keys, non-ASCI and arrow keycodes that don't generate WM_CHAR under Windows (TODO: find a table?)
 	return !bCmdKeyPressed && Sym < 128 &&
-		(Sym != SDLK_DOWN && Sym != SDLK_LEFT && Sym != SDLK_RIGHT && Sym != SDLK_UP);
+		(Sym != SDLK_DOWN && Sym != SDLK_LEFT && Sym != SDLK_RIGHT && Sym != SDLK_UP && Sym != SDLK_DELETE);
 }
 
 void FLinuxApplication::ProcessDeferredMessage( SDL_Event Event )
