Fix using the editor in fullscreen mode on Mac OS X without changing other libaries. Regular windows with an OS titlebar are set as fullscreen primary, but all others being as secondary windows. so that they don't try and go fullscreen and break OS X's window handling. When we explicitly ask for fullscreen mode we change the window collection behaviour before toggling to make the window fullscreen primary. After switching back to windowed mode we set the window back to auxiliary if it doesn't have an OS titlebar. This ensures that only one window is ever fullscreen in the Editor & lets standalone games work properly too.

diff --git a/Engine/Source/Runtime/Core/Private/Mac/MacWindow.cpp b/Engine/Source/Runtime/Core/Private/Mac/MacWindow.cpp
--- a/Engine/Source/Runtime/Core/Private/Mac/MacWindow.cpp
+++ b/Engine/Source/Runtime/Core/Private/Mac/MacWindow.cpp
@@ -139,7 +139,7 @@
 
 - (void)redrawContents
 {
-	if(self.bForwardEvents)
+	if(self.bForwardEvents && ([self isVisible] && [super alphaValue] > 0.0f))
 	{
 		MacApplication->OnWindowRedrawContents( self );
 	}
@@ -825,11 +825,18 @@ void FMacWindow::Initialize( FMacApplication* const Application, const TSharedRe
 		// Only makes sense for regular windows (windows that last a while.)
 		[WindowHandle registerForDraggedTypes: [NSArray arrayWithObject: NSFilenamesPboardType]];
 
-		[WindowHandle setCollectionBehavior: NSWindowCollectionBehaviorFullScreenPrimary|NSWindowCollectionBehaviorDefault|NSWindowCollectionBehaviorManaged|NSWindowCollectionBehaviorParticipatesInCycle];
+		if( Definition->HasOSWindowBorder )
+		{
+			[WindowHandle setCollectionBehavior: NSWindowCollectionBehaviorFullScreenPrimary|NSWindowCollectionBehaviorDefault|NSWindowCollectionBehaviorManaged|NSWindowCollectionBehaviorParticipatesInCycle];
+		}
+		else
+		{
+			[WindowHandle setCollectionBehavior: NSWindowCollectionBehaviorFullScreenAuxiliary|NSWindowCollectionBehaviorDefault|NSWindowCollectionBehaviorManaged|NSWindowCollectionBehaviorParticipatesInCycle];
+		}
 	}
 	else if(Definition->AppearsInTaskbar)
 	{
-		[WindowHandle setCollectionBehavior:NSWindowCollectionBehaviorDefault|NSWindowCollectionBehaviorManaged|NSWindowCollectionBehaviorParticipatesInCycle];
+		[WindowHandle setCollectionBehavior:NSWindowCollectionBehaviorFullScreenAuxiliary|NSWindowCollectionBehaviorDefault|NSWindowCollectionBehaviorManaged|NSWindowCollectionBehaviorParticipatesInCycle];
 	}
 	else
 	{
@@ -1064,6 +1071,16 @@ void FMacWindow::SetWindowMode( EWindowMode::Type NewWindowMode )
 	if( bIsFullscreen != bMakeFullscreen )
 	{
 		bool WindowIsFullScreen = !bMakeFullscreen;
+		
+		NSWindowCollectionBehavior Behaviour = [WindowHandle collectionBehavior];
+		if(bMakeFullscreen)
+		{
+			Behaviour &= ~(NSWindowCollectionBehaviorFullScreenAuxiliary);
+			Behaviour |= NSWindowCollectionBehaviorFullScreenPrimary;
+		}
+		
+		[WindowHandle setCollectionBehavior: Behaviour];
+		
 		[WindowHandle toggleFullScreen:nil];
 		// Ensure that the window has transitioned BEFORE leaving this function
 		// this prevents problems with failure to correctly update mouse locks
@@ -1074,6 +1091,13 @@ void FMacWindow::SetWindowMode( EWindowMode::Type NewWindowMode )
 			OwningApplication->ProcessDeferredEvents(0.f);
 			WindowIsFullScreen = [WindowHandle windowMode] != EWindowMode::Windowed;
 		} while(WindowIsFullScreen != bMakeFullscreen);
+		
+		if(!bMakeFullscreen && !Definition->HasOSWindowBorder)
+		{
+			Behaviour &= ~(NSWindowCollectionBehaviorFullScreenPrimary);
+			Behaviour |= NSWindowCollectionBehaviorFullScreenAuxiliary;
+		}
+		
 		WindowMode = NewWindowMode;
 	}
 }
