Core - Fixed comment

diff --git a/Engine/Source/Runtime/Core/Public/HAL/FileManager.h b/Engine/Source/Runtime/Core/Public/HAL/FileManager.h
--- a/Engine/Source/Runtime/Core/Public/HAL/FileManager.h
+++ b/Engine/Source/Runtime/Core/Public/HAL/FileManager.h
@@ -97,7 +97,7 @@ public:
 	/** Finds file or directories. */
 	virtual void FindFiles( TArray<FString>& FileNames, const TCHAR* Filename, bool Files, bool Directories)=0;
 
-	/** Finds file or directories recusively. */
+	/** Finds file or directories recursively. */
 	virtual void FindFilesRecursive( TArray<FString>& FileNames, const TCHAR* StartDirectory, const TCHAR* Filename, bool Files, bool Directories, bool bClearFileNames=true) = 0; // utility
 
 	/** Gets the age of a file measured in seconds. */
