Fix text entry in Mac OS X fullscreen mode.
We need to return the correct input-method aware subview from the window in FSlateCocoaWindow's openGLView, so iterate over the subviews to find the right one. This is necessary as OS X's fullscreen mode changes the subview array. In the input method system check that the view is of the correct type and report the failure rather than just crashing.

diff --git a/Engine/Source/Runtime/Core/Private/Mac/MacTextInputMethodSystem.cpp b/Engine/Source/Runtime/Core/Private/Mac/MacTextInputMethodSystem.cpp
--- a/Engine/Source/Runtime/Core/Private/Mac/MacTextInputMethodSystem.cpp
+++ b/Engine/Source/Runtime/Core/Private/Mac/MacTextInputMethodSystem.cpp
@@ -483,9 +483,13 @@ void FMacTextInputMethodSystem::ActivateContext(const TSharedRef<ITextInputMetho
 		FSlateCocoaWindow* CocoaWindow = (FSlateCocoaWindow*)GenericWindow->GetOSWindowHandle();
 		if(CocoaWindow && [CocoaWindow openGLView])
 		{
-			FSlateTextView* TextView = (FSlateTextView*)[CocoaWindow openGLView];
-			[TextView activateInputMethod:Context];
-			return;
+			NSView* GLView = [CocoaWindow openGLView];
+			if([GLView isKindOfClass:[FSlateTextView class]])
+			{
+				FSlateTextView* TextView = (FSlateTextView*)GLView;
+				[TextView activateInputMethod:Context];
+				return;
+			}
 		}
 	}
 	UE_LOG(LogMacTextInputMethodSystem, Error, TEXT("Activating a context failed when its window couldn't be found."));
@@ -508,9 +512,13 @@ void FMacTextInputMethodSystem::DeactivateContext(const TSharedRef<ITextInputMet
 		FSlateCocoaWindow* CocoaWindow = (FSlateCocoaWindow*)GenericWindow->GetOSWindowHandle();
 		if(CocoaWindow && [CocoaWindow openGLView])
 		{
-			FSlateTextView* TextView = (FSlateTextView*)[CocoaWindow openGLView];
-			[TextView deactivateInputMethod];
-			return;
+			NSView* GLView = [CocoaWindow openGLView];
+			if([GLView isKindOfClass:[FSlateTextView class]])
+			{
+				FSlateTextView* TextView = (FSlateTextView*)GLView;
+				[TextView deactivateInputMethod];
+				return;
+			}
 		}
 	}
 	UE_LOG(LogMacTextInputMethodSystem, Error, TEXT("Deactivating a context failed when its window couldn't be found."));
diff --git a/Engine/Source/Runtime/Core/Private/Mac/MacWindow.cpp b/Engine/Source/Runtime/Core/Private/Mac/MacWindow.cpp
--- a/Engine/Source/Runtime/Core/Private/Mac/MacWindow.cpp
+++ b/Engine/Source/Runtime/Core/Private/Mac/MacWindow.cpp
@@ -57,7 +57,14 @@
 	if([self styleMask] & (NSTexturedBackgroundWindowMask))
 	{
 		NSView* SuperView = [[self contentView] superview];
-		return [[SuperView subviews] lastObject];
+		for(NSView* View in [SuperView subviews])
+		{
+			if([View isKindOfClass:[FSlateTextView class]])
+			{
+				return View;
+			}
+		}
+		return nil;
 	}
 	else
 	{
