Made TMap::RemoveAndCopyValue and TMap::FindAndRemoveChecked move the value out instead of copying as an optimisation.

diff --git a/Engine/Source/Runtime/Core/Public/Containers/Map.h b/Engine/Source/Runtime/Core/Public/Containers/Map.h
--- a/Engine/Source/Runtime/Core/Public/Containers/Map.h
+++ b/Engine/Source/Runtime/Core/Public/Containers/Map.h
@@ -835,7 +835,7 @@ public:
 		if(!PairId.IsValidId())
 			return false;
 
-		OutRemovedValue = Super::Pairs[PairId].Value;
+		OutRemovedValue = MoveTemp(Super::Pairs[PairId].Value);
 		Super::Pairs.Remove(PairId);
 		return true;
 	}
@@ -850,7 +850,7 @@ public:
 	{
 		const FSetElementId PairId = Super::Pairs.FindId(Key);
 		check(PairId.IsValidId());
-		ValueType Result = Super::Pairs[PairId].Value;
+		ValueType Result = MoveTemp(Super::Pairs[PairId].Value);
 		Super::Pairs.Remove(PairId);
 		return Result;
 	}
