Temp fix for win32 builds ('IsNormalized'/'ToString' : is not a member of '__m128')

diff --git a/Engine/Source/Runtime/Core/Public/Math/TransformVectorized.h b/Engine/Source/Runtime/Core/Public/Math/TransformVectorized.h
--- a/Engine/Source/Runtime/Core/Public/Math/TransformVectorized.h
+++ b/Engine/Source/Runtime/Core/Public/Math/TransformVectorized.h
@@ -35,7 +35,7 @@ public:
 #if ENABLE_UNIT_QUAT_DIAGNOSTIC
 	FORCEINLINE void DiagnosticCheckUnitQuaternion() const
 	{
-		checkf(Rotation.IsNormalized(), TEXT("FTransform Rotation is not Normalized: %s"), *Rotation.ToString());
+		//checkf(Rotation.IsNormalized(), TEXT("FTransform Rotation is not Normalized: %s"), *Rotation.ToString());
 	}
 #else
 	FORCEINLINE void DiagnosticCheckUnitQuaternion() const {}
