Proper fix for ('IsNormalized'/'ToString' : is not a member of '__m128')

diff --git a/Engine/Source/Runtime/Core/Public/Math/TransformVectorized.h b/Engine/Source/Runtime/Core/Public/Math/TransformVectorized.h
--- a/Engine/Source/Runtime/Core/Public/Math/TransformVectorized.h
+++ b/Engine/Source/Runtime/Core/Public/Math/TransformVectorized.h
@@ -35,7 +35,10 @@ public:
 #if ENABLE_UNIT_QUAT_DIAGNOSTIC
 	FORCEINLINE void DiagnosticCheckUnitQuaternion() const
 	{
-		//checkf(Rotation.IsNormalized(), TEXT("FTransform Rotation is not Normalized: %s"), *Rotation.ToString());
+		FQuat QuatRotation;
+		VectorStoreAligned(Rotation, &QuatRotation);
+
+		checkf(QuatRotation.IsNormalized(), TEXT("FTransform Rotation is not Normalized: %s"), *QuatRotation.ToString());
 	}
 #else
 	FORCEINLINE void DiagnosticCheckUnitQuaternion() const {}
