Stats - Added a guard to prevent from using SCOPE_COUNTER in the global scope as this may really hurt stats (only for debug builds, but should be enough to catch issues early)

diff --git a/Engine/Source/Runtime/Core/Public/Stats/Stats2.h b/Engine/Source/Runtime/Core/Public/Stats/Stats2.h
--- a/Engine/Source/Runtime/Core/Public/Stats/Stats2.h
+++ b/Engine/Source/Runtime/Core/Public/Stats/Stats2.h
@@ -1407,19 +1407,28 @@ struct FStat_##StatName\
 #define DECLARE_STATS_GROUP_MAYBE_COMPILED_OUT(GroupDesc, GroupId, GroupCat, CompileIn) \
 	DECLARE_STAT_GROUP(GroupDesc, GroupId, GroupCat, false, CompileIn);
 
+#ifdef UE_BUILD_DEBUG
+	#define SCOPE_CYCLE_COUNTER_GUARD {const char* ReadTheText = "SCOPE_CYCLE_COUNTER can't be used in the global scope.";}
+#else
+	#define SCOPE_CYCLE_COUNTER_GUARD
+#endif // UE_BUILD_DEBUG
 
 #define DECLARE_SCOPE_CYCLE_COUNTER(CounterName,StatId,GroupId) \
+	SCOPE_CYCLE_COUNTER_GUARD \
 	DECLARE_STAT(CounterName,StatId,GroupId,EStatDataType::ST_int64, true, true, FPlatformMemory::MCR_Invalid); \
 	static DEFINE_STAT(StatId) \
 	FScopeCycleCounter CycleCount_##StatId(GET_STATID(StatId));
 
 #define QUICK_SCOPE_CYCLE_COUNTER(Stat) \
+	SCOPE_CYCLE_COUNTER_GUARD \
 	DECLARE_SCOPE_CYCLE_COUNTER(TEXT(#Stat),Stat,STATGROUP_Quick)
 
 #define SCOPE_CYCLE_COUNTER(Stat) \
+	SCOPE_CYCLE_COUNTER_GUARD \
 	FScopeCycleCounter CycleCount_##Stat(GET_STATID(Stat));
 
 #define CONDITIONAL_SCOPE_CYCLE_COUNTER(Stat,bCondition) \
+	SCOPE_CYCLE_COUNTER_GUARD \
 	FScopeCycleCounter CycleCount_##Stat(bCondition ? GET_STATID(Stat) : TStatId());
 
 
