Randomize crash handler directory.

Makes it even more collision-resistant - crash report server can have reports from multiple machines and just pid is not enough.
(Part of TTP #309293).

diff --git a/Engine/Source/Runtime/Core/Private/Linux/LinuxPlatformMisc.cpp b/Engine/Source/Runtime/Core/Private/Linux/LinuxPlatformMisc.cpp
--- a/Engine/Source/Runtime/Core/Private/Linux/LinuxPlatformMisc.cpp
+++ b/Engine/Source/Runtime/Core/Private/Linux/LinuxPlatformMisc.cpp
@@ -949,7 +949,8 @@ void GenerateCrashInfoAndLaunchReporter(const FLinuxCrashContext & Context)
 #if !IS_PROGRAM
 
 	// create a crash-specific directory
-	FString CrashInfoFolder = FString::Printf(TEXT("crashinfo-%s-pid-%d"), FApp::GetGameName(), getpid());
+	FString CrashInfoFolder = FString::Printf(TEXT("crashinfo-%s-pid-%d-%s-%s"), FApp::GetGameName(), getpid(), 
+		*FDateTime::Now().ToString(), *FGuid::NewGuid().ToString());
 	FString CrashInfoAbsolute = FPaths::ConvertRelativePathToFull(CrashInfoFolder);
 	if (IFileManager::Get().MakeDirectory(*CrashInfoFolder))
 	{
