Linux: Do not init a static variable so it can moved to BSS.

Minor savings of about 16K binary size.

diff --git a/Engine/Source/Runtime/Core/Private/Linux/LinuxPlatformMisc.cpp b/Engine/Source/Runtime/Core/Private/Linux/LinuxPlatformMisc.cpp
--- a/Engine/Source/Runtime/Core/Private/Linux/LinuxPlatformMisc.cpp
+++ b/Engine/Source/Runtime/Core/Private/Linux/LinuxPlatformMisc.cpp
@@ -83,10 +83,13 @@ void FLinuxMisc::NormalizePath(FString& InPath)
 	if (InPath.Contains(TEXT("~"), ESearchCase::CaseSensitive))	// case sensitive is quicker, and our substring doesn't care
 	{
 		static bool bHaveHome = false;
-		static TCHAR CachedResult[ PlatformMiscLimits::MaxUserHomeDirLength ] = TEXT("~");	// init with a default value that changes nothing
+		static TCHAR CachedResult[PlatformMiscLimits::MaxUserHomeDirLength];
 
 		if (!bHaveHome)
 		{
+			CachedResult[0] = TEXT('~');	// init with a default value that changes nothing
+			CachedResult[1] = TEXT('\0');
+
 			//  get user $HOME var first
 			const char * VarValue = secure_getenv("HOME");
 			if (NULL != VarValue)
