Core - FThreadSingleton can't be used outside of Core

diff --git a/Engine/Source/Runtime/Core/Private/HAL/ThreadingBase.cpp b/Engine/Source/Runtime/Core/Private/HAL/ThreadingBase.cpp
--- a/Engine/Source/Runtime/Core/Private/HAL/ThreadingBase.cpp
+++ b/Engine/Source/Runtime/Core/Private/HAL/ThreadingBase.cpp
@@ -596,3 +596,33 @@ FQueuedThreadPool* FQueuedThreadPool::Allocate()
 {
 	return new FQueuedThreadPoolBase;
 }
+
+/*-----------------------------------------------------------------------------
+	FThreadSingletonInitializer
+-----------------------------------------------------------------------------*/
+
+FThreadSingleton* FThreadSingletonInitializer::Get( const FThreadSingleton::TCreateSingletonFuncPtr CreateFunc, const FThreadSingleton::TDestroySingletonFuncPtr DestroyFunc, uint32& TlsSlot )
+{
+	check( CreateFunc != nullptr );
+	check( DestroyFunc != nullptr );
+	if( TlsSlot == 0 )
+	{
+		check( IsInGameThread() );
+		TlsSlot = FPlatformTLS::AllocTlsSlot();
+	}
+	FThreadSingleton* ThreadSingleton = (FThreadSingleton*)FPlatformTLS::GetTlsValue( TlsSlot );
+	if( !ThreadSingleton )
+	{
+		const uint32 ThreadId = FPlatformTLS::GetCurrentThreadId();
+		ThreadSingleton = (*CreateFunc)();
+		FRunnableThread::GetThreadRegistry().Lock();
+		FRunnableThread* RunnableThread = FRunnableThread::GetThreadRegistry().GetThread( ThreadId );
+		if( RunnableThread )
+		{
+			RunnableThread->OnThreadDestroyed().AddRaw( ThreadSingleton, DestroyFunc );
+		}
+		FRunnableThread::GetThreadRegistry().Unlock();
+		FPlatformTLS::SetTlsValue( TlsSlot, ThreadSingleton );
+	}
+	return ThreadSingleton;
+}
diff --git a/Engine/Source/Runtime/Core/Private/Misc/MemStack.cpp b/Engine/Source/Runtime/Core/Private/Misc/MemStack.cpp
--- a/Engine/Source/Runtime/Core/Private/Misc/MemStack.cpp
+++ b/Engine/Source/Runtime/Core/Private/Misc/MemStack.cpp
@@ -6,10 +6,7 @@
 
 #include "CorePrivate.h"
 
-static struct FForceInitAtBootFMemStack : public TForceInitAtBoot<FMemStack>
-{} FForceInitAtBootFMemStack;
-
-template<> uint32 FThreadSingleton<FMemStack>::TlsSlot = 0;
+DECLARE_THREAD_SINGLETON( FMemStack );
 
 DECLARE_MEMORY_STAT(TEXT("MemStack Allocated (all threads)"), STAT_MemStackAllocated,STATGROUP_Memory);
 DECLARE_MEMORY_STAT(TEXT("MemStack Used (all threads)"), STAT_MemStackUsed,STATGROUP_Memory);
diff --git a/Engine/Source/Runtime/Core/Private/Stats/Stats2.cpp b/Engine/Source/Runtime/Core/Private/Stats/Stats2.cpp
--- a/Engine/Source/Runtime/Core/Private/Stats/Stats2.cpp
+++ b/Engine/Source/Runtime/Core/Private/Stats/Stats2.cpp
@@ -7,13 +7,12 @@
 
 #include "CorePrivate.h"
 
-template<> uint32 FThreadSingleton<FThreadIdleStats>::TlsSlot = 0;
+DECLARE_THREAD_SINGLETON( FThreadIdleStats );
 
 struct FStats2Globals
 {
 	static void Get()
 	{
-		FThreadIdleStats::Get();
 #if	STATS
 		FStartupMessages::Get();
 		IStatGroupEnableManager::Get();
diff --git a/Engine/Source/Runtime/Core/Public/HAL/ThreadingBase.h b/Engine/Source/Runtime/Core/Public/HAL/ThreadingBase.h
--- a/Engine/Source/Runtime/Core/Public/HAL/ThreadingBase.h
+++ b/Engine/Source/Runtime/Core/Public/HAL/ThreadingBase.h
@@ -835,56 +835,81 @@ extern CORE_API bool IsInGameThread ();
 /** @return True if called from the slate thread, and not merely a thread calling slate functions. */
 extern CORE_API bool IsInSlateThread ();
 
-/**
-* This a special version of singleton. It means that there is created only one instance for each thread.
-* Calling Get() method is thread-safe, but first call should be done on the game thread. 
-* @see TForceInitAtBoot usage.
-*/
-template< class T >
-class CORE_API FThreadSingleton
+
+/** Minimal base class for the thread singleton. */
+struct FThreadSingleton
 {
-	/** TLS slot that holds a FThreadSingleton. */
-	static uint32 TlsSlot;
+	typedef FThreadSingleton* (*TCreateSingletonFuncPtr)();
+	typedef void (FThreadSingleton::*TDestroySingletonFuncPtr)();
 
-protected:
-	FThreadSingleton()
-		: ThreadId(FPlatformTLS::GetCurrentThreadId())
+	CORE_API virtual ~FThreadSingleton()
 	{}
 
-	/** Thread ID this singleton was created on. */
-	const uint32 ThreadId;
+	virtual void DeleteThis() = 0;
+};
 
+/** Thread singleton initializer. */
+class FThreadSingletonInitializer
+{
 public:
 	/**
-	 * @return an instance of a singleton for the current thread.
+	* @return an instance of a singleton for the current thread.
+	*/
+	static CORE_API FThreadSingleton* Get( const FThreadSingleton::TCreateSingletonFuncPtr CreateFunc, const FThreadSingleton::TDestroySingletonFuncPtr DestroyFunc, uint32& TlsSlot );
+};
+
+/**
+ * This a special version of singleton. It means that there is created only one instance for each thread.
+ * Calling Get() method is thread-safe, but first call should be done on the game thread.
+ * @see DECLARE_THREAD_SINGLETON usage.
+ */
+template < class T >
+class TThreadSingleton : public FThreadSingleton
+{
+	/**
+	 * @return TLS slot that holds a TThreadSingleton.
 	 */
-	static T& Get()
+	static uint32& GetTlsSlot()
 	{
-		if( !TlsSlot )
-		{
-			check(IsInGameThread());
-			TlsSlot = FPlatformTLS::AllocTlsSlot();
-		}
-		T* ThreadSingleton = (T*)FPlatformTLS::GetTlsValue(TlsSlot);
-		if( !ThreadSingleton )
-		{
-			const uint32 ThreadId = FPlatformTLS::GetCurrentThreadId();
-			ThreadSingleton = new T();
-			FRunnableThread::GetThreadRegistry().Lock();
-			FRunnableThread* RunnableThread = FRunnableThread::GetThreadRegistry().GetThread(ThreadId);
-			if( RunnableThread )
-			{
-				RunnableThread->OnThreadDestroyed().AddRaw( ThreadSingleton, &FThreadSingleton<T>::Shutdown );
-			}
-			FRunnableThread::GetThreadRegistry().Unlock();
-			FPlatformTLS::SetTlsValue(TlsSlot, ThreadSingleton);
-		}
-		return *ThreadSingleton;
+		static uint32 TlsSlot = 0;
+		return TlsSlot;
+	}
+
+protected:
+	/** Default constructor. */
+	TThreadSingleton()
+		: ThreadId( FPlatformTLS::GetCurrentThreadId() )
+	{}
+
+	/**
+	 * @return a new instance of the thread singleton.
+	 */
+	static FThreadSingleton* Create()
+	{
+		return new T();
 	}
 
-	void Shutdown()
+	/**
+	 *	Deletes this instance of the thread singleton. 
+	 */
+	virtual void DeleteThis()
 	{
 		T* This = (T*)this;
 		delete This;
 	}
+
+	/** Thread ID of this thread singleton. */
+	const uint32 ThreadId;
+
+public:
+	/**
+	 *	@return an instance of a singleton for the current thread.
+	 */
+	FORCEINLINE static T& Get()
+	{
+		return *(T*)FThreadSingletonInitializer::Get( &T::Create, &FThreadSingleton::DeleteThis, T::GetTlsSlot() );
+	}
 };
+
+#define DECLARE_THREAD_SINGLETON(ClassType) \
+	static auto& GForceInitAtBoot_ThreadSingletonInitializer_##ClassType = ClassType::Get();
\ No newline at end of file
diff --git a/Engine/Source/Runtime/Core/Public/Misc/MemStack.h b/Engine/Source/Runtime/Core/Public/Misc/MemStack.h
--- a/Engine/Source/Runtime/Core/Public/Misc/MemStack.h
+++ b/Engine/Source/Runtime/Core/Public/Misc/MemStack.h
@@ -19,7 +19,7 @@ enum EMemOned   {MEM_Oned  =1};
  * Items are allocated via PushBytes() or the specialized operator new()s.
  * Items are freed en masse by using FMemMark to Pop() them.
  **/
-class CORE_API FMemStack : public FThreadSingleton<FMemStack>
+class CORE_API FMemStack : public TThreadSingleton<FMemStack>
 {
 public:
 
diff --git a/Engine/Source/Runtime/Core/Public/Stats/Stats2.h b/Engine/Source/Runtime/Core/Public/Stats/Stats2.h
--- a/Engine/Source/Runtime/Core/Public/Stats/Stats2.h
+++ b/Engine/Source/Runtime/Core/Public/Stats/Stats2.h
@@ -8,9 +8,9 @@
 /**
 * This is thread-private information about the thread idle stats, which we always collect, even in final builds
 */
-class CORE_API FThreadIdleStats : public FThreadSingleton<FThreadIdleStats>
+class FThreadIdleStats : public TThreadSingleton<FThreadIdleStats>
 {
-	friend class FThreadSingleton<FThreadIdleStats>;
+	friend class TThreadSingleton<FThreadIdleStats>;
 
 	FThreadIdleStats()
 		: Waits(0)
diff --git a/Engine/Source/Runtime/Launch/Private/Android/AndroidJNI.cpp b/Engine/Source/Runtime/Launch/Private/Android/AndroidJNI.cpp
--- a/Engine/Source/Runtime/Launch/Private/Android/AndroidJNI.cpp
+++ b/Engine/Source/Runtime/Launch/Private/Android/AndroidJNI.cpp
@@ -19,7 +19,7 @@ extern bool GOBBinAPK;
 // FJNIHelper
 
 // Caches access to the environment, attached to the current thread
-class FJNIHelper : public FThreadSingleton<FJNIHelper>
+class FJNIHelper : public TThreadSingleton<FJNIHelper>
 {
 public:
 	static JNIEnv* GetEnvironment()
@@ -31,7 +31,7 @@ private:
 	JNIEnv* CachedEnv = NULL;
 
 private:
-	friend class FThreadSingleton<FJNIHelper>;
+	friend class TThreadSingleton<FJNIHelper>;
 
 	FJNIHelper()
 		: CachedEnv(nullptr)
@@ -61,7 +61,7 @@ private:
 	}
 };
 
-template<> uint32 FThreadSingleton<FJNIHelper>::TlsSlot = 0;
+template<> uint32 TThreadSingleton<FJNIHelper>::TlsSlot = 0;
 
 JNIEnv* GetJavaEnv(bool bRequireGlobalThis)
 {
