Removing all remaining references to 'BasedOn' (ini)

diff --git a/Engine/Programs/UnrealFrontend/Config/DefaultEngine.ini b/Engine/Programs/UnrealFrontend/Config/DefaultEngine.ini
--- a/Engine/Programs/UnrealFrontend/Config/DefaultEngine.ini
+++ b/Engine/Programs/UnrealFrontend/Config/DefaultEngine.ini
@@ -1,6 +1,3 @@
-[Configuration]
-BasedOn=..\Engine\Config\BaseEngine.ini
-
 [URL]
 GameName=UnrealFrontend
 GameNameShort=UFE
diff --git a/Engine/Programs/UnrealLightmass/Config/DefaultEngine.ini b/Engine/Programs/UnrealLightmass/Config/DefaultEngine.ini
--- a/Engine/Programs/UnrealLightmass/Config/DefaultEngine.ini
+++ b/Engine/Programs/UnrealLightmass/Config/DefaultEngine.ini
@@ -1,6 +1,3 @@
-[Configuration]
-BasedOn=..\Engine\Config\BaseEngine.ini
-
 [Plugins]
 -EnabledPlugins="ExampleDeviceProfileSelector"
 +ProgramEnabledPlugins="UdpMessaging"
diff --git a/Engine/Source/Runtime/Core/Public/Misc/ConfigCacheIni.h b/Engine/Source/Runtime/Core/Public/Misc/ConfigCacheIni.h
--- a/Engine/Source/Runtime/Core/Public/Misc/ConfigCacheIni.h
+++ b/Engine/Source/Runtime/Core/Public/Misc/ConfigCacheIni.h
@@ -466,8 +466,7 @@ public:
 	 * Loads and generates a destination ini file and adds it to GConfig:
 	 *   - Looking on commandline for override source/dest .ini filenames
 	 *   - Generating the name for the engine to refer to the ini
-	 *   - Loading a source .ini file
-	 *   - Processing all BasedOn's
+	 *   - Loading a source .ini file hierarchy
 	 *   - Filling out an FConfigFile
 	 *   - Save the generated ini
 	 *   - Adds the FConfigFile to GConfig
