Remove unused TRANSFORM_TRACK_NAN_ISSUES from TransformVectorized.h. This is controlled by ENABLE_NAN_DIAGNOSTIC in UnrealMathUtility now.

diff --git a/Engine/Source/Runtime/Core/Public/Math/TransformVectorized.h b/Engine/Source/Runtime/Core/Public/Math/TransformVectorized.h
--- a/Engine/Source/Runtime/Core/Public/Math/TransformVectorized.h
+++ b/Engine/Source/Runtime/Core/Public/Math/TransformVectorized.h
@@ -2,7 +2,6 @@
 
 
 #pragma once
-#define TRANSFORM_TRACK_NAN_ISSUES (0 && !UE_BUILD_SHIPPING && !UE_BUILD_TEST)
 
 /** This returns Quaternion Inverse of X **/
 #define MAKE_QUATINV_VECTORREGISTER(X) VectorMultiply(GlobalVectorConstants::QINV_SIGN_MASK, X)
