Core - Fixed a crash introduces by the recent thread singleton change, added missing CORE_API to thread idle stats

diff --git a/Engine/Source/Runtime/Core/Public/Stats/Stats2.h b/Engine/Source/Runtime/Core/Public/Stats/Stats2.h
--- a/Engine/Source/Runtime/Core/Public/Stats/Stats2.h
+++ b/Engine/Source/Runtime/Core/Public/Stats/Stats2.h
@@ -8,7 +8,7 @@
 /**
 * This is thread-private information about the thread idle stats, which we always collect, even in final builds
 */
-class FThreadIdleStats : public TThreadSingleton<FThreadIdleStats>
+class CORE_API FThreadIdleStats : public TThreadSingleton<FThreadIdleStats>
 {
 	friend class TThreadSingleton<FThreadIdleStats>;
 
