Add explicit copy constructor to TPair to make the error messages for copying a pair containing a noncopyable type show the caller's code and not just "This diagnostic occurred in the compiler generated function 'TPair<KeyType,ValueType>::TPair(const TPair<KeyType,ValueType> &)'" (VC++)

diff --git a/Engine/Source/Runtime/Core/Public/Containers/Map.h b/Engine/Source/Runtime/Core/Public/Containers/Map.h
--- a/Engine/Source/Runtime/Core/Public/Containers/Map.h
+++ b/Engine/Source/Runtime/Core/Public/Containers/Map.h
@@ -77,6 +77,22 @@ public:
 		// the initializer is itself an rvalue reference.
 	}
 
+	// Explicit copy constructor to make the error messages for copying a pair containing a noncopyable type show the caller's code and not only
+	// "This diagnostic occurred in the compiler generated function 'TPair<KeyType,ValueType>::TPair(const TPair<KeyType,ValueType> &)'" (VC++)
+	FORCEINLINE TPair(const TPair& InInitializer)
+		: Key(InInitializer.Key)
+		, Value(InInitializer.Value)
+	{
+	}
+
+	FORCEINLINE TPair& operator=(const TPair& Other)
+	{
+		Key = Other.Key;
+		Value = Other.Value;
+
+		return *this;
+	}
+
 	/** Serializer. */
 	FORCEINLINE friend FArchive& operator<<(FArchive& Ar,TPair& Pair)
 	{
