Add FBox::MoveTo
https://github.com/EpicGames/UnrealEngine/pull/188 ([REDACTED_PROFILE_URL])

diff --git a/Engine/Source/Runtime/Core/Public/Math/Box.h b/Engine/Source/Runtime/Core/Public/Math/Box.h
--- a/Engine/Source/Runtime/Core/Public/Math/Box.h
+++ b/Engine/Source/Runtime/Core/Public/Math/Box.h
@@ -189,6 +189,17 @@ public:
 		return FBox(Min + Offset, Max + Offset);
 	}
 
+	/** 
+	 * Moves center of bounding box to new destination
+	 *
+	 * @param   Destination point to move center of box to
+	 * @return  new moved bounding box 
+	 */
+	FBox MoveTo( const FVector& Destination ) const
+	{
+		const FVector Offset = Destination - GetCenter();
+		return FBox(Min + Offset, Max + Offset);
+	}
 
 	/**
 	 * Gets the box's center point.
