Transient FTexts can be regenerated when cultures change.

AddedFTextHistory (and child classes) to manage the history of an FText's creation.

#ttp 243860 - L10N: FText should efficiently maintain a list of it's creation parameters making it possible for traditionally transient text to be safely serialized

diff --git a/Engine/Source/Editor/UnrealEd/Private/Commandlets/GatherTextFromAssetsCommandlet.cpp b/Engine/Source/Editor/UnrealEd/Private/Commandlets/GatherTextFromAssetsCommandlet.cpp
--- a/Engine/Source/Editor/UnrealEd/Private/Commandlets/GatherTextFromAssetsCommandlet.cpp
+++ b/Engine/Source/Editor/UnrealEd/Private/Commandlets/GatherTextFromAssetsCommandlet.cpp
@@ -515,22 +515,26 @@ bool UGatherTextFromAssetsCommandlet::ProcessTextProperty(UTextProperty* InTextP
 	{
 		FConflictTracker::FEntry NewEntry;
 		NewEntry.ObjectPath = Object->GetPathName();
-		NewEntry.SourceString = Data->SourceString;
+		NewEntry.SourceString = Data->GetSourceString();
 		NewEntry.Status = EAssetTextGatherStatus::None;
 
 		// Fix missing key if broken and allowed.
-		if( !( Data->Key.IsValid() ) || Data->Key->IsEmpty() )
+		TSharedPtr< FString > Namespace;
+		TSharedPtr< FString > Key;
+		FTextLocalizationManager::Get().FindKeyNamespaceFromDisplayString(Data->DisplayString, Namespace, Key);
+
+		if( !( Key.IsValid() ) || Key->IsEmpty() )
 		{
 			// Key fix.
 			if (bFixBroken)
 			{
 				// Create key if needed.
-				if( !( Data->Key.IsValid() ) )
+				if( !( Key.IsValid() ) )
 				{
-					Data->Key = MakeShareable( new FString() );
+					Key = MakeShareable( new FString() );
 				}
 				// Generate new GUID for key.
-				*(Data->Key) = FGuid::NewGuid().ToString();
+				*(Key) = FGuid::NewGuid().ToString();
 
 				// Fixed.
 				NewEntry.Status = EAssetTextGatherStatus::MissingKey_Resolved;
@@ -543,29 +547,29 @@ bool UGatherTextFromAssetsCommandlet::ProcessTextProperty(UTextProperty* InTextP
 		}
 
 		// Must have valid key.
-		if( Data->Key.IsValid() && !( Data->Key->IsEmpty() ) )
+		if( Key.IsValid() && !( Key->IsEmpty() ) )
 		{
 			FContext SearchContext;
-			SearchContext.Key = Data->Key.IsValid() ? *Data->Key : TEXT("");
+			SearchContext.Key = Key.IsValid() ? *Key : TEXT("");
 
 			// Find existing entry from manifest or manifest dependencies.
-			TSharedPtr< FManifestEntry > ExistingEntry = ManifestInfo->GetManifest()->FindEntryByContext( Data->Namespace.IsValid() ? *(Data->Namespace) : TEXT(""), SearchContext );
+			TSharedPtr< FManifestEntry > ExistingEntry = ManifestInfo->GetManifest()->FindEntryByContext( Namespace.IsValid() ? *Namespace : TEXT(""), SearchContext );
 			if( !ExistingEntry.IsValid() )
 			{
 				FString FileInfo;
-				ExistingEntry = ManifestInfo->FindDependencyEntrybyContext( Data->Namespace.IsValid() ? *(Data->Namespace) : TEXT(""), SearchContext, FileInfo );
+				ExistingEntry = ManifestInfo->FindDependencyEntrybyContext( Namespace.IsValid() ? *Namespace : TEXT(""), SearchContext, FileInfo );
 			}
 
 			// Entry already exists, check for conflict.
 			if( ExistingEntry.IsValid() )
 			{
 				// Fix conflict if present and allowed.
-				if( ExistingEntry->Source.Text != ( Data->SourceString.IsValid() ? **(Data->SourceString) : TEXT("") ) )
+				if( ExistingEntry->Source.Text != ( Data->GetSourceString().IsValid() ? **(Data->GetSourceString()) : TEXT("") ) )
 				{
 					if (bFixBroken)
 					{
 						// Generate new GUID for key.
-						*(Data->Key) = FGuid::NewGuid().ToString();
+						*(Key) = FGuid::NewGuid().ToString();
 
 						// Fixed.
 						NewEntry.Status = EAssetTextGatherStatus::IdentityConflict_Resolved;
@@ -585,7 +589,7 @@ bool UGatherTextFromAssetsCommandlet::ProcessTextProperty(UTextProperty* InTextP
 			if( !( ExistingEntry.IsValid() ) )
 			{
 				// Check for valid string.
-				if( Data->SourceString.IsValid() && !( Data->SourceString->IsEmpty() ) )
+				if( Data->GetSourceString().IsValid() && !( Data->GetSourceString()->IsEmpty() ) )
 				{
 					FString SrcLocation = Object->GetPathName() + TEXT(".") + InTextProp->GetName();
 
@@ -600,8 +604,11 @@ bool UGatherTextFromAssetsCommandlet::ProcessTextProperty(UTextProperty* InTextP
 							{
 								UTextProperty* TextProp	=  Cast<UTextProperty>( *(PropIt) );
 								FText* DataCDO = TextProp->ContainerPtrToValuePtr<FText>( CDO );
+								TSharedPtr< FString > NamespaceCDO;
+								TSharedPtr< FString > KeyCDO;
+								FTextLocalizationManager::Get().FindKeyNamespaceFromDisplayString(DataCDO->DisplayString, NamespaceCDO, KeyCDO);
 
-								if( DataCDO->Key == Data->Key || ( DataCDO->Key.Get() && Data->Key.Get() && ( *(DataCDO->Key) == *(Data->Key) ) ) )
+								if( KeyCDO == Key || ( KeyCDO.Get() && Key.Get() && ( *KeyCDO == *Key ) ) )
 								{
 									SrcLocation = CDO->GetPathName() + TEXT(".") + TextProp->GetName();
 									break;
@@ -611,18 +618,18 @@ bool UGatherTextFromAssetsCommandlet::ProcessTextProperty(UTextProperty* InTextP
 					}
 
 					FContext Context;
-					Context.Key = Data->Key.IsValid() ? *Data->Key : TEXT("");
+					Context.Key = Key.IsValid() ? *Key : TEXT("");
 					Context.SourceLocation = SrcLocation;
 
 					FString EntryDescription = FString::Printf( TEXT("In %s"), *Object->GetFullName());
-					ManifestInfo->AddEntry(EntryDescription, Data->Namespace.Get() ? *Data->Namespace : TEXT(""), Data->SourceString.Get() ? *(Data->SourceString) : TEXT(""), Context );
+					ManifestInfo->AddEntry(EntryDescription, Namespace.Get() ? *Namespace : TEXT(""), Data->GetSourceString().Get() ? *(Data->GetSourceString()) : TEXT(""), Context );
 				}
 			}
 		}
 
 		// Add to conflict tracker.
-		FConflictTracker::FKeyTable& KeyTable = ConflictTracker.Namespaces.FindOrAdd( Data->Namespace.IsValid() ? *(Data->Namespace) : TEXT("") );
-		FConflictTracker::FEntryArray& EntryArray = KeyTable.FindOrAdd( Data->Key.IsValid() ? *(Data->Key) : TEXT("") );
+		FConflictTracker::FKeyTable& KeyTable = ConflictTracker.Namespaces.FindOrAdd( Namespace.IsValid() ? *Namespace : TEXT("") );
+		FConflictTracker::FEntryArray& EntryArray = KeyTable.FindOrAdd( Key.IsValid() ? *Key : TEXT("") );
 		EntryArray.Add(NewEntry);
 
 		OutFixed = (NewEntry.Status == EAssetTextGatherStatus::MissingKey_Resolved || NewEntry.Status == EAssetTextGatherStatus::IdentityConflict_Resolved);
diff --git a/Engine/Source/Runtime/Core/Private/Internationalization/ICUText.cpp b/Engine/Source/Runtime/Core/Private/Internationalization/ICUText.cpp
--- a/Engine/Source/Runtime/Core/Private/Internationalization/ICUText.cpp
+++ b/Engine/Source/Runtime/Core/Private/Internationalization/ICUText.cpp
@@ -4,6 +4,7 @@
 
 #if UE_ENABLE_ICU
 #include "Text.h"
+#include "TextHistory.h"
 #include "ICUCulture.h"
 #include <unicode/coll.h>
 #include <unicode/sortkey.h>
@@ -29,7 +30,10 @@ FText FText::AsDate(const FDateTime::FDate& Date, const EDateTimeStyle::Type Dat
 	FString NativeString;
 	ICUUtilities::Convert(FormattedString, NativeString);
 
-	return FText::CreateChronologicalText(NativeString);
+	FText ResultText = FText::CreateChronologicalText(NativeString);
+	ResultText.History = MakeShareable(new FTextHistory_AsDate(Date, DateStyle, TargetCulture));
+
+	return ResultText;
 }
 
 FText FText::AsTime(const FDateTime::FTime& Time, const EDateTimeStyle::Type TimeStyle, const FString& TimeZone, const TSharedPtr<FCulture>& TargetCulture)
@@ -50,7 +54,10 @@ FText FText::AsTime(const FDateTime::FTime& Time, const EDateTimeStyle::Type Tim
 	FString NativeString;
 	ICUUtilities::Convert(FormattedString, NativeString);
 
-	return FText::CreateChronologicalText(NativeString);
+	FText ResultText = FText::CreateChronologicalText(NativeString);
+	ResultText.History = MakeShareable(new FTextHistory_AsTime(Time, TimeStyle, TimeZone, TargetCulture));
+
+	return ResultText;
 }
 
 FText FText::AsTime(const FTimespan& Time, const TSharedPtr<FCulture>& TargetCulture)
@@ -94,7 +101,10 @@ FText FText::AsDateTime(const FDateTime& DateTime, const EDateTimeStyle::Type Da
 	FString NativeString;
 	ICUUtilities::Convert(FormattedString, NativeString);
 
-	return FText::CreateChronologicalText(NativeString);
+	FText ResultText = FText::CreateChronologicalText(NativeString);
+	ResultText.History = MakeShareable(new FTextHistory_AsDateTime(DateTime, DateStyle, TimeStyle, TimeZone, TargetCulture));
+
+	return ResultText;
 }
 
 FText FText::AsMemory(SIZE_T NumBytes, const FNumberFormattingOptions* const Options, const TSharedPtr<FCulture>& TargetCulture)
@@ -250,7 +260,7 @@ void FText::GetFormatPatternParameters(const FText& Pattern, TArray<FString>& Pa
 	}
 }
 
-FText FText::Format(const FText& Pattern, const FFormatNamedArguments& Arguments)
+FText FText::FormatInternal(const FText& Pattern, const FFormatNamedArguments& Arguments, bool bInRebuildText)
 {
 	checkf(FInternationalization::Get().IsInitialized() == true, TEXT("FInternationalization is not initialized. An FText formatting method was likely used in static object initialization - this is not supported."));
 	//SCOPE_CYCLE_COUNTER( STAT_TextFormat );
@@ -295,6 +305,12 @@ FText FText::Format(const FText& Pattern, const FFormatNamedArguments& Arguments
 			break;
 		case EFormatArgumentType::Text:
 			{
+				// When doing a rebuild, all FText arguments need to be rebuilt during the Format
+				if(bInRebuildText)
+				{
+					ArgumentValue.TextValue->Rebuild();
+				}
+
 				FString NativeStringValue = ArgumentValue.TextValue->ToString();
 				icu::UnicodeString ICUStringValue;
 				ICUUtilities::Convert(NativeStringValue, ICUStringValue, false);
@@ -326,16 +342,19 @@ FText FText::Format(const FText& Pattern, const FFormatNamedArguments& Arguments
 	FString NativeResultString;
 	ICUUtilities::Convert(ICUResultString, NativeResultString);
 
-	return FText(NativeResultString);
+	FText ResultText(NativeResultString);
+	ResultText.History = MakeShareable(new FTextHistory_NamedFormat(Pattern, Arguments));
+
+	return ResultText;
 }
 
-FText FText::Format(const FText& Pattern, const FFormatOrderedArguments& Arguments)
+FText FText::FormatInternal(const FText& Pattern, const FFormatOrderedArguments& Arguments, bool bInRebuildText)
 {
 	checkf(FInternationalization::Get().IsInitialized() == true, TEXT("FInternationalization is not initialized. An FText formatting method was likely used in static object initialization - this is not supported."));
 	//SCOPE_CYCLE_COUNTER( STAT_TextFormat );
 
 	const bool bEnableErrorResults = ENABLE_TEXT_ERROR_CHECKING_RESULTS && bEnableErrorCheckingResults;
-	
+
 	TArray<icu::Formattable> ArgumentValues;
 	ArgumentValues.Reserve(Arguments.Num());
 
@@ -366,6 +385,12 @@ FText FText::Format(const FText& Pattern, const FFormatOrderedArguments& Argumen
 			break;
 		case EFormatArgumentType::Text:
 			{
+				// When doing a rebuild, all FText arguments need to be rebuilt during the Format
+				if(bInRebuildText)
+				{
+					ArgumentValue.TextValue->Rebuild();
+				}
+
 				FString NativeStringValue = ArgumentValue.TextValue->ToString();
 				icu::UnicodeString ICUStringValue;
 				ICUUtilities::Convert(NativeStringValue, ICUStringValue, false);
@@ -404,10 +429,11 @@ FText FText::Format(const FText& Pattern, const FFormatOrderedArguments& Argumen
 		ResultText.Flags = ResultText.Flags | ETextFlag::Transient;
 	}
 
+	ResultText.History = MakeShareable(new FTextHistory_OrderedFormat(Pattern, Arguments));
 	return ResultText;
 }
 
-FText FText::Format(const FText& Pattern, const TArray< FFormatArgumentData > InArguments)
+FText FText::FormatInternal(const FText& Pattern, const TArray< FFormatArgumentData > InArguments, bool bInRebuildText)
 {
 	checkf(FInternationalization::Get().IsInitialized() == true, TEXT("FInternationalization is not initialized. An FText formatting method was likely used in static object initialization - this is not supported."));
 	//SCOPE_CYCLE_COUNTER( STAT_TextFormat );
@@ -452,6 +478,12 @@ FText FText::Format(const FText& Pattern, const TArray< FFormatArgumentData > In
 			break;
 		case EFormatArgumentType::Text:
 			{
+				// When doing a rebuild, all FText arguments need to be rebuilt during the Format
+				if(bInRebuildText)
+				{
+					ArgumentValue.TextValue->Rebuild();
+				}
+
 				FString StringValue = ArgumentValue.TextValue->ToString();
 				icu::UnicodeString ICUStringValue;
 				ICUUtilities::Convert(StringValue, ICUStringValue, false);
@@ -490,7 +522,23 @@ FText FText::Format(const FText& Pattern, const TArray< FFormatArgumentData > In
 		ResultText.Flags = ResultText.Flags | ETextFlag::Transient;
 	}
 
+	ResultText.History = MakeShareable(new FTextHistory_ArgumentDataFormat(Pattern, InArguments));
 	return ResultText;
 }
 
+FText FText::Format(const FText& Pattern, const FFormatNamedArguments& Arguments)
+{
+	return FormatInternal(Pattern, Arguments, false);
+}
+
+FText FText::Format(const FText& Pattern, const FFormatOrderedArguments& Arguments)
+{
+	return FormatInternal(Pattern, Arguments, false);
+}
+
+FText FText::Format(const FText& Pattern, const TArray< struct FFormatArgumentData > InArguments)
+{
+	return FormatInternal(Pattern, InArguments, false);
+}
+
 #endif
diff --git a/Engine/Source/Runtime/Core/Private/Internationalization/ICUText.h b/Engine/Source/Runtime/Core/Private/Internationalization/ICUText.h
--- a/Engine/Source/Runtime/Core/Private/Internationalization/ICUText.h
+++ b/Engine/Source/Runtime/Core/Private/Internationalization/ICUText.h
@@ -24,7 +24,9 @@ FText FText::AsNumberTemplate(T1 Val, const FNumberFormattingOptions* const Opti
 	FString NativeString;
 	ICUUtilities::Convert(FormattedString, NativeString);
 
-	return FText::CreateNumericalText(NativeString);
+	FText ReturnText = FText::CreateNumericalText(NativeString);
+	ReturnText.History = MakeShareable(new FTextHistory_AsNumber(Val, Options, TargetCulture));
+	return ReturnText;
 }
 
 template<typename T1, typename T2>
@@ -42,7 +44,9 @@ FText FText::AsCurrencyTemplate(T1 Val, const FNumberFormattingOptions* const Op
 	FString NativeString;
 	ICUUtilities::Convert(FormattedString, NativeString);
 
-	return FText::CreateNumericalText(NativeString);
+	FText ReturnText = FText::CreateNumericalText(NativeString);
+	ReturnText.History = MakeShareable(new FTextHistory_AsCurrency(Val, Options, TargetCulture));
+	return ReturnText;
 }
 
 template<typename T1, typename T2>
@@ -60,5 +64,7 @@ FText FText::AsPercentTemplate(T1 Val, const FNumberFormattingOptions* const Opt
 	FString NativeString;
 	ICUUtilities::Convert(FormattedString, NativeString);
 
-	return FText::CreateNumericalText(NativeString);
+	FText ReturnText = FText::CreateNumericalText(NativeString);
+	ReturnText.History = MakeShareable(new FTextHistory_AsPercent(Val, Options, TargetCulture));
+	return ReturnText;
 }
\ No newline at end of file
diff --git a/Engine/Source/Runtime/Core/Private/Internationalization/Text.cpp b/Engine/Source/Runtime/Core/Private/Internationalization/Text.cpp
--- a/Engine/Source/Runtime/Core/Private/Internationalization/Text.cpp
+++ b/Engine/Source/Runtime/Core/Private/Internationalization/Text.cpp
@@ -8,25 +8,40 @@
 #include "LegacyText.h"
 #endif
 
+#include "TextHistory.h"
+
 //DEFINE_LOG_CATEGORY(LogText);
 
 //DEFINE_STAT(STAT_TextFormat);
 
 #define LOCTEXT_NAMESPACE "Core.Text"
 
+FArchive& operator<<( FArchive& Ar, FFormatArgumentData& Value )
+{
+	Ar << Value.ArgumentName;
+	Ar << Value.ArgumentValue;
+	return Ar;
+}
+
 const FString* FTextInspector::GetNamespace(const FText& Text)
 {
-	return Text.Namespace.Get();
+	TSharedPtr< FString > Namespace;
+	TSharedPtr< FString > Key;
+	FTextLocalizationManager::Get().FindKeyNamespaceFromDisplayString(Text.DisplayString, Namespace, Key);
+	return Namespace.Get();
 }
 
 const FString* FTextInspector::GetKey(const FText& Text)
 {
-	return Text.Key.Get();
+	TSharedPtr< FString > Namespace;
+	TSharedPtr< FString > Key;
+	FTextLocalizationManager::Get().FindKeyNamespaceFromDisplayString(Text.DisplayString, Namespace, Key);
+	return Key.Get();
 }
 
 const FString* FTextInspector::GetSourceString(const FText& Text)
 {
-	return Text.SourceString.Get();
+	return Text.GetSourceString().Get();
 }
 
 const FString& FTextInspector::GetDisplayString(const FText& Text)
@@ -39,6 +54,13 @@ int32 FTextInspector::GetFlags(const FText& Text)
 	return Text.Flags;
 }
 
+FFormatArgumentValue::FFormatArgumentValue()
+	: Type(EFormatArgumentType::Int)
+{
+	IntValue = 0;
+};
+
+
 FFormatArgumentValue::FFormatArgumentValue(const int Value)
 	: Type(EFormatArgumentType::Int)
 {
@@ -51,6 +73,18 @@ FFormatArgumentValue::FFormatArgumentValue(const unsigned int Value)
 	UIntValue = Value;
 }
 
+FFormatArgumentValue::FFormatArgumentValue(const int64 Value)
+	: Type(EFormatArgumentType::Int)
+{
+	IntValue = Value;
+}
+
+FFormatArgumentValue::FFormatArgumentValue(const uint64 Value)
+	: Type(EFormatArgumentType::UInt)
+{
+	UIntValue = Value;
+}
+
 FFormatArgumentValue::FFormatArgumentValue(const float Value)
 	: Type(EFormatArgumentType::Float)
 {
@@ -109,6 +143,48 @@ FFormatArgumentValue::~FFormatArgumentValue()
 	}
 }
 
+FArchive& operator<<( FArchive& Ar, FFormatArgumentValue& Value )
+{
+	int8 TypeAsInt8 = Value.Type;
+	Ar << TypeAsInt8;
+	Value.Type = (EFormatArgumentType::Type)TypeAsInt8;
+
+	switch(Value.Type)
+	{
+	case EFormatArgumentType::Double:
+		{
+			Ar << Value.DoubleValue;
+			break;
+		}
+	case EFormatArgumentType::Float:
+		{
+			Ar << Value.FloatValue;
+			break;
+		}
+	case EFormatArgumentType::Int:
+		{
+			Ar << Value.IntValue;
+			break;
+		}
+	case EFormatArgumentType::UInt:
+		{
+			Ar << Value.UIntValue;
+			break;
+		}
+	case EFormatArgumentType::Text:
+		{
+			if(Ar.IsLoading())
+			{
+				Value.TextValue = new FText;
+			}
+			Ar << *Value.TextValue;
+			break;
+		}
+	}
+	
+	return Ar;
+}
+
 // These default values have been duplicated to the KismetTextLibrary functions for Blueprints. Please replicate any changes there!
 FNumberFormattingOptions::FNumberFormattingOptions()
 	:   UseGrouping(true),
@@ -121,6 +197,22 @@ FNumberFormattingOptions::FNumberFormattingOptions()
 
 }
 
+FArchive& operator<<(FArchive& Ar, FNumberFormattingOptions& Value)
+{
+	Ar << Value.UseGrouping;
+
+	int8 RoundingModeInt8 = (int8)Value.RoundingMode;
+	Ar << RoundingModeInt8;
+	Value.RoundingMode = (ERoundingMode)RoundingModeInt8;
+
+	Ar << Value.MinimumIntegralDigits;
+	Ar << Value.MaximumIntegralDigits;
+	Ar << Value.MinimumFractionalDigits;
+	Ar << Value.MaximumFractionalDigits;
+
+	return Ar;
+}
+
 bool FText::bEnableErrorCheckingResults = ENABLE_TEXT_ERROR_CHECKING_RESULTS;
 bool FText::bSuppressWarnings = false;
 
@@ -136,43 +228,34 @@ const FText FText::UnEscapedCloseBraceOutsideOfArgumentBlock = LOCTEXT("Error_Un
 const FText FText::SerializationFailureError = LOCTEXT("Error_SerializationFailure", "ERR: Transient text cannot be serialized \"{0}\".");
 
 FText::FText()
-	:
-	SourceString( NULL ) ,
-	DisplayString( new FString() ) ,
-	Flags(0)
+	: DisplayString( new FString() )
+	, Flags(0)
+	, Revision(FTextLocalizationManager::Get().GetHeadCultureRevision())
 {
 }
 
 FText::FText( const FText& Source )
-	:
-	SourceString( Source.SourceString )
-	, DisplayString( Source.DisplayString )
-	, Namespace( Source.Namespace )
-	, Key( Source.Key )
+	: DisplayString( Source.DisplayString )
 	, Flags( Source.Flags )
+	, History( Source.History )
+	, Revision( Source.Revision )
 {
-
 }
 
 FText::FText( FString InSourceString )
-	:
-	SourceString( new FString( MoveTemp( InSourceString ) ) ) ,
-	DisplayString( SourceString.ToSharedRef() ) ,
-	Flags(0)
+	: DisplayString( new FString( MoveTemp(InSourceString) ))
+	, Flags(0)
+	, Revision(FTextLocalizationManager::Get().GetHeadCultureRevision())
 {
-	if (GIsEditor)
-	{
-		Key = MakeShareable( new FString(FGuid::NewGuid().ToString()) );
-	}
+	History = MakeShareable(new FTextHistory_Base(DisplayString));
 }
 
 FText::FText( FString InSourceString, FString InNamespace, FString InKey, int32 InFlags )
-	: SourceString( new FString( MoveTemp( InSourceString ) ) )
-	, DisplayString( FTextLocalizationManager::Get().GetString(InNamespace, InKey, SourceString.Get()) )
-	, Namespace( new FString( MoveTemp( InNamespace) ) )
-	, Key( new FString( MoveTemp( InKey ) ) )
+	: DisplayString( FTextLocalizationManager::Get().GetString(InNamespace, InKey, &InSourceString) )
 	, Flags(InFlags)
+	, Revision(FTextLocalizationManager::Get().GetHeadCultureRevision())
 {
+	History = MakeShareable(new FTextHistory_Base(InSourceString));
 }
 
 FText FText::TrimPreceding( const FText& InText )
@@ -344,81 +427,129 @@ CORE_API FArchive& operator<<( FArchive& Ar, FText& Value )
 	//When duplicating, the CDO is used as the template, then values for the instance are assigned.
 	//If we don't duplicate the string, the CDO and the instance are both pointing at the same thing.
 	//This would result in all subsequently duplicated objects stamping over formerly duplicated ones.
-	if(Ar.IsLoading())
+
+	// Older FText's stored their "SourceString", that is now stored in a history class so move it there
+	if(Ar.IsLoading() && Ar.UE4Ver() < VER_UE4_FTEXT_HISTORY)
 	{
-		Value.SourceString = MakeShareable( new FString() );
-		Ar << *(Value.SourceString);
-		if( Value.SourceString->IsEmpty() )
+		FString SourceStringToImplantIntoHistory;
+		Ar << SourceStringToImplantIntoHistory;
+
+		Value.History = MakeShareable(new FTextHistory_Base(SourceStringToImplantIntoHistory));
+
+		// Namespaces and keys are no longer stored in the FText, we need to read them in and discard
+		if( Ar.UE4Ver() >= VER_UE4_ADDED_NAMESPACE_AND_KEY_DATA_TO_FTEXT )
 		{
-			Value.SourceString.Reset();
+			TSharedPtr< FString > Namespace;
+			TSharedPtr< FString > Key;
+
+			Namespace = MakeShareable( new FString() );
+			Ar << *Namespace;
+
+			Key = MakeShareable( new FString() );
+			Ar << *Key;
+
+			// Get the DisplayString using the namespace, key, and source string.
+			const FString& NamespaceStr = Namespace.Get() ? *Namespace : TEXT("");
+			const FString& KeyStr = Key.Get() ? *Key : TEXT("");
+			Value.DisplayString = FTextLocalizationManager::Get().GetString(NamespaceStr, KeyStr, &SourceStringToImplantIntoHistory);
 		}
 	}
-	else if( Ar.IsSaving() )
+
+	if(Ar.IsPersistent())
 	{
-		if( Value.SourceString.IsValid() )
-		{
-			Ar << *(Value.SourceString);
-		}
-		else
-		{
-			FString EmptySourceString;
-			Ar << EmptySourceString;
-		}
+		Value.Flags &= ~(ETextFlag::ConvertedProperty); // Remove conversion flag before saving.
 	}
+	Ar << Value.Flags;
 
-	if( Ar.UE4Ver() >= VER_UE4_ADDED_NAMESPACE_AND_KEY_DATA_TO_FTEXT )
+	if( Ar.UE4Ver() >= VER_UE4_FTEXT_HISTORY )
 	{
-		if(Ar.IsLoading())
+		if(Ar.IsSaving())
 		{
-			Value.Namespace = MakeShareable( new FString() );
-			Ar << *Value.Namespace;
-			if( Value.Namespace->IsEmpty() )
+			// If there is no history, mark it, otherwise the history will serialize it's type
+			if(!Value.History.IsValid())
 			{
-				Value.Namespace.Reset();
-			}
-
-			Value.Key = MakeShareable( new FString() );
-			Ar << *Value.Key;
-			if( Value.Key->IsEmpty() )
-			{
-				Value.Key.Reset();
+				int8 NoHistory = INDEX_NONE;
+				Ar << NoHistory;
 			}
 		}
-		else if(Ar.IsSaving())
+		else
 		{
-			if( Value.Namespace.IsValid() )
-			{
-				Ar << *(Value.Namespace);
-			}
-			else
-			{
-				FString Empty;
-				Ar << Empty;
-			}
+			// The type is serialized during the serialization of the history, during deserialization we need to deserialize it and create the correct history
+			int8 HistoryType = INDEX_NONE;
+			Ar << HistoryType;
 
-			if( Value.Key.IsValid() )
+			// Create the history class based on the serialized type
+			switch((ETextHistoryType)HistoryType)
 			{
-				Ar << *(Value.Key);
-			}
-			else
-			{
-				FString Empty;
-				Ar << Empty;
+			case ETextHistoryType::Base:
+				{
+					Value.History = MakeShareable(new FTextHistory_Base);
+					break;
+				}
+			case ETextHistoryType::NamedFormat:
+				{
+					Value.History = MakeShareable(new FTextHistory_NamedFormat);
+					break;
+				}
+			case ETextHistoryType::OrderedFormat:
+				{
+					Value.History = MakeShareable(new FTextHistory_OrderedFormat);
+					break;
+				}
+			case ETextHistoryType::ArgumentFormat:
+				{
+					Value.History = MakeShareable(new FTextHistory_ArgumentDataFormat);
+					break;
+				}
+			case ETextHistoryType::AsNumber:
+				{
+					Value.History = MakeShareable(new FTextHistory_AsNumber);
+					break;
+				}
+			case ETextHistoryType::AsPercent:
+				{
+					Value.History = MakeShareable(new FTextHistory_AsPercent);
+					break;
+				}
+			case ETextHistoryType::AsCurrency:
+				{
+					Value.History = MakeShareable(new FTextHistory_AsCurrency);
+					break;
+				}
+			case ETextHistoryType::AsDate:
+				{
+					Value.History = MakeShareable(new FTextHistory_AsDate);
+					break;
+				}
+			case ETextHistoryType::AsTime:
+				{
+					Value.History = MakeShareable(new FTextHistory_AsTime);
+					break;
+				}
+			case ETextHistoryType::AsDateTime:
+				{
+					Value.History = MakeShareable(new FTextHistory_AsDateTime);
+					break;
+				}
+			default:
+				{
+					Value.History.Reset();
+				}
 			}
 		}
 
-		if(Ar.IsPersistent())
+		if(Value.History.IsValid())
 		{
-			Value.Flags &= ~(ETextFlag::ConvertedProperty); // Remove conversion flag before saving.
+			Value.History->Serialize(Ar);
+			Value.History->SerializeForDisplayString(Ar, Value.DisplayString);
 		}
-		Ar << Value.Flags;
 	}
 
 	if(Ar.IsLoading())
 	{
-		const FString& Namespace = Value.Namespace.Get() ? *(Value.Namespace) : TEXT("");
-		const FString& Key = Value.Key.Get() ? *(Value.Key) : TEXT("");
-		Value.DisplayString = FTextLocalizationManager::Get().GetString(Namespace, Key, Value.SourceString.Get());
+		// Keeps the revision up-to-date
+		Value.Revision = INDEX_NONE;
+		Value.Rebuild();
 	}
 
 	// Validate
@@ -433,7 +564,7 @@ CORE_API FArchive& operator<<( FArchive& Ar, FText& Value )
 #if WITH_EDITOR
 FText FText::ChangeKey( FString Namespace, FString Key, const FText& Text )
 {
-	return FText( *Text.SourceString, MoveTemp( Namespace ), MoveTemp( Key ) );
+	return FText( *Text.History->GetSourceString().Get(), MoveTemp( Namespace ), MoveTemp( Key ) );
 }
 #endif
 
@@ -484,9 +615,34 @@ FText FText::FromString( FString String )
 
 const FString& FText::ToString() const
 {
+	Rebuild();
+
 	return DisplayString.Get();
 }
 
+void FText::Rebuild() const
+{
+	if(History.IsValid() && History->IsOutOfDate(Revision))
+	{
+		int32& RevisionNonCast = (int32)Revision;
+		RevisionNonCast = FTextLocalizationManager::Get().GetHeadCultureRevision();
+		
+		DisplayString.Get() = History->ToText().DisplayString.Get();
+	}
+}
 
+TSharedPtr<FString> FText::GetSourceString() const
+{
+	if(History.IsValid())
+	{
+		TSharedPtr<FString> SourceString = History->GetSourceString();
+		if(SourceString.IsValid())
+		{
+			return History->GetSourceString();
+		}
+	}
+
+	return DisplayString;
+}
 
 #undef LOCTEXT_NAMESPACE
diff --git a/Engine/Source/Runtime/Core/Private/Internationalization/TextHistory.cpp b/Engine/Source/Runtime/Core/Private/Internationalization/TextHistory.cpp
new file mode 100644
--- /dev/null
+++ b/Engine/Source/Runtime/Core/Private/Internationalization/TextHistory.cpp
@@ -0,0 +1,551 @@
+// Copyright 1998-2014 Epic Games, Inc. All Rights Reserved.
+
+#include "CorePrivate.h"
+
+#include "TextHistory.h"
+
+///////////////////////////////////////
+// FTextHistory
+
+/** Base class for all FText history types */
+bool FTextHistory::IsOutOfDate(int32 InRevision)
+{
+	return InRevision < FTextLocalizationManager::Get().GetHeadCultureRevision();
+}
+
+TSharedPtr< FString > FTextHistory::GetSourceString()
+{
+	return NULL;
+}
+
+///////////////////////////////////////
+// FTextHistory_NamedFormat
+
+FTextHistory_Base::FTextHistory_Base(FString InSourceString)
+	: SourceString(new FString( MoveTemp( InSourceString ) ))
+{
+
+}
+
+FTextHistory_Base::FTextHistory_Base(TSharedPtr< FString > InSourceString)
+	: SourceString(InSourceString)
+{
+
+}
+
+FText FTextHistory_Base::ToText() const
+{
+	// This should never be called
+	check(0);
+	return FText::GetEmpty();
+}
+
+TSharedPtr< FString > FTextHistory_Base::GetSourceString()
+{
+	return SourceString;
+}
+
+void FTextHistory_Base::Serialize( FArchive& Ar )
+{
+	// If I serialize out the Namespace and Key HERE, then we can load it up.
+	if(Ar.IsSaving())
+	{
+		int8 HistoryType = (int8)ETextHistoryType::Base;
+		Ar << HistoryType;
+	}
+}
+
+void FTextHistory_Base::SerializeForDisplayString(FArchive& Ar, TSharedRef<FString>& InOutDisplayString)
+{
+	if(Ar.IsLoading())
+	{
+		FString Namespace;
+		FString Key;
+		FString SourceStringRaw;
+
+		Ar << Namespace;
+		Ar << Key;
+		Ar << SourceStringRaw;
+
+		// If there was a SourceString, store it in the member variable
+		if(!SourceStringRaw.IsEmpty())
+		{
+			SourceString = MakeShareable(new FString(SourceStringRaw));
+		}
+
+		// Using the deserialized namespace and key, find the DisplayString.
+		InOutDisplayString = FTextLocalizationManager::Get().GetString(Namespace, Key, SourceString.Get());
+	}
+	else if(Ar.IsSaving())
+	{
+		TSharedPtr< FString > Namespace;
+		TSharedPtr< FString > Key;
+
+		FTextLocalizationManager::Get().FindKeyNamespaceFromDisplayString(InOutDisplayString, Namespace, Key);
+
+		// Serialize the Namespace, or an empty string if there is none
+		if( Namespace.IsValid() )
+		{
+			Ar << *(Namespace);
+		}
+		else
+		{
+			FString Empty;
+			Ar << Empty;
+		}
+
+		// Serialize the Key, if there is not one and this is in the editor, generate a new key.
+		if( Key.IsValid() )
+		{
+			Ar << *(Key);
+		}
+		else
+		{
+			if(GIsEditor)
+			{
+				FString SerializedKey = FGuid::NewGuid().ToString();
+				Ar << SerializedKey;
+			}
+			else
+			{
+				FString Empty;
+				Ar << Empty;
+			}
+		}
+
+		// Serialize the SourceString, or an empty string if there is none
+		if( SourceString.IsValid() )
+		{
+			Ar << *(SourceString);
+		}
+		else
+		{
+			FString Empty;
+			Ar << Empty;
+		}
+	}
+}
+
+///////////////////////////////////////
+// FTextHistory_NamedFormat
+
+FTextHistory_NamedFormat::FTextHistory_NamedFormat(const FText& InSourceText, const FFormatNamedArguments& InArguments)
+	: SourceText(InSourceText)
+	, Arguments(InArguments)
+{
+
+}
+
+FText FTextHistory_NamedFormat::ToText() const
+{
+	return FText::FormatInternal(SourceText, Arguments, true);
+}
+
+void FTextHistory_NamedFormat::Serialize( FArchive& Ar )
+{
+	if(Ar.IsSaving())
+	{
+		int8 HistoryType = (int8)ETextHistoryType::NamedFormat;
+		Ar << HistoryType;
+	}
+
+	Ar << SourceText;
+	Ar << Arguments;
+}
+
+///////////////////////////////////////
+// FTextHistory_OrderedFormat
+
+FTextHistory_OrderedFormat::FTextHistory_OrderedFormat(const FText& InSourceText, const FFormatOrderedArguments& InArguments)
+	: SourceText(InSourceText)
+	, Arguments(InArguments)
+{
+
+}
+
+FText FTextHistory_OrderedFormat::ToText() const
+{
+	return FText::FormatInternal(SourceText, Arguments, true);
+}
+
+void FTextHistory_OrderedFormat::Serialize( FArchive& Ar )
+{
+	if(Ar.IsSaving())
+	{
+		int8 HistoryType = (int8)ETextHistoryType::OrderedFormat;
+		Ar << HistoryType;
+	}
+
+	Ar << SourceText;
+	Ar << Arguments;
+}
+
+///////////////////////////////////////
+// FTextHistory_ArgumentDataFormat
+
+FTextHistory_ArgumentDataFormat::FTextHistory_ArgumentDataFormat(const FText& InSourceText, const TArray< struct FFormatArgumentData >& InArguments)
+	: SourceText(InSourceText)
+	, Arguments(InArguments)
+{
+
+}
+
+FText FTextHistory_ArgumentDataFormat::ToText() const
+{
+	return FText::FormatInternal(SourceText, Arguments, true);
+}
+
+void FTextHistory_ArgumentDataFormat::Serialize( FArchive& Ar )
+{
+	if(Ar.IsSaving())
+	{
+		int8 HistoryType = (int8)ETextHistoryType::ArgumentFormat;
+		Ar << HistoryType;
+	}
+
+	Ar << SourceText;
+	Ar << Arguments;
+}
+
+///////////////////////////////////////
+// FTextHistory_FormatNumber
+
+FTextHistory_FormatNumber::FTextHistory_FormatNumber()
+	: FormatOptions(nullptr)
+{
+}
+
+FTextHistory_FormatNumber::FTextHistory_FormatNumber(const FFormatArgumentValue& InSourceValue, const FNumberFormattingOptions* const InFormatOptions, const TSharedPtr<FCulture> InTargetCulture)
+	: SourceValue(InSourceValue)
+	, FormatOptions(nullptr)
+	, TargetCulture(InTargetCulture)
+{
+	if(InFormatOptions)
+	{
+		FormatOptions = new FNumberFormattingOptions;
+		(*FormatOptions) = *InFormatOptions;
+	}
+}
+
+FTextHistory_FormatNumber::~FTextHistory_FormatNumber()
+{
+	delete FormatOptions;
+}
+
+void FTextHistory_FormatNumber::Serialize(FArchive& Ar)
+{
+	Ar << SourceValue;
+
+	bool bHasFormatOptions = FormatOptions != nullptr;
+	Ar << bHasFormatOptions;
+
+	if(bHasFormatOptions)
+	{
+		if(Ar.IsLoading())
+		{
+			FormatOptions = new FNumberFormattingOptions;
+		}
+		Ar << *FormatOptions;
+	}
+
+	if(Ar.IsSaving())
+	{
+		FString CultureName = TargetCulture.IsValid()? TargetCulture->GetName() : FString();
+		Ar << CultureName;
+	}
+	else if(Ar.IsLoading())
+	{
+		FString CultureName;
+		Ar << CultureName;
+
+		if(!CultureName.IsEmpty())
+		{
+			TargetCulture = FInternationalization::Get().GetCulture(CultureName);
+		}
+	}
+}
+
+///////////////////////////////////////
+// FTextHistory_AsNumber
+
+FTextHistory_AsNumber::FTextHistory_AsNumber(const FFormatArgumentValue& InSourceValue, const FNumberFormattingOptions* const InFormatOptions, const TSharedPtr<FCulture> InTargetCulture)
+	: FTextHistory_FormatNumber(InSourceValue, InFormatOptions, InTargetCulture)
+{
+}
+
+FText FTextHistory_AsNumber::ToText() const
+{
+	switch(SourceValue.Type)
+	{
+	case EFormatArgumentType::UInt:
+		{
+			return FText::AsNumber(SourceValue.UIntValue, FormatOptions, TargetCulture);
+		}
+	case EFormatArgumentType::Int:
+		{
+			return FText::AsNumber(SourceValue.IntValue, FormatOptions, TargetCulture);
+		}
+	case EFormatArgumentType::Float:
+		{
+			return FText::AsNumber(SourceValue.FloatValue, FormatOptions, TargetCulture);
+		}
+	case EFormatArgumentType::Double:
+		{
+			return FText::AsNumber(SourceValue.DoubleValue, FormatOptions, TargetCulture);
+		}
+	default:
+		{
+			// Should never reach this point
+			check(0);
+		}
+	}
+	return FText();
+}
+
+void FTextHistory_AsNumber::Serialize( FArchive& Ar )
+{
+	if(Ar.IsSaving())
+	{
+		int8 HistoryType = (int8)ETextHistoryType::AsNumber;
+		Ar << HistoryType;
+	}
+
+	FTextHistory_FormatNumber::Serialize(Ar);
+}
+
+///////////////////////////////////////
+// FTextHistory_AsPercent
+
+FTextHistory_AsPercent::FTextHistory_AsPercent(const FFormatArgumentValue& InSourceValue, const FNumberFormattingOptions* const InFormatOptions, const TSharedPtr<FCulture> InTargetCulture)
+	: FTextHistory_FormatNumber(InSourceValue, InFormatOptions, InTargetCulture)
+{
+}
+
+FText FTextHistory_AsPercent::ToText() const
+{
+	switch(SourceValue.Type)
+	{
+	case EFormatArgumentType::Float:
+		{
+			return FText::AsPercent(SourceValue.FloatValue, FormatOptions, TargetCulture);
+		}
+	case EFormatArgumentType::Double:
+		{
+			return FText::AsPercent(SourceValue.DoubleValue, FormatOptions, TargetCulture);
+		}
+	default:
+		{
+			// Should never reach this point
+			check(0);
+		}
+	}
+	return FText();
+}
+
+void FTextHistory_AsPercent::Serialize( FArchive& Ar )
+{
+	if(Ar.IsSaving())
+	{
+		int8 HistoryType = (int8)ETextHistoryType::AsPercent;
+		Ar << HistoryType;
+	}
+
+	FTextHistory_FormatNumber::Serialize(Ar);
+}
+
+///////////////////////////////////////
+// FTextHistory_AsCurrency
+
+FTextHistory_AsCurrency::FTextHistory_AsCurrency(const FFormatArgumentValue& InSourceValue, const FNumberFormattingOptions* const InFormatOptions, const TSharedPtr<FCulture> InTargetCulture)
+	: FTextHistory_FormatNumber(InSourceValue, InFormatOptions, InTargetCulture)
+{
+}
+
+FText FTextHistory_AsCurrency::ToText() const
+{
+	switch(SourceValue.Type)
+	{
+	case EFormatArgumentType::UInt:
+		{
+			return FText::AsCurrency(SourceValue.UIntValue, FormatOptions, TargetCulture);
+		}
+	case EFormatArgumentType::Int:
+		{
+			return FText::AsCurrency(SourceValue.IntValue, FormatOptions, TargetCulture);
+		}
+	case EFormatArgumentType::Float:
+		{
+			return FText::AsCurrency(SourceValue.FloatValue, FormatOptions, TargetCulture);
+		}
+	case EFormatArgumentType::Double:
+		{
+			return FText::AsCurrency(SourceValue.DoubleValue, FormatOptions, TargetCulture);
+		}
+	default:
+		{
+			// Should never reach this point
+			check(0);
+		}
+	}
+	return FText();
+}
+
+void FTextHistory_AsCurrency::Serialize( FArchive& Ar )
+{
+	if(Ar.IsSaving())
+	{
+		int8 HistoryType = (int8)ETextHistoryType::AsCurrency;
+		Ar << HistoryType;
+	}
+
+	FTextHistory_FormatNumber::Serialize(Ar);
+}
+
+///////////////////////////////////////
+// FTextHistory_AsDate
+
+FTextHistory_AsDate::FTextHistory_AsDate(const FDateTime::FDate& InSourceDate, const EDateTimeStyle::Type InDateStyle, const TSharedPtr<FCulture> InTargetCulture)
+	: SourceDate(InSourceDate)
+	, DateStyle(InDateStyle)
+	, TargetCulture(InTargetCulture)
+{
+}
+
+void FTextHistory_AsDate::Serialize(FArchive& Ar)
+{
+	if(Ar.IsSaving())
+	{
+		int8 HistoryType = (int8)ETextHistoryType::AsDate;
+		Ar << HistoryType;
+	}
+
+	Ar << SourceDate;
+
+	int8 DateStyleInt8 = (int8)DateStyle;
+	Ar << DateStyleInt8;
+	DateStyle = (EDateTimeStyle::Type)DateStyleInt8;
+
+	if(Ar.IsSaving())
+	{
+		FString CultureName = TargetCulture.IsValid()? TargetCulture->GetName() : FString();
+		Ar << CultureName;
+	}
+	else if(Ar.IsLoading())
+	{
+		FString CultureName;
+		Ar << CultureName;
+
+		if(!CultureName.IsEmpty())
+		{
+			TargetCulture = FInternationalization::Get().GetCulture(CultureName);
+		}
+	}
+}
+
+FText FTextHistory_AsDate::ToText() const
+{
+	return FText::AsDate(SourceDate, DateStyle, TargetCulture);
+}
+
+///////////////////////////////////////
+// FTextHistory_AsTime
+
+FTextHistory_AsTime::FTextHistory_AsTime(const FDateTime::FTime& InSourceTime, const EDateTimeStyle::Type InTimeStyle, const FString& InTimeZone, const TSharedPtr<FCulture> InTargetCulture)
+	: SourceTime(InSourceTime)
+	, TimeStyle(InTimeStyle)
+	, TimeZone(InTimeZone)
+	, TargetCulture(InTargetCulture)
+{
+}
+
+void FTextHistory_AsTime::Serialize(FArchive& Ar)
+{
+	if(Ar.IsSaving())
+	{
+		int8 HistoryType = (int8)ETextHistoryType::AsTime;
+		Ar << HistoryType;
+	}
+
+	Ar << SourceTime;
+
+	int8 TimeStyleInt8 = (int8)TimeStyle;
+	Ar << TimeStyleInt8;
+	TimeStyle = (EDateTimeStyle::Type)TimeStyleInt8;
+
+	Ar << TimeZone;
+
+	if(Ar.IsSaving())
+	{
+		FString CultureName = TargetCulture.IsValid()? TargetCulture->GetName() : FString();
+		Ar << CultureName;
+	}
+	else if(Ar.IsLoading())
+	{
+		FString CultureName;
+		Ar << CultureName;
+
+		if(!CultureName.IsEmpty())
+		{
+			TargetCulture = FInternationalization::Get().GetCulture(CultureName);
+		}
+	}
+}
+
+FText FTextHistory_AsTime::ToText() const
+{
+	return FText::AsTime(SourceTime, TimeStyle, TimeZone, TargetCulture);
+}
+
+///////////////////////////////////////
+// FTextHistory_AsDateTime
+
+FTextHistory_AsDateTime::FTextHistory_AsDateTime(const FDateTime& InSourceDateTime, const EDateTimeStyle::Type InDateStyle, const EDateTimeStyle::Type InTimeStyle, const FString& InTimeZone, const TSharedPtr<FCulture> InTargetCulture)
+	: SourceDateTime(InSourceDateTime)
+	, DateStyle(InDateStyle)
+	, TimeStyle(InTimeStyle)
+	, TimeZone(InTimeZone)
+	, TargetCulture(InTargetCulture)
+{
+}
+
+void FTextHistory_AsDateTime::Serialize(FArchive& Ar)
+{
+	if(Ar.IsSaving())
+	{
+		int8 HistoryType = (int8)ETextHistoryType::AsDateTime;
+		Ar << HistoryType;
+	}
+
+	Ar << SourceDateTime;
+
+	int8 DateStyleInt8 = (int8)DateStyle;
+	Ar << DateStyleInt8;
+	DateStyle = (EDateTimeStyle::Type)DateStyleInt8;
+
+	int8 TimeStyleInt8 = (int8)TimeStyle;
+	Ar << TimeStyleInt8;
+	TimeStyle = (EDateTimeStyle::Type)TimeStyleInt8;
+
+	Ar << TimeZone;
+
+	if(Ar.IsSaving())
+	{
+		FString CultureName = TargetCulture.IsValid()? TargetCulture->GetName() : FString();
+		Ar << CultureName;
+	}
+	else if(Ar.IsLoading())
+	{
+		FString CultureName;
+		Ar << CultureName;
+
+		if(!CultureName.IsEmpty())
+		{
+			TargetCulture = FInternationalization::Get().GetCulture(CultureName);
+		}
+	}
+}
+
+FText FTextHistory_AsDateTime::ToText() const
+{
+	return FText::AsDateTime(SourceDateTime, DateStyle, TimeStyle, TimeZone, TargetCulture);
+}
\ No newline at end of file
diff --git a/Engine/Source/Runtime/Core/Private/Internationalization/TextHistory.h b/Engine/Source/Runtime/Core/Private/Internationalization/TextHistory.h
new file mode 100644
--- /dev/null
+++ b/Engine/Source/Runtime/Core/Private/Internationalization/TextHistory.h
@@ -0,0 +1,247 @@
+// Copyright 1998-2014 Epic Games, Inc. All Rights Reserved.
+#pragma once
+
+enum class ETextHistoryType
+{
+	Base = 0,
+	NamedFormat,
+	OrderedFormat,
+	ArgumentFormat,
+	AsNumber,
+	AsPercent,
+	AsCurrency,
+	AsDate,
+	AsTime,
+	AsDateTime
+
+	// Add new enum types at the end only! They are serialized by index.
+};
+
+/** Base interface class for all FText history types */
+class CORE_API FTextHistory
+{
+public:
+	virtual ~FTextHistory() {};
+
+	/** Rebuilds the FText from the hierarchical history, the result should be in the current locale */
+	virtual FText ToText() const = 0;
+	
+	/** Serializes the history to/from an FArchive */
+	virtual void Serialize(FArchive& Ar) = 0;
+
+	/** Serializes data needed to get the FText's DisplayString, most history classes do not need to do anything here */
+	virtual void SerializeForDisplayString(FArchive& Ar, TSharedRef<FString>& InOutDisplayString) {};
+
+	/** Returns TRUE if the Revision is out of date */
+	virtual bool IsOutOfDate(int32 InRevision);
+
+	/** Returns the source string managed by the history (if any). */
+	virtual TSharedPtr< FString > GetSourceString();
+};
+
+/** No complexity to it, just holds the source string. */
+class CORE_API FTextHistory_Base : public FTextHistory
+{
+public:
+	FTextHistory_Base() {};
+	FTextHistory_Base(FString InSourceString);
+	FTextHistory_Base(TSharedPtr< FString > InSourceString);
+
+	// Begin FTextHistory interface
+	virtual FText ToText() const OVERRIDE;
+	virtual void Serialize(FArchive& Ar) OVERRIDE;
+	virtual void SerializeForDisplayString(FArchive& Ar, TSharedRef<FString>& InOutDisplayString) OVERRIDE;
+	virtual bool IsOutOfDate(int32 InRevision) OVERRIDE { return false; }
+	virtual TSharedPtr< FString > GetSourceString() OVERRIDE;
+	// End FTextHistory interface
+
+private:
+	/** The source string for an FText */
+	TSharedPtr<FString> SourceString;
+};
+
+/** Handles history for FText::Format when passing named arguments */
+class CORE_API FTextHistory_NamedFormat : public FTextHistory
+{
+public:
+	FTextHistory_NamedFormat() {};
+	FTextHistory_NamedFormat(const FText& InSourceText, const FFormatNamedArguments& InArguments);
+
+	// Begin FTextHistory interface
+	virtual FText ToText() const OVERRIDE;
+	virtual void Serialize(FArchive& Ar) OVERRIDE;
+	// End FTextHistory interface
+
+private:
+	/** The pattern used to format the text */
+	FText SourceText;
+	/** Arguments to replace in the pattern string */
+	FFormatNamedArguments Arguments;
+};
+
+/** Handles history for FText::Format when passing ordered arguments */
+class CORE_API FTextHistory_OrderedFormat : public FTextHistory
+{
+public:
+	FTextHistory_OrderedFormat() {};
+	FTextHistory_OrderedFormat(const FText& InSourceText, const FFormatOrderedArguments& InArguments);
+
+	// Begin FTextHistory interface
+	virtual FText ToText() const OVERRIDE;
+	virtual void Serialize(FArchive& Ar) OVERRIDE;
+	// End FTextHistory interface
+
+private:
+	/** The pattern used to format the text */
+	FText SourceText;
+	/** Arguments to replace in the pattern string */
+	FFormatOrderedArguments Arguments;
+};
+
+/** Handles history for FText::Format when passing raw argument data */
+class CORE_API FTextHistory_ArgumentDataFormat : public FTextHistory
+{
+public:
+	FTextHistory_ArgumentDataFormat() {};
+	FTextHistory_ArgumentDataFormat(const FText& InSourceText, const TArray< struct FFormatArgumentData >& InArguments);
+
+	// Begin FTextHistory interface
+	virtual FText ToText() const OVERRIDE;
+	virtual void Serialize(FArchive& Ar) OVERRIDE;
+	// End FTextHistory interface
+
+private:
+	/** The pattern used to format the text */
+	FText SourceText;
+	/** Arguments to replace in the pattern string */
+	TArray< struct FFormatArgumentData > Arguments;
+};
+
+/** Base class for managing formatting FText's from: AsNumber, AsPercent, and AsCurrency. Manages data serialization of these history events */
+class CORE_API FTextHistory_FormatNumber : public FTextHistory
+{
+public:
+	FTextHistory_FormatNumber();
+	FTextHistory_FormatNumber(const FFormatArgumentValue& InSourceValue, const FNumberFormattingOptions* const InFormatOptions, const TSharedPtr<FCulture> InTargetCulture);
+
+	~FTextHistory_FormatNumber();
+
+	// Begin FTextHistory interface
+	virtual void Serialize(FArchive& Ar) OVERRIDE;
+	// End FTextHistory interface
+protected:
+	/** The source value to format from */
+	FFormatArgumentValue SourceValue;
+	/** All the formatting options available to format using. This can be NULL */
+	FNumberFormattingOptions* FormatOptions;
+	/** The culture to format using */
+	TSharedPtr<FCulture> TargetCulture;
+};
+
+/**  Handles history for formatting using AsNumber */
+class CORE_API FTextHistory_AsNumber : public FTextHistory_FormatNumber
+{
+public:
+	FTextHistory_AsNumber() {};
+	FTextHistory_AsNumber(const FFormatArgumentValue& InSourceValue, const FNumberFormattingOptions* const InFormatOptions, const TSharedPtr<FCulture> InTargetCulture);
+
+	// Begin FTextHistory interface
+	virtual FText ToText() const OVERRIDE;
+	virtual void Serialize(FArchive& Ar) OVERRIDE;
+	// End FTextHistory interface
+};
+
+/**  Handles history for formatting using AsPercent */
+class CORE_API FTextHistory_AsPercent : public FTextHistory_FormatNumber
+{
+public:
+	FTextHistory_AsPercent() {};
+	FTextHistory_AsPercent(const FFormatArgumentValue& InSourceValue, const FNumberFormattingOptions* const InFormatOptions, const TSharedPtr<FCulture> InTargetCulture);
+
+	// Begin FTextHistory interface
+	virtual FText ToText() const OVERRIDE;
+	virtual void Serialize(FArchive& Ar) OVERRIDE;
+	// End FTextHistory interface
+};
+
+/**  Handles history for formatting using AsCurrency */
+class CORE_API FTextHistory_AsCurrency : public FTextHistory_FormatNumber
+{
+public:
+	FTextHistory_AsCurrency() {};
+	FTextHistory_AsCurrency(const FFormatArgumentValue& InSourceValue, const FNumberFormattingOptions* const InFormatOptions, const TSharedPtr<FCulture> InTargetCulture);
+
+	// Begin FTextHistory interface
+	virtual FText ToText() const OVERRIDE;
+	virtual void Serialize(FArchive& Ar) OVERRIDE;
+	// End FTextHistory interface
+};
+
+/**  Handles history for formatting using AsDate */
+class CORE_API FTextHistory_AsDate : public FTextHistory
+{
+public:
+	FTextHistory_AsDate() {};
+	FTextHistory_AsDate(const FDateTime::FDate& InSourceDate, const EDateTimeStyle::Type InDateStyle, const TSharedPtr<FCulture> InTargetCulture);
+
+	// Begin FTextHistory interface
+	virtual FText ToText() const OVERRIDE;
+	virtual void Serialize(FArchive& Ar) OVERRIDE;
+	// End FTextHistory interface
+
+private:
+	/** The source date structure to format */
+	FDateTime::FDate SourceDate;
+	/** Style to format the date using */
+	EDateTimeStyle::Type DateStyle;
+	/** Culture to format the date in */
+	TSharedPtr<FCulture> TargetCulture;
+};
+
+/**  Handles history for formatting using AsTime */
+class CORE_API FTextHistory_AsTime : public FTextHistory
+{
+public:
+	FTextHistory_AsTime() {};
+	FTextHistory_AsTime(const FDateTime::FTime& InSourceTime, const EDateTimeStyle::Type InTimeStyle, const FString& InTimeZone, const TSharedPtr<FCulture> InTargetCulture);
+
+	// Begin FTextHistory interface
+	virtual FText ToText() const OVERRIDE;
+	virtual void Serialize(FArchive& Ar) OVERRIDE;
+	// End FTextHistory interface
+
+private:
+	/** The source time structure to format */
+	FDateTime::FTime SourceTime;
+	/** Style to format the time using */
+	EDateTimeStyle::Type TimeStyle;
+	/** Timezone to put the time in */
+	FString TimeZone;
+	/** Culture to format the time in */
+	TSharedPtr<FCulture> TargetCulture;
+};
+
+/**  Handles history for formatting using AsDateTime */
+class CORE_API FTextHistory_AsDateTime : public FTextHistory
+{
+public:
+	FTextHistory_AsDateTime() {};
+	FTextHistory_AsDateTime(const FDateTime& InSourceDateTime, const EDateTimeStyle::Type InDateStyle, const EDateTimeStyle::Type InTimeStyle, const FString& InTimeZone, const TSharedPtr<FCulture> InTargetCulture);
+
+	// Begin FTextHistory interface
+	virtual FText ToText() const OVERRIDE;
+	virtual void Serialize(FArchive& Ar) OVERRIDE;
+	// End FTextHistory interfaces
+
+private:
+	/** The source date and time structure to format */
+	FDateTime SourceDateTime;
+	/** Style to format the date using */
+	EDateTimeStyle::Type DateStyle;
+	/** Style to format the time using */
+	EDateTimeStyle::Type TimeStyle;
+	/** Timezone to put the time in */
+	FString TimeZone;
+	/** Culture to format the time in */
+	TSharedPtr<FCulture> TargetCulture;
+};
\ No newline at end of file
diff --git a/Engine/Source/Runtime/Core/Private/Internationalization/TextLocalizationManager.cpp b/Engine/Source/Runtime/Core/Private/Internationalization/TextLocalizationManager.cpp
--- a/Engine/Source/Runtime/Core/Private/Internationalization/TextLocalizationManager.cpp
+++ b/Engine/Source/Runtime/Core/Private/Internationalization/TextLocalizationManager.cpp
@@ -249,6 +249,9 @@ void FTextLocalizationManager::FLocalizationEntryTracker::ReadFromArchive(FArchi
 
 void FTextLocalizationManager::LoadResources(const bool ShouldLoadEditor, const bool ShouldLoadGame)
 {
+	// Add one to the revision index, so all FText's refresh.
+	++HeadCultureRevisionIndex;
+
 	FInternationalization& I18N = FInternationalization::Get();
 
 	const FString& CultureName = I18N.GetCurrentCulture()->GetName();
@@ -408,6 +411,8 @@ void FTextLocalizationManager::UpdateLiveTable(const TArray<FLocalizationEntryTr
 						MakeShareable( new FString(NewEntry.LocalizedString) )	/*String*/
 						);
 					LiveKeyTable.Add( Key, NewLiveEntry );
+
+					ReverseLiveTable.Add(NewLiveEntry.String, FNamespaceKeyEntry( NamespaceName.IsEmpty()? nullptr : MakeShareable( new FString( NamespaceName ) ), Key.IsEmpty()? nullptr : MakeShareable( new FString( Key ) )));
 				}
 			}
 		}
@@ -544,12 +549,26 @@ TSharedRef<FString> FTextLocalizationManager::GetString(const FString& Namespace
 			}
 
 			LiveKeyTable->Add( Key, NewEntry );
+
+			ReverseLiveTable.Add(NewEntry.String, FNamespaceKeyEntry( Namespace.IsEmpty()? nullptr : MakeShareable( new FString( Namespace ) ), Key.IsEmpty()? nullptr : MakeShareable( new FString( Key ) )));
+
 		}
 
 		return UnlocalizedString;
 	}
 }
 
+void FTextLocalizationManager::FindKeyNamespaceFromDisplayString(TSharedRef<FString> InDisplayString, TSharedPtr<FString>& OutNamespace, TSharedPtr<FString>& OutKey)
+{
+	FNamespaceKeyEntry* NamespaceKeyEntry = ReverseLiveTable.Find(InDisplayString);
+
+	if(NamespaceKeyEntry)
+	{
+		OutNamespace = NamespaceKeyEntry->Namespace;
+		OutKey = NamespaceKeyEntry->Key;
+	}
+}
+
 namespace
 {
 	TSharedPtr<FJsonObject> ReadJSONTextFile(const FString& InFilePath)
diff --git a/Engine/Source/Runtime/Core/Private/Misc/DateTime.cpp b/Engine/Source/Runtime/Core/Private/Misc/DateTime.cpp
--- a/Engine/Source/Runtime/Core/Private/Misc/DateTime.cpp
+++ b/Engine/Source/Runtime/Core/Private/Misc/DateTime.cpp
@@ -6,6 +6,24 @@
 
 #include "CorePrivate.h"
 
+FArchive& operator<<(FArchive& Ar, FDateTime::FDate& Value)
+{
+	Ar << Value.Day;
+	Ar << Value.Month;
+	Ar << Value.Year;
+
+	return Ar;
+}
+
+FArchive& operator<<(FArchive& Ar, FDateTime::FTime& Value)
+{
+	Ar << Value.Hour;
+	Ar << Value.Minute;
+	Ar << Value.Second;
+	Ar << Value.Millisecond;
+
+	return Ar;
+}
 
 /* FDateTime constants
  *****************************************************************************/
diff --git a/Engine/Source/Runtime/Core/Private/Tests/Internationalization/TextTest.cpp b/Engine/Source/Runtime/Core/Private/Tests/Internationalization/TextTest.cpp
--- a/Engine/Source/Runtime/Core/Private/Tests/Internationalization/TextTest.cpp
+++ b/Engine/Source/Runtime/Core/Private/Tests/Internationalization/TextTest.cpp
@@ -395,6 +395,177 @@ bool FTextTest::RunTest (const FString& Parameters)
 			}
 		}
 	}
+
+	{
+		I18N.SetCurrentCulture(OriginalCulture);
+
+		TArray<uint8> FormattedHistoryAsEnglish;
+		TArray<uint8> FormattedHistoryAsFrenchCanadian;
+		TArray<uint8> InvariantFTextData;
+
+		FString InvariantString = TEXT("This is a culture invariant string.");
+
+		// Scoping to allow all locals to leave scope after we serialize at the end
+		{
+			FFormatNamedArguments Args;
+			Args.Add(TEXT("String1"), LOCTEXT("RebuildFTextTest1_Lorem", "Lorem"));
+			Args.Add(TEXT("String2"), LOCTEXT("RebuildFTextTest1_Ipsum", "Ipsum"));
+			FText FormattedTest1 = FText::Format(LOCTEXT("RebuildNamedText1", "{String1} \"Lorem Ipsum\" {String2}"), Args);
+
+			FFormatOrderedArguments ArgsOrdered;
+			ArgsOrdered.Add(LOCTEXT("RebuildFTextTest1_Lorem", "Lorem"));
+			ArgsOrdered.Add(LOCTEXT("RebuildFTextTest1_Ipsum", "Ipsum"));
+			FText FormattedTestOrdered1 = FText::Format(LOCTEXT("RebuildOrderedText1", "{0} \"Lorem Ipsum\" {1}"), ArgsOrdered);
+
+			// Will change to 5.542 due to default settings for numbers
+			FText AsNumberTest1 = FText::AsNumber(5.5421);
+
+			FText AsPercentTest1 = FText::AsPercent(0.925);
+			FText AsCurrencyTest1 = FText::AsCurrency(100.25);
+
+			FDateTime::FDate DateInfo;
+			DateInfo.Day = 20;
+			DateInfo.Month = 5;
+			DateInfo.Year = 2080;
+			FText AsDateTest1 = FText::AsDate(DateInfo);
+
+			FDateTime::FTime TimeInfo;
+			TimeInfo.Hour = 8;
+			TimeInfo.Minute = 25;
+			TimeInfo.Second = 44;
+			TimeInfo.Millisecond = 5;
+			FText AsTimeTest1 = FText::AsTime(TimeInfo);
+
+			FDateTime DateTimeInfo(2080, 8, 20, 9, 33, 22);
+			FText AsDateTimeTest1 = FText::AsDateTime(DateTimeInfo);
+
+			FFormatNamedArguments ArgsLayer2;
+			ArgsLayer2.Add("NamedLayer1", FormattedTest1);
+			ArgsLayer2.Add("OrderedLayer1", FormattedTestOrdered1);
+			ArgsLayer2.Add("FTextNumber", AsNumberTest1);
+			ArgsLayer2.Add("Number", 5010.89221);
+			ArgsLayer2.Add("Date", AsDateTest1);
+			ArgsLayer2.Add("Time", AsTimeTest1);
+			ArgsLayer2.Add("DateTime", AsDateTimeTest1);
+			ArgsLayer2.Add("Percent", AsPercentTest1);
+			ArgsLayer2.Add("Currency", AsCurrencyTest1);
+			FText FormattedTestLayer2 = FText::Format(LOCTEXT("RebuildTextLayer2", "{NamedLayer1} | {OrderedLayer1} | {FTextNumber} | {Number} | {Date} | {Time} | {DateTime} | {Percent} | {Currency}"), ArgsLayer2);
+
+			{
+				// Serialize the full, bulky FText that is a composite of most of the other FTextHistories.
+				FMemoryWriter Ar(FormattedHistoryAsEnglish, /*bIsPersistent=*/ true);
+				Ar << FormattedTestLayer2;
+				Ar.Close();
+			}
+
+			FString FormattedTestLayer2_Baked = FormattedTestLayer2.ToString();
+
+			{
+				// Swap to "LEET" culture to check if rebuilding works (verify the whole)
+				I18N.SetCurrentCulture("LEET");
+
+				// Convert the baked string into an FText, which will be leetified, then compare it to the rebuilt FText
+				if(FormattedTestLayer2.ToString() != TEXT("‡««L0r3m» \"L0r3m 1p$um\" «Ip$um»» | ««L0r3m» \"L0r3m 1p$um\" «Ip$um»» | «5.5421» | «5010.89221» | «1944 Month4 13» | «02:25:44» | «1944 Month7 14 22:05:06» | «92%» | «¤ 100.25»‡"))
+				{
+					AddError( TEXT("FormattedTestLayer2 did not rebuild to correctly in LEET!") );
+				}
+			}
+
+			// Swap to French-Canadian to check if rebuilding works (verify each numerical component)
+			{
+				I18N.SetCurrentCulture("fr-CA");
+
+				// Need the FText to be rebuilt in fr-CA.
+				FormattedTestLayer2.ToString();
+
+				if(AsNumberTest1.CompareTo(FText::AsNumber(5.5421)) != 0)
+				{
+					AddError( TEXT("AsNumberTest1 did not rebuild correctly in French-Canadian") );
+				}
+
+				if(AsPercentTest1.CompareTo(FText::AsPercent(0.925)) != 0)
+				{
+					AddError( TEXT("AsPercentTest1 did not rebuild correctly in French-Canadian") );
+				}
+
+				if(AsCurrencyTest1.CompareTo(FText::AsCurrency(100.25)) != 0)
+				{
+					AddError( TEXT("AsCurrencyTest1 did not rebuild correctly in French-Canadian") );
+				}
+
+				if(AsDateTest1.CompareTo(FText::AsDate(DateInfo)) != 0)
+				{
+					AddError( TEXT("AsDateTest1 did not rebuild correctly in French-Canadian") );
+				}
+
+				if(AsTimeTest1.CompareTo(FText::AsTime(TimeInfo)) != 0)
+				{
+					AddError( TEXT("AsTimeTest1 did not rebuild correctly in French-Canadian") );
+				}
+
+				if(AsDateTimeTest1.CompareTo(FText::AsDateTime(DateTimeInfo)) != 0)
+				{
+					AddError( TEXT("AsDateTimeTest1 did not rebuild correctly in French-Canadian") );
+				}
+
+				{
+					// Serialize the full, bulky FText that is a composite of most of the other FTextHistories.
+					// We don't care how this may be translated, we will be serializing this in as LEET.
+					FMemoryWriter Ar(FormattedHistoryAsFrenchCanadian, /*bIsPersistent=*/ true);
+					Ar << FormattedTestLayer2;
+					Ar.Close();
+				}
+
+				{
+					FText InvariantFText = FText::FromString(InvariantString);
+
+					// Serialize an invariant FText
+					FMemoryWriter Ar(InvariantFTextData, /*bIsPersistent=*/ true);
+					Ar << InvariantFText;
+					Ar.Close();
+				}
+			}
+		}
+
+		{
+			I18N.SetCurrentCulture("LEET");
+
+			FText FormattedEnglishTextHistoryAsLeet;
+			FText FormattedFrenchCanadianTextHistoryAsLeet;
+
+			{
+				FMemoryReader Ar(FormattedHistoryAsEnglish, /*bIsPersistent=*/ true);
+				Ar << FormattedEnglishTextHistoryAsLeet;
+				Ar.Close();
+			}
+			{
+				FMemoryReader Ar(FormattedHistoryAsFrenchCanadian, /*bIsPersistent=*/ true);
+				Ar << FormattedFrenchCanadianTextHistoryAsLeet;
+				Ar.Close();
+			}
+
+			// Confirm the two FText's serialize in and get translated into the current (LEET) translation. One originated in English, the other in French-Canadian locales.
+			if(FormattedEnglishTextHistoryAsLeet.CompareTo(FormattedFrenchCanadianTextHistoryAsLeet) != 0)
+			{
+				AddError( TEXT("Serialization of text histories from source English and source French-Canadian to LEET did not produce the same results!") );
+			}
+
+			{
+				I18N.SetCurrentCulture(OriginalCulture);
+
+				FText InvariantFText;
+
+				FMemoryReader Ar(InvariantFTextData, /*bIsPersistent=*/ true);
+				Ar << InvariantFText;
+				Ar.Close();
+
+				if(InvariantFText.ToString() != InvariantString)
+				{
+					AddError( TEXT("Invariant FText did not match the original FString after serialization!") );
+				}
+			}
+		}
+	}
 #else
 	AddWarning("ICU is disabled thus locale-aware string collation is disabled.");
 #endif
diff --git a/Engine/Source/Runtime/Core/Public/Internationalization/Text.h b/Engine/Source/Runtime/Core/Public/Internationalization/Text.h
--- a/Engine/Source/Runtime/Core/Public/Internationalization/Text.h
+++ b/Engine/Source/Runtime/Core/Public/Internationalization/Text.h
@@ -41,6 +41,7 @@ namespace EDateTimeStyle
 		Medium,
 		Long,
 		Full
+		// Add new enum types at the end only! They are serialized by index.
 	};
 }
 
@@ -53,6 +54,7 @@ namespace EFormatArgumentType
 		Float,
 		Double,
 		Text
+		// Add new enum types at the end only! They are serialized by index.
 	};
 }
 
@@ -78,6 +80,9 @@ enum ERoundingMode
 	ToNegativeInfinity,
 	/** Rounds to the value which is more positive: 0.1 becomes 1, -0.1 becomes 0 */
 	ToPositiveInfinity,
+
+
+	// Add new enum types at the end only! They are serialized by index.
 };
 
 struct CORE_API FNumberFormattingOptions
@@ -90,6 +95,8 @@ struct CORE_API FNumberFormattingOptions
 	int32 MaximumIntegralDigits;
 	int32 MinimumFractionalDigits;
 	int32 MaximumFractionalDigits;
+
+	friend FArchive& operator<<(FArchive& Ar, FNumberFormattingOptions& Value);
 };
 
 class FCulture;
@@ -222,6 +229,7 @@ public:
 	static FText TrimPrecedingAndTrailing( const FText& );
 
 	static void GetFormatPatternParameters(const FText& Pattern, TArray<FString>& ParameterNames);
+
 	static FText Format(const FText& Pattern, const FFormatNamedArguments& Arguments);
 	static FText Format(const FText& Pattern, const FFormatOrderedArguments& Arguments);
 	static FText Format(const FText& Pattern, const TArray< struct FFormatArgumentData > InArguments);
@@ -265,6 +273,16 @@ private:
 	 */
 	static FText CreateChronologicalText(FString InSourceString);
 
+	/** Returns the source string of the FText */
+	TSharedPtr<FString> GetSourceString() const;
+
+	/** Rebuilds the FText under the current culture if needed */
+	void Rebuild() const;
+
+	static FText FormatInternal(const FText& Pattern, const FFormatNamedArguments& Arguments, bool bInRebuildText);
+	static FText FormatInternal(const FText& Pattern, const FFormatOrderedArguments& Arguments, bool bInRebuildText);
+	static FText FormatInternal(const FText& Pattern, const TArray< struct FFormatArgumentData > InArguments, bool bInRebuildText);
+
 private:
 	template<typename T1, typename T2>
 	static FText AsNumberTemplate(T1 Val, const FNumberFormattingOptions* const Options = NULL, const TSharedPtr<FCulture>& TargetCulture = NULL);
@@ -274,13 +292,18 @@ private:
 	static FText AsPercentTemplate(T1 Val, const FNumberFormattingOptions* const Options = NULL, const TSharedPtr<FCulture>& TargetCulture = NULL);
 
 private:
-
-	TSharedPtr<FString> SourceString;
+	/** The visible display string for this FText */
 	TSharedRef<FString> DisplayString;
-	TSharedPtr<FString> Namespace;
-	TSharedPtr<FString> Key;
+
+	/** The FText's history, to allow it to rebuild under a new culture */
+	TSharedPtr<class FTextHistory> History;
+
+	/** Flags with various information on what sort of FText this is */
 	int32 Flags;
 
+	/** Revision index of this FText, rebuilds when it is out of sync with the FTextLocalizationManager */
+	int32 Revision;
+
 	static bool bEnableErrorCheckingResults;
 	static bool bSuppressWarnings;
 
@@ -302,6 +325,9 @@ public:
 	friend class UGatherTextFromAssetsCommandlet;
 	friend class UEdGraphSchema;
 	friend class UEdGraphSchema_K2;
+	friend class FTextHistory_NamedFormat;
+	friend class FTextHistory_ArgumentDataFormat;
+	friend class FTextHistory_OrderedFormat;
 
 #if !UE_ENABLE_ICU
 	friend class FLegacyTextHelper;
@@ -324,29 +350,37 @@ public:
 
 struct CORE_API FFormatArgumentValue
 {
-	const EFormatArgumentType::Type Type;
+	EFormatArgumentType::Type Type;
 	union
 	{
-		int IntValue;
-		unsigned int UIntValue;
+		int64 IntValue;
+		uint64 UIntValue;
 		float FloatValue;
 		double DoubleValue;
 		FText* TextValue;
 	};
 
+	FFormatArgumentValue();
+
 	FFormatArgumentValue(const int Value);
 	FFormatArgumentValue(const unsigned int Value);
+	FFormatArgumentValue(const int64 Value);
+	FFormatArgumentValue(const uint64 Value);
 	FFormatArgumentValue(const float Value);
 	FFormatArgumentValue(const double Value);
 	FFormatArgumentValue(const FText& Value);
 	FFormatArgumentValue(const FFormatArgumentValue& Source);
 	~FFormatArgumentValue();
+
+	friend FArchive& operator<<( FArchive& Ar, FFormatArgumentValue& Value );
 };
 
 struct FFormatArgumentData
 {
 	FText ArgumentName;
 	FText ArgumentValue;
+
+	friend FArchive& operator<<( FArchive& Ar, FFormatArgumentData& Value );
 };
 
 class CORE_API FTextBuilder
diff --git a/Engine/Source/Runtime/Core/Public/Internationalization/TextLocalizationManager.h b/Engine/Source/Runtime/Core/Public/Internationalization/TextLocalizationManager.h
--- a/Engine/Source/Runtime/Core/Public/Internationalization/TextLocalizationManager.h
+++ b/Engine/Source/Runtime/Core/Public/Internationalization/TextLocalizationManager.h
@@ -60,18 +60,43 @@ private:
 		const TSharedRef<FString>* GetString(const FString& Namespace, const FString& Key, const uint32 SourceStringHash) const;
 	};
 
+	/** Stores a Namespace and Key, used in reverse lookups from a DisplayString to find the Namespace and Key */
+	struct FNamespaceKeyEntry
+	{
+		TSharedPtr<FString> Namespace;
+		TSharedPtr<FString> Key;
+
+		FNamespaceKeyEntry(TSharedPtr<FString> InNamespace, TSharedPtr<FString> InKey)
+			: Namespace(InNamespace)
+			, Key(InKey)
+		{}
+	};
+
 public:
 	static CORE_API FTextLocalizationManager& Get();
 
 	TSharedPtr<FString> FindString( const FString& Namespace, const FString& Key, const FString* const SourceString = nullptr );
 	TSharedRef<FString> GetString(const FString& Namespace, const FString& Key, const FString* const SourceString);
 
+	/**
+	 * Using an FText SourceString's, returns the associated Namespace and Key
+	 *
+	 * @param InSourceString		The FText DisplayString
+	 * @param OutNamespace			Returns with the associated Namespace
+	 * @param OutKey				Returns with the associated Key
+	 */
+	void CORE_API FindKeyNamespaceFromDisplayString(TSharedRef<FString> InDisplayString, TSharedPtr<FString>& OutNamespace, TSharedPtr<FString>& OutKey);
+
 	void CORE_API RegenerateResources(const FString& ConfigFilePath);
 
+	/** Returns the current culture revision index */
+	int CORE_API GetHeadCultureRevision() const { return HeadCultureRevisionIndex; }
+
 private:
 	FTextLocalizationManager() 
 		: bIsInitialized(false)
 		, SynchronizationObject()
+		, HeadCultureRevisionIndex(0)
 	{}
 
 	void LoadResources(const bool ShouldLoadEditor, const bool ShouldLoadGame);
@@ -81,5 +106,12 @@ private:
 private:
 	bool bIsInitialized;
 	FTextLookupTable LiveTable;
+	
+	/** Table for looking up namespaces and keys using the DisplayString of an FText */
+	TMap< TSharedPtr< FString >, FNamespaceKeyEntry > ReverseLiveTable;
+
 	FCriticalSection SynchronizationObject;
+
+	/** The current culture revision, increments every time the culture is changed and allows for FTexts to know when to rebuild under a new culture */
+	int32 HeadCultureRevisionIndex;
 };
\ No newline at end of file
diff --git a/Engine/Source/Runtime/Core/Public/Misc/DateTime.h b/Engine/Source/Runtime/Core/Public/Misc/DateTime.h
--- a/Engine/Source/Runtime/Core/Public/Misc/DateTime.h
+++ b/Engine/Source/Runtime/Core/Public/Misc/DateTime.h
@@ -97,6 +97,8 @@ public:
 		 * Holds the date's year.
 		 */
 		int32 Year;
+
+		friend FArchive& operator<<(FArchive& Ar, FDate& Value);
 	};
 
 
@@ -124,6 +126,8 @@ public:
 		 * Holds the time of day's millisecond.
 		 */
 		int32 Millisecond;
+
+		friend FArchive& operator<<(FArchive& Ar, FTime& Value);
 	};
 
 public:
diff --git a/Engine/Source/Runtime/Core/Public/UObject/ObjectVersion.h b/Engine/Source/Runtime/Core/Public/UObject/ObjectVersion.h
--- a/Engine/Source/Runtime/Core/Public/UObject/ObjectVersion.h
+++ b/Engine/Source/Runtime/Core/Public/UObject/ObjectVersion.h
@@ -548,7 +548,8 @@ enum EUnrealEngineObjectUE4Version
 	VER_UE4_SERIALIZE_RICH_CURVE_KEY,
 	// Change the outer of ULandscapeMaterialInstanceConstants and Landscape-related textures to the level in which they reside
 	VER_UE4_MOVE_LANDSCAPE_MICS_AND_TEXTURES_WITHIN_LEVEL,
-
+	// FTexts have creation history data, removed Key, Namespaces, and SourceString
+	VER_UE4_FTEXT_HISTORY,
 	// -----<new versions can be added before this line>-------------------------------------------------
 	// - this needs to be the last line (see note below)
 	VER_UE4_AUTOMATIC_VERSION_PLUS_ONE,
