Prevent the high-precision mouse mode from revealing the main menu bar and dock when used in fullscreen mode on OS X.

diff --git a/Engine/Source/Runtime/Core/Private/Mac/MacApplication.cpp b/Engine/Source/Runtime/Core/Private/Mac/MacApplication.cpp
--- a/Engine/Source/Runtime/Core/Private/Mac/MacApplication.cpp
+++ b/Engine/Source/Runtime/Core/Private/Mac/MacApplication.cpp
@@ -339,6 +339,12 @@ void FMacApplication::ProcessEvent( NSEvent* Event )
 					// Clamp to the current screen and avoid the menu bar and dock to prevent popups and other
 					// assorted potential for mouse abuse.
 					NSRect VisibleFrame = [screen visibleFrame];
+					// Avoid the menu bar & dock disclosure borders at the top & bottom of fullscreen windows
+					if(CurrentEventWindow->GetWindowMode() != EWindowMode::Windowed)
+					{
+						VisibleFrame.origin.y += 5;
+						VisibleFrame.size.height -= 10;
+					}
 					NSRect FullFrame = [screen frame];
 					VisibleFrame.origin.y = (FullFrame.origin.y+FullFrame.size.height) - (VisibleFrame.origin.y + VisibleFrame.size.height);
 					
