Made FText::Revision properly mutable to resolve Clang issues.

diff --git a/Engine/Source/Runtime/Core/Private/Internationalization/Text.cpp b/Engine/Source/Runtime/Core/Private/Internationalization/Text.cpp
--- a/Engine/Source/Runtime/Core/Private/Internationalization/Text.cpp
+++ b/Engine/Source/Runtime/Core/Private/Internationalization/Text.cpp
@@ -624,8 +624,7 @@ void FText::Rebuild() const
 {
 	if(History.IsValid() && History->IsOutOfDate(Revision))
 	{
-		int32& RevisionNonCast = (int32)Revision;
-		RevisionNonCast = FTextLocalizationManager::Get().GetHeadCultureRevision();
+		Revision = FTextLocalizationManager::Get().GetHeadCultureRevision();
 		
 		DisplayString.Get() = History->ToText().DisplayString.Get();
 	}
diff --git a/Engine/Source/Runtime/Core/Public/Internationalization/Text.h b/Engine/Source/Runtime/Core/Public/Internationalization/Text.h
--- a/Engine/Source/Runtime/Core/Public/Internationalization/Text.h
+++ b/Engine/Source/Runtime/Core/Public/Internationalization/Text.h
@@ -302,7 +302,7 @@ private:
 	int32 Flags;
 
 	/** Revision index of this FText, rebuilds when it is out of sync with the FTextLocalizationManager */
-	int32 Revision;
+	mutable int32 Revision;
 
 	static bool bEnableErrorCheckingResults;
 	static bool bSuppressWarnings;
