#UE4 #HTML5
  -  Remove warning from HTML5 Win32 build.

diff --git a/Engine/Source/Runtime/Core/Public/HTML5/HTML5PlatformCompilerSetup.h b/Engine/Source/Runtime/Core/Public/HTML5/HTML5PlatformCompilerSetup.h
--- a/Engine/Source/Runtime/Core/Public/HTML5/HTML5PlatformCompilerSetup.h
+++ b/Engine/Source/Runtime/Core/Public/HTML5/HTML5PlatformCompilerSetup.h
@@ -55,6 +55,7 @@
 #pragma warning(default : 4289) // nonstandard extension used : 'var' : loop control variable declared in the for-loop is used outside the for-loop scope
 #pragma warning(default : 4302) // 'conversion' : truncation from 'type 1' to 'type 2'
 #pragma warning(disable : 4315) // 'this' pointer for member may not be aligned 8 as expected by the constructor
+#pragma warning(disable : 4316) //  over aligned allocation warning 
 //#pragma warning(disable : 4339) // 'type' : use of undefined type detected in CLR meta-data - use of this type may lead to a runtime exception
 #pragma warning(disable : 4347) // behavior change: 'function template' is called instead of 'function
 #pragma warning(disable : 4366) // The result of the unary '&' operator may be unaligned
