Fixing Build errors on platforms without ICU enabled.

diff --git a/Engine/Source/Runtime/Core/Private/Internationalization/LegacyText.cpp b/Engine/Source/Runtime/Core/Private/Internationalization/LegacyText.cpp
--- a/Engine/Source/Runtime/Core/Private/Internationalization/LegacyText.cpp
+++ b/Engine/Source/Runtime/Core/Private/Internationalization/LegacyText.cpp
@@ -481,4 +481,19 @@ FText FText::Format(const FText& Pattern, const TArray< FFormatArgumentData > In
 	return FLegacyTextHelper::Format(Pattern, FGetArgumentValue::CreateRaw(&ArgumentGetter, &FArgumentGetter::GetArgumentValue));
 }
 
+FText FText::FormatInternal(const FText& Pattern, const FFormatNamedArguments& Arguments, bool bInRebuildText)
+{
+	Format(Pattern, Arguments);
+}
+
+FText FText::FormatInternal(const FText& Pattern, const FFormatOrderedArguments& Arguments, bool bInRebuildText)
+{
+	Format(Pattern, Arguments);
+}
+
+FText FText::FormatInternal(const FText& Pattern, const TArray< struct FFormatArgumentData > InArguments, bool bInRebuildText)
+{
+	Format(Pattern, InArguments);
+}
+
 #endif
