LegacyICU not returning values in functions causing build isuses.

diff --git a/Engine/Source/Runtime/Core/Private/Internationalization/LegacyText.cpp b/Engine/Source/Runtime/Core/Private/Internationalization/LegacyText.cpp
--- a/Engine/Source/Runtime/Core/Private/Internationalization/LegacyText.cpp
+++ b/Engine/Source/Runtime/Core/Private/Internationalization/LegacyText.cpp
@@ -483,17 +483,17 @@ FText FText::Format(const FText& Pattern, const TArray< FFormatArgumentData > In
 
 FText FText::FormatInternal(const FText& Pattern, const FFormatNamedArguments& Arguments, bool bInRebuildText)
 {
-	Format(Pattern, Arguments);
+	return Format(Pattern, Arguments);
 }
 
 FText FText::FormatInternal(const FText& Pattern, const FFormatOrderedArguments& Arguments, bool bInRebuildText)
 {
-	Format(Pattern, Arguments);
+	return Format(Pattern, Arguments);
 }
 
 FText FText::FormatInternal(const FText& Pattern, const TArray< struct FFormatArgumentData > InArguments, bool bInRebuildText)
 {
-	Format(Pattern, InArguments);
+	return Format(Pattern, InArguments);
 }
 
 #endif
