UnrealMatrix.h did not have NaN checking for Inverse()

diff --git a/Engine/Source/Runtime/Core/Public/Math/UnrealMatrix.h b/Engine/Source/Runtime/Core/Public/Math/UnrealMatrix.h
--- a/Engine/Source/Runtime/Core/Public/Math/UnrealMatrix.h
+++ b/Engine/Source/Runtime/Core/Public/Math/UnrealMatrix.h
@@ -8,6 +8,9 @@
 
 #pragma once
 
+#include "OutputDevices.h"
+#include "AssertionMacros.h"
+
 /**
  * FMatrix inline functions.
  */
@@ -255,13 +258,13 @@ inline float FMatrix::RotDeterminant() const
 inline FMatrix FMatrix::Inverse() const
 {
 	// If we're in non final release, then make sure we're not creating NaNs
-#if 0 && !(UE_BUILD_SHIPPING || UE_BUILD_TEST)
+#if 1 && !(UE_BUILD_SHIPPING || UE_BUILD_TEST)
 	// Check for zero scale matrix to invert
 	if(	GetScaledAxis( EAxis::X ).IsNearlyZero(SMALL_NUMBER) && 
 		GetScaledAxis( EAxis::Y ).IsNearlyZero(SMALL_NUMBER) && 
 		GetScaledAxis( EAxis::Z ).IsNearlyZero(SMALL_NUMBER) ) 
 	{
-		UE_LOG(LogUnrealMatrix, Fatal,TEXT("FMatrix::Inverse(), trying to invert a NIL matrix, this results in NaNs! Use InverseSafe() instead."));
+		UE_LOG(LogUnrealMath, Fatal, TEXT("FMatrix::Inverse(), trying to invert a NIL matrix, this results in NaNs! Use InverseSafe() instead."));
 	}
 #endif
 	FMatrix Result;
