Enable PLATFORM_ENABLE_VECTORINTRINSICS on Mac so we use SSE instructions, I don't see why the won't work as all Macs have this instruction set.

diff --git a/Engine/Source/Runtime/Core/Public/Mac/MacPlatform.h b/Engine/Source/Runtime/Core/Public/Mac/MacPlatform.h
--- a/Engine/Source/Runtime/Core/Public/Mac/MacPlatform.h
+++ b/Engine/Source/Runtime/Core/Public/Mac/MacPlatform.h
@@ -29,7 +29,7 @@ typedef FMacPlatformTypes FPlatformTypes;
 #define PLATFORM_LITTLE_ENDIAN						1
 //#define PLATFORM_EXCEPTIONS_DISABLED				!PLATFORM_DESKTOP
 #define PLATFORM_SUPPORTS_PRAGMA_PACK				1
-//#define PLATFORM_ENABLE_VECTORINTRINSICS			0
+#define PLATFORM_ENABLE_VECTORINTRINSICS			1
 //#define PLATFORM_USE_LS_SPEC_FOR_WIDECHAR			1
 #define PLATFORM_USE_SYSTEM_VSWPRINTF				0
 #define PLATFORM_COMPILER_DISTINGUISHES_INT_AND_LONG			1
