Make FText movable

diff --git a/Engine/Source/Runtime/Core/Private/Internationalization/Text.cpp b/Engine/Source/Runtime/Core/Private/Internationalization/Text.cpp
--- a/Engine/Source/Runtime/Core/Private/Internationalization/Text.cpp
+++ b/Engine/Source/Runtime/Core/Private/Internationalization/Text.cpp
@@ -234,6 +234,12 @@ FText::FText()
 {
 }
 
+#if PLATFORM_COMPILER_HAS_DEFAULTED_FUNCTIONS
+	FText::FText(const FText& Other) = default;
+	FText::FText(FText&& Other) = default;
+	FText& operator=(const FText& Other) = default;
+	FText& operator=(FText&& Other) = default;
+#else
 FText::FText( const FText& Source )
 	: DisplayString( Source.DisplayString )
 	, History( Source.History )
@@ -242,6 +248,35 @@ FText::FText( const FText& Source )
 {
 }
 
+FText::FText(FText&& Source)
+	: DisplayString(MoveTemp(Source.DisplayString))
+	, History(MoveTemp(Source.History))
+	, Flags(MoveTemp(Source.Flags))
+	, Revision(MoveTemp(Source.Revision))
+{
+}
+
+FText& FText::operator=(const FText& Source)
+{
+	DisplayString = Source.DisplayString;
+	History = Source.History;
+	Flags = Source.Flags;
+	Revision = Source.Revision;
+
+	return *this;
+}
+
+FText& FText::operator=(FText&& Source)
+{
+	DisplayString = MoveTemp(Source.DisplayString);
+	History = MoveTemp(Source.History);
+	Flags = MoveTemp(Source.Flags);
+	Revision = MoveTemp(Source.Revision);
+
+	return *this;
+}
+#endif
+
 FText::FText( FString InSourceString )
 	: DisplayString( new FString( MoveTemp(InSourceString) ))
 	, Flags(0)
diff --git a/Engine/Source/Runtime/Core/Public/Internationalization/Text.h b/Engine/Source/Runtime/Core/Public/Internationalization/Text.h
--- a/Engine/Source/Runtime/Core/Public/Internationalization/Text.h
+++ b/Engine/Source/Runtime/Core/Public/Internationalization/Text.h
@@ -115,6 +115,10 @@ public:
 
 	FText();
 	FText( const FText& Source );
+	FText(FText&& Source);
+
+	FText& operator=(const FText& Source);
+	FText& operator=(FText&& Source);
 
 	/**
 	 * Generate an FText that represents the passed number in the current culture
