Core - Removed unused COMPRESS_ForcePPUDecompressZLib from the compressing method

diff --git a/Engine/Source/Runtime/Core/Public/Misc/Compression.h b/Engine/Source/Runtime/Core/Public/Misc/Compression.h
--- a/Engine/Source/Runtime/Core/Public/Misc/Compression.h
+++ b/Engine/Source/Runtime/Core/Public/Misc/Compression.h
@@ -19,8 +19,6 @@ enum ECompressionFlags
 	COMPRESS_BiasMemory 			= 0x10,
 	/** Prefer compression that compresses faster (ONLY VALID FOR COMPRESSION)		*/
 	COMPRESS_BiasSpeed				= 0x20,
-	/** If this flag is present, decompression will not happen on the SPUs.			*/
-	COMPRESS_ForcePPUDecompressZLib	= 0x80
 };
 
 // Define global current platform default to current platform.
