Fix FText on non-windows platforms...

diff --git a/Engine/Source/Runtime/Core/Private/Internationalization/Text.cpp b/Engine/Source/Runtime/Core/Private/Internationalization/Text.cpp
--- a/Engine/Source/Runtime/Core/Private/Internationalization/Text.cpp
+++ b/Engine/Source/Runtime/Core/Private/Internationalization/Text.cpp
@@ -237,14 +237,14 @@ FText::FText()
 #if PLATFORM_COMPILER_HAS_DEFAULTED_FUNCTIONS
 	FText::FText(const FText& Other) = default;
 	FText::FText(FText&& Other) = default;
-	FText& operator=(const FText& Other) = default;
-	FText& operator=(FText&& Other) = default;
+	FText& FText::operator=(const FText& Other) = default;
+	FText& FText::operator=(FText&& Other) = default;
 #else
-FText::FText( const FText& Source )
-	: DisplayString( Source.DisplayString )
-	, History( Source.History )
-	, Flags( Source.Flags )
-	, Revision( Source.Revision )
+FText::FText(const FText& Source)
+	: DisplayString(Source.DisplayString)
+	, History(Source.History)
+	, Flags(Source.Flags)
+	, Revision(Source.Revision)
 {
 }
 
