Changed some FText smoketests to be normal Application tests, they are expensive to run (> 0.5 seconds) and destructive.

Fixes issues with culture invariant FTexts appearing in LEET in the game/editor.

diff --git a/Engine/Source/Runtime/Core/Private/Tests/Internationalization/NumberFormatingRulesTest.cpp b/Engine/Source/Runtime/Core/Private/Tests/Internationalization/NumberFormatingRulesTest.cpp
--- a/Engine/Source/Runtime/Core/Private/Tests/Internationalization/NumberFormatingRulesTest.cpp
+++ b/Engine/Source/Runtime/Core/Private/Tests/Internationalization/NumberFormatingRulesTest.cpp
@@ -10,7 +10,7 @@
 // Disable optimization for NumberFormatingRulesTest as it compiles very slowly in development builds
 PRAGMA_DISABLE_OPTIMIZATION
 
-IMPLEMENT_SIMPLE_AUTOMATION_TEST(FNumberFormattingRulesTest, "Core.Misc.Number Formatting Rules", EAutomationTestFlags::ATF_SmokeTest)
+IMPLEMENT_SIMPLE_AUTOMATION_TEST(FNumberFormattingRulesTest, "Core.Misc.Number Formatting Rules", EAutomationTestFlags::ATF_ApplicationMask)
 
 namespace
 {
@@ -322,4 +322,4 @@ bool FNumberFormattingRulesTest::RunTest (const FString& Parameters)
 }
 
 
-PRAGMA_ENABLE_OPTIMIZATION
\ No newline at end of file
+PRAGMA_ENABLE_OPTIMIZATION
diff --git a/Engine/Source/Runtime/Core/Private/Tests/Internationalization/TextTest.cpp b/Engine/Source/Runtime/Core/Private/Tests/Internationalization/TextTest.cpp
--- a/Engine/Source/Runtime/Core/Private/Tests/Internationalization/TextTest.cpp
+++ b/Engine/Source/Runtime/Core/Private/Tests/Internationalization/TextTest.cpp
@@ -9,7 +9,7 @@ PRAGMA_DISABLE_OPTIMIZATION
 
 #define LOCTEXT_NAMESPACE "Core.Tests.TextFormatTest"
 
-IMPLEMENT_SIMPLE_AUTOMATION_TEST(FTextTest, "Core.Misc.Text", EAutomationTestFlags::ATF_SmokeTest)
+IMPLEMENT_SIMPLE_AUTOMATION_TEST(FTextTest, "Core.Misc.Text", EAutomationTestFlags::ATF_ApplicationMask)
 
 namespace
 {
@@ -49,6 +49,8 @@ namespace
 
 bool FTextTest::RunTest (const FString& Parameters)
 {
+	AddLogItem(TEXT("This test is destructive to existing culture invariant text! All culture invariant text will appear in LEET afterwards!"));
+
 	FInternationalization& I18N = FInternationalization::Get();
 	const bool OriginalEnableErrorCheckingValue = FText::GetEnableErrorCheckingResults();
 	const bool OriginalSuppressWarningsValue = FText::GetSuppressWarnings();
