Core: Removed FDateTime::FDate and FDateTime::FTime as they encouraged abuse and bad coding style

- also fixed const correctness on getters
- renamed localization for time spans to FText::AsTimespan (time spans are not times of day!)
- documentation improvements

diff --git a/Engine/Source/Editor/Kismet/Private/SBlueprintEditorToolbar.cpp b/Engine/Source/Editor/Kismet/Private/SBlueprintEditorToolbar.cpp
--- a/Engine/Source/Editor/Kismet/Private/SBlueprintEditorToolbar.cpp
+++ b/Engine/Source/Editor/Kismet/Private/SBlueprintEditorToolbar.cpp
@@ -348,7 +348,7 @@ class SBlueprintDiffMenu : public SCompoundWidget
 						FFormatNamedArguments Args;
 						Args.Add( TEXT("CheckInNumber"), FText::AsNumber( Revision->GetCheckInIdentifier(), NULL, I18N.GetInvariantCulture() ) );
 						Args.Add( TEXT("UserName"), FText::FromString( Revision->GetUserName() ) );
-						Args.Add( TEXT("DateTime"), FText::AsDate( Revision->GetDate().ToDate() ) );
+						Args.Add( TEXT("DateTime"), FText::AsDate( Revision->GetDate() ) );
 						Args.Add( TEXT("ChanglistDescription"), FText::FromString( Revision->GetDescription() ) );
 						const FText ToolTip = FText::Format( LOCTEXT("RevisionToolTip", "CL #{CheckInNumber} {UserName} \n{DateTime} \n{ChanglistDescription}"), Args );
 						
diff --git a/Engine/Source/Editor/NewsFeed/Private/Widgets/SNewsFeedListRow.h b/Engine/Source/Editor/NewsFeed/Private/Widgets/SNewsFeedListRow.h
--- a/Engine/Source/Editor/NewsFeed/Private/Widgets/SNewsFeedListRow.h
+++ b/Engine/Source/Editor/NewsFeed/Private/Widgets/SNewsFeedListRow.h
@@ -97,7 +97,7 @@ public:
 							SNew(STextBlock)
 								.ColorAndOpacity(FSlateColor::UseSubduedForeground())
 								.Font(FSlateFontInfo(FPaths::EngineContentDir() / TEXT("Slate/Fonts/Roboto-Regular.ttf"), 8))
-								.Text(FText::AsDate(NewsFeedItem->Issued.ToDate()))
+								.Text(FText::AsDate(NewsFeedItem->Issued))
 						]
 				]
 				
diff --git a/Engine/Source/Editor/UnrealEd/Private/Dialogs/SBuildProgress.cpp b/Engine/Source/Editor/UnrealEd/Private/Dialogs/SBuildProgress.cpp
--- a/Engine/Source/Editor/UnrealEd/Private/Dialogs/SBuildProgress.cpp
+++ b/Engine/Source/Editor/UnrealEd/Private/Dialogs/SBuildProgress.cpp
@@ -163,7 +163,7 @@ void SBuildProgressWidget::SetBuildType(EBuildType InBuildType)
 FText SBuildProgressWidget::BuildElapsedTimeText() const
 {
 	// Display elapsed build time.
-	return FText::AsTime( FDateTime::Now() - BuildStartTime );
+	return FText::AsTimespan( FDateTime::Now() - BuildStartTime );
 }
 
 void SBuildProgressWidget::UpdateTime()
diff --git a/Engine/Source/Runtime/Core/Private/Internationalization/ICUText.cpp b/Engine/Source/Runtime/Core/Private/Internationalization/ICUText.cpp
--- a/Engine/Source/Runtime/Core/Private/Internationalization/ICUText.cpp
+++ b/Engine/Source/Runtime/Core/Private/Internationalization/ICUText.cpp
@@ -12,13 +12,12 @@
 #include <unicode/msgfmt.h>
 #include "ICUUtilities.h"
 
-FText FText::AsDate(const FDateTime::FDate& Date, const EDateTimeStyle::Type DateStyle, const TSharedPtr<FCulture>& TargetCulture)
+FText FText::AsDate(const FDateTime& DateTime, const EDateTimeStyle::Type DateStyle, const TSharedPtr<FCulture>& TargetCulture)
 {
 	FInternationalization& I18N = FInternationalization::Get();
 	checkf(I18N.IsInitialized() == true, TEXT("FInternationalization is not initialized. An FText formatting method was likely used in static object initialization - this is not supported."));
 	TSharedPtr<FCulture> Culture = TargetCulture.IsValid() ? TargetCulture : I18N.GetCurrentCulture();
 
-	FDateTime DateTime(Date.Year, Date.Month, Date.Day);
 	int32 UNIXTimestamp = DateTime.ToUnixTimestamp();
 	UDate ICUDate = static_cast<double>(UNIXTimestamp) * U_MILLIS_PER_SECOND;
 
@@ -31,18 +30,17 @@ FText FText::AsDate(const FDateTime::FDate& Date, const EDateTimeStyle::Type Dat
 	ICUUtilities::Convert(FormattedString, NativeString);
 
 	FText ResultText = FText::CreateChronologicalText(NativeString);
-	ResultText.History = MakeShareable(new FTextHistory_AsDate(Date, DateStyle, TargetCulture));
+	ResultText.History = MakeShareable(new FTextHistory_AsDate(DateTime, DateStyle, TargetCulture));
 
 	return ResultText;
 }
 
-FText FText::AsTime(const FDateTime::FTime& Time, const EDateTimeStyle::Type TimeStyle, const FString& TimeZone, const TSharedPtr<FCulture>& TargetCulture)
+FText FText::AsTime(const FDateTime& DateTime, const EDateTimeStyle::Type TimeStyle, const FString& TimeZone, const TSharedPtr<FCulture>& TargetCulture)
 {
 	FInternationalization& I18N = FInternationalization::Get();
 	checkf(I18N.IsInitialized() == true, TEXT("FInternationalization is not initialized. An FText formatting method was likely used in static object initialization - this is not supported."));
 	TSharedPtr<FCulture> Culture = TargetCulture.IsValid() ? TargetCulture : I18N.GetCurrentCulture();
 
-	FDateTime DateTime(1970, 1, 1, Time.Hour, Time.Minute, Time.Second, Time.Millisecond);
 	int32 UNIXTimestamp = DateTime.ToUnixTimestamp();
 	UDate ICUDate = static_cast<double>(UNIXTimestamp) * U_MILLIS_PER_SECOND;
 
@@ -55,12 +53,12 @@ FText FText::AsTime(const FDateTime::FTime& Time, const EDateTimeStyle::Type Tim
 	ICUUtilities::Convert(FormattedString, NativeString);
 
 	FText ResultText = FText::CreateChronologicalText(NativeString);
-	ResultText.History = MakeShareable(new FTextHistory_AsTime(Time, TimeStyle, TimeZone, TargetCulture));
+	ResultText.History = MakeShareable(new FTextHistory_AsTime(DateTime, TimeStyle, TimeZone, TargetCulture));
 
 	return ResultText;
 }
 
-FText FText::AsTime(const FTimespan& Time, const TSharedPtr<FCulture>& TargetCulture)
+FText FText::AsTimespan(const FTimespan& Time, const TSharedPtr<FCulture>& TargetCulture)
 {
 	FInternationalization& I18N = FInternationalization::Get();
 	checkf(I18N.IsInitialized() == true, TEXT("FInternationalization is not initialized. An FText formatting method was likely used in static object initialization - this is not supported."));
diff --git a/Engine/Source/Runtime/Core/Private/Internationalization/LegacyText.cpp b/Engine/Source/Runtime/Core/Private/Internationalization/LegacyText.cpp
--- a/Engine/Source/Runtime/Core/Private/Internationalization/LegacyText.cpp
+++ b/Engine/Source/Runtime/Core/Private/Internationalization/LegacyText.cpp
@@ -5,34 +5,32 @@
 #if !UE_ENABLE_ICU
 #include "Text.h"
 
-FText FText::AsDate(const FDateTime::FDate& Date, const EDateTimeStyle::Type DateStyle, const TSharedPtr<FCulture>& TargetCulture)
+FText FText::AsDate( const FDateTime& DateTime, const EDateTimeStyle::Type DateStyle, const TSharedPtr<FCulture>& TargetCulture )
 {
 	checkf(FInternationalization::Get().IsInitialized() == true, TEXT("FInternationalization is not initialized. An FText formatting method was likely used in static object initialization - this is not supported."));
-	FDateTime DateTime(Date.Year, Date.Month, Date.Day);
 	return FText::FromString( DateTime.ToString( TEXT("%Y.%m.%d") ) );
 }
 
-FText FText::AsTime(const FDateTime::FTime& Time, const EDateTimeStyle::Type TimeStyle, const FString& TimeZone, const TSharedPtr<FCulture>& TargetCulture)
+FText FText::AsTime( const FDateTime& DateTime, const EDateTimeStyle::Type TimeStyle, const FString& TimeZone, const TSharedPtr<FCulture>& TargetCulture )
 {
 	checkf(FInternationalization::Get().IsInitialized() == true, TEXT("FInternationalization is not initialized. An FText formatting method was likely used in static object initialization - this is not supported."));
-	FDateTime DateTime(1, 1, 1, Time.Hour, Time.Minute, Time.Second, Time.Millisecond);
 	return FText::FromString( DateTime.ToString( TEXT("%H.%M.%S") ) );
 }
 
-FText FText::AsTime(const FTimespan& Time, const TSharedPtr<FCulture>& TargetCulture)
+FText FText::AsTimespan( const FTimespan& Time, const TSharedPtr<FCulture>& TargetCulture)
 {
 	checkf(FInternationalization::Get().IsInitialized() == true, TEXT("FInternationalization is not initialized. An FText formatting method was likely used in static object initialization - this is not supported."));
 	FDateTime DateTime(Time.GetTicks());
 	return FText::FromString( DateTime.ToString( TEXT("%H.%M.%S") ) );
 }
 
-FText FText::AsDateTime(const FDateTime& DateTime, const EDateTimeStyle::Type DateStyle, const EDateTimeStyle::Type TimeStyle, const FString& TimeZone, const TSharedPtr<FCulture>& TargetCulture)
+FText FText::AsDateTime( const FDateTime& DateTime, const EDateTimeStyle::Type DateStyle, const EDateTimeStyle::Type TimeStyle, const FString& TimeZone, const TSharedPtr<FCulture>& TargetCulture )
 {
 	checkf(FInternationalization::Get().IsInitialized() == true, TEXT("FInternationalization is not initialized. An FText formatting method was likely used in static object initialization - this is not supported."));
-	return FText::FromString( DateTime.ToString( TEXT("%Y.%m.%d-%H.%M.%S") ) );
+	return FText::FromString(DateTime.ToString(TEXT("%Y.%m.%d-%H.%M.%S")));
 }
 
-FText FText::AsMemory(SIZE_T NumBytes, const FNumberFormattingOptions* const Options, const TSharedPtr<FCulture>& TargetCulture)
+FText FText::AsMemory( SIZE_T NumBytes, const FNumberFormattingOptions* const Options, const TSharedPtr<FCulture>& TargetCulture )
 {
 	checkf(FInternationalization::Get().IsInitialized() == true, TEXT("FInternationalization is not initialized. An FText formatting method was likely used in static object initialization - this is not supported."));
 	FFormatNamedArguments Args;
diff --git a/Engine/Source/Runtime/Core/Private/Internationalization/TextHistory.cpp b/Engine/Source/Runtime/Core/Private/Internationalization/TextHistory.cpp
--- a/Engine/Source/Runtime/Core/Private/Internationalization/TextHistory.cpp
+++ b/Engine/Source/Runtime/Core/Private/Internationalization/TextHistory.cpp
@@ -405,8 +405,8 @@ void FTextHistory_AsCurrency::Serialize( FArchive& Ar )
 ///////////////////////////////////////
 // FTextHistory_AsDate
 
-FTextHistory_AsDate::FTextHistory_AsDate(const FDateTime::FDate& InSourceDate, const EDateTimeStyle::Type InDateStyle, const TSharedPtr<FCulture> InTargetCulture)
-	: SourceDate(InSourceDate)
+FTextHistory_AsDate::FTextHistory_AsDate(const FDateTime& InSourceDateTime, const EDateTimeStyle::Type InDateStyle, const TSharedPtr<FCulture> InTargetCulture)
+	: SourceDateTime(InSourceDateTime)
 	, DateStyle(InDateStyle)
 	, TargetCulture(InTargetCulture)
 {
@@ -420,7 +420,7 @@ void FTextHistory_AsDate::Serialize(FArchive& Ar)
 		Ar << HistoryType;
 	}
 
-	Ar << SourceDate;
+	Ar << SourceDateTime;
 
 	int8 DateStyleInt8 = (int8)DateStyle;
 	Ar << DateStyleInt8;
@@ -445,14 +445,14 @@ void FTextHistory_AsDate::Serialize(FArchive& Ar)
 
 FText FTextHistory_AsDate::ToText() const
 {
-	return FText::AsDate(SourceDate, DateStyle, TargetCulture);
+	return FText::AsDate(SourceDateTime, DateStyle, TargetCulture);
 }
 
 ///////////////////////////////////////
 // FTextHistory_AsTime
 
-FTextHistory_AsTime::FTextHistory_AsTime(const FDateTime::FTime& InSourceTime, const EDateTimeStyle::Type InTimeStyle, const FString& InTimeZone, const TSharedPtr<FCulture> InTargetCulture)
-	: SourceTime(InSourceTime)
+FTextHistory_AsTime::FTextHistory_AsTime(const FDateTime& InSourceDateTime, const EDateTimeStyle::Type InTimeStyle, const FString& InTimeZone, const TSharedPtr<FCulture> InTargetCulture)
+	: SourceDateTime(InSourceDateTime)
 	, TimeStyle(InTimeStyle)
 	, TimeZone(InTimeZone)
 	, TargetCulture(InTargetCulture)
@@ -467,7 +467,7 @@ void FTextHistory_AsTime::Serialize(FArchive& Ar)
 		Ar << HistoryType;
 	}
 
-	Ar << SourceTime;
+	Ar << SourceDateTime;
 
 	int8 TimeStyleInt8 = (int8)TimeStyle;
 	Ar << TimeStyleInt8;
@@ -494,7 +494,7 @@ void FTextHistory_AsTime::Serialize(FArchive& Ar)
 
 FText FTextHistory_AsTime::ToText() const
 {
-	return FText::AsTime(SourceTime, TimeStyle, TimeZone, TargetCulture);
+	return FText::AsTime(SourceDateTime, TimeStyle, TimeZone, TargetCulture);
 }
 
 ///////////////////////////////////////
diff --git a/Engine/Source/Runtime/Core/Private/Internationalization/TextHistory.h b/Engine/Source/Runtime/Core/Private/Internationalization/TextHistory.h
--- a/Engine/Source/Runtime/Core/Private/Internationalization/TextHistory.h
+++ b/Engine/Source/Runtime/Core/Private/Internationalization/TextHistory.h
@@ -183,7 +183,7 @@ class CORE_API FTextHistory_AsDate : public FTextHistory
 {
 public:
 	FTextHistory_AsDate() {};
-	FTextHistory_AsDate(const FDateTime::FDate& InSourceDate, const EDateTimeStyle::Type InDateStyle, const TSharedPtr<FCulture> InTargetCulture);
+	FTextHistory_AsDate(const FDateTime& InSourceDateTime, const EDateTimeStyle::Type InDateStyle, const TSharedPtr<FCulture> InTargetCulture);
 
 	// Begin FTextHistory interface
 	virtual FText ToText() const OVERRIDE;
@@ -192,7 +192,7 @@ public:
 
 private:
 	/** The source date structure to format */
-	FDateTime::FDate SourceDate;
+	FDateTime SourceDateTime;
 	/** Style to format the date using */
 	EDateTimeStyle::Type DateStyle;
 	/** Culture to format the date in */
@@ -204,7 +204,7 @@ class CORE_API FTextHistory_AsTime : public FTextHistory
 {
 public:
 	FTextHistory_AsTime() {};
-	FTextHistory_AsTime(const FDateTime::FTime& InSourceTime, const EDateTimeStyle::Type InTimeStyle, const FString& InTimeZone, const TSharedPtr<FCulture> InTargetCulture);
+	FTextHistory_AsTime(const FDateTime& InSourceDateTime, const EDateTimeStyle::Type InTimeStyle, const FString& InTimeZone, const TSharedPtr<FCulture> InTargetCulture);
 
 	// Begin FTextHistory interface
 	virtual FText ToText() const OVERRIDE;
@@ -213,7 +213,7 @@ public:
 
 private:
 	/** The source time structure to format */
-	FDateTime::FTime SourceTime;
+	FDateTime SourceDateTime;
 	/** Style to format the time using */
 	EDateTimeStyle::Type TimeStyle;
 	/** Timezone to put the time in */
diff --git a/Engine/Source/Runtime/Core/Private/Misc/DateTime.cpp b/Engine/Source/Runtime/Core/Private/Misc/DateTime.cpp
--- a/Engine/Source/Runtime/Core/Private/Misc/DateTime.cpp
+++ b/Engine/Source/Runtime/Core/Private/Misc/DateTime.cpp
@@ -6,24 +6,6 @@
 
 #include "CorePrivate.h"
 
-FArchive& operator<<(FArchive& Ar, FDateTime::FDate& Value)
-{
-	Ar << Value.Day;
-	Ar << Value.Month;
-	Ar << Value.Year;
-
-	return Ar;
-}
-
-FArchive& operator<<(FArchive& Ar, FDateTime::FTime& Value)
-{
-	Ar << Value.Hour;
-	Ar << Value.Minute;
-	Ar << Value.Second;
-	Ar << Value.Millisecond;
-
-	return Ar;
-}
 
 /* FDateTime constants
  *****************************************************************************/
@@ -82,6 +64,40 @@ FArchive& operator<<( FArchive& Ar, FDateTime& DateTime )
 /* FDateTime interface
  *****************************************************************************/
 
+void FDateTime::GetDate( int32& OutYear, int32& OutMonth, int32& OutDay ) const
+{
+	// Based on FORTRAN code in:
+	// Fliegel, H. F. and van Flandern, T. C.,
+	// Communications of the ACM, Vol. 11, No. 10 (October 1968).
+
+	int32 i, j, k, l, n;
+
+	l = FMath::FloorToInt(GetJulianDay() + 0.5) + 68569;
+	n = 4 * l / 146097;
+	l = l - (146097 * n + 3) / 4;
+	i = 4000 * (l + 1) / 1461001;
+	l = l - 1461 * i / 4 + 31;
+	j = 80 * l / 2447;
+	k = l - 2447 * j / 80;
+	l = j / 11;
+	j = j + 2 - 12 * l;
+	i = 100 * (n - 49) + i + l;
+
+	OutYear = i;
+	OutMonth = j;
+	OutDay = k;
+}
+
+
+int32 FDateTime::GetDay( ) const
+{
+	int32 Year, Month, Day;
+	GetDate(Year, Month, Day);
+
+	return Day;
+}
+
+
 EDayOfWeek::Type FDateTime::GetDayOfWeek( ) const
 {
 	// January 1, 0001 was a Monday
@@ -91,12 +107,12 @@ EDayOfWeek::Type FDateTime::GetDayOfWeek( ) const
 
 int32 FDateTime::GetDayOfYear( ) const
 {
-	FDate Date = ToDate();
-	int32 Day = Date.Day;
+	int32 Year, Month, Day;
+	GetDate(Year, Month, Day);
 
-	for (int32 Month = 1; Month < Date.Month; Month++)
+	for (int32 CurrentMonth = 1; CurrentMonth < Month; ++CurrentMonth)
 	{
-		Day += DaysInMonth(Date.Year, Month);
+		Day += DaysInMonth(Year, CurrentMonth);
 	}
 
 	return Day;
@@ -121,32 +137,21 @@ int32 FDateTime::GetHour12( ) const
 }
 
 
-FDateTime::FDate FDateTime::ToDate( ) const
+int32 FDateTime::GetMonth( ) const
 {
-	FDate Result;
+	int32 Year, Month, Day;
+	GetDate(Year, Month, Day);
 
-	// Based on FORTRAN code in:
-	// Fliegel, H. F. and van Flandern, T. C.,
-	// Communications of the ACM, Vol. 11, No. 10 (October 1968).
+	return Month;
+}
 
-	int32 i, j, k, l, n;
 
-	l = FMath::FloorToInt(GetJulianDay() + 0.5) + 68569;
-	n = 4 * l / 146097;
-	l = l - (146097 * n + 3) / 4;
-	i = 4000 * (l + 1) / 1461001;
-	l = l - 1461 * i / 4 + 31;
-	j = 80 * l / 2447;
-	k = l - 2447 * j / 80;
-	l = j / 11;
-	j = j + 2 - 12 * l;
-	i = 100 * (n - 49) + i + l;
-
-	Result.Year = i;
-	Result.Month = j;
-	Result.Day = k;
+int32 FDateTime::GetYear( ) const
+{
+	int32 Year, Month, Day;
+	GetDate(Year, Month, Day);
 
-	return Result;
+	return Year;
 }
 
 
@@ -203,18 +208,6 @@ FString FDateTime::ToString( const TCHAR* Format ) const
 }
 
 
-FDateTime::FTime FDateTime::ToTime( ) const
-{
-	FTime Result;
-
-	Result.Hour = GetHour();
-	Result.Minute = GetMinute();
-	Result.Second = GetSecond();
-	Result.Millisecond = GetMillisecond();
-
-	return Result;
-}
-
 /* FDateTime static interface
  *****************************************************************************/
 
@@ -255,13 +248,12 @@ bool FDateTime::IsLeapYear( int32 Year )
 
 FDateTime FDateTime::Now( )
 {
-	FDate Date;
-	int32 DayOfWeek;
-	FTime Time;
+	int32 Year, Month, Day, DayOfWeek;
+	int32 Hour, Minute, Second, Millisecond;
 
-	FPlatformTime::SystemTime(Date.Year, Date.Month, DayOfWeek, Date.Day, Time.Hour, Time.Minute, Time.Second, Time.Millisecond);
+	FPlatformTime::SystemTime(Year, Month, DayOfWeek, Day, Hour, Minute, Second, Millisecond);
 
-	return FDateTime(Date.Year, Date.Month, Date.Day, Time.Hour, Time.Minute, Time.Second, Time.Millisecond);
+	return FDateTime(Year, Month, Day, Hour, Minute, Second, Millisecond);
 }
 
 
@@ -428,13 +420,12 @@ bool FDateTime::ParseIso8601( const TCHAR* DateTimeString, FDateTime& OutDateTim
 
 FDateTime FDateTime::UtcNow( )
 {
-	FDate Date;
-	int32 DayOfWeek;
-	FTime Time;
+	int32 Year, Month, Day, DayOfWeek;
+	int32 Hour, Minute, Second, Millisecond;
 
-	FPlatformTime::UtcTime(Date.Year, Date.Month, DayOfWeek, Date.Day, Time.Hour, Time.Minute, Time.Second, Time.Millisecond);
+	FPlatformTime::UtcTime(Year, Month, DayOfWeek, Day, Hour, Minute, Second, Millisecond);
 
-	return FDateTime(Date.Year, Date.Month, Date.Day, Time.Hour, Time.Minute, Time.Second, Time.Millisecond);
+	return FDateTime(Year, Month, Day, Hour, Minute, Second, Millisecond);
 }
 
 
diff --git a/Engine/Source/Runtime/Core/Private/Tests/Internationalization/DateTimeFormattingRulesTest.cpp b/Engine/Source/Runtime/Core/Private/Tests/Internationalization/DateTimeFormattingRulesTest.cpp
--- a/Engine/Source/Runtime/Core/Private/Tests/Internationalization/DateTimeFormattingRulesTest.cpp
+++ b/Engine/Source/Runtime/Core/Private/Tests/Internationalization/DateTimeFormattingRulesTest.cpp
@@ -32,9 +32,7 @@ bool FDateTimeFormattingRulesTest::RunTest (const FString& Parameters)
 	const FDateTime UnixBillennium = FDateTime::FromUnixTimestamp(1000000000);
 	const FDateTime UnixOnes = FDateTime::FromUnixTimestamp(1111111111);
 	const FDateTime UnixDecimalSequence = FDateTime::FromUnixTimestamp(1234567890);
-	const FDateTime::FDate TestDate = { 13, 6, 1990 };
-	const FDateTime::FTime TestTime = { 12, 34, 56, 789 };
-	const FDateTime TestDateTime( TestDate.Year, TestDate.Month, TestDate.Day, TestTime.Hour, TestTime.Minute, TestTime.Second, TestTime.Millisecond );
+	const FDateTime TestDateTime( 1990, 6, 13, 12, 34, 56, 789 );
 
 	I18N.SetCurrentCulture("en-US");
 
@@ -59,18 +57,6 @@ bool FDateTimeFormattingRulesTest::RunTest (const FString& Parameters)
 	Test( this, TEXT("Testing Unix Decimal Sequence"), FText::AsDateTime(UnixDecimalSequence, EDateTimeStyle::Long, EDateTimeStyle::Long, "GMT"), FText::FromString( TEXT("February 13, 2009 at 11:31:30 PM GMT") ) );
 	Test( this, TEXT("Testing Unix Decimal Sequence"), FText::AsDateTime(UnixDecimalSequence, EDateTimeStyle::Full, EDateTimeStyle::Full, "GMT"), FText::FromString( TEXT("Friday, February 13, 2009 at 11:31:30 PM GMT") ) );
 
-	// Date
-	Test( this, TEXT("Testing Date"), FText::AsDate(TestDate, EDateTimeStyle::Short), FText::FromString( TEXT("6/13/90") ) );
-	Test( this, TEXT("Testing Date"), FText::AsDate(TestDate, EDateTimeStyle::Medium), FText::FromString( TEXT("Jun 13, 1990") ) );
-	Test( this, TEXT("Testing Date"), FText::AsDate(TestDate, EDateTimeStyle::Long), FText::FromString( TEXT("June 13, 1990") ) );
-	Test( this, TEXT("Testing Date"), FText::AsDate(TestDate, EDateTimeStyle::Full), FText::FromString( TEXT("Wednesday, June 13, 1990") ) );
-
-	// Time
-	Test( this, TEXT("Testing Time"), FText::AsTime(TestTime, EDateTimeStyle::Short, "GMT"), FText::FromString( TEXT("12:34 PM") ) );
-	Test( this, TEXT("Testing Time"), FText::AsTime(TestTime, EDateTimeStyle::Medium, "GMT"), FText::FromString( TEXT("12:34:56 PM") ) );
-	Test( this, TEXT("Testing Time"), FText::AsTime(TestTime, EDateTimeStyle::Long, "GMT"), FText::FromString( TEXT("12:34:56 PM GMT") ) );
-	Test( this, TEXT("Testing Time"), FText::AsTime(TestTime, EDateTimeStyle::Full, "GMT"), FText::FromString( TEXT("12:34:56 PM GMT") ) );
-
 	// Date Time
 	Test( this, TEXT("Testing Date-Time"), FText::AsDateTime(TestDateTime, EDateTimeStyle::Short,EDateTimeStyle::Short, "GMT"), FText::FromString( TEXT("6/13/90, 12:34 PM") ) );
 	Test( this, TEXT("Testing Date-Time"), FText::AsDateTime(TestDateTime, EDateTimeStyle::Medium, EDateTimeStyle::Medium, "GMT"), FText::FromString( TEXT("Jun 13, 1990, 12:34:56 PM") ) );
@@ -100,18 +86,6 @@ bool FDateTimeFormattingRulesTest::RunTest (const FString& Parameters)
 	Test( this, TEXT("Testing Unix Decimal Sequence"), FText::AsDateTime(UnixDecimalSequence, EDateTimeStyle::Long, EDateTimeStyle::Long, "GMT"), FText::FromString( TEXT("2009\x5E74") TEXT("2\x6708") TEXT("13\x65E5 23:31:30 GMT") ) );
 	Test( this, TEXT("Testing Unix Decimal Sequence"), FText::AsDateTime(UnixDecimalSequence, EDateTimeStyle::Full, EDateTimeStyle::Full, "GMT"), FText::FromString( TEXT("2009\x5E74") TEXT("2\x6708") TEXT("13\x65E5\x91D1\x66DC\x65E5") TEXT(" ") TEXT("23\x6642") TEXT("31\x5206") TEXT("30\x79D2") TEXT(" GMT") ) );
 
-	// Date
-	Test( this, TEXT("Testing Date"), FText::AsDate(TestDate, EDateTimeStyle::Short), FText::FromString( TEXT("1990/06/13") ) );
-	Test( this, TEXT("Testing Date"), FText::AsDate(TestDate, EDateTimeStyle::Medium), FText::FromString( TEXT("1990/06/13") ) );
-	Test( this, TEXT("Testing Date"), FText::AsDate(TestDate, EDateTimeStyle::Long), FText::FromString( TEXT("1990\x5E74") TEXT("6\x6708") TEXT("13\x65E5") ) );
-	Test( this, TEXT("Testing Date"), FText::AsDate(TestDate, EDateTimeStyle::Full), FText::FromString( TEXT("1990\x5E74") TEXT("6\x6708") TEXT("13\x65E5\x6C34\x66DC\x65E5") ) );
-
-	// Time
-	Test( this, TEXT("Testing Time"), FText::AsTime(TestTime, EDateTimeStyle::Short, "GMT"), FText::FromString( TEXT("12:34") ) );
-	Test( this, TEXT("Testing Time"), FText::AsTime(TestTime, EDateTimeStyle::Medium, "GMT"), FText::FromString( TEXT("12:34:56") ) );
-	Test( this, TEXT("Testing Time"), FText::AsTime(TestTime, EDateTimeStyle::Long, "GMT"), FText::FromString( TEXT("12:34:56 GMT") ) );
-	Test( this, TEXT("Testing Time"), FText::AsTime(TestTime, EDateTimeStyle::Full, "GMT"), FText::FromString( TEXT("12\x6642") TEXT("34\x5206") TEXT("56\x79D2") TEXT(" GMT") ) );
-
 	// Date Time
 	Test( this, TEXT("Testing Date-Time"), FText::AsDateTime(TestDateTime, EDateTimeStyle::Short,EDateTimeStyle::Short, "GMT"), FText::FromString( TEXT("1990/06/13 12:34") ) );
 	Test( this, TEXT("Testing Date-Time"), FText::AsDateTime(TestDateTime, EDateTimeStyle::Medium, EDateTimeStyle::Medium, "GMT"), FText::FromString( TEXT("1990/06/13 12:34:56") ) );
diff --git a/Engine/Source/Runtime/Core/Private/Tests/Internationalization/TextTest.cpp b/Engine/Source/Runtime/Core/Private/Tests/Internationalization/TextTest.cpp
--- a/Engine/Source/Runtime/Core/Private/Tests/Internationalization/TextTest.cpp
+++ b/Engine/Source/Runtime/Core/Private/Tests/Internationalization/TextTest.cpp
@@ -425,19 +425,6 @@ bool FTextTest::RunTest (const FString& Parameters)
 			FText AsPercentTest1 = FText::AsPercent(0.925);
 			FText AsCurrencyTest1 = FText::AsCurrency(100.25);
 
-			FDateTime::FDate DateInfo;
-			DateInfo.Day = 20;
-			DateInfo.Month = 5;
-			DateInfo.Year = 2080;
-			FText AsDateTest1 = FText::AsDate(DateInfo);
-
-			FDateTime::FTime TimeInfo;
-			TimeInfo.Hour = 8;
-			TimeInfo.Minute = 25;
-			TimeInfo.Second = 44;
-			TimeInfo.Millisecond = 5;
-			FText AsTimeTest1 = FText::AsTime(TimeInfo);
-
 			FDateTime DateTimeInfo(2080, 8, 20, 9, 33, 22);
 			FText AsDateTimeTest1 = FText::AsDateTime(DateTimeInfo);
 
@@ -446,8 +433,6 @@ bool FTextTest::RunTest (const FString& Parameters)
 			ArgsLayer2.Add("OrderedLayer1", FormattedTestOrdered1);
 			ArgsLayer2.Add("FTextNumber", AsNumberTest1);
 			ArgsLayer2.Add("Number", 5010.89221);
-			ArgsLayer2.Add("Date", AsDateTest1);
-			ArgsLayer2.Add("Time", AsTimeTest1);
 			ArgsLayer2.Add("DateTime", AsDateTimeTest1);
 			ArgsLayer2.Add("Percent", AsPercentTest1);
 			ArgsLayer2.Add("Currency", AsCurrencyTest1);
@@ -495,16 +480,6 @@ bool FTextTest::RunTest (const FString& Parameters)
 					AddError( TEXT("AsCurrencyTest1 did not rebuild correctly in French-Canadian") );
 				}
 
-				if(AsDateTest1.CompareTo(FText::AsDate(DateInfo)) != 0)
-				{
-					AddError( TEXT("AsDateTest1 did not rebuild correctly in French-Canadian") );
-				}
-
-				if(AsTimeTest1.CompareTo(FText::AsTime(TimeInfo)) != 0)
-				{
-					AddError( TEXT("AsTimeTest1 did not rebuild correctly in French-Canadian") );
-				}
-
 				if(AsDateTimeTest1.CompareTo(FText::AsDateTime(DateTimeInfo)) != 0)
 				{
 					AddError( TEXT("AsDateTimeTest1 did not rebuild correctly in French-Canadian") );
diff --git a/Engine/Source/Runtime/Core/Private/Tests/Misc/DateTimeTest.cpp b/Engine/Source/Runtime/Core/Private/Tests/Misc/DateTimeTest.cpp
--- a/Engine/Source/Runtime/Core/Private/Tests/Misc/DateTimeTest.cpp
+++ b/Engine/Source/Runtime/Core/Private/Tests/Misc/DateTimeTest.cpp
@@ -22,9 +22,7 @@ bool FDateTimeTest::RunTest( const FString& Parameters )
 	const FDateTime UnixOnes = FDateTime::FromUnixTimestamp(UnixOnesTimestamp);
 	const FDateTime UnixDecimalSequence = FDateTime::FromUnixTimestamp(UnixDecimalSequenceTimestamp);
 
-	const FDateTime::FDate TestDate = { 14, 8, 2013 };
-	const FDateTime::FTime TestTime = { 12, 34, 56, 789 };
-	const FDateTime TestDateTime( TestDate.Year, TestDate.Month, TestDate.Day, TestTime.Hour, TestTime.Minute, TestTime.Second, TestTime.Millisecond );
+	const FDateTime TestDateTime(2013, 8, 14, 12, 34, 56, 789);
 	
 #define TestUnixEquivalent( Desc, A, B ) if( (A).ToUnixTimestamp() != (B) ) AddError(FString::Printf(TEXT("%s - A=%d B=%d"),Desc,(A).ToUnixTimestamp(),(B)));
 
@@ -75,13 +73,13 @@ bool FDateTimeTest::RunTest( const FString& Parameters )
 	TestSecond( TEXT("Testing Unix Decimal Sequence Second"), UnixDecimalSequence, 30 );
 	TestMillisecond( TEXT("Testing Unix Decimal Sequence Millisecond"), UnixDecimalSequence, 0 );
 
-	TestYear( TEXT("Testing Test Date Time Year"), TestDateTime, TestDate.Year );
-	TestMonth( TEXT("Testing Test Date Time Month"), TestDateTime, TestDate.Month );
-	TestDay( TEXT("Testing Test Date Time Day"), TestDateTime, TestDate.Day );
-	TestHour( TEXT("Testing Test Date Time Hour"), TestDateTime, TestTime.Hour );
-	TestMinute( TEXT("Testing Test Date Time Minute"), TestDateTime, TestTime.Minute );
-	TestSecond( TEXT("Testing Test Date Time Second"), TestDateTime, TestTime.Second );
-	TestMillisecond( TEXT("Testing Test Date Time Millisecond"), TestDateTime, TestTime.Millisecond );
+	TestYear( TEXT("Testing Test Date Time Year"), TestDateTime, 2013 );
+	TestMonth( TEXT("Testing Test Date Time Month"), TestDateTime, 8 );
+	TestDay( TEXT("Testing Test Date Time Day"), TestDateTime, 14 );
+	TestHour( TEXT("Testing Test Date Time Hour"), TestDateTime, 12 );
+	TestMinute( TEXT("Testing Test Date Time Minute"), TestDateTime, 34 );
+	TestSecond( TEXT("Testing Test Date Time Second"), TestDateTime, 56 );
+	TestMillisecond( TEXT("Testing Test Date Time Millisecond"), TestDateTime, 789 );
 
 #undef TestYear
 #undef TestMonth
diff --git a/Engine/Source/Runtime/Core/Public/Internationalization/Text.h b/Engine/Source/Runtime/Core/Public/Internationalization/Text.h
--- a/Engine/Source/Runtime/Core/Public/Internationalization/Text.h
+++ b/Engine/Source/Runtime/Core/Public/Internationalization/Text.h
@@ -159,10 +159,10 @@ public:
 	/**
 	 * Generate an FText that represents the passed number as a date and/or time in the current culture
 	 */
-	static FText AsDate(const FDateTime::FDate& Date, const EDateTimeStyle::Type DateStyle = EDateTimeStyle::Default, const TSharedPtr<FCulture>& TargetCulture = NULL);
-	static FText AsTime(const FDateTime::FTime& Time, const EDateTimeStyle::Type TimeStyle = EDateTimeStyle::Default, const FString& TimeZone = TEXT(""), const TSharedPtr<FCulture>& TargetCulture = NULL);
-	static FText AsTime(const FTimespan& Time, const TSharedPtr<FCulture>& TargetCulture = NULL);
+	static FText AsDate(const FDateTime& DateTime, const EDateTimeStyle::Type DateStyle = EDateTimeStyle::Default, const TSharedPtr<FCulture>& TargetCulture = NULL);
 	static FText AsDateTime(const FDateTime& DateTime, const EDateTimeStyle::Type DateStyle = EDateTimeStyle::Default, const EDateTimeStyle::Type TimeStyle = EDateTimeStyle::Default, const FString& TimeZone = TEXT(""), const TSharedPtr<FCulture>& TargetCulture = NULL);
+	static FText AsTime(const FDateTime& DateTime, const EDateTimeStyle::Type TimeStyle = EDateTimeStyle::Default, const FString& TimeZone = TEXT(""), const TSharedPtr<FCulture>& TargetCulture = NULL);
+	static FText AsTimespan(const FTimespan& Time, const TSharedPtr<FCulture>& TargetCulture = NULL);
 
 	/**
 	 * Generate an FText that represents the passed number as a memory size in the current culture
diff --git a/Engine/Source/Runtime/Core/Public/Misc/DateTime.h b/Engine/Source/Runtime/Core/Public/Misc/DateTime.h
--- a/Engine/Source/Runtime/Core/Public/Misc/DateTime.h
+++ b/Engine/Source/Runtime/Core/Public/Misc/DateTime.h
@@ -79,68 +79,14 @@ class FDateTime
 public:
 
 	/**
-	 * Holds the components of a Gregorian date.
-	 */
-	struct FDate
-	{
-		/**
-		 * Holds the date's day.
-		 */
-		int32 Day;
-
-		/**
-		 * Holds the date's month.
-		 */
-		int32 Month;
-
-		/**
-		 * Holds the date's year.
-		 */
-		int32 Year;
-
-		friend FArchive& operator<<(FArchive& Ar, FDate& Value);
-	};
-
-
-	/**
-	 * Holds the components of a time of day.
-	 */
-	struct FTime
-	{
-		/**
-		 * Holds the time of day's hour.
-		 */
-		int32 Hour;
-
-		/**
-		 * Holds the time of day's minute.
-		 */
-		int32 Minute;
-
-		/**
-		 * Holds the time of day's second.
-		 */
-		int32 Second;
-
-		/**
-		 * Holds the time of day's millisecond.
-		 */
-		int32 Millisecond;
-
-		friend FArchive& operator<<(FArchive& Ar, FTime& Value);
-	};
-
-public:
-
-	/**
-	 * Default constructor.
+	 * Default constructor (no initialization).
 	 */
 	FDateTime( ) { }
 
 	/**
 	 * Creates and initializes a new instance with the specified number of ticks.
 	 *
-	 * @param Ticks - The ticks representing the date and time.
+	 * @param Ticks The ticks representing the date and time.
 	 */
 	FDateTime( int64 InTicks )
 		: Ticks(InTicks)
@@ -149,13 +95,13 @@ public:
 	/**
 	 * Creates and initializes a new instance with the specified year, month, day, hour, minute, second and millisecond.
 	 *
-	 * @param Year - The year.
-	 * @param Month - The month.
-	 * @param Day - The day.
-	 * @param Hour - The hour (optional).
-	 * @param Minute - The minute (optional).
-	 * @param Second - The second (optional).
-	 * @param Millisecond - The millisecond (optional).
+	 * @param Year The year.
+	 * @param Month The month.
+	 * @param Day The day.
+	 * @param Hour The hour (optional).
+	 * @param Minute The minute (optional).
+	 * @param Second The second (optional).
+	 * @param Millisecond The millisecond (optional).
 	 */
 	CORE_API FDateTime( int32 Year, int32 Month, int32 Day, int32 Hour = 0, int32 Minute = 0, int32 Second = 0, int32 Millisecond = 0 );
 
@@ -177,7 +123,6 @@ public:
 	 * Adds the given time span to this date.
 	 *
 	 * @return This date.
-	 *
 	 * @see FTimespan
 	 */
 	FDateTime& operator+=( const FTimespan& Other )
@@ -191,7 +136,6 @@ public:
 	 * Returns time span between this date and the given date.
 	 *
 	 * @return A time span whose value is the difference of this date and the given date.
-	 *
 	 * @see FTimespan
 	 */
 	FTimespan operator-( const FDateTime& Other ) const
@@ -203,7 +147,6 @@ public:
 	 * Returns result of subtracting the given time span from this date.
 	 *
 	 * @return A date whose value is the difference of this date and the given time span.
-	 *
 	 * @see FTimespan
 	 */
 	FDateTime operator-( const FTimespan& Other ) const
@@ -215,7 +158,6 @@ public:
 	 * Subtracts the given time span from this date.
 	 *
 	 * @return This date.
-	 *
 	 * @see FTimespan
 	 */
 	FDateTime& operator-=( const FTimespan& Other )
@@ -228,8 +170,7 @@ public:
 	/**
 	 * Compares this date with the given date for equality.
 	 *
-	 * @param other - The date to compare with.
-	 *
+	 * @param other The date to compare with.
 	 * @return true if the dates are equal, false otherwise.
 	 */
 	bool operator==( const FDateTime& Other ) const
@@ -240,8 +181,7 @@ public:
 	/**
 	 * Compares this date with the given date for inequality.
 	 *
-	 * @param other - The date to compare with.
-	 *
+	 * @param other The date to compare with.
 	 * @return true if the dates are not equal, false otherwise.
 	 */
 	bool operator!=( const FDateTime& Other ) const
@@ -252,8 +192,7 @@ public:
 	/**
 	 * Checks whether this date is greater than the given date.
 	 *
-	 * @param other - The date to compare with.
-	 *
+	 * @param other The date to compare with.
 	 * @return true if this date is greater, false otherwise.
 	 */
 	bool operator>( const FDateTime& Other ) const
@@ -264,8 +203,7 @@ public:
 	/**
 	 * Checks whether this date is greater than or equal to the date span.
 	 *
-	 * @param other - The date to compare with.
-	 *
+	 * @param other The date to compare with.
 	 * @return true if this date is greater or equal, false otherwise.
 	 */
 	bool operator>=( const FDateTime& Other ) const
@@ -276,8 +214,7 @@ public:
 	/**
 	 * Checks whether this date is less than the given date.
 	 *
-	 * @param other - The date to compare with.
-	 *
+	 * @param other The date to compare with.
 	 * @return true if this date is less, false otherwise.
 	 */
 	bool operator<( const FDateTime& Other ) const
@@ -288,8 +225,7 @@ public:
 	/**
 	 * Checks whether this date is less than or equal to the given date.
 	 *
-	 * @param other - The date to compare with.
-	 *
+	 * @param other The date to compare with.
 	 * @return true if this date is less or equal, false otherwise.
 	 */
 	bool operator<=( const FDateTime& Other ) const
@@ -305,28 +241,34 @@ public:
 	 * The time part is truncated and becomes 00:00:00.000.
 	 *
 	 * @return A FDateTime object containing the date.
-	 *
-	 * @see FTimespan
 	 */
-	FDateTime GetDate( )
+	FDateTime GetDate( ) const
 	{
 		return FDateTime(Ticks - (Ticks % ETimespan::TicksPerDay));
 	}
 
+	/**
+	 * Gets the date components of this date.
+	 *
+	 * @param OutYear Will contain the year.
+	 * @param OutMonth Will contain the number of the month (1-12).
+	 * @param OutDay Will contain the number of the day (1-31).
+	 */
+	CORE_API void GetDate( int32& OutYear, int32& OutMonth, int32& OutDay ) const;
+
 	/**
 	 * Gets this date's day part (1 to 31).
 	 *
 	 * @return Day of the month.
+	 * @see GetHour, GetHour12, GetMillisecond, GetMinute, GetMonth, GetSecond, GetYear
 	 */
-	int32 GetDay( ) const
-	{
-		return ToDate().Day;
-	}
+	CORE_API int32 GetDay( ) const;
 
 	/**
 	 * Calculates this date's day of the week (Sunday - Saturday).
 	 *
 	 * @return The week day.
+	 * @see GetDayOfYear, GetMonthOfYear, GetTimeOfDay
 	 */
 	CORE_API EDayOfWeek::Type GetDayOfWeek( ) const;
 
@@ -334,6 +276,7 @@ public:
 	 * Gets this date's day of the year.
 	 *
 	 * @return The day of year.
+	 * @see GetDayOfWeek, GetMonthOfYear, GetTimeOfDay
 	 */
 	CORE_API int32 GetDayOfYear( ) const;
 
@@ -341,6 +284,7 @@ public:
 	 * Gets this date's hour part in 24-hour clock format (0 to 23).
 	 *
 	 * @return The hour.
+	 * @see GetDay, GetDayOfWeek, GetDayOfYear, GetHour12, GetMillisecond, GetMinute, GetMonth, GetSecond, GetYear
 	 */
 	int32 GetHour( ) const
 	{
@@ -351,6 +295,7 @@ public:
 	 * Gets this date's hour part in 12-hour clock format (1 to 12).
 	 *
 	 * @return The hour in AM/PM format.
+	 * @see GetDay, GetHour, GetMillisecond, GetMinute, GetMonth, GetSecond, GetYear
 	 */
 	CORE_API int32 GetHour12( ) const;
 
@@ -362,8 +307,7 @@ public:
 	 * 1721425.5, which corresponds to Monday, January 1, 0001 in the Gregorian calendar.
 	 *
 	 * @return Julian Day.
-	 *
-	 * @see GetModifiedJulianDay
+	 * @see FromJulianDay, GetModifiedJulianDay
 	 */
 	double GetJulianDay( ) const
 	{
@@ -377,7 +321,6 @@ public:
 	 * November 17, 1858 in the Gregorian calendar.
 	 *
 	 * @return Modified Julian Day
-	 *
 	 * @see GetJulianDay
 	 */
 	double GetModifiedJulianDay( ) const
@@ -389,6 +332,7 @@ public:
 	 * Gets this date's millisecond part (0 to 999).
 	 *
 	 * @return The millisecond.
+	 * @see GetDay, GetHour, GetHour12, GetMinute, GetMonth, GetSecond, GetYear
 	 */
 	int32 GetMillisecond( ) const
 	{
@@ -399,6 +343,7 @@ public:
 	 * Gets this date's minute part (0 to 59).
 	 *
 	 * @return The minute.
+	 * @see GetDay, GetHour, GetHour12, GetMillisecond, GetMonth, GetSecond, GetYear
 	 */
 	int32 GetMinute( ) const
 	{
@@ -409,16 +354,15 @@ public:
 	 * Gets this date's the month part (1 to 12).
 	 *
 	 * @return The month.
+	 * @see GetDay, GetHour, GetHour12, GetMillisecond, GetMinute, GetSecond, GetYear
 	 */
-	int32 GetMonth( ) const
-	{
-		return ToDate().Month;
-	}
+	CORE_API int32 GetMonth( ) const;
 
 	/**
 	 * Gets the date's month of the year (January to December).
 	 *
 	 * @return Month of year.
+	 * @see GetDayOfWeek, GetDayOfYear, GetTimeOfDay
 	 */
 	EMonthOfYear::Type GetMonthOfYear( ) const
 	{
@@ -429,6 +373,7 @@ public:
 	 * Gets this date's second part.
 	 *
 	 * @return The second.
+	 * @see GetDay, GetHour, GetHour12, GetMillisecond, GetMinute, GetMonth, GetYear
 	 */
 	int32 GetSecond( ) const
 	{
@@ -449,6 +394,7 @@ public:
 	 * Gets the time elapsed since midnight of this date.
 	 *
 	 * @param Time of day since midnight.
+	 * @see GetDayOfWeek, GetDayOfYear, GetMonthOfYear
 	 */
 	FTimespan GetTimeOfDay( ) const
 	{
@@ -459,16 +405,15 @@ public:
 	 * Gets this date's year part.
 	 *
 	 * @return The year.
+	 * @see GetDay, GetHour, GetHour12, GetMillisecond, GetMinute, GetMonth, GetSecond
 	 */
-	int32 GetYear( ) const
-	{
-		return ToDate().Year;
-	}
+	CORE_API int32 GetYear( ) const;
 
 	/**
 	 * Gets whether this date's time is in the afternoon.
 	 *
 	 * @param true if it is in the afternoon, false otherwise.
+	 * @see IsMorning
 	 */
 	bool IsAfternoon( ) const
 	{
@@ -479,25 +424,20 @@ public:
 	 * Gets whether this date's time is in the morning.
 	 *
 	 * @param true if it is in the morning, false otherwise.
+	 * @see IsAfternoon
 	 */
 	bool IsMorning( ) const
 	{
 		return (GetHour() < 12);
 	}
 
-	/**
-	 * Returns this date as a Gregorian date structure.
-	 */
-	CORE_API FDate ToDate( ) const;
-
 	/**
 	 * Returns the ISO-8601 string representation of the FDateTime.
 	 *
 	 * The resulting string assumes that the FDateTime is in UTC.
 	 * 
 	 * @return String representation.
-	 *
-	 * @see FromIso8601
+	 * @see ParseIso8601, ToString
 	 */
 	CORE_API FString ToIso8601( ) const;
 
@@ -508,29 +448,24 @@ public:
 	 *		yyyy.mm.dd-hh.mm.ss
 	 *
 	 * @return String representation.
+	 * @see Parse, ToIso8601
 	 */
 	CORE_API FString ToString( ) const;
 
 	/**
 	 * Returns the string representation of this date.
 	 *
-	 * @param Format - The format of the returned string.
-	 *
+	 * @param Format The format of the returned string.
 	 * @return String representation.
+	 * @see Parse, ToIso8601
 	 */
 	CORE_API FString ToString( const TCHAR* Format ) const;
 
-	/**
-	 * Returns this date's time part as a time of day structure.
-	 *
-	 * @return Time of day.
-	 */
-	CORE_API FTime ToTime( ) const;
-
 	/**
 	 * Returns this date as the number of seconds since the Unix Epoch (January 1st of 1970).
 	 *
 	 * @return Time of day.
+	 * @see FromUnixTimestamp
 	 */
 	CORE_API int32 ToUnixTimestamp( ) const
 	{
@@ -542,28 +477,28 @@ public:
 	/**
 	 * Gets the number of days in the year and month.
 	 *
-	 * @param Year - The year.
-	 * @param Month - The month.
-	 *
+	 * @param Year The year.
+	 * @param Month The month.
 	 * @return The number of days
+	 * @see DaysInYear
 	 */
 	static CORE_API int32 DaysInMonth( int32 Year, int32 Month );
 
 	/**
 	 * Gets the number of days in the given year.
 	 *
-	 * @param Year - The year.
-	 *
+	 * @param Year The year.
 	 * @return The number of days.
+	 * @see DaysInMonth
 	 */
 	static CORE_API int32 DaysInYear( int32 Year );
 
 	/**
 	 * Returns the proleptic Gregorian date for the given Julian Day.
 	 *
-	 * @param JulianDay - The Julian Day.
-	 *
+	 * @param JulianDay The Julian Day.
 	 * @return Gregorian date and time.
+	 * @see GetJulianDay
 	 */
 	static FDateTime FromJulianDay( double JulianDay )
 	{
@@ -573,10 +508,8 @@ public:
 	/**
 	 * Returns the date from Unix time (seconds from midnight 1970-01-01)
 	 *
-	 * @param UnixTime - Unix time (seconds from midnight 1970-01-01)
-	 *
+	 * @param UnixTime Unix time (seconds from midnight 1970-01-01)
 	 * @return Gregorian date and time.
-	 *
 	 * @see ToUnixTimestamp
 	 */
 	static FDateTime FromUnixTimestamp( int32 UnixTime )
@@ -591,8 +524,7 @@ public:
 	 * with the astronomical year. All years divisible by 4, but not divisible by 100 - except if they
 	 * are also divisible by 400 - are leap years.
 	 *
-	 * @param Year - The year to check.
-	 *
+	 * @param Year The year to check.
 	 * @return true if the year is a leap year, false otherwise.
 	 */
 	static CORE_API bool IsLeapYear( int32 Year );
@@ -632,6 +564,7 @@ public:
 	 * Gets the local date and time on this computer.
 	 *
 	 * @return Current date and time.
+	 * @see Today
 	 */
 	static CORE_API FDateTime Now( );
 
@@ -641,9 +574,8 @@ public:
 	 * Currently, the string must be in the format written by either FDateTime.ToString() or
 	 * FTimeStamp.TimestampToFString(). Other formats are not supported at this time.
 	 *
-	 * @param DateTimeString - The string to convert.
-	 * @param OutDateTime - Will contain the parsed date and time.
-	 *
+	 * @param DateTimeString The string to convert.
+	 * @param OutDateTime Will contain the parsed date and time.
 	 * @return true if the string was converted successfully, false otherwise.
 	 */
 	static CORE_API bool Parse( const FString& DateTimeString, FDateTime& OutDateTime );
@@ -653,10 +585,8 @@ public:
 	 * 
 	 * @param DateTimeString The string to be parsed
 	 * @param OutDateTime FDateTime object (in UTC) corresponding to the input string (which may have been in any timezone).
-	 *
 	 * @return true if the string was converted successfully, false otherwise.
-	 *
-	 * @see ToIso8601
+	 * @see Parse, ToIso8601
 	 */
 	static CORE_API bool ParseIso8601( const TCHAR* DateTimeString, FDateTime& OutDateTime );
 
@@ -666,6 +596,7 @@ public:
 	 * The time component is set to 00:00:00
 	 *
 	 * @return Current date.
+	 * @see Now
 	 */
 	static FDateTime Today( )
 	{
@@ -684,9 +615,8 @@ public:
 	/**
 	 * Serializes the given date and time from or into the specified archive.
 	 *
-	 * @param Ar - The archive to serialize from or into.
-	 * @param DateTime - The date and time value to serialize.
-	 *
+	 * @param Ar The archive to serialize from or into.
+	 * @param DateTime The date and time value to serialize.
 	 * @return The archive.
 	 *
 	 * @todo gmp: Figure out better include order in Core.h so this can be inlined.
@@ -696,8 +626,7 @@ public:
 	/**
 	 * Gets the hash for the specified date and time.
 	 *
-	 * @param DateTime - The date and time to get the hash for.
-	 *
+	 * @param DateTime The date and time to get the hash for.
 	 * @return Hash value.
 	 */
 	friend uint32 GetTypeHash( const FDateTime& DateTime )
@@ -708,25 +637,25 @@ public:
 	/**
 	 * Exports the DateTime value to a string.
 	 *
-	 * @param ValueStr - Will hold the string value.
-	 * @param DefaultValue - The default value.
-	 * @param Parent - Not used.
-	 * @param PortFlags - Not used.
-	 * @param ExportRootScope - Not used.
-	 *
+	 * @param ValueStr Will hold the string value.
+	 * @param DefaultValue The default value.
+	 * @param Parent Not used.
+	 * @param PortFlags Not used.
+	 * @param ExportRootScope Not used.
 	 * @return true on success, false otherwise.
+	 * @see ImportTextItem
 	 */
 	CORE_API bool ExportTextItem( FString& ValueStr, FDateTime const& DefaultValue, UObject* Parent, int32 PortFlags, class UObject* ExportRootScope ) const;
 
 	/**
 	 * Imports the DateTime value from a text buffer.
 	 *
-	 * @param Buffer - The text buffer to import from.
-	 * @param PortFlags - Not used.
-	 * @param Parent - Not used.
-	 * @param ErrorText - The output device for error logging.
-	 *
+	 * @param Buffer The text buffer to import from.
+	 * @param PortFlags Not used.
+	 * @param Parent Not used.
+	 * @param ErrorText The output device for error logging.
 	 * @return true on success, false otherwise.
+	 * @see ExportTextItem
 	 */
 	CORE_API bool ImportTextItem( const TCHAR*& Buffer, int32 PortFlags, class UObject* Parent, FOutputDevice* ErrorText );
 
diff --git a/Engine/Source/Runtime/Core/Public/Misc/Timespan.h b/Engine/Source/Runtime/Core/Public/Misc/Timespan.h
--- a/Engine/Source/Runtime/Core/Public/Misc/Timespan.h
+++ b/Engine/Source/Runtime/Core/Public/Misc/Timespan.h
@@ -9,34 +9,22 @@
 
 namespace ETimespan
 {
-	/**
-	 * The number of timespan ticks per day.
-	 */
+	/** The number of timespan ticks per day. */
 	const int64 TicksPerDay = 864000000000;
 
-	/**
-	 * The number of timespan ticks per hour.
-	 */
+	/** The number of timespan ticks per hour. */
 	const int64 TicksPerHour = 36000000000;
 
-	/**
-	 * The number of timespan ticks per millisecond.
-	 */
+	/** The number of timespan ticks per millisecond. */
 	const int64 TicksPerMillisecond = 10000;
 
-	/**
-	 * The number of timespan ticks per minute.
-	 */
+	/** The number of timespan ticks per minute. */
 	const int64 TicksPerMinute = 600000000;
 
-	/**
-	 * The number of timespan ticks per second.
-	 */
+	/** The number of timespan ticks per second. */
 	const int64 TicksPerSecond = 10000000;
 
-	/**
-	 * The number of timespan ticks per week.
-	 */
+	/** The number of timespan ticks per week. */
 	const int64 TicksPerWeek = 6048000000000;
 }
 
@@ -67,7 +55,7 @@ public:
 	/**
 	 * Creates and initializes a new time interval with the specified number of ticks.
 	 *
-	 * @param Ticks - The number of ticks.
+	 * @param Ticks The number of ticks.
 	 */
 	FTimespan( int64 InTicks )
 		: Ticks(InTicks)
@@ -76,9 +64,9 @@ public:
 	/**
 	 * Creates and initializes a new time interval with the specified number of hours, minutes and seconds.
 	 *
-	 * @param Hours - The hours component.
-	 * @param Minutes - The minutes component.
-	 * @param Seconds - The seconds component.	 */
+	 * @param Hours The hours component.
+	 * @param Minutes The minutes component.
+	 * @param Seconds The seconds component.	 */
 	FTimespan( int32 Hours, int32 Minutes, int32 Seconds )
 	{
 		Assign(0, Hours, Minutes, Seconds, 0);
@@ -87,10 +75,10 @@ public:
 	/**
 	 * Creates and initializes a new time interval with the specified number of days, hours, minutes and seconds.
 	 *
-	 * @param Days - The days component.
-	 * @param Hours - The hours component.
-	 * @param Minutes - The minutes component.
-	 * @param Seconds - The seconds component.
+	 * @param Days The days component.
+	 * @param Hours The hours component.
+	 * @param Minutes The minutes component.
+	 * @param Seconds The seconds component.
 	 */
 	FTimespan( int32 Days, int32 Hours, int32 Minutes, int32 Seconds )
 	{
@@ -100,11 +88,11 @@ public:
 	/**
 	 * Creates and initializes a new time interval with the specified number of days, hours, minutes and seconds.
 	 *
-	 * @param Days - The days component.
-	 * @param Hours - The hours component.
-	 * @param Minutes - The minutes component.
-	 * @param Seconds - The seconds component.
-	 * @param Milliseconds - The milliseconds component.
+	 * @param Days The days component.
+	 * @param Hours The hours component.
+	 * @param Minutes The minutes component.
+	 * @param Seconds The seconds component.
+	 * @param Milliseconds The milliseconds component.
 	 */
 	FTimespan( int32 Days, int32 Hours, int32 Minutes, int32 Seconds, int32 Milliseconds )
 	{
@@ -150,8 +138,7 @@ public:
 	/**
 	 * Returns the result of subtracting the given time span from this time span.
 	 *
-	 * @param Other - The time span to compare with.
-	 *
+	 * @param Other The time span to compare with.
 	 * @return A time span whose value is the difference of this time span and the given time span.
 	 */
 	FTimespan operator-( const FTimespan& Other ) const
@@ -162,8 +149,7 @@ public:
 	/**
 	 * Subtracts the given time span from this time span.
 	 *
-	 * @param Other - The time span to subtract.
-	 *
+	 * @param Other The time span to subtract.
 	 * @return This time span.
 	 */
 	FTimespan& operator-=( const FTimespan& Other )
@@ -176,8 +162,7 @@ public:
 	/**
 	 * Returns the result of multiplying the this time span with the given scalar.
 	 *
-	 * @param Scalar - The scalar to multiply with.
-	 *
+	 * @param Scalar The scalar to multiply with.
 	 * @return A time span whose value is the product of this time span and the given scalar.
 	 */
 	FTimespan operator*( float Scalar ) const
@@ -188,8 +173,7 @@ public:
 	/**
 	 * Multiplies this time span with the given scalar.
 	 *
-	 * @param Scalar - The scalar to multiply with.
-	 *
+	 * @param Scalar The scalar to multiply with.
 	 * @return This time span.
 	 */
 	FTimespan& operator*=( float Scalar )
@@ -202,8 +186,7 @@ public:
 	/**
 	 * Compares this time span with the given time span for equality.
 	 *
-	 * @param Other - The time span to compare with.
-	 *
+	 * @param Other The time span to compare with.
 	 * @return true if the time spans are equal, false otherwise.
 	 */
 	bool operator==( const FTimespan& Other ) const
@@ -214,8 +197,7 @@ public:
 	/**
 	 * Compares this time span with the given time span for inequality.
 	 *
-	 * @param Other - The time span to compare with.
-	 *
+	 * @param Other The time span to compare with.
 	 * @return true if the time spans are not equal, false otherwise.
 	 */
 	bool operator!=( const FTimespan& Other ) const
@@ -226,8 +208,7 @@ public:
 	/**
 	 * Checks whether this time span is greater than the given time span.
 	 *
-	 * @param Other - The time span to compare with.
-	 *
+	 * @param Other The time span to compare with.
 	 * @return true if this time span is greater, false otherwise.
 	 */
 	bool operator>( const FTimespan& Other ) const
@@ -238,8 +219,7 @@ public:
 	/**
 	 * Checks whether this time span is greater than or equal to the given time span.
 	 *
-	 * @param Other - The time span to compare with.
-	 *
+	 * @param Other The time span to compare with.
 	 * @return true if this time span is greater or equal, false otherwise.
 	 */
 	bool operator>=( const FTimespan& Other ) const
@@ -250,8 +230,7 @@ public:
 	/**
 	 * Checks whether this time span is less than the given time span.
 	 *
-	 * @param Other - The time span to compare with.
-	 *
+	 * @param Other The time span to compare with.
 	 * @return true if this time span is less, false otherwise.
 	 */
 	bool operator<( const FTimespan& Other ) const
@@ -262,8 +241,7 @@ public:
 	/**
 	 * Checks whether this time span is less than or equal to the given time span.
 	 *
-	 * @param Other - The time span to compare with.
-	 *
+	 * @param Other The time span to compare with.
 	 * @return true if this time span is less or equal, false otherwise.
 	 */
 	bool operator<=( const FTimespan& Other ) const
@@ -289,7 +267,6 @@ public:
 	 * This method may overflow the timespan if its value is equal to MinValue.
 	 *
 	 * @return Duration of this time span.
-	 *
 	 * @see MinValue
 	 */
 	FTimespan GetDuration( )
@@ -301,7 +278,6 @@ public:
 	 * Gets the hours component of this time span.
 	 *
 	 * @return Hours component.
-	 *
 	 * @see GetTotalHours
 	 */
 	int32 GetHours( ) const
@@ -313,7 +289,6 @@ public:
 	 * Gets the milliseconds component of this time span.
 	 *
 	 * @return Milliseconds component.
-	 *
 	 * @see GetTotalMilliseconds
 	 */
 	int32 GetMilliseconds( ) const
@@ -325,7 +300,6 @@ public:
 	 * Gets the minutes component of this time span.
 	 *
 	 * @return Minutes component.
-	 *
 	 * @see GetTotalMinutes
 	 */
 	int32 GetMinutes( ) const
@@ -337,7 +311,6 @@ public:
 	 * Gets the seconds component of this time span.
 	 *
 	 * @return Seconds component.
-	 *
 	 * @see GetTotalSeconds
 	 */
 	int32 GetSeconds( ) const
@@ -359,7 +332,6 @@ public:
 	 * Gets the total number of days represented by this time span.
 	 *
 	 * @return Number of days.
-	 *
 	 * @see GetDays
 	 */
 	double GetTotalDays( ) const
@@ -371,7 +343,6 @@ public:
 	 * Gets the total number of hours represented by this time span.
 	 *
 	 * @return Number of hours.
-	 *
 	 * @see GetHours
 	 */
 	double GetTotalHours( ) const
@@ -383,7 +354,6 @@ public:
 	 * Gets the total number of milliseconds represented by this time span.
 	 *
 	 * @return Number of milliseconds.
-	 *
 	 * @see GetMilliseconds
 	 */
 	double GetTotalMilliseconds( ) const
@@ -395,7 +365,6 @@ public:
 	 * Gets the total number of minutes represented by this time span.
 	 *
 	 * @return Number of minutes.
-	 *
 	 * @see GetMinutes
 	 */
 	double GetTotalMinutes( ) const
@@ -407,7 +376,6 @@ public:
 	 * Gets the total number of seconds represented by this time span.
 	 *
 	 * @return Number of seconds.
-	 *
 	 * @see GetSeconds
 	 */
 	double GetTotalSeconds( ) const
@@ -422,6 +390,7 @@ public:
 	 *		[-][d.]hh:mm:ss.fff
 	 *
 	 * @return String representation.
+	 * @see Parse
 	 */
 	CORE_API FString ToString( ) const;
 
@@ -443,8 +412,8 @@ public:
 	 *		%F - prints the total number of milliseconds (0..999)
 	 *
 	 * @param Format - The format of the returned string.
-	 *
 	 * @return String representation.
+	 * @see Parse
 	 */
 	CORE_API FString ToString( const TCHAR* Format ) const;
 
@@ -453,70 +422,45 @@ public:
 	/**
 	 * Creates a time span that represents the specified number of days.
 	 *
-	 * @param Days - The number of days.
-	 *
+	 * @param Days The number of days.
 	 * @return Time span.
-	 *
-	 * @see FromHours
-	 * @see FromMilliseconds
-	 * @see FromMinutes
-	 * @see FromSeconds
+	 * @see FromHours, FromMilliseconds, FromMinutes, FromSeconds
 	 */
 	static CORE_API FTimespan FromDays( double Days );
 
 	/**
 	 * Creates a time span that represents the specified number of hours.
 	 *
-	 * @param Hours - The number of hours.
-	 *
+	 * @param Hours The number of hours.
 	 * @return Time span.
-	 *
-	 * @see FromDays
-	 * @see FromMilliseconds
-	 * @see FromMinutes
-	 * @see FromSeconds
+	 * @see FromDays, FromMilliseconds, FromMinutes, FromSeconds
 	 */
 	static CORE_API FTimespan FromHours( double Hours );
 
 	/**
 	 * Creates a time span that represents the specified number of milliseconds.
 	 *
-	 * @param Milliseconds - The number of milliseconds.
-	 *
+	 * @param Milliseconds The number of milliseconds.
 	 * @return Time span.
-	 *
-	 * @see FromDays
-	 * @see FromHours
-	 * @see FromMinutes
-	 * @see FromSeconds
+	 * @see FromDays, FromHours, FromMinutes, FromSeconds
 	 */
 	static CORE_API FTimespan FromMilliseconds( double Milliseconds );
 
 	/**
 	 * Creates a time span that represents the specified number of minutes.
 	 *
-	 * @param Minutes - The number of minutes.
-	 *
+	 * @param Minutes The number of minutes.
 	 * @return Time span.
-	 *
-	 * @see FromDays
-	 * @see FromHours
-	 * @see FromMilliseconds
-	 * @see FromSeconds
+	 * @see FromDays, FromHours, FromMilliseconds, FromSeconds
 	 */
 	static CORE_API FTimespan FromMinutes( double Minutes );
 
 	/**
 	 * Creates a time span that represents the specified number of seconds.
 	 *
-	 * @param Seconds - The number of seconds.
-	 *
+	 * @param Seconds The number of seconds.
 	 * @return Time span.
-	 *
-	 * @see FromDays
-	 * @see FromHours
-	 * @see FromMilliseconds
-	 * @see FromMinutes
+	 * @see FromDays, FromHours, FromMilliseconds, FromMinutes
 	 */
 	static CORE_API FTimespan FromSeconds( double Seconds );
 
@@ -526,9 +470,7 @@ public:
 	 * The maximum time span value is slightly more than 10,675,199 days.
 	 *
 	 * @return Maximum time span.
-	 *
-	 * @see MinValue
-	 * @see Zero
+	 * @see MinValue,Zero
 	 */
 	static FTimespan MaxValue( )
 	{
@@ -541,9 +483,7 @@ public:
 	 * The minimum time span value is slightly less than -10,675,199 days.
 	 *
 	 * @return Minimum time span.
-	 *
-	 * @see MaxValue
-	 * @see ZeroValue
+	 * @see MaxValue, ZeroValue
 	 */
 	static FTimespan MinValue( )
 	{
@@ -556,10 +496,10 @@ public:
 	 * Currently, the string must be in the format written by FTimespan.ToString().
 	 * Other formats are not supported at this time.
 	 *
-	 * @param TimespanString - The string to convert.
-	 * @param OutTimespan - Will contain the parsed time span.
-	 *
+	 * @param TimespanString The string to convert.
+	 * @param OutTimespan Will contain the parsed time span.
 	 * @return true if the string was converted successfully, false otherwise.
+	 * @see ToString
 	 */
 	static CORE_API bool Parse( const FString& TimespanString, FTimespan& OutTimespan );
 
@@ -569,9 +509,7 @@ public:
 	 * The zero time span value can be used in comparison operations with other time spans.
 	 *
 	 * @return Zero time span.
-	 *
-	 * @see MaxValue
-	 * @see MinValue
+	 * @see MaxValue, MinValue
 	 */
 	static FTimespan Zero( )
 	{
@@ -583,9 +521,8 @@ public:
 	/**
 	 * Serializes the given time span from or into the specified archive.
 	 *
-	 * @param Ar - The archive to serialize from or into.
-	 * @param Timespan - The time span value to serialize.
-	 *
+	 * @param Ar The archive to serialize from or into.
+	 * @param Timespan The time span value to serialize.
 	 * @return The archive.
 	 *
 	 * @todo gmp: Figure out better include order in Core.h so this can be inlined.
@@ -595,8 +532,7 @@ public:
 	/**
 	 * Gets the hash for the specified time span.
 	 *
-	 * @param Timespan - The timespan to get the hash for.
-	 *
+	 * @param Timespan The timespan to get the hash for.
 	 * @return Hash value.
 	 */
 	friend uint32 GetTypeHash( const FTimespan& Timespan )
@@ -609,11 +545,11 @@ protected:
 	/**
 	 * Assigns the specified components to this time span.
 	 *
-	 * @param Days - The days component.
-	 * @param Hours - The hours component.
-	 * @param Minutes - The minutes component.
-	 * @param Seconds - The seconds component.
-	 * @param Milliseconds - The milliseconds component.
+	 * @param Days The days component.
+	 * @param Hours The hours component.
+	 * @param Minutes The minutes component.
+	 * @param Seconds The seconds component.
+	 * @param Milliseconds The milliseconds component.
 	 */
 	void CORE_API Assign( int32 Days, int32 Hours, int32 Minutes, int32 Seconds, int32 Milliseconds );
 
@@ -627,8 +563,8 @@ private:
 /**
  * Pre-multiplies a time span with the given scalar.
  *
- * @param Scalar - The scalar to pre-multiply with.
- * @param Timespan - The time span to multiply.
+ * @param Scalar The scalar to pre-multiply with.
+ * @param Timespan The time span to multiply.
  */
 inline FTimespan operator*( float Scalar, const FTimespan& Timespan )
 {
