Core: Documentation and formatting cleanup pass on commonly used types

diff --git a/Engine/Source/Runtime/Core/Public/Containers/ArrayBuilder.h b/Engine/Source/Runtime/Core/Public/Containers/ArrayBuilder.h
--- a/Engine/Source/Runtime/Core/Public/Containers/ArrayBuilder.h
+++ b/Engine/Source/Runtime/Core/Public/Containers/ArrayBuilder.h
@@ -10,8 +10,8 @@
 /**
  * Template for fluent array builders.
  *
- * @param ElementType - The type of elements stored in the array.
- * @param Allocator - The allocator to use for elements.
+ * @param ElementType The type of elements stored in the array.
+ * @param Allocator The allocator to use for elements.
  */
 template<typename ElementType, typename Allocator = FDefaultAllocator> class TArrayBuilder
 {
@@ -20,28 +20,27 @@ public:
 	/**
 	 * Default constructor.
 	 */
-	TArrayBuilder () { }
+	TArrayBuilder( ) { }
 
 	/**
 	 * Creates and initializes an array builder from an array of items.
 	 *
-	 * @param InArray - The array of items to copy.
+	 * @param InArray The array of items to copy.
 	 */
-	template<typename OtherAllocator> TArrayBuilder (const TArray<ElementType, OtherAllocator>& InArray)
+	template<typename OtherAllocator> TArrayBuilder( const TArray<ElementType, OtherAllocator>& InArray )
 		: Array(InArray)
 	{ }
 
-
 public:
 
 	/**
 	 * Adds an item.
 	 *
-	 * @param Item - The item to add.
-	 *
+	 * @param Item The item to add.
 	 * @return This instance (for method chaining).
+	 * @see AddUnique
 	 */
-	TArrayBuilder& Add (const ElementType& Item)
+	TArrayBuilder& Add( const ElementType& Item )
 	{
 		Array.Add(Item);
 
@@ -51,11 +50,11 @@ public:
 	/**
 	 * Adds an unique item.
 	 *
-	 * @param Item - The unique item to add.
-	 *
+	 * @param Item The unique item to add.
 	 * @return This instance (for method chaining).
+	 * @see Add
 	 */
-	TArrayBuilder& AddUnique (const ElementType& Item)
+	TArrayBuilder& AddUnique( const ElementType& Item )
 	{
 		Array.AddUnique(Item);
 
@@ -65,18 +64,16 @@ public:
 	/**
 	 * Appends an array of items.
 	 *
-	 * @param OtherArray - The array to append.
-	 *
+	 * @param OtherArray The array to append.
 	 * @return This instance (for method chaining).
 	 */
-	template<typename OtherAllocator> TArrayBuilder& Append (const TArray<ElementType, OtherAllocator>& OtherArray)
+	template<typename OtherAllocator> TArrayBuilder& Append( const TArray<ElementType, OtherAllocator>& OtherArray )
 	{
 		Array.Append(OtherArray);
 
 		return *this;
 	}
 
-
 public:
 
 	/**
@@ -84,7 +81,7 @@ public:
 	 *
 	 * @return A new array.
 	 */
-	TArray<ElementType, Allocator> Build ()
+	TArray<ElementType, Allocator> Build( )
 	{
 		return Array;
 	}
@@ -94,12 +91,11 @@ public:
 	 *
 	 * @return A new array.
 	 */
-	operator TArray<ElementType, Allocator> ()
+	operator TArray<ElementType, Allocator>( )
 	{
 		return Build();
 	}
 
-
 private:
 
 	// Holds the array.
diff --git a/Engine/Source/Runtime/Core/Public/Containers/CircularBuffer.h b/Engine/Source/Runtime/Core/Public/Containers/CircularBuffer.h
--- a/Engine/Source/Runtime/Core/Public/Containers/CircularBuffer.h
+++ b/Engine/Source/Runtime/Core/Public/Containers/CircularBuffer.h
@@ -21,7 +21,7 @@ public:
 	/**
 	 * Creates and initializes a new instance of the TCircularBuffer class.
 	 *
-	 * @param Capacity - The number of elements that the buffer can store (will be rounded up to the next power of 2).
+	 * @param Capacity The number of elements that the buffer can store (will be rounded up to the next power of 2).
 	 */
 	TCircularBuffer( uint32 Capacity )
 	{
@@ -36,8 +36,8 @@ public:
 	/**
 	 * Creates and initializes a new instance of the TCircularBuffer class.
 	 *
-	 * @param Capacity -The number of elements that the buffer can store (will be rounded up to the next power of 2).
-	 * @param InitialValue - The initial value for the buffer's elements.
+	 * @param Capacity The number of elements that the buffer can store (will be rounded up to the next power of 2).
+	 * @param InitialValue The initial value for the buffer's elements.
 	 */
 	TCircularBuffer( uint32 Capacity, const ElementType& InitialValue )
 	{
@@ -53,7 +53,7 @@ public:
 	/**
 	 * Returns the mutable element at the specified index.
 	 *
-	 * @param Index - The index of the element to return.
+	 * @param Index The index of the element to return.
 	 */
 	FORCEINLINE ElementType& operator[]( uint32 Index )
 	{
@@ -63,7 +63,7 @@ public:
 	/**
 	 * Returns the immutable element at the specified index.
 	 *
-	 * @param Index - The index of the element to return.
+	 * @param Index The index of the element to return.
 	 */
 	FORCEINLINE const ElementType& operator[]( uint32 Index ) const
 	{
@@ -85,8 +85,7 @@ public:
 	/**
 	 * Calculates the index that follows the given index.
 	 *
-	 * @param CurrentIndex - The current index.
-	 *
+	 * @param CurrentIndex The current index.
 	 * @return The next index.
 	 */
 	FORCEINLINE uint32 GetNextIndex( uint32 CurrentIndex ) const
diff --git a/Engine/Source/Runtime/Core/Public/Containers/CircularQueue.h b/Engine/Source/Runtime/Core/Public/Containers/CircularQueue.h
--- a/Engine/Source/Runtime/Core/Public/Containers/CircularQueue.h
+++ b/Engine/Source/Runtime/Core/Public/Containers/CircularQueue.h
@@ -27,7 +27,7 @@ public:
 	/**
 	 * Default constructor.
 	 *
-	 * @param Size - The number of elements that the queue can hold (will be rounded up to the next power of 2).
+	 * @param Size The number of elements that the queue can hold (will be rounded up to the next power of 2).
 	 */
 	TCircularQueue( uint32 Size)
 		: Buffer(Size)
@@ -62,9 +62,8 @@ public:
 	/**
 	 * Removes an item from the front of the queue.
 	 *
-	 * @param OutElement - Will contain the element if the queue is not empty.
-	 *
-	 * @return TRUE if an element has been returned, FALSE if the queue was empty.
+	 * @param OutElement Will contain the element if the queue is not empty.
+	 * @return true if an element has been returned, false if the queue was empty.
 	 */
 	bool Dequeue( ElementType& OutElement )
 	{
@@ -92,9 +91,8 @@ public:
 	/**
 	 * Adds an item to the end of the queue.
 	 *
-	 * @param Element - The element to add.
-	 *
-	 * @return TRUE if the item was added, FALSE if the queue was full.
+	 * @param Element The element to add.
+	 * @return true if the item was added, false if the queue was full.
 	 */
 	bool Enqueue( const ElementType& Element )
 	{
@@ -114,10 +112,8 @@ public:
 	/**
 	 * Checks whether the queue is empty.
 	 *
-	 * @return TRUE if the queue is empty, FALSE otherwise.
-	 *
-	 * @see Empty
-	 * @see IsFull
+	 * @return true if the queue is empty, false otherwise.
+	 * @see Empty, IsFull
 	 */
 	FORCEINLINE bool IsEmpty( ) const
 	{
@@ -127,8 +123,7 @@ public:
 	/**
 	 * Checks whether the queue is full.
 	 *
-	 * @return TRUE if the queue is full, FALSE otherwise.
-	 *
+	 * @return true if the queue is full, false otherwise.
 	 * @see IsEmpty
 	 */
 	bool IsFull( ) const
@@ -139,9 +134,8 @@ public:
 	/**
 	 * Returns the oldest item in the queue without removing it.
 	 *
-	 * @param OutItem - Will contain the item if the queue is not empty.
-	 *
-	 * @return TRUE if an item has been returned, FALSE if the queue was empty.
+	 * @param OutItem Will contain the item if the queue is not empty.
+	 * @return true if an item has been returned, false if the queue was empty.
 	 */
 	bool Peek( ElementType& OutItem )
 	{
diff --git a/Engine/Source/Runtime/Core/Public/Containers/EnumAsByte.h b/Engine/Source/Runtime/Core/Public/Containers/EnumAsByte.h
--- a/Engine/Source/Runtime/Core/Public/Containers/EnumAsByte.h
+++ b/Engine/Source/Runtime/Core/Public/Containers/EnumAsByte.h
@@ -6,6 +6,7 @@
 
 #pragma once
 
+
 /**
  * Template to store enums as BYTEs with type safety
 **/
@@ -16,87 +17,112 @@ public:
 	typedef TEnum EnumType;
 
 	/**
-	 * Default Constructor, like an enum it is uninitialized
-	**/
-	FORCEINLINE TEnumAsByte()
-	{
-	}
+	 * Default Constructor (no initialization).
+	 */
+	FORCEINLINE TEnumAsByte( ) { }
+
 	/**
-	 * Copy constructor
+	 * Copy constructor.
+	 *
 	 * @param InValue value to construct with 
-	**/
-	FORCEINLINE TEnumAsByte(const TEnumAsByte &InValue)
+	 */
+	FORCEINLINE TEnumAsByte( const TEnumAsByte &InValue )
 		: Value(InValue.Value)
-	{
-	}
+	{ }
+
 	/**
-	 * Constructor, initialize to the enum value
-	 * @param InValue value to construct with 
-	**/
-	FORCEINLINE TEnumAsByte(TEnum InValue)
+	 * Constructor, initialize to the enum value.
+	 *
+	 * @param InValue value to construct with .
+	 */
+	FORCEINLINE TEnumAsByte( TEnum InValue )
 		: Value(InValue)
-	{
-	}
+	{ }
+
 	/**
-	 * Constructor, initialize to the int32 value
+	 * Constructor, initialize to the int32 value.
+	 *
 	 * @param InValue value to construct with 
-	**/
-	explicit FORCEINLINE TEnumAsByte(int32 InValue)
+	 */
+	explicit FORCEINLINE TEnumAsByte( int32 InValue )
 		: Value(InValue)
-	{
-	}
+	{ }
+
 	/**
 	 * Constructor, initialize to the int32 value
-	 * @param InValue value to construct with 
-	**/
-	explicit FORCEINLINE TEnumAsByte(uint8 InValue)
+	 *
+	 * @param InValue value to construct with.
+	 */
+	explicit FORCEINLINE TEnumAsByte( uint8 InValue )
 		: Value(InValue)
-	{
-	}
+	{ }
+
+public:
+
 	/**
-	 * Assignment operator
-	 * @param InValue value to set 
-	**/
-	FORCEINLINE TEnumAsByte& operator=(TEnumAsByte InValue)
+	 * Assignment operator.
+	 *
+	 * @param InValue value to set.
+	 */
+	FORCEINLINE TEnumAsByte& operator=( TEnumAsByte InValue )
 	{
 		Value = InValue.Value;
 		return *this;
 	}
 	/**
-	 * Assignment operator
-	 * @param InValue value to set 
-	**/
-	FORCEINLINE TEnumAsByte& operator=(TEnum InValue)
+	 * Assignment operator.
+	 *
+	 * @param InValue value to set.
+	 */
+	FORCEINLINE TEnumAsByte& operator=( TEnum InValue )
 	{
 		Value = InValue;
 		return *this;
 	}
 
-	/** Implicit conversion to TEnum. */
-	operator TEnum() const
+	/**
+	 * Compares two enumeration values for equality.
+	 *
+	 * @return true if the two values are equal, false otherwise.
+	 */
+	bool operator==( TEnum InValue ) const
 	{
-		return TEnum(Value);
+		return (Value == InValue);
 	}
-	/** Comparison operator **/
-	bool operator==(TEnum InValue) const
+
+	/**
+	 * Compares two enumeration values for equality.
+	 *
+	 * @return true if the two values are equal, false otherwise.
+	 */
+	bool operator==(TEnumAsByte InValue) const
 	{
-		return Value == InValue;
+		return (Value == InValue.Value);
 	}
-	/** Comparison operator **/
-	bool operator==(TEnumAsByte InValue) const
+
+	/**
+	 * Implicit conversion to TEnum.
+	 */
+	operator TEnum( ) const
 	{
-		return Value == InValue.Value;
+		return TEnum(Value);
 	}
-	/** Get the value */
-	TEnum GetValue() const
+
+public:
+
+	/**
+	 * Gets the value.
+	 */
+	TEnum GetValue( ) const
 	{
 		return TEnum(Value);
 	}
+
 private:
-	/** The actual value **/
+
+	/** Holds the actual value **/
 	uint8 Value;
 };
 
-template<class T> struct TIsPODType<TEnumAsByte<T> > { enum { Value = true }; };
-
 
+template<class T> struct TIsPODType<TEnumAsByte<T>> { enum { Value = true }; };
diff --git a/Engine/Source/Runtime/Core/Public/Containers/Queue.h b/Engine/Source/Runtime/Core/Public/Containers/Queue.h
--- a/Engine/Source/Runtime/Core/Public/Containers/Queue.h
+++ b/Engine/Source/Runtime/Core/Public/Containers/Queue.h
@@ -14,14 +14,10 @@ namespace EQueueMode
 	 */
 	enum Type
 	{
-		/**
-		 * Multiple-producers, single-consumer queue.
-		 */
+		/** Multiple-producers, single-consumer queue. */
 		Mpsc,
 
-		/**
-		 * Single-producer, single-consumer queue.
-		 */
+		/** Single-producer, single-consumer queue. */
 		Spsc
 	};
 }
@@ -75,8 +71,7 @@ public:
 	/**
 	 * Removes and returns the item from the tail of the queue.
 	 *
-	 * @param OutValue - Will hold the returned value.
-	 *
+	 * @param OutValue Will hold the returned value.
 	 * @return true if a value was returned, false if the queue was empty.
 	 */
 	bool Dequeue( ItemType& OutItem )
@@ -103,8 +98,7 @@ public:
 	/**
 	 * Adds an item to the head of the queue.
 	 *
-	 * @param Item - The item to add.
-	 *
+	 * @param Item The item to add.
 	 * @return true if the item was added, false otherwise.
 	 */
 	bool Enqueue( const ItemType& Item )
@@ -146,8 +140,7 @@ public:
 	/**
 	 * Peeks at the queue's tail item without removing it.
 	 *
-	 * @param OutItem - Will hold the peeked at item.
-	 *
+	 * @param OutItem Will hold the peeked at item.
 	 * @return true if an item was returned, false if the queue was empty.
 	 */
 	bool Peek( ItemType& OutItem )
diff --git a/Engine/Source/Runtime/Core/Public/Math/Box.h b/Engine/Source/Runtime/Core/Public/Math/Box.h
--- a/Engine/Source/Runtime/Core/Public/Math/Box.h
+++ b/Engine/Source/Runtime/Core/Public/Math/Box.h
@@ -8,37 +8,33 @@
 
 
 /**
- * Implements a rectangular box.
+ * Implements an axis-aligned box.
+ *
+ * Boxes describe an axis-aligned extent in three dimensions. They are used for many different things in the
+ * Engine and in games, such as bounding volumes, collision detection and visibility calculation.
  */
 class FBox
 {
 public:
 
-	/**
-	 * Holds the box's minimum point.
-	 */
+	/** Holds the box's minimum point. */
 	FVector Min;
 
-	/**
-	 * Holds the box's maximum point.
-	 */
+	/** Holds the box's maximum point. */
 	FVector Max;
 
-	/**
-	 * Holds a flag indicating whether this box is valid.
-	 */
+	/** Holds a flag indicating whether this box is valid. */
 	uint8 IsValid;
 
-
 public:
 
 	/**
-	 * Default constructor.
+	 * Default constructor (no initialization).
 	 */
 	FBox( ) { }
 
 	/**
-	 * Creates and initializes a new box.
+	 * Creates and initializes a new box with zero extent and marks it as invalid.
 	 */
 	FBox( int32 )
 	{
@@ -46,9 +42,9 @@ public:
 	}
 
 	/**
-	 * Creates and initializes a new box.
+	 * Creates and initializes a new box with zero extent and marks it as invalid.
 	 *
-	 * @param EForceInit - Force Init Enum.
+	 * @param EForceInit Force Init Enum.
 	 */
 	explicit FBox( EForceInit )
 	{
@@ -56,10 +52,10 @@ public:
 	}
 
 	/**
-	 * Creates and initializes a new box from the specified parameters.
+	 * Creates and initializes a new box from the specified extents.
 	 *
-	 * @param InMin - The box's minimum point.
-	 * @param InMax - The box's maximum point.
+	 * @param InMin The box's minimum point.
+	 * @param InMax The box's maximum point.
 	 */
 	FBox( const FVector& InMin, const FVector& InMax )
 		: Min(InMin)
@@ -70,26 +66,34 @@ public:
 	/**
 	 * Creates and initializes a new box from the given set of points.
 	 *
-	 * @param Points - Array of Points to create for the bounding volume.
-	 * @param Count - The number of points.
+	 * @param Points Array of Points to create for the bounding volume.
+	 * @param Count The number of points.
 	 */
 	CORE_API FBox( const FVector* Points, int32 Count );
 
 	/**
 	 * Creates and initializes a new box from an array of points.
 	 *
-	 * @param Points - Array of Points to create for the bounding volume.
+	 * @param Points Array of Points to create for the bounding volume.
 	 */
 	CORE_API FBox( const TArray<FVector>& Points );
 
-
 public:
 
+	/**
+	 * Compares two boxes for equality.
+	 *
+	 * @return true if the boxes are equal, false otherwise.
+	 */
+	bool operator==( const FBox& Other ) const
+	{
+		return (Min == Other.Min) && (Max == Other.Max);
+	}
+
 	/**
 	 * Adds to this bounding box to include a given point.
 	 *
 	 * @param Other the point to increase the bounding volume to.
-	 *
 	 * @return Reference to this bounding box after resizing to include the other point.
 	 */
 	FORCEINLINE FBox& operator+=( const FVector &Other );
@@ -98,7 +102,6 @@ public:
 	 * Gets the result of addition to this bounding volume.
 	 *
 	 * @param Other The other point to add to this.
-	 *
 	 * @return A new bounding volume.
 	 */
 	FBox operator+( const FVector& Other ) const
@@ -110,7 +113,6 @@ public:
 	 * Adds to this bounding box to include a new bounding volume.
 	 *
 	 * @param Other the bounding volume to increase the bounding volume to.
-	 *
 	 * @return Reference to this bounding volume after resizing to include the other bounding volume.
 	 */
 	FORCEINLINE FBox& operator+=( const FBox& Other );
@@ -119,7 +121,6 @@ public:
 	 * Gets the result of addition to this bounding volume.
 	 *
 	 * @param Other The other volume to add to this.
-	 *
 	 * @return A new bounding volume.
 	 */
 	FBox operator+( const FBox& Other ) const
@@ -131,7 +132,6 @@ public:
 	 * Gets reference to the min or max of this bounding volume.
 	 *
 	 * @param Index the index into points of the bounding volume.
-	 *
 	 * @return a reference to a point of the bounding volume.
 	 */
     FVector& operator[]( int32 Index )
@@ -145,21 +145,13 @@ public:
 
 		return Max;
 	}
-	bool operator==( const FBox& Other ) const
-	{
-		return ( Min == Other.Min ) && ( Max == Other.Max );
-	}
-
 
 public:
 
-
-
 	/** 
 	 * Calculates the distance of a point to this box.
 	 *
-	 * @param Point - The point.
-	 *
+	 * @param Point The point.
 	 * @return The distance.
 	 */
 	FORCEINLINE float ComputeSquaredDistanceToPoint( const FVector& Point ) const
@@ -168,10 +160,10 @@ public:
 	}
 
 	/** 
-	 * Increase the bounding box volume
+	 * Increases the box size.
 	 *
-	 * @param	W size to increase volume by
-	 * @return	new bounding box increased in size
+	 * @param W The size to increase the volume by.
+	 * @return A new bounding box.
 	 */
 	FBox ExpandBy( float W ) const
 	{
@@ -179,10 +171,10 @@ public:
 	}
 
 	/** 
-	 * Shift bounding box position
+	 * Shifts the bounding box position.
 	 *
-	 * @param	Offset vector to shift by
-	 * @return	new shifted bounding box 
+	 * @param Offset The vector to shift the box by.
+	 * @return A new bounding box.
 	 */
 	FBox ShiftBy( const FVector& Offset ) const
 	{
@@ -190,10 +182,10 @@ public:
 	}
 
 	/** 
-	 * Moves center of bounding box to new destination
+	 * Moves the center of bounding box to new destination.
 	 *
-	 * @param   Destination point to move center of box to
-	 * @return  new moved bounding box 
+	 * @param The destination point to move center of box to.
+	 * @return A new bounding box.
 	 */
 	FBox MoveTo( const FVector& Destination ) const
 	{
@@ -202,9 +194,10 @@ public:
 	}
 
 	/**
-	 * Gets the box's center point.
+	 * Gets the center point of this box.
 	 *
-	 * @return Center point.
+	 * @return The center point.
+	 * @see GetCenterAndExtents, GetExtent, GetSize, GetVolume
 	 */
 	FVector GetCenter( ) const
 	{
@@ -212,10 +205,11 @@ public:
 	}
 
 	/**
-	 * Get the center and extents
+	 * Gets the center and extents of this box.
 	 *
-	 * @param center[out] reference to center point
-	 * @param Extents[out] reference to the extent around the center
+	 * @param center[out] Will contain the box center point.
+	 * @param Extents[out] Will contain the extent around the center.
+	 * @see GetCenter, GetExtent, GetSize, GetVolume
 	 */
 	void GetCenterAndExtents( FVector & center, FVector & Extents ) const
 	{
@@ -226,16 +220,16 @@ public:
 	/**
 	 * Calculates the closest point on or inside the box to a given point in space.
 	 *
-	 * @param Point - The point in space.
-	 *
+	 * @param Point The point in space.
 	 * @return The closest point on or inside the box.
 	 */
 	FORCEINLINE FVector GetClosestPointTo( const FVector& Point ) const;
 
 	/**
-	 * Gets the box extents around the center.
+	 * Gets the extents of this box.
 	 *
-	 * @return Box extents.
+	 * @return The box extents.
+	 * @see GetCenter, GetCenterAndExtents, GetSize, GetVolume
 	 */
 	FVector GetExtent( ) const
 	{
@@ -245,8 +239,7 @@ public:
 	/**
 	 * Gets a reference to the specified point of the bounding box.
 	 *
-	 * @param PointIndex - Index of point
-	 *
+	 * @param PointIndex The index of the extrema point to return.
 	 * @return A reference to the point.
 	 */
 	FVector& GetExtrema( int PointIndex )
@@ -257,8 +250,7 @@ public:
 	/**
 	 * Gets a read-only reference to the specified point of the bounding box.
 	 *
-	 * @param PointIndex - Index of point
-	 *
+	 * @param PointIndex The index of extrema point to return.
 	 * @return A read-only reference to the point.
 	 */
 	const FVector& GetExtrema( int PointIndex ) const
@@ -267,9 +259,10 @@ public:
 	}
 
 	/**
-	 * Gets the box size.
+	 * Gets the size of this box.
 	 *
-	 * @return Box size.
+	 * @return The box size.
+	 * @see GetCenter, GetCenterAndExtents, GetExtent, GetVolume
 	 */
 	FVector GetSize( ) const
 	{
@@ -277,9 +270,10 @@ public:
 	}
 
 	/**
-	 * Gets the box volume.
+	 * Gets the volume of this box.
 	 *
-	 * @return Box volume.
+	 * @return The box volume.
+	 * @see GetCenter, GetCenterAndExtents, GetExtent, GetSize
 	 */
 	float GetVolume( ) const
 	{
@@ -296,59 +290,35 @@ public:
 	}
 
 	/**
-	 * Test for intersection of bounding box and this bounding box
+	 * Checks whether the given bounding box intersects this bounding box.
 	 *
-	 * @param other bounding box to test intersection
-	 * @return true if the other box intersects this box, otherwise false
+	 * @param Other The bounding box to intersect with.
+	 * @return true if the boxes intersect, false otherwise.
 	 */
 	FORCEINLINE bool Intersect( const FBox & other ) const;
 
 	/**
-	 * Test for intersection of bounding box and this bounding box in the XY plane
+	 * Checks whether the given bounding box intersects this bounding box in the XY plane.
 	 *
-	 * @param other bounding box to test intersection
-	 * @return true if the other box intersects this box in the XY Plane, otherwise false
+	 * @param Other The bounding box to test intersection.
+	 * @return true if the boxes intersect in the XY Plane, false otherwise.
 	 */
-	FORCEINLINE bool IntersectXY( const FBox& other ) const;
+	FORCEINLINE bool IntersectXY( const FBox& Other ) const;
 
 	/**
 	  * Gets a bounding volume transformed by an inverted FTransform object.
 	  *
-	  * @param M the FTransform object which translation will be inverted.
-	  * @return	transformed box
+	  * @param M The transformation object to perform the inversely transform this box with.
+	  * @return	The transformed box.
 	  */
 	CORE_API FBox InverseTransformBy( const FTransform & M ) const;
 
-	/**
-	 * Gets a bounding volume transformed by a matrix.
-	 *
-	 * @param M the matrix .
-	 * @return	transformed box
-	 */
-	CORE_API FBox TransformBy( const FMatrix& M ) const;
-
-	/**
-	 * Gets a bounding volume transformed by a FTransform object.
-	 *
-	 * @param M the FTransform object.
-	 * @return	transformed box
-	 */
-	CORE_API FBox TransformBy( const FTransform & M ) const;
-
-	/** 
-	 * Transforms and projects a world bounding box to screen space
-	 *
-	 * @param	ProjM - projection matrix
-	 * @return	transformed box
-	 */
-	CORE_API FBox TransformProjectBy( const FMatrix& ProjM ) const;
-	
 	/** 
-	 * Checks to see if the location is inside this box
+	 * Checks whether the given location is inside this box.
 	 * 
-	 * @param In - The location to test for inside the bounding volume
-	 *
-	 * @return true if location is inside this volume
+	 * @param In The location to test for inside the bounding volume.
+	 * @return true if location is inside this volume.
+	 * @see IsInsideXY
 	 */
 	bool IsInside( const FVector& In ) const
 	{
@@ -356,40 +326,66 @@ public:
 	}
 
 	/** 
-	 * Checks to see if the location is inside this box in the XY plane
+	 * Checks whether a given box is fully encapsulated by this box.
 	 * 
-	 * @param In - The location to test for inside the bounding box
-	 *
-	 * @return true if location is inside this box in the XY plane
+	 * @param Other The box to test for encapsulation within the bounding volume.
+	 * @return true if box is inside this volume.
 	 */
-	bool IsInsideXY( const FVector& In ) const
+	bool IsInside( const FBox& Other ) const
 	{
-		return ((In.X > Min.X) && (In.X < Max.X) && (In.Y > Min.Y) && (In.Y < Max.Y));
+		return (IsInside(Other.Min) && IsInside(Other.Max));
 	}
-	
 
 	/** 
-	 * Checks to see if given box is fully encapsulated by this box
+	 * Checks whether the given location is inside this box in the XY plane.
 	 * 
-	 * @param Other box to test for encapsulation within the bounding volume
-	 * @return true if box is inside this volume
+	 * @param In The location to test for inside the bounding box.
+	 * @return true if location is inside this box in the XY plane.
+	 * @see IsInside
 	 */
-	bool IsInside( const FBox& Other ) const
+	bool IsInsideXY( const FVector& In ) const
 	{
-		return (IsInside(Other.Min) && IsInside(Other.Max));
+		return ((In.X > Min.X) && (In.X < Max.X) && (In.Y > Min.Y) && (In.Y < Max.Y));
 	}
 
 	/** 
-	 * Checks to see if given box is fully encapsulated by this box in the XY plane
+	 * Checks whether the given box is fully encapsulated by this box in the XY plane.
 	 * 
-	 * @param Other box to test for encapsulation within the bounding box
-	 * @return true if box is inside this box in the XY plane
+	 * @param Other The box to test for encapsulation within the bounding box.
+	 * @return true if box is inside this box in the XY plane.
 	 */
 	bool IsInsideXY( const FBox& Other ) const
 	{
 		return (IsInsideXY(Other.Min) && IsInsideXY(Other.Max));
 	}
 
+	/**
+	 * Gets a bounding volume transformed by a matrix.
+	 *
+	 * @param M The matrix to transform by.
+	 * @return The transformed box.
+	 * @see TransformProjectBy
+	 */
+	CORE_API FBox TransformBy( const FMatrix& M ) const;
+
+	/**
+	 * Gets a bounding volume transformed by a FTransform object.
+	 *
+	 * @param M The transformation object.
+	 * @return The transformed box.
+	 * @see TransformProjectBy
+	 */
+	CORE_API FBox TransformBy( const FTransform & M ) const;
+
+	/** 
+	 * Transforms and projects a world bounding box to screen space
+	 *
+	 * @param ProjM The projection matrix.
+	 * @return The transformed box.
+	 * @see TransformBy
+	 */
+	CORE_API FBox TransformProjectBy( const FMatrix& ProjM ) const;
+
 	/**
 	 * Get a textual representation of this box.
 	 *
@@ -397,16 +393,14 @@ public:
 	 */
 	FString ToString( ) const;
 
-
 public:
 
 	/** 
 	 * Utility function to build an AABB from Origin and Extent 
 	 *
-	 * @param Origin the location of the bounding box
-	 * @param Extent half size of the bounding box
-	 *
-	 * @return a new AABB
+	 * @param Origin The location of the bounding box.
+	 * @param Extent Half size of the bounding box.
+	 * @return A new axis-aligned bounding box.
 	 */
 	static FBox BuildAABB( const FVector& Origin, const FVector& Extent )
 	{
@@ -415,15 +409,13 @@ public:
 		return NewBox;
 	}
 
-
 public:
 
 	/**
 	 * Serializes the bounding box.
 	 *
-	 * @param Ar - The archive to serialize into.
-	 * @param Box - The box to serialize.
-	 *
+	 * @param Ar The archive to serialize into.
+	 * @param Box The box to serialize.
 	 * @return Reference to the Archive after serialization.
 	 */
 	friend FArchive& operator<<( FArchive& Ar, FBox& Box )
@@ -564,4 +556,4 @@ FORCEINLINE bool FBox::IntersectXY( const FBox& Other ) const
 FORCEINLINE FString FBox::ToString( ) const
 {
 	return FString::Printf(TEXT("IsValid=%s, Min=(%s), Max=(%s)"), IsValid ? TEXT("true") : TEXT("false"), *Min.ToString(), *Max.ToString());
-}
\ No newline at end of file
+}
diff --git a/Engine/Source/Runtime/Core/Public/Math/Box2D.h b/Engine/Source/Runtime/Core/Public/Math/Box2D.h
--- a/Engine/Source/Runtime/Core/Public/Math/Box2D.h
+++ b/Engine/Source/Runtime/Core/Public/Math/Box2D.h
@@ -12,8 +12,6 @@
  */
 class FBox2D
 {
-public:
-
 	/** Holds the box's minimum point. */
 	FVector2D Min;
 
@@ -23,13 +21,18 @@ public:
 	/** Holds a flag indicating whether this box is valid. */
 	bool bIsValid;
 
-
 public:
 
-	/** Default constructor. */
+	/**
+	 * Default constructor (no initialization).
+	 */
 	FBox2D( ) { }
 
-	/** Creates and initializes a new box. */
+	/**
+	 * Creates and initializes a new box.
+	 *
+	 * The box extents are initialized to zero and the box is marked as invalid.
+	 */
 	FBox2D( int32 )
 	{
 		Init();
@@ -38,7 +41,9 @@ public:
 	/**
 	 * Creates and initializes a new box.
 	 *
-	 * @param EForceInit - Force Init Enum.
+	 * The box extents are initialized to zero and the box is marked as invalid.
+	 *
+	 * @param EForceInit Force Init Enum.
 	 */
 	explicit FBox2D( EForceInit )
 	{
@@ -48,8 +53,8 @@ public:
 	/**
 	 * Creates and initializes a new box from the specified parameters.
 	 *
-	 * @param InMin - The box's minimum point.
-	 * @param InMax - The box's maximum point.
+	 * @param InMin The box's minimum point.
+	 * @param InMax The box's maximum point.
 	 */
 	FBox2D( const FVector2D& InMin, const FVector2D& InMax )
 		: Min(InMin)
@@ -57,31 +62,38 @@ public:
 		, bIsValid(true)
 	{ }
 
-	
-	 
 	/**
 	 * Creates and initializes a new box from the given set of points.
 	 *
-	 * @param Points - Array of Points to create for the bounding volume.
-	 * @param Count - The number of points.
+	 * @param Points Array of Points to create for the bounding volume.
+	 * @param Count The number of points.
 	 */
 	CORE_API FBox2D( const FVector2D* Points, const int32 Count );
 
 	/**
 	 * Creates and initializes a new box from an array of points.
 	 *
-	 * @param Points - Array of Points to create for the bounding volume.
+	 * @param Points Array of Points to create for the bounding volume.
 	 */
 	CORE_API FBox2D( const TArray<FVector2D>& Points );
 
-
 public:
 
 	/**
-	 * Adds to this bounding box to include a given point.
+	 * Compares two boxes for equality.
 	 *
-	 * @param Other the point to increase the bounding volume to.
+	 * @param Other The other box to compare with.
+	 * @return true if the boxes are equal, false otherwise.
+	 */
+	bool operator==( const FBox2D& Other ) const
+	{
+		return (Min == Other.Min) && (Max == Other.Max);
+	}
+
+	/**
+	 * Adds to this bounding box to include a given point.
 	 *
+	 * @param Other The point to increase the bounding volume to.
 	 * @return Reference to this bounding box after resizing to include the other point.
 	 */
 	FORCEINLINE FBox2D& operator+=( const FVector2D &Other );
@@ -90,7 +102,6 @@ public:
 	 * Gets the result of addition to this bounding volume.
 	 *
 	 * @param Other The other point to add to this.
-	 *
 	 * @return A new bounding volume.
 	 */
 	FBox2D operator+( const FVector2D& Other ) const
@@ -101,8 +112,7 @@ public:
 	/**
 	 * Adds to this bounding box to include a new bounding volume.
 	 *
-	 * @param Other the bounding volume to increase the bounding volume to.
-	 *
+	 * @param Other The bounding volume to increase the bounding volume to.
 	 * @return Reference to this bounding volume after resizing to include the other bounding volume.
 	 */
 	FORCEINLINE FBox2D& operator+=( const FBox2D& Other );
@@ -111,7 +121,6 @@ public:
 	 * Gets the result of addition to this bounding volume.
 	 *
 	 * @param Other The other volume to add to this.
-	 *
 	 * @return A new bounding volume.
 	 */
 	FBox2D operator+( const FBox2D& Other ) const
@@ -122,9 +131,8 @@ public:
 	/**
 	 * Gets reference to the min or max of this bounding volume.
 	 *
-	 * @param Index the index into points of the bounding volume.
-	 *
-	 * @return a reference to a point of the bounding volume.
+	 * @param Index The index into points of the bounding volume.
+	 * @return A reference to a point of the bounding volume.
 	 */
    FVector2D& operator[]( int32 Index )
 	{
@@ -137,21 +145,13 @@ public:
 
 		return Max;
 	}
-	bool operator==( const FBox2D& Other ) const
-	{
-		return ( Min == Other.Min ) && ( Max == Other.Max );
-	}
-
 
 public:
 
-
-
 	/** 
 	 * Calculates the distance of a point to this box.
 	 *
-	 * @param Point - The point.
-	 *
+	 * @param Point The point.
 	 * @return The distance.
 	 */
 	FORCEINLINE float ComputeSquaredDistanceToPoint( const FVector2D& Point ) const
@@ -181,32 +181,32 @@ public:
 	}
 
 	/** 
-	 * Increase the bounding box volume
+	 * Increase the bounding box volume.
 	 *
-	 * @param	W size to increase volume by
-	 * @return	new bounding box increased in size
+	 * @param W The size to increase volume by.
+	 * @return A new bounding box increased in size.
 	 */
 	FBox2D ExpandBy( const float W ) const
 	{
 		return FBox2D(Min - FVector2D(W, W), Max + FVector2D(W, W));
 	}
 
-	/** 
-	 * Shift bounding box position
+	/**
+	 * Gets the box area.
 	 *
-	 * @param	Offset vector to shift by
-	 * @return	new shifted bounding box 
+	 * @return Box area.
+	 * @see GetCenter, GetCenterAndExtents, GetExtent, GetSize
 	 */
-	FBox2D ShiftBy( const FVector2D& Offset ) const
+	float GetArea( ) const
 	{
-		return FBox2D(Min + Offset, Max + Offset);
+		return (Max.X - Min.X) * (Max.Y - Min.Y);
 	}
 
-
 	/**
 	 * Gets the box's center point.
 	 *
-	 * @return Center point.
+	 * @return Th center point.
+	 * @see GetArea, GetCenterAndExtents, GetExtent, GetSize
 	 */
 	FVector2D GetCenter( ) const
 	{
@@ -218,6 +218,7 @@ public:
 	 *
 	 * @param center[out] reference to center point
 	 * @param Extents[out] reference to the extent around the center
+	 * @see GetArea, GetCenter, GetExtent, GetSize
 	 */
 	void GetCenterAndExtents( FVector2D & center, FVector2D & Extents ) const
 	{
@@ -238,6 +239,7 @@ public:
 	 * Gets the box extents around the center.
 	 *
 	 * @return Box extents.
+	 * @see GetArea, GetCenter, GetCenterAndExtents, GetSize
 	 */
 	FVector2D GetExtent( ) const
 	{
@@ -249,22 +251,13 @@ public:
 	 * Gets the box size.
 	 *
 	 * @return Box size.
+	 * @see GetArea, GetCenter, GetCenterAndExtents, GetExtent
 	 */
 	FVector2D GetSize( ) const
 	{
 		return (Max - Min);
 	}
 
-	/**
-	 * Gets the box area.
-	 *
-	 * @return Box area.
-	 */
-	float GetArea( ) const
-	{
-		return (Max.X - Min.X) * (Max.Y - Min.Y);
-	}
-
 	/**
 	 * Set the initial values of the bounding box to Zero.
 	 */
@@ -274,38 +267,46 @@ public:
 		bIsValid = false;
 	}
 
-	
 	/**
-	 * Test for intersection of bounding box and this bounding box
+	 * Checks whether the given box intersects this box.
 	 *
 	 * @param other bounding box to test intersection
-	 * @return true if the other box intersects this box, otherwise false
+	 * @return true if boxes intersect, false otherwise.
 	 */
 	FORCEINLINE bool Intersect( const FBox2D & other ) const;
 
 	/**
-	 * Test if point is inside this box
+	 * Checks whether the given point is inside this box.
 	 *
-	 * @param Point to test
-	 * @return true if the point is inside this box, otherwise false
+	 * @param Point The point to test.
+	 * @return true if the point is inside this box, otherwise false.
 	 */
 	bool IsInside( const FVector2D & TestPoint ) const
 	{
 		return ((TestPoint.X > Min.X) && (TestPoint.X < Max.X) && (TestPoint.Y > Min.Y) && (TestPoint.Y < Max.Y));
 	}
 
-
 	/** 
-	 * Checks to see if given box is fully encapsulated by this box
+	 * Checks whether the given box is fully encapsulated by this box.
 	 * 
-	 * @param Other box to test for encapsulation within the bounding volume
-	 * @return true if box is inside this volume
+	 * @param Other The box to test for encapsulation within the bounding volume.
+	 * @return true if box is inside this volume, false otherwise.
 	 */
 	bool IsInside( const FBox2D& Other ) const
 	{
 		return (IsInside(Other.Min) && IsInside(Other.Max));
 	}
 
+	/** 
+	 * Shift bounding box position.
+	 *
+	 * @param The offset vector to shift by.
+	 * @return A new shifted bounding box.
+	 */
+	FBox2D ShiftBy( const FVector2D& Offset ) const
+	{
+		return FBox2D(Min + Offset, Max + Offset);
+	}
 
 	/**
 	 * Get a textual representation of this box.
@@ -331,7 +332,6 @@ public:
 };
 
 
-
 /* FBox2D inline functions
  *****************************************************************************/
 
@@ -421,4 +421,4 @@ FORCEINLINE bool FBox2D::Intersect( const FBox2D & Other ) const
 FORCEINLINE FString FBox2D::ToString( ) const
 {
 	return FString::Printf(TEXT("bIsValid=%s, Min=(%s), Max=(%s)"), bIsValid ? TEXT("true") : TEXT("false"), *Min.ToString(), *Max.ToString());
-}
\ No newline at end of file
+}
diff --git a/Engine/Source/Runtime/Core/Public/Math/BoxSphereBounds.h b/Engine/Source/Runtime/Core/Public/Math/BoxSphereBounds.h
--- a/Engine/Source/Runtime/Core/Public/Math/BoxSphereBounds.h
+++ b/Engine/Source/Runtime/Core/Public/Math/BoxSphereBounds.h
@@ -12,35 +12,17 @@
  */
 struct FBoxSphereBounds
 {
-	/**
-	 * Holds the origin of the bounding box and sphere.
-	 */
+	/** Holds the origin of the bounding box and sphere. */
 	FVector	Origin;
 
-	/**
-	 * Holds the extent of the bounding box.
-	 */
+	/** Holds the extent of the bounding box. */
 	FVector BoxExtent;
 
-	/**
-	 * Holds the radius of the bounding sphere.
-	 */
+	/** Holds the radius of the bounding sphere. */
 	float SphereRadius;
 
-
 public:
 
-#if ENABLE_NAN_DIAGNOSTIC
-	FORCEINLINE void DiagnosticCheckNaN() const
-	{
-		checkf(!Origin.ContainsNaN(), TEXT("Origin contains NaN: %s"), *Origin.ToString());
-		checkf(!BoxExtent.ContainsNaN(), TEXT("BoxExtent contains NaN: %s"), *BoxExtent.ToString());
-		checkf(!FMath::IsNaN(SphereRadius) && FMath::IsFinite(SphereRadius), TEXT("SphereRadius contains NaN: %f"), SphereRadius);
-	}
-#else
-	FORCEINLINE void DiagnosticCheckNaN() const {}
-#endif
-
 	/**
 	 * Default constructor.
 	 */
@@ -77,8 +59,8 @@ public:
 	/**
 	 * Creates and initializes a new instance from the given Box and Sphere.
 	 *
-	 * @param Box - The bounding box.
-	 * @param Sphere - The bounding sphere.
+	 * @param Box The bounding box.
+	 * @param Sphere The bounding sphere.
 	 */
 	FBoxSphereBounds( const FBox& Box, const FSphere& Sphere )
 	{
@@ -93,7 +75,7 @@ public:
 	 *
 	 * The sphere radius is taken from the extent of the box.
 	 *
-	 * @param Box - The bounding box.
+	 * @param Box The bounding box.
 	 */
 	FBoxSphereBounds( const FBox& Box )
 	{
@@ -120,8 +102,8 @@ public:
 	 *
 	 * The sphere radius is taken from the extent of the box.
 	 *
-	 * @param Points - The points to be considered for the bounding box.
-	 * @param NumPoints - Number of points in the Points array.
+	 * @param Points The points to be considered for the bounding box.
+	 * @param NumPoints Number of points in the Points array.
 	 */
 	FBoxSphereBounds( const FVector* Points, uint32 NumPoints );
 
@@ -131,8 +113,7 @@ public:
 	/**
 	 * Constructs a bounding volume containing both this and B.
 	 *
-	 * @param Other - The other bounding volume.
-	 *
+	 * @param Other The other bounding volume.
 	 * @return The combined bounding volume.
 	 */
 	FORCEINLINE FBoxSphereBounds operator+( const FBoxSphereBounds& Other ) const;
@@ -143,8 +124,7 @@ public:
 	/**
 	 * Calculates the squared distance from a point to a bounding box
 	 *
-	 * @param Point - The point.
-	 *
+	 * @param Point The point.
 	 * @return The distance.
 	 */
 	FORCEINLINE float ComputeSquaredDistanceFromBoxToPoint( const FVector& Point ) const
@@ -168,8 +148,7 @@ public:
 	/**
 	 * Gets the extrema for the bounding box.
 	 *
-	 * @param Extrema - 1 for positive extrema from the origin, else negative
-	 *
+	 * @param Extrema 1 for positive extrema from the origin, else negative
 	 * @return The boxes extrema
 	 */
 	FVector GetBoxExtrema( uint32 Extrema ) const
@@ -195,8 +174,7 @@ public:
 	/**
 	 * Increase the size of the box and sphere by a given size.
 	 *
-	 * @param ExpandAmount - The size to increase by.
-	 *
+	 * @param ExpandAmount The size to increase by.
 	 * @return A new box with the expanded size.
 	 */
 	FORCEINLINE FBoxSphereBounds ExpandBy( float ExpandAmount ) const
@@ -207,8 +185,7 @@ public:
 	/**
 	 * Gets a bounding volume transformed by a matrix.
 	 *
-	 * @param M - The matrix.
-	 *
+	 * @param M The matrix.
 	 * @return The transformed volume.
 	 */
 	CORE_API FBoxSphereBounds TransformBy( const FMatrix& M ) const;
@@ -216,8 +193,7 @@ public:
 	/**
 	 * Gets a bounding volume transformed by a FTransform object.
 	 *
-	 * @param M - The FTransform object.
-	 *
+	 * @param M The FTransform object.
 	 * @return The transformed volume.
 	 */
 	CORE_API FBoxSphereBounds TransformBy( const FTransform & M ) const;
@@ -252,15 +228,24 @@ public:
 		return Result;
 	}
 
+#if ENABLE_NAN_DIAGNOSTIC
+	FORCEINLINE void DiagnosticCheckNaN( ) const
+	{
+		checkf(!Origin.ContainsNaN(), TEXT("Origin contains NaN: %s"), *Origin.ToString());
+		checkf(!BoxExtent.ContainsNaN(), TEXT("BoxExtent contains NaN: %s"), *BoxExtent.ToString());
+		checkf(!FMath::IsNaN(SphereRadius) && FMath::IsFinite(SphereRadius), TEXT("SphereRadius contains NaN: %f"), SphereRadius);
+	}
+#else
+	FORCEINLINE void DiagnosticCheckNaN( ) const {}
+#endif
 
 public:
 
 	/**
 	 * Serializes the given bounding volume from or into the specified archive.
 	 *
-	 * @param Ar - The archive to serialize from or into.
-	 * @param Bounds - The bounding volume to serialize.
-	 *
+	 * @param Ar The archive to serialize from or into.
+	 * @param Bounds The bounding volume to serialize.
 	 * @return The archive..
 	 */
 	friend FArchive& operator<<( FArchive& Ar, FBoxSphereBounds& Bounds )
diff --git a/Engine/Source/Runtime/Core/Public/Math/Edge.h b/Engine/Source/Runtime/Core/Public/Math/Edge.h
--- a/Engine/Source/Runtime/Core/Public/Math/Edge.h
+++ b/Engine/Source/Runtime/Core/Public/Math/Edge.h
@@ -1,7 +1,7 @@
 // Copyright 1998-2014 Epic Games, Inc. All Rights Reserved.
 
 /*=============================================================================
-	Edge.h: Declares the FEdge class.
+	Edge.h: Declares the FEdge structure.
 =============================================================================*/
 
 #pragma once
@@ -10,20 +10,26 @@
 /**
  * Implements an edge consisting of two vertices.
  */
-class FEdge
+struct FEdge
 {
+	/** Holds the edge vertices. */
+	FVector Vertex[2];
+
+	/** Holds a temporary variable used when creating arrays of unique edges. */
+	int32 Count;
+
 public:
 
 	/**
-	 * Default constructor.
+	 * Default constructor (no initialization).
 	 */
 	FEdge () { }
 
 	/**
 	 * Creates and initializes a new edge from two vertices.
 	 *
-	 * @param V1 - The first vertex.
-	 * @param V2 - The second vertex.
+	 * @param V1 The first vertex.
+	 * @param V2 The second vertex.
 	 */
 	FEdge (FVector V1, FVector V2)
 	{
@@ -32,31 +38,16 @@ public:
 		Count = 0;
 	}
 
-
 public:
 
 	/**
 	 * Compares this edge with another.
 	 *
-	 * @param E - The other edge.
-	 *
+	 * @param E The other edge.
 	 * @return true if the two edges are identical, false otherwise.
 	 */
 	bool operator== (const FEdge& E) const
 	{
 		return (((E.Vertex[0] == Vertex[0]) && (E.Vertex[1] == Vertex[1])) || ((E.Vertex[0] == Vertex[1]) && (E.Vertex[1] == Vertex[0])));
 	}
-
-
-public:
-
-	/**
-	 * Holds the edge  vertices.
-	 */
-	FVector Vertex[2];
-
-	/**
-	 * Holds a temporary variable used when creating arrays of unique edges.
-	 */
-	int32 Count;
-};
\ No newline at end of file
+};
diff --git a/Engine/Source/Runtime/Core/Public/Math/IntPoint.h b/Engine/Source/Runtime/Core/Public/Math/IntPoint.h
--- a/Engine/Source/Runtime/Core/Public/Math/IntPoint.h
+++ b/Engine/Source/Runtime/Core/Public/Math/IntPoint.h
@@ -1,24 +1,25 @@
 // Copyright 1998-2014 Epic Games, Inc. All Rights Reserved.
 
+/*=============================================================================
+	IntPoint.h: Declares the FIntPoint structure.
+=============================================================================*/
+
 #pragma once
 
 
 /**
  * Structure for integer points in 2-d space.
+ *
+ * @todo Docs: The operators need better documentation, i.e. what does it mean to divide a point?
  */
 struct FIntPoint
 {
-	/**
-	 * Holds the point's x-coordinate.
-	 */
+	/** Holds the point's x-coordinate. */
 	int32 X;
 	
-	/**
-	 * Holds the point's y-coordinate.
-	 */
+	/** Holds the point's y-coordinate. */
 	int32 Y;
 
-
 public:
 
 	/** An int point with zeroed values. */
@@ -27,19 +28,18 @@ public:
 	/** An int point with INDEX_NONE values. */
 	CORE_API static const FIntPoint NoneValue;
 
-
 public:
 
 	/**
-	 * Default constructor.
+	 * Default constructor (no initialization).
 	 */
 	FIntPoint( );
 
 	/**
 	 * Creates and initializes a new instance with the specified coordinates.
 	 *
-	 * @param InX - The x-coordinate.
-	 * @param InY - The y-coordinate.
+	 * @param InX The x-coordinate.
+	 * @param InY The y-coordinate.
 	 */
 	FIntPoint( int32 InX, int32 InY );
 
@@ -50,14 +50,12 @@ public:
 	 */
 	explicit FORCEINLINE FIntPoint( EForceInit );
 
-
 public:
 
 	/**
 	 * Gets specific component of a point.
 	 *
-	 * @param PointIndex - Index of point component
-	 *
+	 * @param PointIndex Index of point component.
 	 * @return const reference to component.
 	 */
 	const int32& operator()( int32 PointIndex ) const;
@@ -65,8 +63,7 @@ public:
 	/**
 	 * Gets specific component of a point.
 	 *
-	 * @param PointIndex - Index of point component
-	 *
+	 * @param PointIndex Index of point component
 	 * @return reference to component.
 	 */
 	int32& operator()( int32 PointIndex );
@@ -75,8 +72,7 @@ public:
 	 * Compares points for equality.
 	 *
 	 * @param Other The other int point being compared.
-	 *
-	 * @return true if the points are equal otherwise false.
+	 * @return true if the points are equal, false otherwise.
 	 */
 	bool operator==( const FIntPoint& Other ) const;
 
@@ -84,8 +80,7 @@ public:
 	 * Compares points for inequality.
 	 *
 	 * @param Other The other int point being compared.
-	 *
-	 * @return true if the points are not equal otherwise false.
+	 * @return true if the points are not equal, false otherwise.
 	 */
 	bool operator!=( const FIntPoint& Other ) const;
 
@@ -93,7 +88,6 @@ public:
 	 * Scales this point.
 	 *
 	 * @param Scale What to multiply the point by.
-	 *
 	 * @return Reference to this point after multiplication.
 	 */
 	FIntPoint& operator*=( int32 Scale );
@@ -102,7 +96,6 @@ public:
 	 * Divides this point.
 	 *
 	 * @param Divisor What to divide the point by.
-	 *
 	 * @return Reference to this point after division.
 	 */
 	FIntPoint& operator/=( int32 Divisor );
@@ -111,7 +104,6 @@ public:
 	 * Adds to this point.
 	 *
 	 * @param Other The point to add to this point.
-	 *
 	 * @return Reference to this point after addition.
 	 */
 	FIntPoint& operator+=( const FIntPoint& Other );
@@ -120,18 +112,22 @@ public:
 	 * Subtracts from this point.
 	 *
 	 * @param Other The point to subtract from this point.
-	 *
 	 * @return Reference to this point after subtraction.
 	 */
 	FIntPoint& operator-=( const FIntPoint& Other );
 
+	/**
+	 * Divides this point.
+	 *
+	 * @param Other The point to divide with.
+	 * @return Reference to this point after division.
+	 */
 	FIntPoint& operator/=( const FIntPoint& Other );
 
 	/**
 	 * Assigns another point to this one.
 	 *
 	 * @param Other The point to assign this point from.
-	 *
 	 * @return Reference to this point after assignment.
 	 */
 	FIntPoint& operator=( const FIntPoint& Other );
@@ -140,7 +136,6 @@ public:
 	 * Gets the result of scaling on this point.
 	 *
 	 * @param Scale What to multiply the point by.
-	 *
 	 * @return A new scaled int point.
 	 */
 	FIntPoint operator*( int32 Scale ) const;
@@ -149,7 +144,6 @@ public:
 	 * Gets the result of division on this point.
 	 *
 	 * @param Divisor What to divide the point by.
-	 *
 	 * @return A new divided int point.
 	 */
 	FIntPoint operator/( int32 Divisor ) const;
@@ -158,7 +152,6 @@ public:
 	 * Gets the result of addition on this point.
 	 *
 	 * @param Other The other point to add to this.
-	 *
 	 * @return A new combined int point.
 	 */
 	FIntPoint operator+( const FIntPoint& Other ) const;
@@ -167,20 +160,25 @@ public:
 	 * Gets the result of subtraction from this point.
 	 *
 	 * @param Other The other point to subtract from this.
-	 *
 	 * @return A new subtracted int point.
 	 */
 	FIntPoint operator-( const FIntPoint& Other ) const;
 
+	/**
+	 * Gets the result of division on this point.
+	 *
+	 * @param Other The other point to subtract from this.
+	 * @return A new subtracted int point.
+	 */
 	FIntPoint operator/( const FIntPoint& Other ) const;
 
-
 public:
 
 	/**
 	 * Gets the maximum value in the point.
 	 *
 	 * @return The maximum value in the point.
+	 * @see GetMin, Size, SizeSquared
 	 */
 	int32 GetMax( ) const;
 
@@ -188,6 +186,7 @@ public:
 	 * Gets the minimum value in the point.
 	 *
 	 * @return The minimum value in the point.
+	 * @see GetMax, Size, SizeSquared
 	 */
 	int32 GetMin( ) const;
 
@@ -195,6 +194,7 @@ public:
 	 * Gets the distance of this point from (0,0).
 	 *
 	 * @return The distance of this point from (0,0).
+	 * @see GetMax, GetMin, SizeSquared
 	 */
 	int32 Size( ) const;
 
@@ -202,6 +202,7 @@ public:
 	 * Gets the squared distance of this point from (0,0).
 	 *
 	 * @return The squared distance of this point from (0,0).
+	 * @see GetMax, GetMin, Size
 	 */
 	int32 SizeSquared( ) const;
 	
@@ -212,7 +213,6 @@ public:
 	 */
 	FString ToString( ) const;
 
-
 public:
 
 	/**
@@ -220,28 +220,34 @@ public:
 	 *
 	 * @param lhs The int point being divided.
 	 * @param Divisor What to divide the int point by.
-	 *
 	 * @return A new divided int point.
+	 * @see DivideAndRoundDown
 	 */
 	static FIntPoint DivideAndRoundUp( FIntPoint lhs, int32 Divisor );
+
+	/**
+	 * Divide an int point and round down the result.
+	 *
+	 * @param lhs The int point being divided.
+	 * @param Divisor What to divide the int point by.
+	 * @return A new divided int point.
+	 * @see DivideAndRoundUp
+	 */
 	static FIntPoint DivideAndRoundDown( FIntPoint lhs, int32 Divisor );
 
 	/**
 	 * Gets number of components point has.
-	 *
 	 * @return number of components point has.
 	 */
 	static int32 Num( );
 
-
 public:
 
 	/**
 	 * Serializes the point.
 	 *
-	 * @param Ar - The archive to serialize into.
-	 * @param Point - The point to serialize.
-	 *
+	 * @param Ar The archive to serialize into.
+	 * @param Point The point to serialize.
 	 * @return Reference to the Archive after serialization.
 	 */
 	friend FArchive& operator<<( FArchive& Ar, FIntPoint& Point )
@@ -251,6 +257,9 @@ public:
 };
 
 
+/* FInterval inline functions
+ *****************************************************************************/
+
 FORCEINLINE FIntPoint::FIntPoint( ) { }
 
 
diff --git a/Engine/Source/Runtime/Core/Public/Math/IntRect.h b/Engine/Source/Runtime/Core/Public/Math/IntRect.h
--- a/Engine/Source/Runtime/Core/Public/Math/IntRect.h
+++ b/Engine/Source/Runtime/Core/Public/Math/IntRect.h
@@ -1,23 +1,25 @@
 // Copyright 1998-2014 Epic Games, Inc. All Rights Reserved.
 
+/*=============================================================================
+	IntRect.h: Declares the FIntRect structure.
+=============================================================================*/
+
 #pragma once
 
+
 /**
  * Structure for integer rectangles in 2-d space.
+ *
+ * @todo Docs: The operators need better documentation, i.e. what does it mean to divide a rectangle?
  */
 struct FIntRect
 {
-	/**
-	 * Holds the first pixel line/row (like in Win32 RECT).
-	 */
+	/** Holds the first pixel line/row (like in Win32 RECT). */
 	FIntPoint Min;
 
-	/**
-	 * Holds the last pixel line/row (like in Win32 RECT).
-	 */
+	/** Holds the last pixel line/row (like in Win32 RECT). */
 	FIntPoint Max;
 
-
 public:
 
 	/** Constructor */
@@ -41,14 +43,12 @@ public:
 	 */
 	FIntRect( FIntPoint InMin, FIntPoint InMax );
 
-
 public:
 
 	/**
 	 * Gets a specific point in this rectangle.
 	 *
-	 * @param PointIndex - Index of Point in rectangle.
-	 *
+	 * @param PointIndex Index of Point in rectangle.
 	 * @return Const reference to point in rectangle.
 	 */
 	const FIntPoint& operator()( int32 PointIndex ) const;
@@ -56,8 +56,7 @@ public:
 	/**
 	 * Gets a specific point in this rectangle.
 	 *
-	 * @param PointIndex - Index of Point in rectangle.
-	 *
+	 * @param PointIndex Index of Point in rectangle.
 	 * @return Reference to point in rectangle.
 	 */
 	FIntPoint& operator()( int32 PointIndex );
@@ -66,8 +65,7 @@ public:
 	 * Compares Rectangles for equality.
 	 *
 	 * @param Other The Other Rectangle for comparison.
-	 *
-	 * @return true if the rectangles are equal otherwise false.
+	 * @return true if the rectangles are equal, false otherwise..
 	 */
 	bool operator==( const FIntRect& Other ) const;
 
@@ -75,8 +73,7 @@ public:
 	 * Compares Rectangles for inequality.
 	 *
 	 * @param Other The Other Rectangle for comparison.
-	 *
-	 * @return true if the rectangles are not equal otherwise false.
+	 * @return true if the rectangles are not equal, false otherwise..
 	 */
 	bool operator!=( const FIntRect& Other ) const;
 
@@ -84,7 +81,6 @@ public:
 	 * Applies scaling to this rectangle.
 	 *
 	 * @param Scale What to multiply the rectangle by.
-	 *
 	 * @return Reference to this rectangle after scaling.
 	 */
 	FIntRect& operator*=( int32 Scale );
@@ -92,8 +88,7 @@ public:
 	/**
 	 * Adds a point to this rectangle.
 	 *
-	 * @param Point - The point to add onto both points in the rectangle.
-	 *
+	 * @param Point The point to add onto both points in the rectangle.
 	 * @return Reference to this rectangle after addition.
 	 */
 	FIntRect& operator+=( const FIntPoint& Point );
@@ -101,8 +96,7 @@ public:
 	/**
 	 * Subtracts a point from this rectangle.
 	 *
-	 * @param Point - The point to subtract from both points in the rectangle.
-	 *
+	 * @param Point The point to subtract from both points in the rectangle.
 	 * @return Reference to this rectangle after subtraction.
 	 */
 	FIntRect& operator-=( const FIntPoint& Point );
@@ -111,7 +105,6 @@ public:
 	 * Gets the result of scaling on this rectangle.
 	 *
 	 * @param Scale What to multiply this rectangle by.
-	 *
 	 * @return New scaled rectangle.
 	 */
 	FIntRect operator*( int32 Scale ) const;
@@ -120,7 +113,6 @@ public:
 	 * Gets the result of division on this rectangle.
 	 *
 	 * @param Div What to divide this rectangle by.
-	 *
 	 * @return New divided rectangle.
 	 */
 	FIntRect operator/( int32 Div ) const;
@@ -128,19 +120,23 @@ public:
 	/**
 	 * Gets the result of adding a point to this rectangle.
 	 *
-	 * @param Point - The point to add to both points in the rectangle.
-	 *
+	 * @param Point The point to add to both points in the rectangle.
 	 * @return New rectangle with point added to it.
 	 */
 	FIntRect operator+( const FIntPoint& Point ) const;
 
+	/**
+	 * Gets the result of dividing a point with this rectangle.
+	 *
+	 * @param Point The point to divide with.
+	 * @return New rectangle with point divided.
+	 */
 	FIntRect operator/( const FIntPoint& Point ) const;
 
 	/**
 	 * Gets the result of subtracting a point from this rectangle.
 	 *
-	 * @param Point - The point to subtract from both points in the rectangle.
-	 *
+	 * @param Point The point to subtract from both points in the rectangle.
 	 * @return New rectangle with point subtracted from it.
 	 */
 	FIntRect operator-( const FIntPoint& Point ) const;
@@ -148,8 +144,7 @@ public:
 	/**
 	 * Gets the result of adding two rectangles together.
 	 *
-	 * @param Other - The other rectangle to add to this.
-	 *
+	 * @param Other The other rectangle to add to this.
 	 * @return New rectangle after both are added together.
 	 */
 	FIntRect operator+( const FIntRect& Other ) const;
@@ -157,13 +152,11 @@ public:
 	/**
 	 * Gets the result of subtracting a rectangle from this one.
 	 *
-	 * @param Other - The other rectangle to subtract from this.
-	 *
+	 * @param Other The other rectangle to subtract from this.
 	 * @return New rectangle after one is subtracted from this.
 	 */
 	FIntRect operator-( const FIntRect& Other ) const;
 
-
 public:
 
 	/**
@@ -177,32 +170,29 @@ public:
 	 * Creates a rectangle from the bottom part of this rectangle.
 	 *
 	 * @param InHeight Height of the new rectangle (<= rectangles original height).
-	 *
 	 * @return The new rectangle.
 	 */
 	FIntRect Bottom( int32 InHeight ) const;
 
 	/**
 	 * Clip a rectangle using the bounds of another rectangle.
-	 *
-	 * @param Other - The other rectangle to clip against.
+	 * @param Other The other rectangle to clip against.
 	 */
 	void Clip( const FIntRect& Other );
 
 	/**
 	 * Test whether this rectangle contains a point.
 	 *
-	 * @param Point - The point to test against.
-	 *
-	 * @return true if the rectangle contains the specified point, otherwise false.
+	 * @param Point The point to test against.
+	 * @return true if the rectangle contains the specified point,, false otherwise..
 	 */
 	bool Contains( FIntPoint Point ) const;
 
 	/**
 	 * Gets the Center and Extents of this rectangle.
-
-	 * @param OutCenter - Will contain the center point.
-	 * @param OutExtent - Will contain the extent.
+	 *
+	 * @param OutCenter Will contain the center point.
+	 * @param OutExtent Will contain the extent.
 	 */
 	void GetCenterAndExtents( FIntPoint& OutCenter, FIntPoint& OutExtent ) const;
 
@@ -216,22 +206,21 @@ public:
 	/**
 	 * Inflates or deflates the rectangle.
 	 *
-	 * @param Amount - The amount to inflate or deflate the rectangle on each side.
+	 * @param Amount The amount to inflate or deflate the rectangle on each side.
 	 */
 	void InflateRect( int32 Amount );
 
 	/**
 	 * Adds to this rectangle to include a given point.
 	 *
-	 * @param Point the point to increase the rectangle to.
+	 * @param Point The point to increase the rectangle to.
 	 */
 	void Include( FIntPoint Point );
 
 	/**
 	 * Gets a new rectangle from the inner of this one.
 	 *
-	 * @param Shrink - How much to remove from each point of this rectangle.
-	 *
+	 * @param Shrink How much to remove from each point of this rectangle.
 	 * @return New inner Rectangle.
 	 */
 	FIntRect Inner( FIntPoint Shrink ) const;
@@ -240,7 +229,6 @@ public:
 	 * Creates a rectangle from the right hand side of this rectangle.
 	 *
 	 * @param InWidth Width of the new rectangle (<= rectangles original width).
-	 *
 	 * @return The new rectangle.
 	 */
 	FIntRect Right( int32 InWidth ) const;
@@ -249,7 +237,6 @@ public:
 	 * Scales a rectangle using a floating point number.
 	 *
 	 * @param Fraction What to scale the rectangle by
-	 *
 	 * @return New scaled rectangle.
 	 */
 	FIntRect Scale( float Fraction ) const;
@@ -289,7 +276,6 @@ public:
 	 *
 	 * @param lhs The Rectangle to divide.
 	 * @param Div What to divide by.
-	 *
 	 * @return New divided rectangle.
 	 */
 	static FIntRect DivideAndRoundUp( FIntRect lhs, int32 Div );
@@ -301,7 +287,6 @@ public:
 	 */
 	static int32 Num( );
 
-
 public:
 
 	/**
@@ -328,6 +313,9 @@ FORCEINLINE FIntRect FIntRect::Scale( float Fraction ) const
 }
 
 
+/* FIntRect inline functions
+ *****************************************************************************/
+
 FORCEINLINE FIntRect::FIntRect( )
 	: Min(ForceInit)
 	, Max(ForceInit)
diff --git a/Engine/Source/Runtime/Core/Public/Math/IntVector.h b/Engine/Source/Runtime/Core/Public/Math/IntVector.h
--- a/Engine/Source/Runtime/Core/Public/Math/IntVector.h
+++ b/Engine/Source/Runtime/Core/Public/Math/IntVector.h
@@ -1,5 +1,9 @@
 // Copyright 1998-2014 Epic Games, Inc. All Rights Reserved.
 
+/*=============================================================================
+	IntVector.h: Declares the FIntVector structure.
+=============================================================================*/
+
 #pragma once
 
 
@@ -8,22 +12,15 @@
  */
 struct FIntVector
 {
-	/**
-	 * Holds the point's x-coordinate.
-	 */
+	/** Holds the point's x-coordinate. */
 	int32 X;
 	
-	/**
-	 * Holds the point's y-coordinate.
-	 */
+	/** Holds the point's y-coordinate. */
 	int32 Y;
 
-	/**
-	 * Holds the point's z-coordinate.
-	 */
+	/**  Holds the point's z-coordinate. */
 	int32 Z;
 
-
 public:
 
 	/** An int point with zeroed values. */
@@ -32,20 +29,19 @@ public:
 	/** An int point with INDEX_NONE values. */
 	CORE_API static const FIntVector NoneValue;
 
-
 public:
 
 	/**
-	 * Default constructor.
+	 * Default constructor (no initialization).
 	 */
-	FIntVector();
+	FIntVector( );
 
 	/**
 	 * Creates and initializes a new instance with the specified coordinates.
 	 *
-	 * @param InX - The x-coordinate.
-	 * @param InY - The y-coordinate.
-	 * @param InZ - The z-coordinate.
+	 * @param InX The x-coordinate.
+	 * @param InY The y-coordinate.
+	 * @param InZ The z-coordinate.
 	 */
 	FIntVector( int32 InX, int32 InY, int32 InZ );
 
@@ -63,33 +59,29 @@ public:
 	 */
 	explicit FORCEINLINE FIntVector( EForceInit );
 
-
 public:
 
 	/**
 	 * Gets specific component of a point.
 	 *
-	 * @param i Index of point component
-	 *
+	 * @param ComponentIndex Index of point component.
 	 * @return const reference to component.
 	 */
-	const int32& operator()( int32 i ) const;
+	const int32& operator()( int32 ComponentIndex ) const;
 
 	/**
 	 * Gets specific component of a point.
 	 *
-	 * @param i Index of point component
-	 *
+	 * @param ComponentIndex Index of point component.
 	 * @return reference to component.
 	 */
-	int32& operator()( int32 i );
+	int32& operator()( int32 ComponentIndex );
 
 	/**
 	 * Compares points for equality.
 	 *
 	 * @param Other The other int point being compared.
-	 *
-	 * @return true if the points are equal otherwise false.
+	 * @return true if the points are equal, false otherwise..
 	 */
 	bool operator==( const FIntVector& Other ) const;
 
@@ -97,8 +89,7 @@ public:
 	 * Compares points for inequality.
 	 *
 	 * @param Other The other int point being compared.
-	 *
-	 * @return true if the points are not equal otherwise false.
+	 * @return true if the points are not equal, false otherwise..
 	 */
 	bool operator!=( const FIntVector& Other ) const;
 
@@ -106,7 +97,6 @@ public:
 	 * Scales this point.
 	 *
 	 * @param Scale What to multiply the point by.
-	 *
 	 * @return Reference to this point after multiplication.
 	 */
 	FIntVector& operator*=( int32 Scale );
@@ -115,7 +105,6 @@ public:
 	 * Divides this point.
 	 *
 	 * @param Divisor What to divide the point by.
-	 *
 	 * @return Reference to this point after division.
 	 */
 	FIntVector& operator/=( int32 Divisor );
@@ -124,7 +113,6 @@ public:
 	 * Adds to this point.
 	 *
 	 * @param Other The point to add to this point.
-	 *
 	 * @return Reference to this point after addition.
 	 */
 	FIntVector& operator+=( const FIntVector& Other );
@@ -133,7 +121,6 @@ public:
 	 * Subtracts from this point.
 	 *
 	 * @param Other The point to subtract from this point.
-	 *
 	 * @return Reference to this point after subtraction.
 	 */
 	FIntVector& operator-=( const FIntVector& Other );
@@ -142,7 +129,6 @@ public:
 	 * Assigns another point to this one.
 	 *
 	 * @param Other The point to assign this point from.
-	 *
 	 * @return Reference to this point after assignment.
 	 */
 	FIntVector& operator=( const FIntVector& Other );
@@ -151,7 +137,6 @@ public:
 	 * Gets the result of scaling on this point.
 	 *
 	 * @param Scale What to multiply the point by.
-	 *
 	 * @return A new scaled int point.
 	 */
 	FIntVector operator*( int32 Scale ) const;
@@ -160,7 +145,6 @@ public:
 	 * Gets the result of division on this point.
 	 *
 	 * @param Divisor What to divide the point by.
-	 *
 	 * @return A new divided int point.
 	 */
 	FIntVector operator/( int32 Divisor ) const;
@@ -169,7 +153,6 @@ public:
 	 * Gets the result of addition on this point.
 	 *
 	 * @param Other The other point to add to this.
-	 *
 	 * @return A new combined int point.
 	 */
 	FIntVector operator+( const FIntVector& Other ) const;
@@ -178,12 +161,10 @@ public:
 	 * Gets the result of subtraction from this point.
 	 *
 	 * @param Other The other point to subtract from this.
-	 *
 	 * @return A new subtracted int point.
 	 */
 	FIntVector operator-( const FIntVector& Other ) const;
 
-
 public:
 
 	/**
@@ -214,7 +195,6 @@ public:
 	 */
 	FString ToString( ) const;
 
-
 public:
 
 	/**
@@ -222,7 +202,6 @@ public:
 	 *
 	 * @param lhs The int point being divided.
 	 * @param Divisor What to divide the int point by.
-	 *
 	 * @return A new divided int point.
 	 */
 	static FIntVector DivideAndRoundUp( FIntVector lhs, int32 Divisor );
@@ -234,15 +213,13 @@ public:
 	 */
 	static int32 Num( );
 
-
 public:
 
 	/**
 	 * Serializes the Rectangle.
 	 *
-	 * @param Ar - The archive to serialize into.
-	 * @param Vector - The vector to serialize.
-	 *
+	 * @param Ar The archive to serialize into.
+	 * @param Vector The vector to serialize.
 	 * @return Reference to the Archive after serialization.
 	 */
 	friend FArchive& operator<<( FArchive& Ar, FIntVector& Vector )
@@ -252,6 +229,9 @@ public:
 };
 
 
+/* FIntVector inline functions
+ *****************************************************************************/
+
 FORCEINLINE FIntVector::FIntVector( )
 { }
 
@@ -277,15 +257,15 @@ FORCEINLINE FIntVector::FIntVector( EForceInit )
 { }
 
 
-FORCEINLINE const int32& FIntVector::operator()( int32 i ) const
+FORCEINLINE const int32& FIntVector::operator()( int32 ComponentIndex ) const
 {
-	return (&X)[i];
+	return (&X)[ComponentIndex];
 }
 
 
-FORCEINLINE int32& FIntVector::operator()( int32 i )
+FORCEINLINE int32& FIntVector::operator()( int32 ComponentIndex )
 {
-	return (&X)[i];
+	return (&X)[ComponentIndex];
 }
 
 
diff --git a/Engine/Source/Runtime/Core/Public/Math/InterpCurve.h b/Engine/Source/Runtime/Core/Public/Math/InterpCurve.h
--- a/Engine/Source/Runtime/Core/Public/Math/InterpCurve.h
+++ b/Engine/Source/Runtime/Core/Public/Math/InterpCurve.h
@@ -1,22 +1,59 @@
 // Copyright 1998-2014 Epic Games, Inc. All Rights Reserved.
 
+/*=============================================================================
+	InterpCurve.h: Declares the FInterpCurve class.
+=============================================================================*/
+
 #pragma once
 
-template< class T > class FInterpCurve
+
+/**
+ * Template for interpolation curves.
+ *
+ * @see FInterpCurvePoint
+ * @todo Docs: FInterpCurve needs template and function documentation
+ */
+template<class T>
+class FInterpCurve
 {
 public:
-	TArray< FInterpCurvePoint<T> >	Points;
 
-	/** Constructor */
-	FInterpCurve();
+	/** Holds the collection of interpolation points. */
+	TArray<FInterpCurvePoint<T>> Points;
+
+public:
 
-	/** Add a new keypoint to the InterpCurve with the supplied In and Out value. Returns the index of the new key.*/
+	/**
+	 * Default constructor (no initialization).
+	 */
+	FInterpCurve( ) { }
+
+public:
+
+	/**
+	 * Adds a new keypoint to the InterpCurve with the supplied In and Out value.
+	 *
+	 * @param InVal
+	 * @param OutVal
+	 *
+	 * @return The index of the new key.
+	 */
 	int32 AddPoint( const float InVal, const T &OutVal );
 
-	/** Move a keypoint to a new In value. This may change the index of the keypoint, so the new key index is returned. */
+	/**
+	 * Moves a keypoint to a new In value.
+	 *
+	 * This may change the index of the keypoint, so the new key index is returned.
+	 *
+	 * @param PointIndex
+	 * @param NewInVal
+	 * @return
+	 */
 	int32 MovePoint( int32 PointIndex, float NewInVal );
 
-	/** Clear all keypoints from InterpCurve. */
+	/**
+	 * Clears all keypoints from InterpCurve.
+	 */
 	void Reset();
 
 	/** 
@@ -58,6 +95,7 @@ public:
 	/** Calculate the min/max out value that can be returned by this InterpCurve. */
 	void CalcBounds(T& OutMin, T& OutMax, const T& Default) const;
 
+public:
 
 	/**
 	 * Serializes the interp curve.
@@ -78,12 +116,16 @@ public:
 	}
 
 private:
-	/** Find the lower index of the two points whose input values bound the supplied input value */
+
+	/**
+	 * Finds the lower index of the two points whose input values bound the supplied input value.
+	 */
 	int32 GetPointIndexForInputValue(const float InValue) const;
 };
 
-template< class T > 
-FORCEINLINE FInterpCurve<T>::FInterpCurve() {}
+
+/* FInterpCurve inline functions
+ *****************************************************************************/
 
 template< class T > 
 FORCEINLINE int32 FInterpCurve<T>::AddPoint( const float InVal, const T &OutVal )
@@ -94,6 +136,7 @@ FORCEINLINE int32 FInterpCurve<T>::AddPoint( const float InVal, const T &OutVal
 	return i;
 }
 
+
 template< class T > 
 FORCEINLINE int32 FInterpCurve<T>::MovePoint( int32 PointIndex, float NewInVal )
 {
@@ -115,12 +158,14 @@ FORCEINLINE int32 FInterpCurve<T>::MovePoint( int32 PointIndex, float NewInVal )
 	return NewPointIndex;
 }
 
+
 template< class T > 
 FORCEINLINE void FInterpCurve<T>::Reset()
 {
 	Points.Empty();
 }
 
+
 template< class T >
 int32 FInterpCurve<T>::GetPointIndexForInputValue(const float InValue) const
 {
@@ -158,6 +203,7 @@ int32 FInterpCurve<T>::GetPointIndexForInputValue(const float InValue) const
 	return MinIndex;
 }
 
+
 template< class T > 
 T FInterpCurve<T>::Eval( const float InVal, const T& Default ) const
 {
@@ -209,6 +255,7 @@ T FInterpCurve<T>::Eval( const float InVal, const T& Default ) const
 	}
 }
 
+
 template< class T > 
 T FInterpCurve<T>::EvalDerivative( const float InVal, const T& Default ) const
 {
@@ -260,6 +307,7 @@ T FInterpCurve<T>::EvalDerivative( const float InVal, const T& Default ) const
 	}
 }
 
+
 template< class T > 
 T FInterpCurve<T>::EvalSecondDerivative( const float InVal, const T& Default ) const
 {
@@ -312,6 +360,7 @@ T FInterpCurve<T>::EvalSecondDerivative( const float InVal, const T& Default ) c
 	}
 }
 
+
 template< class T > 
 float FInterpCurve<T>::InaccurateFindNearest( const T &PointInSpace, float& OutDistanceSq ) const
 {
@@ -341,6 +390,7 @@ float FInterpCurve<T>::InaccurateFindNearest( const T &PointInSpace, float& OutD
 	return 0.0f;
 }
 
+
 template< class T > 
 float FInterpCurve<T>::InaccurateFindNearestOnSegment( const T &PointInSpace, int32 PtIdx, float& OutSquaredDistance ) const
 {
@@ -423,6 +473,7 @@ float FInterpCurve<T>::InaccurateFindNearestOnSegment( const T &PointInSpace, in
 	}
 }
 
+
 template< class T > 
 void FInterpCurve<T>::AutoSetTangents(float Tension)
 {
@@ -493,6 +544,7 @@ void FInterpCurve<T>::AutoSetTangents(float Tension)
 	}
 }
 
+
 template< class T > 
 void FInterpCurve<T>::CalcBounds(T& OutMin, T& OutMax, const T& Default) const
 {
@@ -519,9 +571,23 @@ void FInterpCurve<T>::CalcBounds(T& OutMin, T& OutMax, const T& Default) const
 }
 
 
-typedef FInterpCurve<float>				FInterpCurveFloat;
-typedef FInterpCurve<FVector2D>			FInterpCurveVector2D;
-typedef FInterpCurve<FVector>			FInterpCurveVector;
-typedef FInterpCurve<FQuat>				FInterpCurveQuat;
-typedef FInterpCurve<FTwoVectors>		FInterpCurveTwoVectors;
-typedef FInterpCurve<FLinearColor>		FInterpCurveLinearColor;
+/* Common type definitions
+ *****************************************************************************/
+
+/** Type definition for floating point interpolation curves. */
+typedef FInterpCurve<float> FInterpCurveFloat;
+
+/** Type definition for 2D vector interpolation curves. */
+typedef FInterpCurve<FVector2D> FInterpCurveVector2D;
+
+/** Type definition for 3D vector interpolation curves. */
+typedef FInterpCurve<FVector> FInterpCurveVector;
+
+/** Type definition for quaternion interpolation curves. */
+typedef FInterpCurve<FQuat> FInterpCurveQuat;
+
+/** Type definition for vector pair interpolation curves. */
+typedef FInterpCurve<FTwoVectors> FInterpCurveTwoVectors;
+
+/** Type definition for color interpolation curves. */
+typedef FInterpCurve<FLinearColor> FInterpCurveLinearColor;
diff --git a/Engine/Source/Runtime/Core/Public/Math/InterpCurvePoint.h b/Engine/Source/Runtime/Core/Public/Math/InterpCurvePoint.h
--- a/Engine/Source/Runtime/Core/Public/Math/InterpCurvePoint.h
+++ b/Engine/Source/Runtime/Core/Public/Math/InterpCurvePoint.h
@@ -1,5 +1,9 @@
 // Copyright 1998-2014 Epic Games, Inc. All Rights Reserved.
 
+/*=============================================================================
+	InterpCurvePoint.h: Declares the FInterpCurvePoint class and related types.
+=============================================================================*/
+
 #pragma once
 
 enum EInterpCurveMode
@@ -28,26 +32,40 @@ enum EInterpCurveMode
 	CIM_Unknown
 };
 
+
+/**
+ * Template for interpolation points.
+ *
+ * Interpolation points are used for describing the shape of interpolation curves.
+ *
+ * @see FInterpCurve
+ * @todo Docs: FInterpCurvePoint needs better function documentation.
+ */
 template< class T > class FInterpCurvePoint
 {
 public:
+
 	/** Float input value that corresponds to this key (eg. time). */
-	float		InVal;
+	float InVal;
 
 	/** Output value of templated type when input is equal to InVal. */
-	T			OutVal;
+	T OutVal;
 
 	/** Tangent of curve arrive this point. */
-	T			ArriveTangent; 
+	T ArriveTangent; 
 
 	/** Tangent of curve leaving this point. */
-	T			LeaveTangent; 
+	T LeaveTangent; 
 
 	/** Interpolation mode between this point and the next one. @see EInterpCurveMode */
 	TEnumAsByte<EInterpCurveMode> InterpMode; 
 
-	/** Constructor */
-	FInterpCurvePoint();
+public:
+
+	/**
+	 * Default constructor (no initialization).
+	 */
+	FInterpCurvePoint( ) { };
 
 	/** 
 	 * Constructor 
@@ -56,7 +74,7 @@ public:
 	 * @param Out Output value of templated type
 	 * @note uses linear interpolation
 	 */
-	FInterpCurvePoint(const float In, const T &Out);
+	FInterpCurvePoint( const float In, const T &Out );
 
 	/** 
 	 * Constructor 
@@ -67,11 +85,14 @@ public:
 	 * @param InLeaveTangent Tangent of curve leaving from this point.
 	 * @param InInterpMode interpolation mode to use
 	 */
-	FInterpCurvePoint(const float In, const T &Out, const T &InArriveTangent, const T &InLeaveTangent, const EInterpCurveMode InInterpMode);
+	FInterpCurvePoint( const float In, const T &Out, const T &InArriveTangent, const T &InLeaveTangent, const EInterpCurveMode InInterpMode );
+
+public:
 
 	/** @return true if the key value is using a curve interp mode, otherwise false */
-	FORCEINLINE bool IsCurveKey() const;
+	FORCEINLINE bool IsCurveKey( ) const;
 
+public:
 
 	/**
 	 * Serializes the Curve Point.
@@ -90,13 +111,14 @@ public:
 	}
 };
 
-template< class T > 
-FORCEINLINE FInterpCurvePoint<T>::FInterpCurvePoint() {}
+
+/* FInterpCurvePoint inline functions
+ *****************************************************************************/
 
 template< class T > 
-FORCEINLINE FInterpCurvePoint<T>::FInterpCurvePoint(const float In, const T &Out) : 
-InVal(In), 
-	OutVal(Out)
+FORCEINLINE FInterpCurvePoint<T>::FInterpCurvePoint( const float In, const T &Out )
+	: InVal(In)
+	, OutVal(Out)
 {
 	FMemory::Memset( &ArriveTangent, 0, sizeof(T) );	
 	FMemory::Memset( &LeaveTangent, 0, sizeof(T) );
@@ -104,18 +126,19 @@ InVal(In),
 	InterpMode = CIM_Linear;
 }
 
+
 template< class T > 
-FORCEINLINE FInterpCurvePoint<T>::FInterpCurvePoint(const float In, const T &Out, const T &InArriveTangent, const T &InLeaveTangent, const EInterpCurveMode InInterpMode) : 
-InVal(In), 
-	OutVal(Out), 
-	ArriveTangent(InArriveTangent),
-	LeaveTangent(InLeaveTangent),
-	InterpMode(InInterpMode)
-{
-}
+FORCEINLINE FInterpCurvePoint<T>::FInterpCurvePoint( const float In, const T &Out, const T &InArriveTangent, const T &InLeaveTangent, const EInterpCurveMode InInterpMode)
+	: InVal(In)
+	, OutVal(Out)
+	, ArriveTangent(InArriveTangent)
+	, LeaveTangent(InLeaveTangent)
+	, InterpMode(InInterpMode)
+{ }
+
 
 template< class T > 
-FORCEINLINE bool FInterpCurvePoint<T>::IsCurveKey() const
+FORCEINLINE bool FInterpCurvePoint<T>::IsCurveKey( ) const
 {
 	return ((InterpMode == CIM_CurveAuto) || (InterpMode == CIM_CurveAutoClamped) || (InterpMode == CIM_CurveUser) || (InterpMode == CIM_CurveBreak));
 }
@@ -125,6 +148,7 @@ FORCEINLINE bool FInterpCurvePoint<T>::IsCurveKey() const
  */
 CORE_API float ClampFloatTangent( float PrevPointVal, float PrevTime, float CurPointVal, float CurTime, float NextPointVal, float NextTime );
 
+
 /** Computes Tangent for a curve segment */
 template< class T, class U > 
 inline void AutoCalcTangent( const T& PrevP, const T& P, const T& NextP, const U& Tension, T& OutTan )
@@ -132,6 +156,7 @@ inline void AutoCalcTangent( const T& PrevP, const T& P, const T& NextP, const U
 	OutTan = (1.f - Tension) * ( (P - PrevP) + (NextP - P) );
 }
 
+
 /**
  * This actually returns the control point not a tangent. This is expected by the CubicInterp function for Quaternions
  */
@@ -141,6 +166,7 @@ inline void AutoCalcTangent( const FQuat& PrevP, const FQuat& P, const FQuat& Ne
 	FQuat::CalcTangents(PrevP, P, NextP, Tension, OutTan);
 }
 
+
 /** Computes a tangent for the specified control point.  General case, doesn't support clamping. */
 template< class T >
 inline void ComputeCurveTangent( float PrevTime, const T& PrevPoint,
@@ -160,7 +186,6 @@ inline void ComputeCurveTangent( float PrevTime, const T& PrevPoint,
 }
 
 
-
 /**
  * Computes a tangent for the specified control point; supports clamping, but only works
  * with floats or contiguous arrays of floats.
@@ -286,6 +311,7 @@ inline void ComputeCurveTangent( float PrevTime, const FTwoVectors& PrevPoint,
  */
 void CORE_API CurveFloatFindIntervalBounds( const FInterpCurvePoint<float>& Start, const FInterpCurvePoint<float>& End, float& CurrentMin, float& CurrentMax );
 
+
 /**
  * Calculate bounds of 2D vector intervals
  *
@@ -296,6 +322,7 @@ void CORE_API CurveFloatFindIntervalBounds( const FInterpCurvePoint<float>& Star
  */
 void CORE_API CurveVector2DFindIntervalBounds( const FInterpCurvePoint<FVector2D>& Start, const FInterpCurvePoint<FVector2D>& End, FVector2D& CurrentMin, FVector2D& CurrentMax );
 
+
 /**
  * Calculate bounds of vector intervals
  *
@@ -306,6 +333,7 @@ void CORE_API CurveVector2DFindIntervalBounds( const FInterpCurvePoint<FVector2D
  */
 void CORE_API CurveVectorFindIntervalBounds( const FInterpCurvePoint<FVector>& Start, const FInterpCurvePoint<FVector>& End, FVector& CurrentMin, FVector& CurrentMax );
 
+
 /**
  * Calculate bounds of twovector intervals
  *
@@ -316,6 +344,7 @@ void CORE_API CurveVectorFindIntervalBounds( const FInterpCurvePoint<FVector>& S
  */
 void CORE_API CurveTwoVectorsFindIntervalBounds(const FInterpCurvePoint<FTwoVectors>& Start, const FInterpCurvePoint<FTwoVectors>& End, FTwoVectors& CurrentMin, FTwoVectors& CurrentMax);
 
+
 /**
  * Calculate bounds of color intervals
  *
@@ -326,34 +355,40 @@ void CORE_API CurveTwoVectorsFindIntervalBounds(const FInterpCurvePoint<FTwoVect
  */
 void CORE_API CurveLinearColorFindIntervalBounds( const FInterpCurvePoint<FLinearColor>& Start, const FInterpCurvePoint<FLinearColor>& End, FLinearColor& CurrentMin, FLinearColor& CurrentMax );
 
+
 template< class T, class U > 
 inline void CurveFindIntervalBounds( const FInterpCurvePoint<T>& Start, const FInterpCurvePoint<T>& End, T& CurrentMin, T& CurrentMax, const U& Dummy )
 { }
 
+
 template< class U > 
 inline void CurveFindIntervalBounds( const FInterpCurvePoint<float>& Start, const FInterpCurvePoint<float>& End, float& CurrentMin, float& CurrentMax, const U& Dummy )
 {
 	CurveFloatFindIntervalBounds(Start, End, CurrentMin, CurrentMax);
 }
 
+
 template< class U > 
 void CurveFindIntervalBounds( const FInterpCurvePoint<FVector2D>& Start, const FInterpCurvePoint<FVector2D>& End, FVector2D& CurrentMin, FVector2D& CurrentMax, const U& Dummy )
 {
 	CurveVector2DFindIntervalBounds(Start, End, CurrentMin, CurrentMax);
 }
 
+
 template< class U > 
 inline void CurveFindIntervalBounds( const FInterpCurvePoint<FVector>& Start, const FInterpCurvePoint<FVector>& End, FVector& CurrentMin, FVector& CurrentMax, const U& Dummy )
 {
 	CurveVectorFindIntervalBounds(Start, End, CurrentMin, CurrentMax);
 }
 
+
 template< class U > 
 inline void CurveFindIntervalBounds( const FInterpCurvePoint<FTwoVectors>& Start, const FInterpCurvePoint<FTwoVectors>& End, FTwoVectors& CurrentMin, FTwoVectors& CurrentMax, const U& Dummy )
 {
 	CurveTwoVectorsFindIntervalBounds(Start, End, CurrentMin, CurrentMax);
 }
 
+
 template< class U > 
 inline void CurveFindIntervalBounds( const FInterpCurvePoint<FLinearColor>& Start, const FInterpCurvePoint<FLinearColor>& End, FLinearColor& CurrentMin, FLinearColor& CurrentMax, const U& Dummy )
 {
diff --git a/Engine/Source/Runtime/Core/Public/Math/Interval.h b/Engine/Source/Runtime/Core/Public/Math/Interval.h
--- a/Engine/Source/Runtime/Core/Public/Math/Interval.h
+++ b/Engine/Source/Runtime/Core/Public/Math/Interval.h
@@ -12,28 +12,23 @@
  */
 struct FInterval
 {
-	/**
-	 * Holds the lower bound of the interval.
-	 */
+	/** Holds the lower bound of the interval. */
 	float Min;
 	
-	/**
-	 * Holds the upper bound of the interval.
-	 */
+	/** Holds the upper bound of the interval. */
 	float Max;
 	
-	/**
-	 * Holds a flag indicating whether the interval is empty.
-	 */
+	/** Holds a flag indicating whether the interval is empty. */
 	bool bIsEmpty;
 
-
 public:
 
 	/**
 	 * Default constructor.
+	 *
+	 * The interval is initialized to [0, 0].
 	 */
-	FInterval()
+	FInterval( )
 		: Min(0.0f)
 		, Max(0.0f)
 		, bIsEmpty(true)
@@ -42,20 +37,23 @@ public:
     /**
 	 * Creates and initializes a new interval with the specified lower and upper bounds.
 	 *
-	 * @param InMin - The lower bound of the constructed interval.
-	 * @param InMax - The upper bound of the constructed interval.
+	 * @param InMin The lower bound of the constructed interval.
+	 * @param InMax The upper bound of the constructed interval.
 	 */
-	FInterval(float InMin, float InMax)
+	FInterval( float InMin, float InMax )
 		: Min(InMin)
 		, Max(InMax)
 		, bIsEmpty(InMin >= InMax)
 	{ }
 
-
 public:
 
-	/** Offset the interval by adding X. */
-	void operator+= (float X)
+	/**
+	 * Offset the interval by adding X.
+	 *
+	 * @param X The offset.
+	 */
+	void operator+= ( float X )
 	{
 		if (!bIsEmpty)
 		{
@@ -64,8 +62,12 @@ public:
 		}
 	}
 
-	/** Offset the interval by subtracting X. */
-	void operator-= (float X)
+	/**
+	 * Offset the interval by subtracting X.
+	 *
+	 * @param X The offset.
+	 */
+	void operator-= ( float X )
 	{
 		if (!bIsEmpty)
 		{
@@ -74,15 +76,14 @@ public:
 		}
 	}
 
-
 public:
 
 	/**
 	 * Expands this interval to both sides by the specified amount.
 	 *
-	 * @param ExpandAmount - The amount to expand by.
+	 * @param ExpandAmount The amount to expand by.
 	 */
-	void Expand(float ExpandAmount)
+	void Expand( float ExpandAmount )
 	{
 		if (!bIsEmpty)
 		{
@@ -94,22 +95,21 @@ public:
 	/**
 	 * Expands this interval if necessary to include the specified element.
 	 *
-	 * @param X - The element to include.
+	 * @param X The element to include.
 	 */
-	FORCEINLINE void Include(float X);
-
+	FORCEINLINE void Include( float X );
 
 public:
 
 	/**
 	 * Calculates the intersection of two intervals.
 	 *
-	 * @param A - The first interval.
-	 * @param B - The second interval.
+	 * @param A The first interval.
+	 * @param B The second interval.
 	 *
 	 * @return The intersection.
 	 */
-	friend FInterval Intersect(const FInterval& A, const FInterval& B)
+	friend FInterval Intersect( const FInterval& A, const FInterval& B )
 	{
 		if (A.bIsEmpty || B.bIsEmpty)
 		{
@@ -124,7 +124,7 @@ public:
 /* FInterval inline functions
  *****************************************************************************/
 
-FORCEINLINE void FInterval::Include(float X)
+FORCEINLINE void FInterval::Include( float X )
 {
 	if (bIsEmpty)
 	{
@@ -144,4 +144,4 @@ FORCEINLINE void FInterval::Include(float X)
 			Max = X;
 		}
 	}
-}
\ No newline at end of file
+}
diff --git a/Engine/Source/Runtime/Core/Public/Math/NumericLimits.h b/Engine/Source/Runtime/Core/Public/Math/NumericLimits.h
--- a/Engine/Source/Runtime/Core/Public/Math/NumericLimits.h
+++ b/Engine/Source/Runtime/Core/Public/Math/NumericLimits.h
@@ -4,10 +4,12 @@
 	NumericLimits.h: Unreal numeric limits
 =============================================================================*/
 
-/*-----------------------------------------------------------------------------
-	Numeric constants.
------------------------------------------------------------------------------*/
 #pragma once
+
+
+/* Numeric constants
+ *****************************************************************************/
+
 #define MIN_uint8		((uint8)	0x00)
 #define	MIN_uint16		((uint16)	0x0000)
 #define	MIN_uint32		((uint32)	0x00000000)
@@ -31,6 +33,10 @@
 #define MIN_dbl			(2.2250738585072014e-308)	/* min positive value */
 #define MAX_dbl			(1.7976931348623158e+308)	
 
+
+/* Numeric type traits
+ *****************************************************************************/
+
 /**
  * Helper class to maps numeric type to it's limits
  */
@@ -65,14 +71,15 @@ struct TNumericLimits
 	}
 };
 
+
 /**
  * Numeric limits for const types
  */
 template <typename NumericType>
 struct TNumericLimits<const NumericType> 
 	: public TNumericLimits<NumericType>
-{
-};
+{ };
+
 
 /**
  * Numeric limits for volatile types
@@ -80,8 +87,8 @@ struct TNumericLimits<const NumericType>
 template <typename NumericType>
 struct TNumericLimits<volatile NumericType> 
 	: public TNumericLimits<NumericType>
-{
-};
+{ };
+
 
 /**
  * Numeric limits for const volatile types
@@ -89,10 +96,10 @@ struct TNumericLimits<volatile NumericType>
 template <typename NumericType>
 struct TNumericLimits<const volatile NumericType> 
 	: public TNumericLimits<NumericType>
-{
-};
+{ };
+
 
-template <>
+template<>
 struct TNumericLimits<uint8> 
 {
 	typedef uint8 NumericType;
@@ -113,7 +120,8 @@ struct TNumericLimits<uint8>
 	}
 };
 
-template <>
+
+template<>
 struct TNumericLimits<uint16> 
 {
 	typedef uint16 NumericType;
@@ -134,7 +142,8 @@ struct TNumericLimits<uint16>
 	}
 };
 
-template <>
+
+template<>
 struct TNumericLimits<uint32> 
 {
 	typedef uint32 NumericType;
@@ -155,7 +164,8 @@ struct TNumericLimits<uint32>
 	}
 };
 
-template <>
+
+template<>
 struct TNumericLimits<uint64> 
 {
 	typedef uint64 NumericType;
@@ -176,7 +186,8 @@ struct TNumericLimits<uint64>
 	}
 };
 
-template <>
+
+template<>
 struct TNumericLimits<int8> 
 {
 	typedef int8 NumericType;
@@ -197,7 +208,8 @@ struct TNumericLimits<int8>
 	}
 };
 
-template <>
+
+template<>
 struct TNumericLimits<int16> 
 {
 	typedef int16 NumericType;
@@ -218,7 +230,8 @@ struct TNumericLimits<int16>
 	}
 };
 
-template <>
+
+template<>
 struct TNumericLimits<int32> 
 {
 	typedef int32 NumericType;
@@ -239,7 +252,8 @@ struct TNumericLimits<int32>
 	}
 };
 
-template <>
+
+template<>
 struct TNumericLimits<int64> 
 {
 	typedef int64 NumericType;
@@ -260,7 +274,8 @@ struct TNumericLimits<int64>
 	}
 };
 
-template <>
+
+template<>
 struct TNumericLimits<float> 
 {
 	typedef float NumericType;
@@ -281,7 +296,8 @@ struct TNumericLimits<float>
 	}
 };
 
-template <>
+
+template<>
 struct TNumericLimits<double> 
 {
 	typedef double NumericType;
diff --git a/Engine/Source/Runtime/Core/Public/Math/OrientedBox.h b/Engine/Source/Runtime/Core/Public/Math/OrientedBox.h
--- a/Engine/Source/Runtime/Core/Public/Math/OrientedBox.h
+++ b/Engine/Source/Runtime/Core/Public/Math/OrientedBox.h
@@ -12,42 +12,27 @@
  */
 struct FOrientedBox
 {
-	/**
-	 * Holds the center of the box.
-	 */
+	/** Holds the center of the box. */
 	FVector Center;
 
-	/**
-	 * Holds the x-axis vector of the box. Must be a unit vector.
-	 */
+	/** Holds the x-axis vector of the box. Must be a unit vector. */
 	FVector AxisX;
 	
-	/**
-	 * Holds the y-axis vector of the box. Must be a unit vector.
-	 */
+	/** Holds the y-axis vector of the box. Must be a unit vector. */
 	FVector AxisY;
 	
-	/**
-	 * Holds the z-axis vector of the box. Must be a unit vector.
-	 */
+	/** Holds the z-axis vector of the box. Must be a unit vector. */
 	FVector AxisZ;
 
-	/**
-	 * Holds the extent of the box along its x-axis.
-	 */
+	/** Holds the extent of the box along its x-axis. */
 	float ExtentX;
 	
-	/**
-	 * Holds the extent of the box along its y-axis.
-	 */
+	/** Holds the extent of the box along its y-axis. */
 	float ExtentY;
 
-	/**
-	 * Holds the extent of the box along its z-axis.
-	 */
+	/** Holds the extent of the box along its z-axis. */
 	float ExtentZ;
 
-
 public:
 
 	/**
@@ -65,20 +50,19 @@ public:
 		, ExtentZ(1.0f)
 	{ }
 
-
 public:
 
 	/**
 	 * Fills in the Verts array with the eight vertices of the box.
 	 *
-	 * @param Verts		The array to fill in with the vertices.
+	 * @param Verts The array to fill in with the vertices.
 	 */
 	FORCEINLINE void CalcVertices (FVector* Verts) const;
 
 	/**
 	 * Finds the projection interval of the box when projected onto Axis.
 	 *
-	 * @param Axis	The unit vector defining the axis to project the box onto.
+	 * @param Axis The unit vector defining the axis to project the box onto.
 	 */
 	FORCEINLINE FInterval Project (const FVector& Axis) const;
 };
@@ -87,7 +71,7 @@ public:
 /* FOrientedBox inline functions
  *****************************************************************************/
 
-FORCEINLINE void FOrientedBox::CalcVertices (FVector* Verts) const
+FORCEINLINE void FOrientedBox::CalcVertices( FVector* Verts ) const
 {
 	static const float Signs[] = { -1.0f, 1.0f };
 
@@ -104,7 +88,7 @@ FORCEINLINE void FOrientedBox::CalcVertices (FVector* Verts) const
 }
 
 
-FORCEINLINE FInterval FOrientedBox::Project (const FVector& Axis) const
+FORCEINLINE FInterval FOrientedBox::Project( const FVector& Axis ) const
 {
 	static const float Signs[] = {-1.0f, 1.0f};
 
@@ -132,4 +116,4 @@ FORCEINLINE FInterval FOrientedBox::Project (const FVector& Axis) const
 	}
 
 	return ProjectionInterval;
-}
\ No newline at end of file
+}
diff --git a/Engine/Source/Runtime/Core/Public/Math/Plane.h b/Engine/Source/Runtime/Core/Public/Math/Plane.h
--- a/Engine/Source/Runtime/Core/Public/Math/Plane.h
+++ b/Engine/Source/Runtime/Core/Public/Math/Plane.h
@@ -1,10 +1,15 @@
 // Copyright 1998-2014 Epic Games, Inc. All Rights Reserved.
 
+/*=============================================================================
+	Plane.h: Declares the FPlane class.
+=============================================================================*/
+
 #pragma once
 
 
 /**
- * FPlane.
+ * Structure for three dimensional planes.
+ *
  * Stores the coeffecients as Ax+By+Cz=D.
  * Note that this is different than many other Plane classes that use Ax+By+Cz+D=0.
  */
@@ -13,13 +18,14 @@ MS_ALIGN(16) class FPlane
 {
 public:
 
-	// Variables.
 	float W;
 
 public:
 
-	/** Constructor.*/
-	FORCEINLINE FPlane();
+	/**
+	 * Default constructor (no initialization).
+	 */
+	FORCEINLINE FPlane( );
 
 	/**
 	 * Copy Constructor.
@@ -83,8 +89,8 @@ public:
 	 * Calculates distance between plane and a point.
 	 *
 	 * @param P The other point.
-	 *
-	 * @return >0: point is in front of the plane, <0: behind, =0: on the plane */
+	 * @return >0: point is in front of the plane, <0: behind, =0: on the plane.
+	 */
 	FORCEINLINE float PlaneDot( const FVector &P ) const;
 
 	/**
@@ -98,7 +104,6 @@ public:
 	 * Get the result of transforming the plane by a Matrix.
 	 *
 	 * @param M The matrix to transform plane with.
-	 *
 	 * @return The result of transform.
 	 */
 	FPlane TransformBy( const FMatrix& M ) const;
@@ -111,7 +116,6 @@ public:
 	 * @param M The Matrix to transform plane with.
 	 * @param DetM Determinant of Matrix.
 	 * @param TA Transpose-adjoint of Matrix.
-	 *
 	 * @return The result of transform.
 	 */
 	FPlane TransformByUsingAdjointT( const FMatrix& M, float DetM, const FMatrix& TA ) const;
@@ -120,7 +124,6 @@ public:
 	 * Check if two planes are identical.
 	 *
 	 * @param V The other plane.
-	 *
 	 * @return true if planes are identical, otherwise false.
 	 */
 	bool operator==( const FPlane& V ) const;
@@ -129,7 +132,6 @@ public:
 	 * Check if two planes are different.
 	 *
 	 * @param V The other plane.
-	 *
 	 * @return true if planes are different, otherwise false.
 	 */
 	bool operator!=( const FPlane& V ) const;
@@ -139,7 +141,6 @@ public:
 	 *
 	 * @param V The other plane.
 	 * @param Tolerance Error Tolerance.
-	 *
 	 * @return true if the two planes are equal within specified tolerance, otherwise false.
 	 */
 	bool Equals(const FPlane& V, float Tolerance=KINDA_SMALL_NUMBER) const;
@@ -148,7 +149,6 @@ public:
 	 * Calculates dot product of two planes.
 	 *
 	 * @param V The other plane.
-	 *
 	 * @return The dot product.
 	 */
 	FORCEINLINE float operator|( const FPlane& V ) const;
@@ -157,7 +157,6 @@ public:
 	 * Gets result of adding a plane to this.
 	 *
 	 * @param V The other plane.
-	 *
 	 * @return The result of adding a plane to this.
 	 */
 	FPlane operator+( const FPlane& V ) const;
@@ -166,7 +165,6 @@ public:
 	 * Gets result of subtracting a plane from this.
 	 *
 	 * @param V The other plane.
-	 *
 	 * @return The result of subtracting a plane from this.
 	 */
 	FPlane operator-( const FPlane& V ) const;
@@ -175,7 +173,6 @@ public:
 	 * Gets result of dividing a plane.
 	 *
 	 * @param Scale What to divide by.
-	 *
 	 * @return The result of division.
 	 */
 	FPlane operator/( float Scale ) const;
@@ -184,7 +181,6 @@ public:
 	 * Gets result of scaling a plane.
 	 *
 	 * @param Scale The scaling factor.
-	 *
 	 * @return The result of scaling.
 	 */
 	FPlane operator*( float Scale ) const;
@@ -193,7 +189,6 @@ public:
 	 * Gets result of multiplying a plane with this.
 	 *
 	 * @param V The other plane.
-	 *
 	 * @return The result of multiplying a plane with this.
 	 */
 	FPlane operator*( const FPlane& V );
@@ -202,7 +197,6 @@ public:
 	 * Add another plane to this.
 	 *
 	 * @param V The other plane.
-	 *
 	 * @return Copy of plane after addition.
 	 */
 	FPlane operator+=( const FPlane& V );
@@ -211,7 +205,6 @@ public:
 	 * Subtract another plane from this.
 	 *
 	 * @param V The other plane.
-	 *
 	 * @return Copy of plane after subtraction.
 	 */
 	FPlane operator-=( const FPlane& V );
@@ -220,7 +213,6 @@ public:
 	 * Scale this plane.
 	 *
 	 * @param Scale The scaling factor.
-	 *
 	 * @return Copy of plane after scaling.
 	 */
 	FPlane operator*=( float Scale );
@@ -229,7 +221,6 @@ public:
 	 * Multiply another plane with this.
 	 *
 	 * @param V The other plane.
-	 *
 	 * @return Copy of plane after multiplication.
 	 */
 	FPlane operator*=( const FPlane& V );
@@ -238,7 +229,6 @@ public:
 	 * Divide this plane.
 	 *
 	 * @param V What to divide by.
-	 *
 	 * @return Copy of plane after division.
 	 */
 	FPlane operator/=( float V );
@@ -248,7 +238,6 @@ public:
 	 *
 	 * @param Ar Serialization Archive.
 	 * @param P Plane to serialize.
-	 *
 	 * @return Reference to Archive after serialization.
 	 */
 	friend FArchive& operator<<( FArchive& Ar, FPlane &P )
@@ -258,8 +247,8 @@ public:
 
 	/**
 	 * Serializes the vector compressed for e.g. network transmission.
-	 * @param	Ar	Archive to serialize to/ from
-	 * @return false to allow the ordinary struct code to run (this never happens)
+	 * @param Ar Archive to serialize to/ from.
+	 * @return false to allow the ordinary struct code to run (this never happens).
 	 */
 	bool NetSerialize(FArchive& Ar, class UPackageMap*, bool& bOutSuccess)
 	{
@@ -286,6 +275,9 @@ public:
 } GCC_ALIGN(16);
 
 
+/* FPlane inline functions
+ *****************************************************************************/
+
 FORCEINLINE FPlane::FPlane()
 {}
 
diff --git a/Engine/Source/Runtime/Core/Public/Math/Quat.h b/Engine/Source/Runtime/Core/Public/Math/Quat.h
--- a/Engine/Source/Runtime/Core/Public/Math/Quat.h
+++ b/Engine/Source/Runtime/Core/Public/Math/Quat.h
@@ -1,39 +1,51 @@
 // Copyright 1998-2014 Epic Games, Inc. All Rights Reserved.
 
+/*=============================================================================
+	Quat.h: Declares the FQuat class.
+=============================================================================*/
+
 #pragma once
 
+
 /**
  * Floating point quaternion.
+ *
+ * The X, Y, Z, W components also double as the Axis/Angle format.
  */
 MS_ALIGN(16) class FQuat 
 {
+public:
+
+	/** Holds the quaternion's X-component. */
+	float X;
+
+	/** Holds the quaternion's Y-component. */
+	float Y;
+
+	/** Holds the quaternion's Z-component. */
+	float Z;
+
+	/** Holds the quaternion's W-component. */
+	float W;
+
 public:
 
 	//Identity quaternion
 	static CORE_API const FQuat Identity;
 
-	// Variables.
-	float X,Y,Z,W;
-	// X,Y,Z, W also doubles as the Axis/Angle format.
-
-#if ENABLE_NAN_DIAGNOSTIC
-	FORCEINLINE void DiagnosticCheckNaN() const
-	{
-		checkf(!ContainsNaN(), TEXT("FQuat contains NaN: %s"), *ToString());
-	}
-#else
-	FORCEINLINE void DiagnosticCheckNaN() const {}
-#endif
+public:
 
-	// Constructors.
-	FORCEINLINE FQuat();
+	/*
+	 * Default constructor (no initialization).
+	 */
+	FORCEINLINE FQuat( ) { }
 
 	/**
-	 * Constructor.
+	 * Creates and initializes a new quaternion to zero.
 	 *
 	 * @param EForceInit Force Init Enum.
 	 */
-	explicit FORCEINLINE FQuat(EForceInit);
+	explicit FORCEINLINE FQuat( EForceInit );
 
 	/**
 	 * Constructor.
@@ -46,33 +58,35 @@ public:
 	FORCEINLINE FQuat( float InX, float InY, float InZ, float InA );
 
 	/**
-	 * Constructor.
+	 * Copy constructor.
 	 *
-	 * @param Q a FQuat object to use to create new quaternion from
+	 * @param Q A FQuat object to use to create new quaternion from.
 	 */
 	FORCEINLINE FQuat( const FQuat& Q );
-
 	
 	/**
-	 * Constructor.
+	 * Creates and initializes a new quaternion from the given matrix.
 	 *
-	 * @param M a rotation matrix
+	 * @param M The rotation matrix to initialize from.
 	 */
 	explicit FQuat( const FMatrix& M );	
 
 	/**
-	 * Constructor.
+	 * Creates and initializes a new quaternion from the given rotator.
 	 *
-	 * @param R a rotator
+	 * @param R The rotator to initialize from.
 	 */
 	explicit FQuat( const FRotator& R);
 
 	/**
-	 * Get a textual representation of the vector.
+	 * Creates and initializes a new quaternion from the a rotation around the given axis.
 	 *
-	 * @return Text describing the vector.
+	 * @param Axis assumed to be a normalized vector
+	 * @param Angle angle to rotate above the given axis (in radians)
 	 */
-	FString ToString() const;
+	FQuat( FVector Axis, float AngleRad );
+
+public:
 
 #ifdef IMPLEMENT_ASSIGNMENT_OPERATOR_MANUALLY
 	/**
@@ -82,30 +96,11 @@ public:
 	 */
 	FORCEINLINE FQuat& operator=(const FQuat& Other);
 #endif
-	/**
-	 * Constructor.
-	 *
-	 * @param Axis assumed to be a normalized vector
-	 * @param Angle angle to rotate above the given axis (in radians)
-	 */
-	FQuat( FVector Axis, float AngleRad );
-
-	/**
-	 * Convert a vector of floating-point Euler angles (in degrees) into a Quaternion.
-	 * 
-	 * @param Euler the Euler angles
-	 * @return constructed FQuat
-	 */
-	static CORE_API FQuat MakeFromEuler(const FVector& Euler);
-
-	/** Convert a Quaternion into floating-point Euler angles (in degrees). */
-	CORE_API FVector Euler() const;
 
 	/**
 	 * Gets the result of adding a Quaternion to this.
 	 *
 	 * @param Q The Quaternion to add.
-	 *
 	 * @return The result of addition.
 	 */
 	FORCEINLINE FQuat operator+( const FQuat& Q ) const;
@@ -114,7 +109,6 @@ public:
 	 * Adds to this quaternion.
 	 *
 	 * @param Other The quaternion to add to this.
-	 *
 	 * @return Reference to this after addition.
 	 */
 	FORCEINLINE FQuat operator+=(const FQuat& Q);
@@ -123,7 +117,6 @@ public:
 	 * Gets the result of subtracting a Quaternion to this.
 	 *
 	 * @param Q The Quaternion to subtract.
-	 *
 	 * @return The result of subtraction.
 	 */
 	FORCEINLINE FQuat operator-( const FQuat& Q ) const;
@@ -133,7 +126,6 @@ public:
 	 *
 	 * @param Q The other Quaternion.
 	 * @param Tolerance Error Tolerance.
-	 *
 	 * @return true if two Quaternion are equal, within specified tolerance, otherwise false.
 	 */
 	FORCEINLINE bool Equals(const FQuat& Q, float Tolerance=KINDA_SMALL_NUMBER) const;
@@ -142,7 +134,6 @@ public:
 	 * Subtracts another quaternion from this.
 	 *
 	 * @param Q The other quaternion.
-	 *
 	 * @return reference to this after subtraction.
 	 */
 	FORCEINLINE FQuat operator-=(const FQuat& Q);
@@ -151,7 +142,6 @@ public:
 	 * Gets the result of multiplying a Quaternion to this.
 	 *
 	 * @param Q The Quaternion to multiply this by.
-	 *
 	 * @return The result of multiplication.
 	 */
 	FORCEINLINE FQuat operator*( const FQuat& Q ) const;
@@ -160,19 +150,23 @@ public:
 	 * Multiply this by a quaternion.
 	 *
 	 * @param Q the quaternion to multiply by this.
-	 *
 	 * @return reference to this after multiply.
 	 */
 	FORCEINLINE FQuat operator*=(const FQuat& Q);
 
 	/*
 	 * This transforms vector, not dot product
+	 *
+	 * @param V
 	 */
 	FVector operator*( const FVector& V ) const;
+
 	/* 
-	 * This matrix conversion came from 
+	 * This matrix conversion came from
 	 * http://www.m-hikari.com/ija/ija-password-2008/ija-password17-20-2008/aristidouIJA17-20-2008.pdf
-	 * used for non-uniform scaling transform
+	 * used for non-uniform scaling transform.
+	 *
+	 * @param M
 	 */	
 	FMatrix operator*( const FMatrix & M ) const;
 	
@@ -180,7 +174,6 @@ public:
 	 * Multiply this quaternion by a scaling factor.
 	 *
 	 * @param Scale The scaling factor.
-	 *
 	 * @return a reference to this after scaling.
 	 */
 	FORCEINLINE FQuat operator*=( const float Scale );
@@ -189,7 +182,6 @@ public:
 	 * Get the result of scaling this rotator.
 	 *
 	 * @param Scale The scaling factor.
-	 *
 	 * @return The result of scaling.
 	 */
 	FORCEINLINE FQuat operator*( const float Scale ) const;
@@ -198,7 +190,6 @@ public:
 	 * Divide this quaternion by scale.
 	 *
 	 * @param Scale What to divide by.
-	 *
 	 * @return new Quaternion of this after division by scale.
 	 */
 	FORCEINLINE FQuat operator/( const float Scale ) const;
@@ -207,7 +198,6 @@ public:
 	 * Checks whether two quaternions are identical.
 	 *
 	 * @param Q The other quaternion.
-	 *
 	 * @return true if two quaternion are identical, otherwise false.
 	 */
 	bool operator==( const FQuat& Q ) const;
@@ -216,7 +206,6 @@ public:
 	 * Checks whether two quaternions are not identical.
 	 *
 	 * @param Q The other quaternion.
-	 *
 	 * @return true if two quaternion are not identical, otherwise false.
 	 */
 	bool operator!=( const FQuat& Q ) const;
@@ -225,47 +214,46 @@ public:
 	 * Calculates dot product of two quaternions.
 	 *
 	 * @param Q The other quaternions.
-	 *
 	 * @return The dot product.
 	 */
 	float operator|( const FQuat& Q ) const;
 
+public:
+
+	/**
+	 * Convert a vector of floating-point Euler angles (in degrees) into a Quaternion.
+	 * 
+	 * @param Euler the Euler angles
+	 * @return constructed FQuat
+	 */
+	static CORE_API FQuat MakeFromEuler( const FVector& Euler );
+
+	/** Convert a Quaternion into floating-point Euler angles (in degrees). */
+	CORE_API FVector Euler( ) const;
+
 	/**
 	 * Normalize this quaternion if it is large enough.
 	 *
 	 * @param Tolerance Minimum squared length of vector for normalization.
 	 */
-	FORCEINLINE void Normalize(float Tolerance=SMALL_NUMBER);
+	FORCEINLINE void Normalize( float Tolerance=SMALL_NUMBER );
 
 	// Return true if this quaternion is normalized
-	bool IsNormalized() const;
+	bool IsNormalized( ) const;
 
 	/**
 	 * Get the length of this quaternion.
 	 *
 	 * @return The length of this quaternion.
 	 */
-	FORCEINLINE float Size() const;
+	FORCEINLINE float Size( ) const;
 
 	/**
 	 * Get the length squared of this quaternion.
 	 *
 	 * @return The length of this quaternion.
 	 */
-	FORCEINLINE float SizeSquared() const;
-
-	/**
-	 * Serializes the quaternion.
-	 *
-	 * @param Ar Reference to the serialization archive.
-	 * @param F Reference to the quaternion being serialized.
-	 *
-	 * @return Reference to the Archive after serialization.
-	 */
-	friend FArchive& operator<<( FArchive& Ar, FQuat& F )
-	{
-		return Ar << F.X << F.Y << F.Z << F.W;
-	}
+	FORCEINLINE float SizeSquared( ) const;
 
 	/** 
 	 * get the axis and angle of rotation of this quaternion
@@ -274,100 +262,126 @@ public:
 	 * @param Angle{out] angle of the quaternion
 	 * @warning : assumes normalized quaternions.
 	 */
-	void ToAxisAndAngle(FVector& Axis, float& Angle) const;
+	void ToAxisAndAngle( FVector& Axis, float& Angle ) const;
 
 	/**
 	 * Return vector rotated by this quaternion
 	 *
-	 * @param v the vector to be rotated
+	 * @param V the vector to be rotated
 	 * @return vector after rotation
 	 */
-	FVector RotateVector(FVector v) const;
+	FVector RotateVector( FVector V ) const;
 
 	/**
 	 * @return quaternion with W=0 and V=theta*v.
 	 */
-	CORE_API FQuat Log() const;
+	CORE_API FQuat Log( ) const;
 
 	/**
 	 * @note Exp should really only be used after Log.
 	 * Assumes a quaternion with W=0 and V=theta*v (where |v| = 1).
 	 * Exp(q) = (sin(theta)*v, cos(theta))
 	 */
-	CORE_API FQuat Exp() const;
+	CORE_API FQuat Exp( ) const;
 
-	FORCEINLINE FQuat Inverse() const;
+	FORCEINLINE FQuat Inverse( ) const;
 
 	/**
 	 * Enforce that the delta between this Quaternion and another one represents
 	 * the shortest possible rotation angle
 	 */
-	void EnforceShortestArcWith(const FQuat& OtherQuat);
+	void EnforceShortestArcWith( const FQuat& OtherQuat );
 	
 	/** Get X Rotation Axis. */
-	FORCEINLINE FVector GetAxisX() const;
+	FORCEINLINE FVector GetAxisX( ) const;
 
 	/** Get Y Rotation Axis. */
-	FORCEINLINE FVector GetAxisY() const;
+	FORCEINLINE FVector GetAxisY( ) const;
 
 	/** Get Z Rotation Axis. */
-	FORCEINLINE FVector GetAxisZ() const;
+	FORCEINLINE FVector GetAxisZ( ) const;
 
 	/** @return rotator representation of this quaternion */
-	FRotator Rotator() const;
-
+	FRotator Rotator( ) const;
 
 	/** @return Vector of the axis of the quaternion */
-	FORCEINLINE FVector GetRotationAxis() const;
+	FORCEINLINE FVector GetRotationAxis( ) const;
 
 	/**
 	 * Serializes the vector compressed for e.g. network transmission.
-	 * @param	Ar	Archive to serialize to/ from
-	 * @return false to allow the ordinary struct code to run (this never happens)
+	 * @param Ar Archive to serialize to/ from.
+	 * @return false to allow the ordinary struct code to run (this never happens).
 	 */
-	CORE_API bool NetSerialize(FArchive& Ar, class UPackageMap* Map, bool& bOutSuccess);
+	CORE_API bool NetSerialize( FArchive& Ar, class UPackageMap* Map, bool& bOutSuccess );
+
+	/**
+	 * Utility to check if there are any NaNs in this Quaternion.
+	 *
+	 * @return true if there are any NaNs in this Quaternion, otherwise false.
+	 */
+	bool ContainsNaN( ) const;
+
+	/**
+	 * Get a textual representation of the vector.
+	 *
+	 * @return Text describing the vector.
+	 */
+	FString ToString( ) const;
+
+public:
+
+#if ENABLE_NAN_DIAGNOSTIC
+	FORCEINLINE void DiagnosticCheckNaN( ) const
+	{
+		checkf(!ContainsNaN(), TEXT("FQuat contains NaN: %s"), *ToString());
+	}
+#else
+	FORCEINLINE void DiagnosticCheckNaN() const {}
+#endif
+
+public:
 
 	/**
 	 * Generates the 'smallest' (geodesic) rotation between these two vectors.
 	 */
-	static CORE_API FQuat FindBetween(const FVector& vec1, const FVector& vec2);
+	static CORE_API FQuat FindBetween( const FVector& vec1, const FVector& vec2 );
 
 	/**
 	 * Error measure (angle) between two quaternions, ranged [0..1].
 	 * Returns the hypersphere-angle between two quaternions; alignment shouldn't matter, though 
 	 * @note normalized input is expected.
 	 */
-	static FORCEINLINE float Error(const FQuat& Q1, const FQuat& Q2);
+	static FORCEINLINE float Error( const FQuat& Q1, const FQuat& Q2 );
 
 	/**
 	 * FQuat::Error with auto-normalization.
 	 */
-	static FORCEINLINE float ErrorAutoNormalize(const FQuat& A, const FQuat& B);
+	static FORCEINLINE float ErrorAutoNormalize( const FQuat& A, const FQuat& B );
 
 	/** 
 	 * Fast Linear Quaternion Interpolation.
 	 * Result is NOT normalized.
 	 */
-	static FORCEINLINE FQuat FastLerp(const FQuat& A, const FQuat& B, const float Alpha);
+	static FORCEINLINE FQuat FastLerp( const FQuat& A, const FQuat& B, const float Alpha );
 
 	/** 
 	 * Bi-Linear Quaternion interpolation.
 	 * Result is NOT normalized.
 	 */
-	static FORCEINLINE FQuat FastBilerp(const FQuat& P00, const FQuat& P10, const FQuat& P01, const FQuat& P11, float FracX, float FracY);
+	static FORCEINLINE FQuat FastBilerp( const FQuat& P00, const FQuat& P10, const FQuat& P01, const FQuat& P11, float FracX, float FracY );
 
 
 	/** Spherical interpolation. Will correct alignment. Output is not normalized. */
-	static CORE_API FQuat Slerp(const FQuat &Quat1,const FQuat &Quat2, float Slerp);
+	static CORE_API FQuat Slerp( const FQuat &Quat1,const FQuat &Quat2, float Slerp );
 
 	/**
 	 * Simpler Slerp that doesn't do any checks for 'shortest distance' etc.
 	 * We need this for the cubic interpolation stuff so that the multiple Slerps dont go in different directions.
 	 */
-	static CORE_API FQuat SlerpFullPath(const FQuat &quat1, const FQuat &quat2, float Alpha);
+	static CORE_API FQuat SlerpFullPath( const FQuat &quat1, const FQuat &quat2, float Alpha );
 	
 	/** Given start and end quaternions of quat1 and quat2, and tangents at those points tang1 and tang2, calculate the point at Alpha (between 0 and 1) between them. */
-	static CORE_API FQuat Squad(const FQuat& quat1, const FQuat& tang1, const FQuat& quat2, const FQuat& tang2, float Alpha);
+	static CORE_API FQuat Squad( const FQuat& quat1, const FQuat& tang1, const FQuat& quat2, const FQuat& tang2, float Alpha );
 
 	/** 
 	 * Calculate tangents between given points
@@ -380,16 +394,26 @@ public:
 	 */
 	static CORE_API void CalcTangents( const FQuat& PrevP, const FQuat& P, const FQuat& NextP, float Tension, FQuat& OutTan );
 
+public:
+
 	/**
-	 * Utility to check if there are any NaNs in this Quaternion.
+	 * Serializes the quaternion.
 	 *
-	 * @return true if there are any NaNs in this Quaternion, otherwise false.
+	 * @param Ar Reference to the serialization archive.
+	 * @param F Reference to the quaternion being serialized.
+	 * @return Reference to the Archive after serialization.
 	 */
-	bool ContainsNaN() const;
+	friend FArchive& operator<<( FArchive& Ar, FQuat& F )
+	{
+		return Ar << F.X << F.Y << F.Z << F.W;
+	}
 
 } GCC_ALIGN(16);
 
 
+/* FQuat inline functions
+ *****************************************************************************/
+
 inline FQuat::FQuat( const FMatrix& M )
 {
 	// If Matrix is NULL, return Identity quaternion. If any of them is 0, you won't be able to construct rotation
@@ -587,39 +611,40 @@ inline FRotator FQuat::Rotator() const
 }
 
 
-FORCEINLINE FQuat::FQuat()
-{}
-
+/* FQuat inline functions
+ *****************************************************************************/
 
-FORCEINLINE FQuat::FQuat(EForceInit ZeroOrNot)
+FORCEINLINE FQuat::FQuat( EForceInit ZeroOrNot )
 	:	X(0), Y(0), Z(0), W(ZeroOrNot == ForceInitToZero ? 0.0f : 1.0f)
-{}
+{ }
 
 
 FORCEINLINE FQuat::FQuat( float InX, float InY, float InZ, float InA )
-:	X(InX), Y(InY), Z(InZ), W(InA)
+	: X(InX)
+	, Y(InY)
+	, Z(InZ)
+	, W(InA)
 {
 	DiagnosticCheckNaN();
 }
 
 
-FORCEINLINE FQuat::FQuat( const FQuat& Q ) :
-	X(Q.X),
-	Y(Q.Y),
-	Z(Q.Z),
-	W(Q.W)
-{
-}
+FORCEINLINE FQuat::FQuat( const FQuat& Q )
+	: X(Q.X)
+	, Y(Q.Y)
+	, Z(Q.Z)
+	, W(Q.W)
+{ }
 
 
-FORCEINLINE FString FQuat::ToString() const
+FORCEINLINE FString FQuat::ToString( ) const
 {
 	return FString::Printf(TEXT("X=%3.3f Y=%3.3f Z=%3.3f W=%3.3f"), X, Y, Z, W);
 }
 
 
 #ifdef IMPLEMENT_ASSIGNMENT_OPERATOR_MANUALLY
-FORCEINLINE FQuat& FQuat::operator=(const FQuat& Other)
+FORCEINLINE FQuat& FQuat::operator=( const FQuat& Other )
 {
 	this->X = Other.X;
 	this->Y = Other.Y;
@@ -673,8 +698,8 @@ FORCEINLINE FQuat FQuat::operator-( const FQuat& Q ) const
 
 FORCEINLINE bool FQuat::Equals(const FQuat& Q, float Tolerance) const
 {
-	return (FMath::Abs(X-Q.X) < Tolerance && FMath::Abs(Y-Q.Y) < Tolerance && FMath::Abs(Z-Q.Z) < Tolerance && FMath::Abs(W-Q.W) < Tolerance)
-		|| (FMath::Abs(X+Q.X) < Tolerance && FMath::Abs(Y+Q.Y) < Tolerance && FMath::Abs(Z+Q.Z) < Tolerance && FMath::Abs(W+Q.W) < Tolerance);
+	return (FMath::Abs(X - Q.X) < Tolerance && FMath::Abs(Y - Q.Y) < Tolerance && FMath::Abs(Z - Q.Z) < Tolerance && FMath::Abs(W - Q.W) < Tolerance)
+		|| (FMath::Abs(X + Q.X) < Tolerance && FMath::Abs(Y + Q.Y) < Tolerance && FMath::Abs(Z + Q.Z) < Tolerance && FMath::Abs(W + Q.W) < Tolerance);
 }
 
 
@@ -702,7 +727,7 @@ FORCEINLINE FQuat FQuat::operator*( const FQuat& Q ) const
 }
 
 
-FORCEINLINE FQuat FQuat::operator*=(const FQuat& Q)
+FORCEINLINE FQuat FQuat::operator*=( const FQuat& Q )
 {
 	/**
 	 * Now this uses VectorQuaternionMultiply that is optimized per platform.
@@ -734,40 +759,42 @@ FORCEINLINE FQuat FQuat::operator*=( const float Scale )
 
 FORCEINLINE FQuat FQuat::operator*( const float Scale ) const
 {
-	return FQuat( Scale*X, Scale*Y, Scale*Z, Scale*W);			
+	return FQuat(Scale * X, Scale * Y, Scale * Z, Scale * W);
 }
 
 
 FORCEINLINE FQuat FQuat::operator/( const float Scale ) const
 {
-	return FQuat( X/Scale, Y/Scale, Z/Scale, W/Scale);			
+	return FQuat(X / Scale, Y / Scale, Z / Scale, W / Scale);
 }
 
 
 FORCEINLINE bool FQuat::operator==( const FQuat& Q ) const
 {
-	return X==Q.X && Y==Q.Y && Z==Q.Z &&  W==Q.W;
+	return X == Q.X && Y == Q.Y && Z == Q.Z && W == Q.W;
 }
 
 
 FORCEINLINE bool FQuat::operator!=( const FQuat& Q ) const
 {
-	return X!=Q.X || Y!=Q.Y || Z!=Q.Z || W!=Q.W;
+	return X != Q.X || Y != Q.Y || Z != Q.Z || W != Q.W;
 }
 
 
 FORCEINLINE float FQuat::operator|( const FQuat& Q ) const
 {
-	return X*Q.X + Y*Q.Y + Z*Q.Z + W*Q.W;
+	return X * Q.X + Y * Q.Y + Z * Q.Z + W * Q.W;
 }
 
 
 FORCEINLINE void FQuat::Normalize(float Tolerance)
 {
-	const float SquareSum = X*X + Y*Y + Z*Z + W*W;
-	if( SquareSum > Tolerance )
+	const float SquareSum = X * X + Y * Y + Z * Z + W * W;
+
+	if (SquareSum > Tolerance)
 	{
 		const float Scale = FMath::InvSqrt(SquareSum);
+
 		X *= Scale; 
 		Y *= Scale; 
 		Z *= Scale;
@@ -780,58 +807,59 @@ FORCEINLINE void FQuat::Normalize(float Tolerance)
 }
 
 
-FORCEINLINE bool FQuat::IsNormalized() const
+FORCEINLINE bool FQuat::IsNormalized( ) const
 {
 	return (FMath::Abs(1.f - SizeSquared()) <= KINDA_SMALL_NUMBER);
 }
 
 
-FORCEINLINE float FQuat::Size() const
+FORCEINLINE float FQuat::Size( ) const
 {
-	return FMath::Sqrt(X*X+Y*Y+Z*Z+W*W);
+	return FMath::Sqrt(X * X + Y * Y + Z * Z + W * W);
 }
 
 
-FORCEINLINE float FQuat::SizeSquared() const
+FORCEINLINE float FQuat::SizeSquared( ) const
 {
-	return (X*X+Y*Y+Z*Z+W*W);
+	return (X * X + Y * Y + Z * Z + W * W);
 }
 
 
-FORCEINLINE void FQuat::ToAxisAndAngle(FVector& Axis, float& Angle) const
+FORCEINLINE void FQuat::ToAxisAndAngle( FVector& Axis, float& Angle ) const
 {
 	Angle = 2.f * FMath::Acos(W);
 	Axis = GetRotationAxis();
 }
 
 
-FORCEINLINE FVector FQuat::GetRotationAxis() const
+FORCEINLINE FVector FQuat::GetRotationAxis( ) const
 {
 	// Ensure we never try to sqrt a neg number
-	const float S = FMath::Sqrt( FMath::Max(1.f-(W*W), 0.f) );
-	if( S >= 0.0001f ) 
+	const float S = FMath::Sqrt(FMath::Max(1.f - (W * W), 0.f));
+
+	if (S >= 0.0001f) 
 	{ 
-		return FVector(X/S, Y/S, Z/S);
+		return FVector(X / S, Y / S, Z / S);
 	} 
 
 	return FVector(1.f, 0.f, 0.f);
 }
 
 
-FORCEINLINE FVector FQuat::RotateVector(FVector v) const
+FORCEINLINE FVector FQuat::RotateVector( FVector V ) const
 {	
 	// (q.W*q.W-qv.qv)v + 2(qv.v)qv + 2 q.W (qv x v)
 
 	const FVector qv(X, Y, Z);
-	FVector vOut = 2.f * W * (qv ^ v);
-	vOut += ((W * W) - (qv | qv)) * v;
-	vOut += (2.f * (qv | v)) * qv;
+	FVector vOut = 2.f * W * (qv ^ V);
+	vOut += ((W * W) - (qv | qv)) * V;
+	vOut += (2.f * (qv | V)) * qv;
 
 	return vOut;
 }
 
 
-FORCEINLINE FQuat FQuat::Inverse() const
+FORCEINLINE FQuat FQuat::Inverse( ) const
 {
 	checkSlow(IsNormalized());
 
@@ -839,10 +867,11 @@ FORCEINLINE FQuat FQuat::Inverse() const
 }
 
 
-FORCEINLINE void FQuat::EnforceShortestArcWith(const FQuat& OtherQuat)
+FORCEINLINE void FQuat::EnforceShortestArcWith( const FQuat& OtherQuat )
 {
 	const float DotResult = (OtherQuat | *this);
 	const float Bias = FMath::FloatSelect(DotResult, 1.0f, -1.0f);
+
 	X *= Bias;
 	Y *= Bias;
 	Z *= Bias;
@@ -850,42 +879,44 @@ FORCEINLINE void FQuat::EnforceShortestArcWith(const FQuat& OtherQuat)
 }
 
 	
-FORCEINLINE FVector FQuat::GetAxisX() const
+FORCEINLINE FVector FQuat::GetAxisX( ) const
 {
-	return RotateVector( FVector(1.f,0.f,0.f) );
+	return RotateVector(FVector(1.f, 0.f, 0.f));
 }
 
 
-FORCEINLINE FVector FQuat::GetAxisY() const
+FORCEINLINE FVector FQuat::GetAxisY( ) const
 {
-	return RotateVector( FVector(0.f,1.f,0.f) );
+	return RotateVector(FVector(0.f, 1.f, 0.f));
 }
 
 
-FORCEINLINE FVector FQuat::GetAxisZ() const
+FORCEINLINE FVector FQuat::GetAxisZ( ) const
 {
-	return RotateVector( FVector(0.f,0.f,1.f) );
+	return RotateVector(FVector(0.f, 0.f, 1.f));
 }
 
 
-FORCEINLINE float FQuat::Error(const FQuat& Q1, const FQuat& Q2)
+FORCEINLINE float FQuat::Error( const FQuat& Q1, const FQuat& Q2 )
 {
-	const float cosom = FMath::Abs(Q1.X*Q2.X + Q1.Y*Q2.Y + Q1.Z*Q2.Z + Q1.W*Q2.W);
-	return (FMath::Abs(cosom) < 0.9999999f) ? FMath::Acos(cosom)*(1.f/PI) : 0.0f;
+	const float cosom = FMath::Abs(Q1.X * Q2.X + Q1.Y * Q2.Y + Q1.Z * Q2.Z + Q1.W * Q2.W);
+	return (FMath::Abs(cosom) < 0.9999999f) ? FMath::Acos(cosom)*(1.f / PI) : 0.0f;
 }
 
 
-FORCEINLINE float FQuat::ErrorAutoNormalize(const FQuat& A, const FQuat& B)
+FORCEINLINE float FQuat::ErrorAutoNormalize( const FQuat& A, const FQuat& B )
 {
 	FQuat Q1 = A;
 	Q1.Normalize();
+
 	FQuat Q2 = B;
 	Q2.Normalize();
+
 	return FQuat::Error(Q1, Q2);
 }
 
 
-FORCEINLINE FQuat FQuat::FastLerp(const FQuat& A, const FQuat& B, const float Alpha)
+FORCEINLINE FQuat FQuat::FastLerp( const FQuat& A, const FQuat& B, const float Alpha )
 {
 	// To ensure the 'shortest route', we make sure the dot product between the both rotations is positive.
 	const float DotResult = (A | B);
@@ -894,23 +925,24 @@ FORCEINLINE FQuat FQuat::FastLerp(const FQuat& A, const FQuat& B, const float Al
 }
 
 
-FORCEINLINE FQuat FQuat::FastBilerp(const FQuat& P00, const FQuat& P10, const FQuat& P01, const FQuat& P11, float FracX, float FracY)
+FORCEINLINE FQuat FQuat::FastBilerp( const FQuat& P00, const FQuat& P10, const FQuat& P01, const FQuat& P11, float FracX, float FracY )
 {
 	return FQuat::FastLerp(
-					FQuat::FastLerp(P00,P10,FracX),
-					FQuat::FastLerp(P01,P11,FracX),
-					FracY
-					);
+		FQuat::FastLerp(P00,P10,FracX),
+		FQuat::FastLerp(P01,P11,FracX),
+		FracY
+	);
 }
 
 
-FORCEINLINE bool FQuat::ContainsNaN() const
+FORCEINLINE bool FQuat::ContainsNaN( ) const
 {
-	return    (FMath::IsNaN(X) || !FMath::IsFinite(X) 
-			|| FMath::IsNaN(Y) || !FMath::IsFinite(Y) 
-			|| FMath::IsNaN(Z) || !FMath::IsFinite(Z) 
-			|| FMath::IsNaN(W) || !FMath::IsFinite(W));
+	return (FMath::IsNaN(X) || !FMath::IsFinite(X) ||
+			FMath::IsNaN(Y) || !FMath::IsFinite(Y) ||
+			FMath::IsNaN(Z) || !FMath::IsFinite(Z) ||
+			FMath::IsNaN(W) || !FMath::IsFinite(W)
+	);
 }
 
 
-template <> struct TIsPODType<FQuat> { enum { Value = true }; };
+template<> struct TIsPODType<FQuat> { enum { Value = true }; };
diff --git a/Engine/Source/Runtime/Core/Public/Math/RandomStream.h b/Engine/Source/Runtime/Core/Public/Math/RandomStream.h
--- a/Engine/Source/Runtime/Core/Public/Math/RandomStream.h
+++ b/Engine/Source/Runtime/Core/Public/Math/RandomStream.h
@@ -6,6 +6,7 @@
 
 #pragma once
 
+
 /**
  * Implements a thread-safe SRand based RNG.
  *
@@ -27,7 +28,7 @@ public:
 	 *
 	 * The seed should be set prior to use.
 	 */
-	FRandomStream()
+	FRandomStream( )
 		: InitialSeed(0)
 		, Seed(0)
 	{ }
@@ -35,22 +36,21 @@ public:
 	/**
 	 * Creates and initializes a new random stream from the specified seed value.
 	 *
-	 * @param InSeed - The seed value.
+	 * @param InSeed The seed value.
 	 */
-	FRandomStream(int32 InSeed)
+	FRandomStream( int32 InSeed )
 		: InitialSeed(InSeed)
 		, Seed(InSeed)
 	{ }
 
-
 public:
 
 	/**
 	 * Initializes this random stream with the specified seed value.
 	 *
-	 * @param InSeed - The seed value.
+	 * @param InSeed The seed value.
 	 */
-	void Initialize(int32 InSeed)
+	void Initialize( int32 InSeed )
 	{
 		InitialSeed = InSeed;
 		Seed = InSeed;
@@ -59,7 +59,7 @@ public:
 	/**
 	 * Resets this random stream to the initial seed value.
 	 */
-	void Reset() const
+	void Reset( ) const
 	{
 		Seed = InitialSeed;
 	}
@@ -67,7 +67,7 @@ public:
 	/**
 	 * Generates a new random seed.
 	 */
-	void GenerateNewSeed()
+	void GenerateNewSeed( )
 	{
 		Initialize(FMath::Rand());
 	}
@@ -77,7 +77,7 @@ public:
 	 *
 	 * @return Random number.
 	 */
-	float GetFraction() const
+	float GetFraction( ) const
 	{
 		MutateSeed();
 
@@ -94,7 +94,7 @@ public:
 	 *
 	 * @return Random number.
 	 */
-	uint32 GetUnsignedInt() const
+	uint32 GetUnsignedInt( ) const
 	{
 		MutateSeed();
 
@@ -106,7 +106,7 @@ public:
 	 *
 	 * @return Random unit vector.
 	 */
-	FVector GetUnitVector() const
+	FVector GetUnitVector( ) const
 	{
 		FVector Result;
 
@@ -116,7 +116,8 @@ public:
 			Result.X = GetFraction() * 2 - 1;
 			Result.Y = GetFraction() * 2 - 1;
 			Result.Z = GetFraction() * 2 - 1;
-		} while(Result.SizeSquared() > 1.f);
+		}
+		while(Result.SizeSquared() > 1.f);
 
 		return Result.UnsafeNormal();
 	}
@@ -126,7 +127,7 @@ public:
 	 *
 	 * @return Current seed.
 	 */
-	int32 GetCurrentSeed() const
+	int32 GetCurrentSeed( ) const
 	{
 		return Seed;
 	}
@@ -136,7 +137,7 @@ public:
 	 *
 	 * @return Random number.
 	 */
-	FORCEINLINE float FRand() const
+	FORCEINLINE float FRand( ) const
 	{
 		return GetFraction();
 	}
@@ -146,7 +147,7 @@ public:
 	 *
 	 * @return A random number in [0..A)
 	 */
-	FORCEINLINE int32 RandHelper(int32 A) const
+	FORCEINLINE int32 RandHelper( int32 A ) const
 	{
 		// Can't just multiply GetFraction by A, as GetFraction could be == 1.0f
 		return ((A > 0) ? FMath::TruncToInt(GetFraction() * ((float)A - DELTA)) : 0);
@@ -157,7 +158,7 @@ public:
 	 *
 	 * @return A random number >= Min and <= Max
 	 */
-	FORCEINLINE int32 RandRange(int32 Min, int32 Max) const
+	FORCEINLINE int32 RandRange( int32 Min, int32 Max ) const
 	{
 		const int32 Range = (Max - Min) + 1;
 
@@ -169,7 +170,7 @@ public:
 	 *
 	 * @return A random number >= Min and <= Max
 	 */
-	FORCEINLINE float FRandRange(float InMin, float InMax) const
+	FORCEINLINE float FRandRange( float InMin, float InMax ) const
 	{
 		return InMin + (InMax - InMin) * FRand();
 	}
@@ -179,7 +180,7 @@ public:
 	 *
 	 * @return Random unit vector.
 	 */
-	FORCEINLINE FVector VRand() const
+	FORCEINLINE FVector VRand( ) const
 	{
 		return GetUnitVector();
 	}
@@ -191,7 +192,7 @@ public:
 	 * @param ConeHalfAngleRad Half-angle of cone, in radians.
 	 * @return Normalized vector within the specified cone.
 	 */
-	FORCEINLINE FVector VRandCone(FVector const& Dir, float ConeHalfAngleRad)
+	FORCEINLINE FVector VRandCone( FVector const& Dir, float ConeHalfAngleRad )
 	{
 		if (ConeHalfAngleRad > 0.f)
 		{
@@ -236,7 +237,7 @@ public:
 	 * @param VerticalConeHalfAngleRad Vertical half-angle of cone, in radians.
 	 * @return Normalized vector within the specified cone.
 	 */
-	FORCEINLINE FVector VRandCone(FVector const& Dir, float HorizontalConeHalfAngleRad, float VerticalConeHalfAngleRad)
+	FORCEINLINE FVector VRandCone( FVector const& Dir, float HorizontalConeHalfAngleRad, float VerticalConeHalfAngleRad )
 	{
 		if ( (VerticalConeHalfAngleRad > 0.f) && (HorizontalConeHalfAngleRad > 0.f) )
 		{
@@ -283,12 +284,11 @@ protected:
 	/**
 	 * Mutates the current seed into the next seed.
 	 */
-	void MutateSeed() const
+	void MutateSeed( ) const
 	{
 		Seed = (Seed * 196314165) + 907633515; 
 	}
 
-
 private:
 
 	// Holds the initial seed.
diff --git a/Engine/Source/Runtime/Core/Public/Math/Range.h b/Engine/Source/Runtime/Core/Public/Math/Range.h
--- a/Engine/Source/Runtime/Core/Public/Math/Range.h
+++ b/Engine/Source/Runtime/Core/Public/Math/Range.h
@@ -60,7 +60,7 @@ public:
 	 *
 	 * The created range is of the form [A, A].
 	 *
-	 * @param A - The element in the range.
+	 * @param A The element in the range.
 	 */
 	explicit TRange( const ElementType& A )
 		: LowerBound(BoundsType::Inclusive(A))
@@ -72,8 +72,8 @@ public:
 	 *
 	 * The created range is of the form [A, B).
 	 *
-	 * @param A - The range's lower bound value (inclusive).
-	 * @param B - The range's upper bound value (exclusive).
+	 * @param A The range's lower bound value (inclusive).
+	 * @param B The range's upper bound value (exclusive).
 	 */
 	explicit TRange( const ElementType& A, const ElementType& B )
 		: LowerBound(BoundsType::Inclusive(A))
@@ -83,8 +83,8 @@ public:
 	/**
 	 * Creates and initializes a new range with the given lower and upper bounds.
 	 *
-	 * @param InLowerBound - The range's lower bound.
-	 * @param InUpperBound - The range's upper bound.
+	 * @param InLowerBound The range's lower bound.
+	 * @param InUpperBound The range's upper bound.
 	 */
 	TRange( const BoundsType& InLowerBound, const BoundsType& InUpperBound )
 		: LowerBound(InLowerBound)
@@ -96,8 +96,7 @@ public:
 	/**
 	 * Compares this range with the specified range for equality.
 	 *
-	 * @param Other - The range to compare with.
-	 *
+	 * @param Other The range to compare with.
 	 * @return true if the ranges are equal, false otherwise.
 	 */
 	bool operator==( const TRange& Other ) const
@@ -113,8 +112,7 @@ public:
 	/**
 	 * Compares this range with the specified range for inequality.
 	 *
-	 * @param Other - The range to compare with.
-	 *
+	 * @param Other The range to compare with.
 	 * @return true if the ranges are not equal, false otherwise.
 	 */
 	bool operator!=( const TRange& Other ) const
@@ -136,8 +134,7 @@ public:
 	 *		[A, B) and [C, C) or
 	 *		[A, B] and (B, C)
 	 *
-	 * @param Other - The other range.
-	 *
+	 * @param Other The other range.
 	 * @return true if this range adjoins the other, false otherwise.
 	 */
 	bool Adjoins( const TRange& Other ) const
@@ -167,9 +164,8 @@ public:
 	 * A range conjoins two non-overlapping ranges if it adjoins both of them, i.e.
 	 *		[B, C) conjoins the two ranges [A, B) and [C, D).
 	 *
-	 * @param X - The first range.
-	 * @param Y - The second range.
-	 *
+	 * @param X The first range.
+	 * @param Y The second range.
 	 * @return true if this range conjoins the two ranges, false otherwise.
 	 */
 	bool Conjoins( const TRange& X, const TRange& Y ) const
@@ -185,8 +181,7 @@ public:
 	/**
 	 * Checks whether this range contains the specified element.
 	 *
-	 * @param Element - The element to check.
-	 *
+	 * @param Element The element to check.
 	 * @return true if the range contains the element, false otherwise.
 	 */
 	bool Contains( const ElementType& Element ) const
@@ -198,8 +193,7 @@ public:
 	/**
 	 * Checks whether this range contains another range.
 	 *
-	 * @param Other - The range to check.
-	 *
+	 * @param Other The range to check.
 	 * @return true if the range contains the other range, false otherwise.
 	 */
 	bool Contains( const TRange& Other ) const
@@ -213,8 +207,7 @@ public:
 	 *
 	 * Two ranges are contiguous if they are adjoint or overlapping.
 	 *
-	 * @param Other - The other range.
-	 *
+	 * @param Other The other range.
 	 * @return true if the ranges are contiguous, false otherwise.
 	 */
 	bool Contiguous( const TRange& Other ) const
@@ -238,9 +231,7 @@ public:
 	 * Use HasLowerBound() to ensure that this range actually has a lower bound.
 	 *
 	 * @return Bound value.
-	 *
-	 * @see GetUpperBoundValue
-	 * @see HasLowerBound
+	 * @see GetUpperBoundValue, HasLowerBound
 	 */
 	const ElementType& GetLowerBoundValue( ) const
 	{
@@ -263,9 +254,7 @@ public:
 	 * Use HasUpperBound() to ensure that this range actually has an upper bound.
 	 *
 	 * @return Bound value.
-	 *
-	 * @see GetLowerBoundValue
-	 * @see HasUpperBound
+	 * @see GetLowerBoundValue, HasUpperBound
 	 */
 	const ElementType& GetUpperBoundValue( ) const
 	{
@@ -333,8 +322,7 @@ public:
 	/**
 	 * Checks whether this range overlaps with another.
 	 *
-	 * @param Other - The other range.
-	 *
+	 * @param Other The other range.
 	 * @return true if the ranges overlap, false otherwise.
 	 */
 	bool Overlaps( const TRange& Other ) const
@@ -394,7 +382,7 @@ public:
 	 * If a range [A, C) does not contain the element B, the original range is returned.
 	 * Otherwise the range is split into two ranges [A, B) and [B, C), each of which may be empty.
 	 *
-	 * @param Element - The element at which to split the range.
+	 * @param Element The element at which to split the range.
 	 */
 	TArray<TRange> Split( const ElementType& Element ) const
 	{
@@ -418,14 +406,10 @@ public:
 	/**
 	 * Calculates the difference between two ranges, i.e. X - Y.
 	 *
-	 * @param X - The first range to subtract from.
-	 * @param Y - The second range to subtract with.
-	 *
-	 * @return Between 0 and 2 remaining ranges
-	 *
-	 * @see Hull
-	 * @see Intersection
-	 * @see Union
+	 * @param X The first range to subtract from.
+	 * @param Y The second range to subtract with.
+	 * @return Between 0 and 2 remaining ranges.
+	 * @see Hull, Intersection, Union
 	 */
 	static FORCEINLINE TArray<TRange> Difference( const TRange& X, const TRange& Y )
 	{
@@ -469,14 +453,10 @@ public:
 	 *
 	 * The hull is the smallest range that contains both ranges.
 	 *
-	 * @param X - The first range.
-	 * @param Y - The second range.
-	 *
+	 * @param X The first range.
+	 * @param Y The second range.
 	 * @return The hull.
-	 *
-	 * @see Difference
-	 * @see Intersection
-	 * @see Union
+	 * @see Difference, Intersection, Union
 	 */
 	static FORCEINLINE TRange Hull( const TRange& X, const TRange& Y )
 	{
@@ -496,13 +476,9 @@ public:
 	/**
 	 * Computes the hull of many ranges.
 	 *
-	 * @param Ranges - The ranges to hull.
-	 *
+	 * @param Ranges The ranges to hull.
 	 * @return The hull.
-	 *
-	 * @see Difference
-	 * @see Intersection
-	 * @see Union
+	 * @see Difference, Intersection, Union
 	 */
 	static FORCEINLINE TRange Hull( const TArray<TRange>& Ranges )
 	{
@@ -526,14 +502,10 @@ public:
 	 *
 	 * The intersection of two ranges is the largest range that is contained by both ranges.
 	 *
-	 * @param X - The first range.
-	 * @param Y - The second range.
-	 *
+	 * @param X The first range.
+	 * @param Y The second range.
 	 * @return The intersection, or an empty range if the ranges do not overlap.
-	 *
-	 * @see Difference
-	 * @see Hull
-	 * @see Union
+	 * @see Difference, Hull, Union
 	 */
 	static FORCEINLINE TRange Intersection( const TRange& X, const TRange& Y )
 	{
@@ -553,13 +525,9 @@ public:
 	/**
 	 * Computes the intersection of many ranges.
 	 *
-	 * @param Ranges - The ranges to intersect.
-	 *
+	 * @param Ranges The ranges to intersect.
 	 * @return The intersection.
-	 *
-	 * @see Difference
-	 * @see Hull
-	 * @see Union
+	 * @see Difference, Hull, Union
 	 */
 	static FORCEINLINE TRange Intersection( const TArray<TRange>& Ranges )
 	{
@@ -583,14 +551,10 @@ public:
 	 *
 	 * A union is a range or series of ranges that contains both ranges.
 	 *
-	 * @param X - The first range.
-	 * @param Y - The second range.
-	 *
+	 * @param X The first range.
+	 * @param Y The second range.
 	 * @return The union, or both ranges if the two ranges are not contiguous, or no ranges if both ranges are empty.
-	 *
-	 * @see Difference
-	 * @see Hull
-	 * @see Intersection
+	 * @see Difference, Hull, Intersection
 	 */
 	static FORCEINLINE TArray<TRange> Union( const TRange& X, const TRange& Y )
 	{
@@ -631,8 +595,7 @@ public:
 	/**
 	 * Creates a left-bounded range that contains all elements greater than or equal to the specified value.
 	 *
-	 * @param Value - The value.
-	 *
+	 * @param Value The value.
 	 * @return A new range.
 	 */
 	static FORCEINLINE TRange AtLeast( const ElementType& Value )
@@ -643,8 +606,7 @@ public:
 	/**
 	 * Creates a right-bounded range that contains all elements less than or equal to the specified value.
 	 *
-	 * @param Value - The value.
-	 *
+	 * @param Value The value.
 	 * @return A new range.
 	 */
 	static FORCEINLINE TRange AtMost( const ElementType& Value )
@@ -655,8 +617,7 @@ public:
 	/**
 	 * Creates a left-bounded range that contains all elements greater than the specified value.
 	 *
-	 * @param Value - The value.
-	 *
+	 * @param Value The value.
 	 * @return A new range.
 	 */
 	static FORCEINLINE TRange GreaterThan( const ElementType& Value )
@@ -667,8 +628,7 @@ public:
 	/**
 	 * Creates a right-bounded range that contains all elements less than the specified value.
 	 *
-	 * @param Value - The value.
-	 *
+	 * @param Value The value.
 	 * @return A new range.
 	 */
 	static FORCEINLINE TRange LessThan( const ElementType& Value )
@@ -681,9 +641,8 @@ public:
 	/**
 	 * Serializes the given range from or into the specified archive.
 	 *
-	 * @param Ar - The archive to serialize from or into.
-	 * @param Range - The range to serialize.
-	 *
+	 * @param Ar The archive to serialize from or into.
+	 * @param Range The range to serialize.
 	 * @return The archive.
 	 */
 	friend class FArchive& operator<<( class FArchive& Ar, TRange& Range )
@@ -694,8 +653,7 @@ public:
 	/**
 	 * Gets the hash for the specified range.
 	 *
-	 * @param Range - The range to get the hash for.
-	 *
+	 * @param Range The range to get the hash for.
 	 * @return Hash value.
 	 */
 	friend uint32 GetTypeHash( const TRange& Range )
diff --git a/Engine/Source/Runtime/Core/Public/Math/RangeBound.h b/Engine/Source/Runtime/Core/Public/Math/RangeBound.h
--- a/Engine/Source/Runtime/Core/Public/Math/RangeBound.h
+++ b/Engine/Source/Runtime/Core/Public/Math/RangeBound.h
@@ -14,19 +14,13 @@ namespace ERangeBoundTypes
 	 */
 	enum Type
 	{
-		/**
-		 * The range excludes the bound.
-		 */
+		/** The range excludes the bound. */
 		Exclusive,
 
-		/**
-		 * The range includes the bound.
-		 */
+		/** The range includes the bound. */
 		Inclusive,
 
-		/**
-		 * The bound is open.
-		 */
+		/** The bound is open. */
 		Open
 	};
 }
@@ -42,9 +36,7 @@ public:
 	/**
 	 * Default constructor.
 	 *
-	 * @see Exclusive
-	 * @see Inclusive
-	 * @see Open
+	 * @see Exclusive, Inclusive, Open
 	 */
 	TRangeBound( )
 		: Type(ERangeBoundTypes::Open)
@@ -54,11 +46,8 @@ public:
 	/**
 	 * Creates a closed bound that includes the specified value.
 	 *
-	 * @param InValue - The bound's value.
-	 *
-	 * @see Exclusive
-	 * @see Inclusive
-	 * @see Open
+	 * @param InValue The bound's value.
+	 * @see Exclusive, Inclusive, Open
 	 */
 	TRangeBound( const ElementType& InValue )
 		: Type(ERangeBoundTypes::Inclusive)
@@ -70,8 +59,7 @@ public:
 	/**
 	 * Compares this bound with the specified bound for equality.
 	 *
-	 * @param Other - The bound to compare with.
-	 *
+	 * @param Other The bound to compare with.
 	 * @return true if the bounds are equal, false otherwise.
 	 */
 	bool operator==( const TRangeBound& Other ) const
@@ -82,8 +70,7 @@ public:
 	/**
 	 * Compares this range with the specified bound for inequality.
 	 *
-	 * @param Other - The bound to compare with.
-	 *
+	 * @param Other The bound to compare with.
 	 * @return true if the bounds are not equal, false otherwise.
 	 */
 	bool operator!=( const TRangeBound& Other ) const
@@ -99,7 +86,6 @@ public:
 	 * Use IsClosed() to verify that this bound is closed before calling this method.
 	 *
 	 * @return Bound value.
-	 *
 	 * @see IsOpen
 	 */
 	const ElementType& GetValue( ) const
@@ -154,9 +140,8 @@ public:
 	/**
 	 * Serializes the given bound from or into the specified archive.
 	 *
-	 * @param Ar - The archive to serialize from or into.
-	 * @param Bound - The bound to serialize.
-	 *
+	 * @param Ar The archive to serialize from or into.
+	 * @param Bound The bound to serialize.
 	 * @return The archive.
 	 */
 	friend class FArchive& operator<<( class FArchive& Ar, TRangeBound& Bound )
@@ -167,8 +152,7 @@ public:
 	/**
 	 * Gets the hash for the specified bound.
 	 *
-	 * @param Bound - The bound to get the hash for.
-	 *
+	 * @param Bound The bound to get the hash for.
 	 * @return Hash value.
 	 */
 	friend uint32 GetTypeHash( const TRangeBound& Bound )
@@ -181,8 +165,7 @@ public:
 	/**
 	 * Returns a closed bound that excludes the specified value.
 	 *
-	 * @param Value - The bound value.
-	 *
+	 * @param Value The bound value.
 	 * @return An exclusive closed bound.
 	 */
 	static FORCEINLINE TRangeBound Exclusive( const ElementType& Value )
@@ -198,8 +181,7 @@ public:
 	/**
 	 * Returns a closed bound that includes the specified value.
 	 *
-	 * @param Value - The bound value.
-	 *
+	 * @param Value The bound value.
 	 * @return An inclusive closed bound.
 	 */
 	static FORCEINLINE TRangeBound Inclusive( const ElementType& Value )
@@ -253,9 +235,8 @@ public:
 	/**
 	 * Returns the greater of two lower bounds.
 	 *
-	 * @param A - The first lower bound.
-	 * @param B - The second lower bound.
-	 *
+	 * @param A The first lower bound.
+	 * @param B The second lower bound.
 	 * @return The greater lower bound.
 	 */
 	static FORCEINLINE const TRangeBound& MaxLower( const TRangeBound& A, const TRangeBound& B )
@@ -272,9 +253,8 @@ public:
 	/**
 	 * Returns the greater of two upper bounds.
 	 *
-	 * @param A - The first upper bound.
-	 * @param B - The second upper bound.
-	 *
+	 * @param A The first upper bound.
+	 * @param B The second upper bound.
 	 * @return The greater upper bound.
 	 */
 	static FORCEINLINE const TRangeBound& MaxUpper( const TRangeBound& A, const TRangeBound& B )
@@ -291,9 +271,8 @@ public:
 	/**
 	 * Returns the lesser of two lower bounds.
 	 *
-	 * @param A - The first lower bound.
-	 * @param B - The second lower bound.
-	 *
+	 * @param A The first lower bound.
+	 * @param B The second lower bound.
 	 * @return The lesser lower bound.
 	 */
 	static FORCEINLINE const TRangeBound& MinLower( const TRangeBound& A, const TRangeBound& B )
@@ -310,9 +289,8 @@ public:
 	/**
 	 * Returns the lesser of two upper bounds.
 	 *
-	 * @param A - The first upper bound.
-	 * @param B - The second upper bound.
-	 *
+	 * @param A The first upper bound.
+	 * @param B The second upper bound.
 	 * @return The lesser upper bound.
 	 */
 	static FORCEINLINE const TRangeBound& MinUpper( const TRangeBound& A, const TRangeBound& B )
diff --git a/Engine/Source/Runtime/Core/Public/Math/RangeSet.h b/Engine/Source/Runtime/Core/Public/Math/RangeSet.h
--- a/Engine/Source/Runtime/Core/Public/Math/RangeSet.h
+++ b/Engine/Source/Runtime/Core/Public/Math/RangeSet.h
@@ -9,12 +9,13 @@
 
 /**
  * Template for range sets.
+ *
+ * @todo gmp: TRangeSet needs an overhaul
  */
 template<typename ElementType> class TRangeSet
 {
 	typedef TRange<ElementType> RangeType;
 
-
 public:
 
 	/**
@@ -35,7 +36,7 @@ public:
 	 * This method merges overlapping ranges into a single range (i.e. {[1, 4], [4, 6]} becomes [1, 6]).
 	 * Adjacent ranges (i.e. {[1, 3], [4, 6]}) are not merged.
 	 *
-	 * @param Range - The range to add.
+	 * @param Range The range to add.
 	 *//*
 	void Add( RangeType Range )
 	{
@@ -59,7 +60,7 @@ public:
 	/**
 	 * Merges another range set into this set.
 	 *
-	 * @param Other - The range set to merge.
+	 * @param Other The range set to merge.
 	 */
 /*	void Merge( const TRangeSet& Other )
 	{
@@ -74,7 +75,7 @@ public:
 	 *
 	 * Ranges that overlap with the removed range will be split.
 	 *
-	 * @param Range - The range to remove.
+	 * @param Range The range to remove.
 	 */
 /*	void Remove( const RangeType& Range )
 	{
@@ -114,8 +115,7 @@ public:
 	/**
 	 * Checks whether this set contains the specified element.
 	 *
-	 * @param Element - The element to check.
-	 *
+	 * @param Element The element to check.
 	 * @return true if the element is in the set, false otherwise.
 	 */
 	bool Contains( const ElementType& Element ) const
@@ -134,8 +134,7 @@ public:
 	/**
 	 * Checks whether this set contains the specified range.
 	 *
-	 * @param Range - The range to check.
-	 *
+	 * @param Range The range to check.
 	 * @return true if the set contains the range, false otherwise.
 	 */
 	bool Contains( const RangeType& Range ) const
@@ -174,8 +173,7 @@ public:
 	/**
 	 * Checks whether this range set overlaps with the specified range.
 	 *
-	 * @param Range - The range to check.
-	 *
+	 * @param Range The range to check.
 	 * @return true if this set overlaps with the range, false otherwise.
 	 */
 	bool Overlaps( const RangeType& Range ) const
@@ -194,11 +192,10 @@ public:
 	/**
 	 * Checks whether this range set overlaps with another.
 	 *
-	 * @param Other - The other range set.
-	 *
+	 * @param Other The other range set.
 	 * @return true if the range sets overlap, false otherwise.
 	 *
-	 * @todo This could be optimized to O(n*logn) using a line sweep on a pre-sorted array of bounds.
+	 * @todo gmp: This could be optimized to O(n*logn) using a line sweep on a pre-sorted array of bounds.
 	 */
 	bool Overlaps( const TRangeSet& Other ) const
 	{
@@ -218,9 +215,8 @@ public:
 	/**
 	 * Serializes the given range set from or into the specified archive.
 	 *
-	 * @param Ar - The archive to serialize from or into.
-	 * @param RangeSet - The range set to serialize.
-	 *
+	 * @param Ar The archive to serialize from or into.
+	 * @param RangeSet The range set to serialize.
 	 * @return The archive.
 	 */
 	friend class FArchive& operator<<( class FArchive& Ar, TRangeSet& RangeSet )
diff --git a/Engine/Source/Runtime/Core/Public/Math/Rotator.h b/Engine/Source/Runtime/Core/Public/Math/Rotator.h
--- a/Engine/Source/Runtime/Core/Public/Math/Rotator.h
+++ b/Engine/Source/Runtime/Core/Public/Math/Rotator.h
@@ -1,48 +1,41 @@
 // Copyright 1998-2014 Epic Games, Inc. All Rights Reserved.
 
+/*=============================================================================
+	Rotator.h: Declares the FRotator class.
+=============================================================================*/
+
 #pragma once
 
-/** Stores Rotation information.*/
+
+/**
+ * Implements a container for rotation information.
+ *
+ * All rotation values are stored in degrees.
+ */
 class FRotator
 {
 public:
-	// Stored in degrees
-	float Pitch; // Looking up and down (0=Straight Ahead, +Up, -Down).
-	float Yaw;   // Rotating around (running in circles), 0=East, +North, -South.
-	float Roll;  // Rotation about axis of screen, 0=Straight, +Clockwise, -CCW.
 
-	/**
-	 * Serializer.
-	 *
-	 * @param Ar Serialization Archive.
-	 * @param R Rotator being serialized.
-	 *
-	 * @return Reference to Archive after serialization.
-	 */
-	friend FArchive& operator<<( FArchive& Ar, FRotator& R )
-	{
-		Ar << R.Pitch << R.Yaw << R.Roll;
-		return Ar;
-	}
+	/** Looking up and down (0=Straight Ahead, +Up, -Down). */
+	float Pitch; 
 
-	/**
-	 * Serializes the rotator compressed for e.g. network transmission.
-	 * 
-	 * @param	Ar	Archive to serialize to/ from
-	 */
-	CORE_API void SerializeCompressed( FArchive& Ar );
+	/** Rotating around (running in circles), 0=East, +North, -South. */
+	float Yaw; 
 
-	/**
-	 * Serializes the rotator compressed for e.g. network transmission (use shorts though).
-	 * 
-	 * @param	Ar	Archive to serialize to/ from
-	 */
-	CORE_API void SerializeCompressedShort( FArchive& Ar );
+	/** Rotation about axis of screen, 0=Straight, +Clockwise, -CCW. */
+	float Roll;
 
-	CORE_API bool NetSerialize(FArchive& Ar, class UPackageMap* Map, bool& bOutSuccess);
+public:
 
-	/** Constructor.*/
-	FORCEINLINE FRotator();
+	/** A rotator of zero degrees on each axis. */
+	static CORE_API const FRotator ZeroRotator;
+
+public:
+
+	/**
+	 * Default constructor (no initialization).
+	 */
+	FORCEINLINE FRotator( ) { };
 
 	/**
 	 * Constructor.
@@ -58,7 +51,7 @@ public:
 	 *
 	 * @param EForceInit Force Init Enum.
 	 */
-	explicit FORCEINLINE FRotator(EForceInit);
+	explicit FORCEINLINE FRotator( EForceInit );
 
 	/**
 	 * Constructor.
@@ -67,13 +60,14 @@ public:
 	 */
 	explicit CORE_API FRotator( const FQuat& Quat );
 
+public:
+
 	// Binary arithmetic operators.
 
 	/**
 	 * Get the result of adding a rotator to this.
 	 *
 	 * @param R The other rotator.
-	 *
 	 * @return The result of adding a rotator to this.
 	 */
 	FRotator operator+( const FRotator &R ) const;
@@ -82,7 +76,6 @@ public:
 	 * Get the result of subtracting a rotator from this.
 	 *
 	 * @param R The other rotator.
-	 *
 	 * @return The result of subtracting a rotator from this.
 	 */
 	FRotator operator-( const FRotator &R ) const;
@@ -91,7 +84,6 @@ public:
 	 * Get the result of scaling this rotator.
 	 *
 	 * @param Scale The scaling factor.
-	 *
 	 * @return The result of scaling.
 	 */
 	FRotator operator*( float Scale ) const;
@@ -100,10 +92,9 @@ public:
 	 * Multiply this rotator by a scaling factor.
 	 *
 	 * @param Scale The scaling factor.
-	 *
 	 * @return Copy of the rotator after scaling.
 	 */
-	FRotator operator*= (float Scale);
+	FRotator operator*=( float Scale );
 
 	// Unary operators.
 
@@ -112,7 +103,7 @@ public:
 	 *
 	 * @return A negated copy of the rotator.
 	 */
-	FORCEINLINE FRotator operator-() const;
+	FORCEINLINE FRotator operator-( ) const;
 
 	// Binary comparison operators.
 
@@ -120,7 +111,6 @@ public:
 	 * Checks whether two rotators are identical.
 	 *
 	 * @param R The other rotator.
-	 *
 	 * @return true if two rotators are identical, otherwise false.
 	 */
 	bool operator==( const FRotator &R ) const;
@@ -129,7 +119,6 @@ public:
 	 * Checks whether two rotators are different.
 	 *
 	 * @param V The other rotator.
-	 *
 	 * @return true if two rotators are different, otherwise false.
 	 */
 	bool operator!=( const FRotator &V ) const;
@@ -140,7 +129,6 @@ public:
 	 * Adds another rotator to this.
 	 *
 	 * @param R The other rotator.
-	 *
 	 * @return Copy of rotator after addition.
 	 */
 	FRotator operator+=( const FRotator &R );
@@ -149,38 +137,37 @@ public:
 	 * Subtracts another rotator from this.
 	 *
 	 * @param R The other rotator.
-	 *
 	 * @return Copy of rotator after subtraction.
 	 */
 	FRotator operator-=( const FRotator &R );
 
+public:
+
 	// Functions.
 
 	/**
 	 * Checks whether rotator is nearly zero, within specified tolerance.
 	 *
 	 * @param Tolerance Error Tolerance.
-	 *
 	 * @return true if rotator is nearly zero, within specified tolerance, otherwise false.
 	 */
-	bool IsNearlyZero(float Tolerance=KINDA_SMALL_NUMBER) const;
+	bool IsNearlyZero( float Tolerance=KINDA_SMALL_NUMBER ) const;
 
 	/**
 	 * Checks whether this has exactly zero rotation.
 	 *
 	 * @return true if this has exactly zero rotation, otherwise false.
 	 */
-	bool IsZero() const;
+	bool IsZero( ) const;
 
 	/**
 	 * Checks whether two rotators are equal, within specified tolerance.
 	 *
 	 * @param R The other rotator.
 	 * @param Tolerance Error Tolerance.
-	 *
 	 * @return true if two rotators are equal, within specified tolerance, otherwise false.
 	 */
-	bool Equals(const FRotator &R, float Tolerance=KINDA_SMALL_NUMBER) const;
+	bool Equals( const FRotator &R, float Tolerance=KINDA_SMALL_NUMBER ) const;
 
 	/**
 	 * Adds to each component of the rotator.
@@ -188,7 +175,6 @@ public:
 	 * @param DeltaPitch Change in pitch. (+/-)
 	 * @param DeltaYaw Change in yaw. (+/-)
 	 * @param DeltaRoll Change in roll. (+/-)
-	 *
 	 * @return Copy of rotator after addition.
 	 */
 	FRotator Add( float DeltaPitch, float DeltaYaw, float DeltaRoll );
@@ -197,7 +183,6 @@ public:
 	 * Get the rotation, snapped to specified degree segments.
 	 *
 	 * @param RotGrid A Rotator specifying how to snap each component.
-	 *
 	 * @return Snapped version of rotation.
 	 */
 	FRotator GridSnap( const FRotator &RotGrid ) const;
@@ -207,163 +192,200 @@ public:
 	 *
 	 * @return Rotation as a direction vector.
 	 */
-	CORE_API FVector Vector() const;
+	CORE_API FVector Vector( ) const;
 
 	/**
 	 * Get Rotation as a quaternion.
 	 *
 	 * @return Rotation as a quaternion.
 	 */
-	CORE_API FQuat Quaternion() const;
+	CORE_API FQuat Quaternion( ) const;
 
 	/**
 	 * Convert a Rotator into floating-point Euler angles (in degrees). Rotator now stored in degrees.
 	 *
 	 * @return Rotation as a Euler angle vector.
 	 */
-	CORE_API FVector Euler() const;
+	CORE_API FVector Euler( ) const;
 
 	/**
-	 * Convert a vector of floating-point Euler angles (in degrees) into a Rotator. Rotator now stored in degrees
+	 * Rotate a vector rotated by this rotator.
 	 *
-	 * @param Euler Euler angle vector.
+	 * @param V The vector to rotate.
+	 * @return The rotated vector.
+	 */
+	CORE_API FVector RotateVector( const FVector& V ) const;
+
+	/**
+	 * Returns the vector rotated by the inverse of this rotator.
+	 *
+	 * @param V The vector to rotate.
+	 * @return The rotated vector.
+	 */
+	CORE_API FVector UnrotateVector( const FVector& V ) const;
+
+	/**
+	 * Gets the rotation values so they fall within the range [0,360]
+	 *
+	 * @return Clamped version of rotator.
+	 */
+	FRotator Clamp( ) const;
+
+	/** 
+	 * Create a copy of this rotator and normalize, removes all winding and creates the "shortest route" rotation. 
 	 *
-	 * @return A rotator from a euler angle.
+	 * @return Normalized copy of this rotator
+	 */
+	FRotator GetNormalized( ) const;
+
+	/** 
+	 * Create a copy of this rotator and denormalize, clamping each axis to 0 - 360. 
+	 *
+	 * @return Denormalized copy of this rotator
+	 */
+	FRotator GetDenormalized( ) const;
+
+	/**
+	 * In-place normalize, removes all winding and creates the "shortest route" rotation.
 	 */
-	static CORE_API FRotator MakeFromEuler(const FVector& Euler);
+	void Normalize( );
 
+	/** 
+	 * Decompose this Rotator into a Winding part (multiples of 360) and a Remainder part. 
+	 * Remainder will always be in [-180, 180] range.
+	 *
+	 * @param Winding[Out] the Winding part of this Rotator
+	 * @param Remainder[Out] the Remainder
+	 */
+	CORE_API void GetWindingAndRemainder( FRotator& Winding, FRotator& Remainder ) const;
 
-	/** Rotate a vector rotated by this rotator */
-	CORE_API FVector RotateVector(const FVector& V) const;
+	/**
+	 * Get a textual representation of the vector.
+	 *
+	 * @return Text describing the vector.
+	 */
+	FString ToString( ) const;
 
-	/* Returns the vector rotated by the inverse of this rotator */
-	CORE_API FVector UnrotateVector(const FVector& V) const;
+	/** Get a short textural representation of this vector, for compact readable logging. */
+	FString ToCompactString( ) const;
 
+	/**
+	 * Initialize this Rotator based on an FString. The String is expected to contain P=, Y=, R=.
+	 * The FRotator will be bogus when InitFromString returns false.
+	 *
+	 * @param InSourceString	FString containing the rotator values.
+	 * @return true if the P,Y,R values were read successfully; false otherwise.
+	 */
+	bool InitFromString( const FString & InSourceString );
 
 	/**
-	 * Gets the rotation values so they fall within the range [0,360]
+	 * Utility to check if there are any NaNs in this Rotator.
 	 *
-	 * @return Clamped version of rotator.
+	 * @return true if there are any NaNs in this Rotator, otherwise false.
+	 */
+	bool ContainsNaN( ) const;
+
+	/**
+	 * Serializes the rotator compressed for e.g. network transmission.
+	 * 
+	 * @param	Ar	Archive to serialize to/ from
+	 */
+	CORE_API void SerializeCompressed( FArchive& Ar );
+
+	/**
+	 * Serializes the rotator compressed for e.g. network transmission (use shorts though).
+	 * 
+	 * @param	Ar	Archive to serialize to/ from
+	 */
+	CORE_API void SerializeCompressedShort( FArchive& Ar );
+
+	/**
 	 */
-	FRotator Clamp() const;
+	CORE_API bool NetSerialize( FArchive& Ar, class UPackageMap* Map, bool& bOutSuccess );
+
+public:
 
 	/**
 	 * Clamps an angle to the range of [0, 360].
 	 *
 	 * @param Angle The angle to clamp.
-	 *
 	 * @return The clamped angle.
 	 */
-	static float ClampAxis(float Angle);
+	static float ClampAxis( float Angle );
 
 	/**
 	 * Clamps an angle to the range of [-180, 180].
 	 *
 	 * @param Angle The Angle to clamp.
-	 *
 	 * @return The clamped angle.
 	 */
-	static float NormalizeAxis(float Angle);
+	static float NormalizeAxis( float Angle );
 
 	/**
 	 * Compresses a floating point angle into a byte.
 	 *
 	 * @param Angle The angle to compress.
-	 *
 	 * @return The angle as a byte.
 	 */
-	static uint8 CompressAxisToByte(float Angle);
+	static uint8 CompressAxisToByte( float Angle );
 
 	/**
 	 * Decompress a word into a floating point angle.
 	 *
 	 * @param Angle The word angle.
-	 *
 	 * @return The decompressed angle.
 	 */
-	static float DecompressAxisFromByte(uint16 Angle);
+	static float DecompressAxisFromByte( uint16 Angle );
 
 	/**
 	 * Compress a floating point angle into a word.
 	 *
 	 * @param Angle The angle to compress.
-	 *
 	 * @return The decompressed angle.
 	 */
-	static uint16 CompressAxisToShort(float Angle);
-
+	static uint16 CompressAxisToShort( float Angle );
 
 	/**
 	 * Decompress a short into a floating point angle.
 	 *
 	 * @param Angle The word angle.
-	 *
 	 * @return The decompressed angle.
 	 */
-	static float DecompressAxisFromShort(uint16 Angle);
-
-	/** 
-	 * Create a copy of this rotator and normalize, removes all winding and creates the "shortest route" rotation. 
-	 *
-	 * @return Normalized copy of this rotator
-	 */
-	FRotator GetNormalized() const;
-
-	/** 
-	 * Create a copy of this rotator and denormalize, clamping each axis to 0 - 360. 
-	 *
-	 * @return Denormalized copy of this rotator
-	 */
-	FRotator GetDenormalized() const;
-
-	/** In-place normalize, removes all winding and creates the "shortest route" rotation. */
-	void Normalize();
-
-	/** 
-	 * Decompose this Rotator into a Winding part (multiples of 360) and a Remainder part. 
-	 * Remainder will always be in [-180, 180] range.
-	 *
-	 * @param Winding[Out] the Winding part of this Rotator
-	 * @param Remainder[Out] the Remainder
-	 */
-	CORE_API void GetWindingAndRemainder(FRotator& Winding, FRotator& Remainder) const;
+	static float DecompressAxisFromShort( uint16 Angle );
 
 	/**
-	 * Get a textual representation of the vector.
+	 * Convert a vector of floating-point Euler angles (in degrees) into a Rotator. Rotator now stored in degrees
 	 *
-	 * @return Text describing the vector.
+	 * @param Euler Euler angle vector.
+	 * @return A rotator from a Euler angle.
 	 */
-	FString ToString() const;
+	static CORE_API FRotator MakeFromEuler( const FVector& Euler );
 
-	/** Get a short textural representation of this vector, for compact readable logging. */
-	FString ToCompactString() const;
+public:
 
 	/**
-	 * Initialize this Rotator based on an FString. The String is expected to contain P=, Y=, R=.
-	 * The FRotator will be bogus when InitFromString returns false.
-	 *
-	 * @param	InSourceString	FString containing the rotator values.
+	 * Serializer.
 	 *
-	 * @return true if the P,Y,R values were read successfully; false otherwise.
+	 * @param Ar Serialization Archive.
+	 * @param R Rotator being serialized.
+	 * @return Reference to Archive after serialization.
 	 */
-	bool InitFromString( const FString & InSourceString );
+	friend FArchive& operator<<( FArchive& Ar, FRotator& R )
+	{
+		Ar << R.Pitch << R.Yaw << R.Roll;
+		return Ar;
+	}
+};
 
-	static CORE_API const FRotator ZeroRotator;
 
-	/**
-	 * Utility to check if there are any NaNs in this Rotator.
-	 *
-	 * @return true if there are any NaNs in this Rotator, otherwise false.
-	 */
-	bool ContainsNaN() const;
-};
+/* FRotator inline functions
+ *****************************************************************************/
 
 /**
- * Scale a rotator and return
- *
- * @param Scale scale to apply to R
- * @param R rotator to be scaled
+ * Scale a rotator and return.
  *
+ * @param Scale scale to apply to R.
+ * @param R rotator to be scaled.
  * @return Scaled rotator.
  */
 FORCEINLINE FRotator operator*( float Scale, const FRotator &R )
@@ -380,65 +402,74 @@ FORCEINLINE FRotator operator*( float Scale, const FRotator &R )
 	}
 #endif
 
-FORCEINLINE FRotator::FRotator() 
-{}
 
-FORCEINLINE FRotator::FRotator( float InPitch, float InYaw, float InRoll )
+	FORCEINLINE FRotator::FRotator( float InPitch, float InYaw, float InRoll )
 	:	Pitch(InPitch), Yaw(InYaw), Roll(InRoll) 
 {}
 
+
 FORCEINLINE FRotator::FRotator(EForceInit)
 	: Pitch(0), Yaw(0), Roll(0)
 {}
 
+
 FORCEINLINE FRotator FRotator::operator+( const FRotator &R ) const
 {
 	return FRotator( Pitch+R.Pitch, Yaw+R.Yaw, Roll+R.Roll );
 }
 
+
 FORCEINLINE FRotator FRotator::operator-( const FRotator &R ) const
 {
 	return FRotator( Pitch-R.Pitch, Yaw-R.Yaw, Roll-R.Roll );
 }
 
+
 FORCEINLINE FRotator FRotator::operator*( float Scale ) const
 {
 	return FRotator( Pitch*Scale, Yaw*Scale, Roll*Scale );
 }
 
+
 FORCEINLINE FRotator FRotator::operator*= (float Scale)
 {
 	Pitch = Pitch*Scale; Yaw = Yaw*Scale; Roll = Roll*Scale;
 	return *this;
 }
 
+
 FORCEINLINE FRotator FRotator::operator-() const
 {
 	return FRotator( -Pitch, -Yaw, -Roll );
 }
 
+
 FORCEINLINE bool FRotator::operator==( const FRotator &R ) const
 {
 	return Pitch==R.Pitch && Yaw==R.Yaw && Roll==R.Roll;
 }
 
+
 FORCEINLINE bool FRotator::operator!=( const FRotator &V ) const
 {
 	return Pitch!=V.Pitch || Yaw!=V.Yaw || Roll!=V.Roll;
 }
 
+
 FORCEINLINE FRotator FRotator::operator+=( const FRotator &R )
 {
 	Pitch += R.Pitch; Yaw += R.Yaw; Roll += R.Roll;
 	return *this;
 }
 
+
 FORCEINLINE FRotator FRotator::operator-=( const FRotator &R )
 {
 	Pitch -= R.Pitch; Yaw -= R.Yaw; Roll -= R.Roll;
 	return *this;
 }
 
+
 FORCEINLINE bool FRotator::IsNearlyZero(float Tolerance) const
 {
 	return
@@ -447,11 +478,13 @@ FORCEINLINE bool FRotator::IsNearlyZero(float Tolerance) const
 		&&	FMath::Abs(NormalizeAxis(Roll))<Tolerance;
 }
 
+
 FORCEINLINE bool FRotator::IsZero() const
 {
 	return (ClampAxis(Pitch)==0.f) && (ClampAxis(Yaw)==0.f) && (ClampAxis(Roll)==0.f);
 }
 
+
 FORCEINLINE bool FRotator::Equals(const FRotator &R, float Tolerance) const
 {
 	return (FMath::Abs(NormalizeAxis(Pitch - R.Pitch)) < Tolerance) 
@@ -459,6 +492,7 @@ FORCEINLINE bool FRotator::Equals(const FRotator &R, float Tolerance) const
 		&& (FMath::Abs(NormalizeAxis(Roll - R.Roll)) < Tolerance);
 }
 
+
 FORCEINLINE FRotator FRotator::Add( float DeltaPitch, float DeltaYaw, float DeltaRoll )
 {
 	Yaw   += DeltaYaw;
@@ -467,6 +501,7 @@ FORCEINLINE FRotator FRotator::Add( float DeltaPitch, float DeltaYaw, float Delt
 	return *this;
 }
 
+
 FORCEINLINE FRotator FRotator::GridSnap( const FRotator &RotGrid ) const
 {
 	return FRotator
@@ -477,12 +512,14 @@ FORCEINLINE FRotator FRotator::GridSnap( const FRotator &RotGrid ) const
 		);
 }
 
-FORCEINLINE FRotator FRotator::Clamp() const
+
+FORCEINLINE FRotator FRotator::Clamp( ) const
 {
-	return FRotator( ClampAxis(Pitch), ClampAxis(Yaw), ClampAxis(Roll) );
+	return FRotator(ClampAxis(Pitch), ClampAxis(Yaw), ClampAxis(Roll));
 }
 
-FORCEINLINE float FRotator::ClampAxis(float Angle)
+
+FORCEINLINE float FRotator::ClampAxis( float Angle )
 {
 	// returns Angle in the range (-360,360)
 	Angle = FMath::Fmod(Angle, 360.f);
@@ -496,7 +533,8 @@ FORCEINLINE float FRotator::ClampAxis(float Angle)
 	return Angle;
 }
 
-FORCEINLINE float FRotator::NormalizeAxis(float Angle)
+
+FORCEINLINE float FRotator::NormalizeAxis( float Angle )
 {
 	// returns Angle in the range [0,360)
 	Angle = ClampAxis(Angle);
@@ -510,39 +548,44 @@ FORCEINLINE float FRotator::NormalizeAxis(float Angle)
 	return Angle;
 }
 
-FORCEINLINE uint8 FRotator::CompressAxisToByte(float Angle)
+
+FORCEINLINE uint8 FRotator::CompressAxisToByte( float Angle )
 {
 	// map [0->360) to [0->256) and mask off any winding
 	return FMath::RoundToInt(Angle * 256.f / 360.f) & 0xFF;
 }
 
-FORCEINLINE float FRotator::DecompressAxisFromByte(uint16 Angle)
+
+FORCEINLINE float FRotator::DecompressAxisFromByte( uint16 Angle )
 {
 	// map [0->256) to [0->360)
 	return (Angle * 360.f / 256.f);
 }
 
-FORCEINLINE uint16 FRotator::CompressAxisToShort(float Angle)
+
+FORCEINLINE uint16 FRotator::CompressAxisToShort( float Angle )
 {
 	// map [0->360) to [0->65536) and mask off any winding
 	return FMath::RoundToInt(Angle * 65536.f / 360.f) & 0xFFFF;
 }
 
 
-FORCEINLINE float FRotator::DecompressAxisFromShort(uint16 Angle)
+FORCEINLINE float FRotator::DecompressAxisFromShort( uint16 Angle )
 {
 	// map [0->65536) to [0->360)
 	return (Angle * 360.f / 65536.f);
 }
 
-FORCEINLINE FRotator FRotator::GetNormalized() const
+
+FORCEINLINE FRotator FRotator::GetNormalized( ) const
 {
 	FRotator Rot = *this;
 	Rot.Normalize();
 	return Rot;
 }
 
-FORCEINLINE FRotator FRotator::GetDenormalized() const
+
+FORCEINLINE FRotator FRotator::GetDenormalized( ) const
 {
 	FRotator Rot = *this;
 	Rot.Pitch	= ClampAxis(Rot.Pitch);
@@ -551,19 +594,21 @@ FORCEINLINE FRotator FRotator::GetDenormalized() const
 	return Rot;
 }
 
-FORCEINLINE void FRotator::Normalize()
+
+FORCEINLINE void FRotator::Normalize( )
 {
 	Yaw = NormalizeAxis(Yaw);
 	Pitch = NormalizeAxis(Pitch);
 	Roll = NormalizeAxis(Roll);
 }
 
-FORCEINLINE FString FRotator::ToString() const
+FORCEINLINE FString FRotator::ToString( ) const
 {
 	return FString::Printf(TEXT("P=%f Y=%f R=%f"), Pitch, Yaw, Roll );
 }
 
-FORCEINLINE FString FRotator::ToCompactString() const
+
+FORCEINLINE FString FRotator::ToCompactString( ) const
 {
 	if( IsNearlyZero() )
 	{
@@ -599,6 +644,7 @@ FORCEINLINE FString FRotator::ToCompactString() const
 	return ReturnString;
 }
 
+
 FORCEINLINE bool FRotator::InitFromString( const FString & InSourceString )
 {
 	Pitch = Yaw = Roll = 0;
@@ -609,11 +655,13 @@ FORCEINLINE bool FRotator::InitFromString( const FString & InSourceString )
 	return bSuccessful;
 }
 
-FORCEINLINE bool FRotator::ContainsNaN() const
+
+FORCEINLINE bool FRotator::ContainsNaN( ) const
 {
 	return (FMath::IsNaN(Pitch) || !FMath::IsFinite(Pitch) 
 			|| FMath::IsNaN(Yaw) || !FMath::IsFinite(Yaw) 
 			|| FMath::IsNaN(Roll) || !FMath::IsFinite(Roll));
 }
 
-template <> struct TIsPODType<FRotator> { enum { Value = true }; };
+
+template<> struct TIsPODType<FRotator> { enum { Value = true }; };
diff --git a/Engine/Source/Runtime/Core/Public/Math/ScaleMatrix.h b/Engine/Source/Runtime/Core/Public/Math/ScaleMatrix.h
--- a/Engine/Source/Runtime/Core/Public/Math/ScaleMatrix.h
+++ b/Engine/Source/Runtime/Core/Public/Math/ScaleMatrix.h
@@ -1,38 +1,50 @@
 // Copyright 1998-2014 Epic Games, Inc. All Rights Reserved.
 
+/*=============================================================================
+	ScaleMatrix.h: Declares the FScaleMatrix class.
+=============================================================================*/
+
 #pragma once
 
-/** Scale matrix */
-class FScaleMatrix : public FMatrix
+
+/**
+ * Scale matrix.
+ */
+class FScaleMatrix
+	: public FMatrix
 {
 public:
 
 	/**
 	 * @param Scale uniform scale to apply to matrix.
 	 */
-	FScaleMatrix(float Scale);
+	FScaleMatrix( float Scale );
 
 	/**
 	 * @param Scale Non-uniform scale to apply to matrix.
 	 */
-	FScaleMatrix(const FVector& Scale);
+	FScaleMatrix( const FVector& Scale );
 };
 
 
-FORCEINLINE FScaleMatrix::FScaleMatrix(float Scale) :
-	FMatrix(
+/* FScaleMatrix inline functions
+ *****************************************************************************/
+
+FORCEINLINE FScaleMatrix::FScaleMatrix( float Scale )
+	: FMatrix(
 		FPlane(Scale,	0.0f,	0.0f,	0.0f),
 		FPlane(0.0f,	Scale,	0.0f,	0.0f),
 		FPlane(0.0f,	0.0f,	Scale,	0.0f),
-		FPlane(0.0f,	0.0f,	0.0f,	1.0f))
-{
-}
+		FPlane(0.0f,	0.0f,	0.0f,	1.0f)
+	)
+{ }
+
 
-FORCEINLINE FScaleMatrix::FScaleMatrix(const FVector& Scale) :
-	FMatrix(
+FORCEINLINE FScaleMatrix::FScaleMatrix( const FVector& Scale )
+	: FMatrix(
 		FPlane(Scale.X,	0.0f,		0.0f,		0.0f),
 		FPlane(0.0f,	Scale.Y,	0.0f,		0.0f),
 		FPlane(0.0f,	0.0f,		Scale.Z,	0.0f),
-		FPlane(0.0f,	0.0f,		0.0f,		1.0f))
-{
-}
+		FPlane(0.0f,	0.0f,		0.0f,		1.0f)
+	)
+{ }
diff --git a/Engine/Source/Runtime/Core/Public/Math/Sphere.h b/Engine/Source/Runtime/Core/Public/Math/Sphere.h
--- a/Engine/Source/Runtime/Core/Public/Math/Sphere.h
+++ b/Engine/Source/Runtime/Core/Public/Math/Sphere.h
@@ -14,41 +14,36 @@ class FSphere
 {
 public:
 
-	/**
-	 * Holds the sphere's center.
-	 */
+	/** Holds the sphere's center. */
 	FVector Center;
 
-	/**
-	 * Holds the sphere's radius.
-	 */
+	/** Holds the sphere's radius. */
 	float W;
 
-
 public:
 
 	/**
-	 * Default constructor.
+	 * Default constructor (no initialization).
 	 */
-	FSphere() { }
+	FSphere( ) { }
 
 	/**
 	 * Creates and initializes a new sphere.
 	 *
 	 * @param int32 Passing int32 sets up zeroed sphere.
 	 */
-	FSphere(int32)
-		: Center(0, 0, 0)
+	FSphere( int32 )
+		: Center(0.0f, 0.0f, 0.0f)
 		, W(0)
 	{ }
 
 	/**
 	 * Creates and initializes a new sphere with the specified parameters.
 	 *
-	 * @param InV - Center of sphere.
-	 * @param InW - Radius of sphere.
+	 * @param InV Center of sphere.
+	 * @param InW Radius of sphere.
 	 */
-	FSphere(FVector InV, float InW)
+	FSphere( FVector InV, float InW )
 		: Center(InV)
 		, W(InW)
 	{ }
@@ -58,9 +53,9 @@ public:
 	 *
 	 * @param EForceInit Force Init Enum.
 	 */
-	explicit FORCEINLINE FSphere(EForceInit)
+	explicit FORCEINLINE FSphere( EForceInit )
 		: Center(ForceInit)
-		, W(0.f)
+		, W(0.0f)
 	{ }
 
 	/**
@@ -69,8 +64,7 @@ public:
 	 * @param Pts Pointer to list of points this sphere must contain.
 	 * @param Count How many points are in the list.
 	 */
-	CORE_API FSphere(const FVector* Pts, int32 Count);
-
+	CORE_API FSphere( const FVector* Pts, int32 Count );
 
 public:
 
@@ -79,10 +73,9 @@ public:
 	 *
 	 * @param Sphere The other sphere.
 	 * @param Tolerance Error Tolerance.
-	 *
 	 * @return true if spheres are equal within specified tolerance, otherwise false.
 	 */
-	bool Equals(const FSphere& Sphere, float Tolerance = KINDA_SMALL_NUMBER) const
+	bool Equals( const FSphere& Sphere, float Tolerance = KINDA_SMALL_NUMBER ) const
 	{
 		return Center.Equals(Sphere.Center, Tolerance) && FMath::Abs(W - Sphere.W) < Tolerance;
 	}
@@ -92,10 +85,9 @@ public:
 	 *
 	 * @param Other The other sphere.
 	 * @param Tolerance Error Tolerance.
-	 *
 	 * @return true if sphere is inside another, otherwise false.
 	 */
-	bool IsInside(const FSphere& Other, float Tolerance = KINDA_SMALL_NUMBER) const
+	bool IsInside( const FSphere& Other, float Tolerance = KINDA_SMALL_NUMBER ) const
 	{
 		if (W > Other.W - Tolerance)
 		{
@@ -109,35 +101,31 @@ public:
 	 * Get result of Transforming sphere by Matrix.
 	 *
 	 * @param M Matrix to transform by.
-	 *
 	 * @return Result of transformation.
 	 */
-	CORE_API FSphere TransformBy(const FMatrix& M) const;
+	CORE_API FSphere TransformBy( const FMatrix& M ) const;
 
 	/**
 	 * Get result of Transforming sphere with Transform.
 	 *
 	 * @param M Transform information.
-	 *
 	 * @return Result of transformation.
 	 */
-	CORE_API FSphere TransformBy(const FTransform& M) const;
-
+	CORE_API FSphere TransformBy( const FTransform& M ) const;
 
 public:
 
 	/**
 	 * Serializes the given sphere from or into the specified archive.
 	 *
-	 * @param Ar - The archive to serialize from or into.
-	 * @param Sphere - The sphere to serialize.
-	 *
+	 * @param Ar The archive to serialize from or into.
+	 * @param Sphere The sphere to serialize.
 	 * @return The archive.
 	 */
-	friend FArchive& operator<< (FArchive& Ar, FSphere& Sphere)
+	friend FArchive& operator<<( FArchive& Ar, FSphere& Sphere )
 	{
 		Ar << Sphere.Center << Sphere.W;
 
 		return Ar;
 	}
-};
\ No newline at end of file
+};
diff --git a/Engine/Source/Runtime/Core/Public/Math/TwoVectors.h b/Engine/Source/Runtime/Core/Public/Math/TwoVectors.h
--- a/Engine/Source/Runtime/Core/Public/Math/TwoVectors.h
+++ b/Engine/Source/Runtime/Core/Public/Math/TwoVectors.h
@@ -1,23 +1,25 @@
 // Copyright 1998-2014 Epic Games, Inc. All Rights Reserved.
 
+/*=============================================================================
+	TwoVectors.h: Declares the FTwoVectors class.
+=============================================================================*/
+
 #pragma once
 
-/** A pair of 3D vectors. */
+
+/**
+ * A pair of 3D vectors.
+ */
 class FTwoVectors
 {
 public:
 
-	/**
-	 * Holds the first vector.
-	 */
+	/**  Holds the first vector. */
 	FVector	v1;
 
-	/**
-	 * Holds the second vector.
-	 */
+	/**  Holds the second vector. */
 	FVector	v2;
 
-
 public:
 
 	/**
@@ -28,26 +30,24 @@ public:
 	/**
 	 * Creates and initializes a new instance with the specified vectors.
 	 *
-	 * @param In1 - The first Vector.
-	 * @param In2 - The second Vector.
+	 * @param In1 The first Vector.
+	 * @param In2 The second Vector.
 	 */
 	FORCEINLINE	FTwoVectors( FVector In1, FVector In2 );
 
 	/**
-	 * Constructor
+	 * Constructor.
 	 *
 	 * @param EForceInit Force Init Enum
 	 */
 	explicit FORCEINLINE FTwoVectors( EForceInit );
 
-
 public:
 
 	/**
 	 * Gets result of addition of two pairs of vectors.
 	 *
 	 * @param V The pair to add.
-	 *
 	 * @return Result of addition.
 	 */
 	FORCEINLINE FTwoVectors operator+( const FTwoVectors& V ) const;
@@ -56,7 +56,6 @@ public:
 	 * Gets result of subtraction of two pairs of vectors.
 	 *
 	 * @param V The pair to subtract.
-	 *
 	 * @return Result of subtraction.
 	 */
 	FORCEINLINE FTwoVectors operator-( const FTwoVectors& V ) const;
@@ -65,7 +64,6 @@ public:
 	 * Gets result of scaling pair of vectors.
 	 *
 	 * @param Scale The scaling factor.
-	 *
 	 * @return Result of Scaling.
 	 */
 	FORCEINLINE FTwoVectors operator*( float Scale ) const;
@@ -74,7 +72,6 @@ public:
 	 * Gets result of dividing pair of vectors.
 	 *
 	 * @param Scale What to divide by.
-	 *
 	 * @return Result of division.
 	 */
 	FTwoVectors operator/( float Scale ) const;
@@ -83,7 +80,6 @@ public:
 	 * Gets result of multiplying two pairs of vectors.
 	 *
 	 * @param V The pair to multiply with.
-	 *
 	 * @return Result of multiplication.
 	 */
 	FORCEINLINE FTwoVectors operator*( const FTwoVectors& V ) const;
@@ -92,7 +88,6 @@ public:
 	 * Gets result of division of two pairs of vectors.
 	 *
 	 * @param V The pair to divide by.
-	 *
 	 * @return Result of division.
 	 */
 	FORCEINLINE FTwoVectors operator/( const FTwoVectors& V ) const;
@@ -103,8 +98,7 @@ public:
 	 * Checks two pairs for equality.
 	 *
 	 * @param V The other pair.
-	 *
-	 * @return true if the two pairs are equal, otherwise false.
+	 * @return true if the two pairs are equal, false otherwise..
 	 */
 	bool operator==( const FTwoVectors& V ) const;
 
@@ -112,8 +106,7 @@ public:
 	 * Checks two pairs for inequality.
 	 *
 	 * @param V The other pair.
-	 *
-	 * @return true if the two pairs are different, otherwise false.
+	 * @return true if the two pairs are different, false otherwise..
 	 */
 	bool operator!=( const FTwoVectors& V ) const;
 
@@ -122,8 +115,7 @@ public:
 	 *
 	 * @param V The other pair.
 	 * @param Tolerance Error Tolerance.
-	 *
-	 * @return true if two pairs are equal within specified tolerance, otherwise false.
+	 * @return true if two pairs are equal within specified tolerance, false otherwise..
 	 */
 	bool Equals( const FTwoVectors& V, float Tolerance = KINDA_SMALL_NUMBER ) const;
 
@@ -142,7 +134,6 @@ public:
 	 * Add a pair to this.
 	 *
 	 * @param The pair to add.
-	 *
 	 * @return Copy of the pair after addition.
 	 */
 	FORCEINLINE FTwoVectors operator+=( const FTwoVectors& V );
@@ -151,7 +142,6 @@ public:
 	 * Subtract a pair from this.
 	 *
 	 * @param The pair to subtract.
-	 *
 	 * @return Copy of the pair after subtraction.
 	 */
 	FORCEINLINE FTwoVectors operator-=( const FTwoVectors& V );
@@ -160,7 +150,6 @@ public:
 	 * Scale the pair.
 	 *
 	 * @param Scale What to scale by.
-	 *
 	 * @return Copy of the pair after scaling.
 	 */
 	FORCEINLINE FTwoVectors operator*=( float Scale );
@@ -169,7 +158,6 @@ public:
 	 * Divide the pair.
 	 *
 	 * @param What to divide by.
-	 *
 	 * @return Copy of the pair after division.
 	 */
 	FTwoVectors operator/=( float V );
@@ -178,7 +166,6 @@ public:
 	 * Multiply the pair by another.
 	 *
 	 * @param The other pair.
-	 *
 	 * @return Copy of the pair after multiplication.
 	 */
 	FTwoVectors operator*=( const FTwoVectors& V );
@@ -187,7 +174,6 @@ public:
 	 * Divide the pair by another.
 	 *
 	 * @param The other pair.
-	 *
 	 * @return Copy of the pair after multiplication.
 	 */
 	FTwoVectors operator/=( const FTwoVectors& V );
@@ -202,7 +188,6 @@ public:
 	 */
     float& operator[]( int32 i );
 
-
 public:
 
 	/**
@@ -226,7 +211,6 @@ public:
 	 */
 	FString ToString( ) const;
 
-
 public:
 
 	/**
@@ -244,6 +228,9 @@ public:
 };
 
 
+/* FTwoVectors inline functions
+ *****************************************************************************/
+
 FORCEINLINE FTwoVectors operator*( float Scale, const FTwoVectors& V )
 {
 	return V.operator*( Scale );
diff --git a/Engine/Source/Runtime/Core/Public/Math/Vector.h b/Engine/Source/Runtime/Core/Public/Math/Vector.h
--- a/Engine/Source/Runtime/Core/Public/Math/Vector.h
+++ b/Engine/Source/Runtime/Core/Public/Math/Vector.h
@@ -1,15 +1,40 @@
 // Copyright 1998-2014 Epic Games, Inc. All Rights Reserved.
 
+/*=============================================================================
+	Vector.h: Declares the FVector class.
+=============================================================================*/
+
 #pragma once
 
+
 /**
  * A 3x1 of FLOATs.
  */
 class FVector 
 {
 public:
-	// Variables.
-	float X,Y,Z;
+
+	/** Holds the vector's X-component. */
+	float X;
+
+	/** Holds the vector's Y-component. */
+	float Y;
+
+	/** Holds the vector's Z-component. */
+	float Z;
+
+public:
+
+	/** A zero vector (0,0,0) */
+	static CORE_API const FVector ZeroVector;
+
+	/** World up vector (0,0,1) */
+	static CORE_API const FVector UpVector;
+
+	/** Local Unreal forward vector (1,0,0) */
+	static CORE_API const FVector ForwardVector;
+
+public:
 
 #if ENABLE_NAN_DIAGNOSTIC
 	FORCEINLINE void DiagnosticCheckNaN() const
@@ -97,7 +122,6 @@ public:
 	 * Calculate Cross product between this and another vector.
 	 *
 	 * @param V The other vector.
-	 *
 	 * @return The Cross product.
 	 */
 	FORCEINLINE FVector operator^( const FVector& V ) const;
@@ -107,7 +131,6 @@ public:
 	 *
 	 * @param A The first vector.
 	 * @param B The second vector.
-	 *
 	 * @return The Cross product.
 	 */
 	FORCEINLINE static FVector CrossProduct( const FVector& A, const FVector& B );
@@ -116,7 +139,6 @@ public:
 	 * Calculate the Dot Product between this and another vector.
 	 *
 	 * @param V The other vector.
-	 *
 	 * @return The Dot Product.
 	 */
 	FORCEINLINE float operator|( const FVector& V ) const;
@@ -126,7 +148,6 @@ public:
 	 *
 	 * @param A The first vector.
 	 * @param B The second vector.
-	 *
 	 * @return The Dot product.
 	 */
 	FORCEINLINE static float DotProduct( const FVector& A, const FVector& B );
@@ -135,7 +156,6 @@ public:
 	 * Gets the result of adding a vector to this.
 	 *
 	 * @param V The vector to add.
-	 *
 	 * @return The result of vector addition.
 	 */
 	FORCEINLINE FVector operator+( const FVector& V ) const;
@@ -144,7 +164,6 @@ public:
 	 * Gets the result of subtracting a vector from this.
 	 *
 	 * @param V The vector to subtract.
-	 *
 	 * @return The result of vector subtraction.
 	 */
 	FORCEINLINE FVector operator-( const FVector& V ) const;
@@ -153,7 +172,6 @@ public:
 	 * Gets the result of subtracting from each component of the vector.
 	 *
 	 * @param Bias What to subtract.
-	 *
 	 * @return The result of subtraction.
 	 */
 	FORCEINLINE FVector operator-( float Bias ) const;
@@ -162,7 +180,6 @@ public:
 	 * Gets the result of adding to each component of the vector.
 	 *
 	 * @param Bias What to add.
-	 *
 	 * @return The result of addition.
 	 */
 	FORCEINLINE FVector operator+( float Bias ) const;
@@ -171,7 +188,6 @@ public:
 	 * Gets the result of multiplying each component of the vector.
 	 *
 	 * @param Scale What to multiply by.
-	 *
 	 * @return The result of multiplication.
 	 */
 	FORCEINLINE FVector operator*( float Scale ) const;
@@ -180,7 +196,6 @@ public:
 	 * Gets the result of dividing each component of the vector.
 	 *
 	 * @param Scale What to divide by.
-	 *
 	 * @return The result of division.
 	 */
 	FVector operator/( float Scale ) const;
@@ -189,7 +204,6 @@ public:
 	 * Gets the result of multiplying vector with this.
 	 *
 	 * @param V The vector to multiply with.
-	 *
 	 * @return The result of multiplication.
 	 */
 	FORCEINLINE FVector operator*( const FVector& V ) const;
@@ -198,7 +212,6 @@ public:
 	 * Gets the result of dividing this vector by another.
 	 *
 	 * @param V The vector to divide by.
-	 *
 	 * @return The result of division.
 	 */
 	FORCEINLINE FVector operator/( const FVector& V ) const;
@@ -209,8 +222,7 @@ public:
 	 * Check against another vector for equality.
 	 *
 	 * @param V The vector to check against.
-	 *
-	 * @return true if the vectors are equal, otherwise false.
+	 * @return true if the vectors are equal, false otherwise..
 	 */
 	bool operator==( const FVector& V ) const;
 
@@ -218,8 +230,7 @@ public:
 	 * Check against another vector for inequality.
 	 *
 	 * @param V The vector to check against.
-	 *
-	 * @return true if the vectors are not equal, otherwise false.
+	 * @return true if the vectors are not equal, false otherwise..
 	 */
 	bool operator!=( const FVector& V ) const;
 
@@ -228,8 +239,7 @@ public:
 	 *
 	 * @param V The vector to check against.
 	 * @param Tolerance Error tolerance.
-	 *
-	 * @return true if the vectors are equal within tolerance limits, otherwise false.
+	 * @return true if the vectors are equal within tolerance limits, false otherwise..
 	 */
 	bool Equals(const FVector& V, float Tolerance=KINDA_SMALL_NUMBER) const;
 
@@ -237,12 +247,10 @@ public:
 	 * Checks whether all components of the vector are the same, within a tolerance.
 	 *
 	 * @param Tolerance Error Tolerance.
-	 *
-	 * @return true if the vectors are equal within tolerance limits, otherwise false.
+	 * @return true if the vectors are equal within tolerance limits, false otherwise..
 	 */
 	bool AllComponentsEqual(float Tolerance=KINDA_SMALL_NUMBER) const;
 
-
 	/**
 	 * Get a negated copy of the vector.
 	 *
@@ -250,12 +258,10 @@ public:
 	 */
 	FORCEINLINE FVector operator-() const;
 
-
 	/**
 	 * Adds another vector to this.
 	 *
 	 * @param V Vector to add.
-	 *
 	 * @return Copy of the vector after addition.
 	 */
 	FORCEINLINE FVector operator+=( const FVector& V );
@@ -264,7 +270,6 @@ public:
 	 * Subtracts another vector from this.
 	 *
 	 * @param V Vector to subtract.
-	 *
 	 * @return Copy of the vector after subtraction.
 	 */
 	FORCEINLINE FVector operator-=( const FVector& V );
@@ -273,7 +278,6 @@ public:
 	 * Scales the vector.
 	 *
 	 * @param Scale What to scale vector by.
-	 *
 	 * @return Copy of the vector after scaling.
 	 */
 	FORCEINLINE FVector operator*=( float Scale );
@@ -282,7 +286,6 @@ public:
 	 * Divides the vector by a number.
 	 *
 	 * @param V What to divide the vector by.
-	 *
 	 * @return Copy of the vector after division.
 	 */
 	FVector operator/=( float V );
@@ -291,7 +294,6 @@ public:
 	 * Multiplies the vector with another vector.
 	 *
 	 * @param V What to multiply vector with.
-	 *
 	 * @return Copy of the vector after multiplication.
 	 */
 	FVector operator*=( const FVector& V );
@@ -300,7 +302,6 @@ public:
 	 * Divides the vector by another vector.
 	 *
 	 * @param V What to divide vector by.
-	 *
 	 * @return Copy of the vector after division.
 	 */
 	FVector operator/=( const FVector& V );
@@ -309,7 +310,6 @@ public:
 	 * Gets specific component of the vector.
 	 *
 	 * @param Index the index of vector component
-	 *
 	 * @return reference to component.
 	 */
 	float& operator[]( int32 Index );
@@ -318,7 +318,6 @@ public:
 	 * Gets specific component of the vector.
 	 *
 	 * @param Index the index of vector component
-	 *
 	 * @return Copy of the component.
 	 */
 	float operator[]( int32 Index )const;
@@ -407,15 +406,14 @@ public:
 	 * Checks whether vector is near to zero within a specified tolerance.
 	 *
 	 * @param Tolerance Error tolerance.
-	 *
-	 * @return true if the vector is near to zero, otherwise false.
+	 * @return true if the vector is near to zero, false otherwise..
 	 */
 	bool IsNearlyZero(float Tolerance=KINDA_SMALL_NUMBER) const;
 
 	/**
 	 * Checks whether vector is exactly zero.
 	 *
-	 * @return true if the vector is exactly zero, otherwise false.
+	 * @return true if the vector is exactly zero, false otherwise..
 	 */
 	bool IsZero() const;
 
@@ -423,8 +421,7 @@ public:
 	 * Normalize this vector if it is large enough.
 	 *
 	 * @param Tolerance Minimum squared length of vector for normalization.
-	 *
-	 * @return true if the vector was normalized correctly, otherwise false.
+	 * @return true if the vector was normalized correctly, false otherwise..
 	 */
 	bool Normalize(float Tolerance=SMALL_NUMBER);
 
@@ -468,7 +465,6 @@ public:
 	 * Gets a copy of this vector snapped to a grid.
 	 *
 	 * @param GridSz Grid dimension.
-	 *
 	 * @return A copy of this vector snapped to a grid.
 	 */
 	FVector GridSnap( const float& GridSz ) const;
@@ -477,7 +473,6 @@ public:
 	 * Get a copy of this vector, clamped inside of a cube.
 	 *
 	 * @param Radius Half size of the cube.
-	 *
 	 * @return A copy of this vector, bound by cube.
 	 */
 	FVector BoundToCube( float Radius ) const;
@@ -494,7 +489,6 @@ public:
 	/** Create a copy of this vector, with the 2D magnitude clamped to MaxSize. Z is unchanged. */
 	FVector ClampMaxSize2D(float MaxSize) const;
 
-
 	/**
 	 * Add a vector to this and clamp the result in a cube.
 	 *
@@ -516,7 +510,6 @@ public:
 	 * Gets a specific component of the vector.
 	 *
 	 * @param Index The index of the component required.
-	 *
 	 * @return Copy of the specified component.
 	 */
 	float Component( int32 Index ) const;
@@ -533,7 +526,6 @@ public:
 	 * Check whether X, Y and Z are nearly equal.
 	 *
 	 * @param Tolerance Specified Tolerance.
-	 *
 	 * @return true if X == Y == Z within the specified tolerance.
 	 */
 	bool IsUniform(float Tolerance=KINDA_SMALL_NUMBER) const;
@@ -542,7 +534,6 @@ public:
 	 * Mirror a vector about a normal vector.
 	 *
 	 * @param MirrorNormal Normal vector to mirror about.
-	 *
 	 * @return Mirrored vector.
 	 */
 	FVector MirrorByVector( const FVector& MirrorNormal ) const;
@@ -551,7 +542,6 @@ public:
 	 * Mirrors a vector about a plane.
 	 *
 	 * @param Plane Plane to mirror about.
-	 *
 	 * @return Mirrored vector.
 	 */
 	FVector MirrorByPlane( const FPlane& Plane ) const;
@@ -561,7 +551,6 @@ public:
 	 *
 	 * @param Angle Angle to rotate (in degrees).
 	 * @param Axis Axis to rotate around.
-	 *
 	 * @return Rotated Vector.
 	 */
 	FVector RotateAngleAxis( const float AngleDeg, const FVector& Axis ) const;
@@ -570,7 +559,6 @@ public:
 	 * Gets a normalized copy of the vector, checking it is safe to do so.
 	 *
 	 * @param Tolerance Minimum squared vector length.
-	 *
 	 * @return Normalized copy if safe, otherwise returns zero vector.
 	 */
 	FVector SafeNormal(float Tolerance=SMALL_NUMBER) const;
@@ -579,7 +567,6 @@ public:
 	 * Gets a normalized copy of the 2D components of the vector, checking it is safe to do so.
 	 *
 	 * @param Tolerance Minimum squared vector length.
-	 *
 	 * @return Normalized copy if safe, otherwise returns zero vector.
 	 */
 	FVector SafeNormal2D(float Tolerance=SMALL_NUMBER) const;
@@ -588,7 +575,6 @@ public:
 	 * Returns the cosine of the angle between this vector and another projected onto the XY plane (no z)
 	 *
 	 * @param B the other vector to find the 2D cosine of the angle with.
-	 *
 	 * @return The cosine.
 	 */
 	FORCEINLINE float CosineAngle2D(FVector B) const;
@@ -597,7 +583,6 @@ public:
 	 * Projects this vector onto the input vector.
 	 *
 	 * @param A Vector to project onto, does not assume it is unnormalized.
-	 *
 	 * @return Projected vector.
 	 */
 	FORCEINLINE FVector ProjectOnTo( const FVector& A ) const ;
@@ -626,7 +611,7 @@ public:
 	/**
 	 * Utility to check if there are any NaNs in this vector.
 	 *
-	 * @return true if there are any NaNs in this vector, otherwise false.
+	 * @return true if there are any NaNs in this vector, false otherwise..
 	 */
 	bool ContainsNaN() const;
 
@@ -634,7 +619,6 @@ public:
 	 * Check if the vector is of unit length, with specified tolerance.
 	 *
 	 * @param LengthSquaredTolerance Tolerance against squared length.
-	 *
 	 * @return true if the vector is a unit vector within the specified tolerance.
 	 */
 	FORCEINLINE bool IsUnit(float LengthSquaredTolerance = KINDA_SMALL_NUMBER ) const;
@@ -664,7 +648,6 @@ public:
 	 * The FVector will be bogus when InitFromString returns false.
 	 *
 	 * @param	InSourceString	FString containing the vector values.
-	 *
 	 * @return true if the X,Y,Z values were read successfully; false otherwise.
 	 */
 	bool InitFromString( const FString & InSourceString );
@@ -699,7 +682,6 @@ public:
 	 *
 	 * @param P First vector.
 	 * @param Q Second vector.
-	 *
 	 * @return 1=yes, 0=no.  Uses fast distance approximation.
 	 */
 	static bool PointsAreSame( const FVector &P, const FVector &Q );
@@ -710,7 +692,6 @@ public:
 	 * @param Point1 First vector.
 	 * @param Point2 Second vector.
 	 * @param Dist Specified distance.
-	 *
 	 * @return 1=yes, 0=no.  Uses fast distance approximation.
 	 */
 	static bool PointsAreNear( const FVector &Point1, const FVector &Point2, float Dist );
@@ -722,7 +703,6 @@ public:
 	 * @param Point The Point we are checking.
 	 * @param PlaneBase The Base Point in the plane.
 	 * @param PlaneNormal The Normal of the plane.
-	 *
 	 * @return Signed distance  between point and plane.
 	 */
 	static float PointPlaneDist( const FVector &Point, const FVector &PlaneBase, const FVector &PlaneNormal );
@@ -732,7 +712,6 @@ public:
 	 *
 	 * @param Point - the point to project onto the plane
 	 * @param Plane - the plane
-	 *
 	 * @return Projection of Point onto Plane
 	 */
 	static FVector PointPlaneProject(const FVector& Point, const FPlane& Plane);
@@ -742,7 +721,6 @@ public:
 	 *
 	 * @param Point - the point to project onto the plane
 	 * @param A,B,C - three points in CCW order defining the plane 
-	 *
 	 * @return Projection of Point onto plane ABC
 	 */
 	static FVector PointPlaneProject(const FVector& Point, const FVector& A, const FVector& B, const FVector& C);
@@ -753,7 +731,6 @@ public:
 	* @param Point - the point to project onto the plane
 	* @param PlaneBase - point on the plane
 	* @param PlaneNorm - normal of the plane
-	*
 	* @return Projection of Point onto plane ABC
 	*/
 	static FVector PointPlaneProject(const FVector& Point, const FVector& PlaneBase, const FVector& PlaneNorm);
@@ -763,7 +740,6 @@ public:
 	 *
 	 * @param V1 The first point.
 	 * @param V2 The second point.
-	 *
 	 * @return The distance between two points.
 	 */
 	static FORCEINLINE float Dist( const FVector &V1, const FVector &V2 );
@@ -773,7 +749,6 @@ public:
 	 *
 	 * @param V1 The first point.
 	 * @param V2 The second point.
-	 *
 	 * @return The squared distance between two points.
 	 */
 	static FORCEINLINE float DistSquared( const FVector &V1, const FVector &V2 );
@@ -783,7 +758,6 @@ public:
 	 *
 	 * @param Normal The plane normal.
 	 * @param Size The size of the box.
-	 *
 	 * @return Pushout required.
 	 */
 	static FORCEINLINE float BoxPushOut( const FVector & Normal, const FVector & Size );
@@ -793,8 +767,7 @@ public:
 	 *
 	 * @param Normal1 First normalized vector.
 	 * @param Normal1 Second normalized vector.
-	 *
-	 * @return true if vectors are parallel, otherwise false.
+	 * @return true if vectors are parallel, false otherwise..
 	 */
 	static bool Parallel( const FVector &Normal1, const FVector &Normal2 );
 
@@ -805,8 +778,7 @@ public:
 	 * @param Normal1 The normal of the first plane.
 	 * @param Base2 The base point in the second plane.
 	 * @param Normal2 The normal of the second plane.
-	 *
-	 * @return true if the planes are coplanar, otherwise false.
+	 * @return true if the planes are coplanar, false otherwise..
 	 */
 	static bool Coplanar( const FVector &Base1, const FVector &Normal1, const FVector &Base2, const FVector &Normal2 );
 
@@ -816,7 +788,6 @@ public:
 	 * @param X The first vector.
 	 * @param Y The second vector.
 	 * @param Z The third vector.
-	 *
 	 * @return The triple product.
 	 */
 	static float Triple( const FVector& X, const FVector& Y, const FVector& Z );
@@ -824,11 +795,10 @@ public:
 	/**
 	 * Generates a list of sample points on a Bezier curve defined by 2 points.
 	 *
-	 * @param	ControlPoints	Array of 4 FVectors (vert1, controlpoint1, controlpoint2, vert2).
-	 * @param	NumPoints		Number of samples.
-	 * @param	OutPoints		Receives the output samples.
-	 *
-	 * @return					Path length.
+	 * @param ControlPoints	Array of 4 FVectors (vert1, controlpoint1, controlpoint2, vert2).
+	 * @param NumPoints Number of samples.
+	 * @param OutPoints Receives the output samples.
+	 * @return The path length.
 	 */
 	static CORE_API float EvaluateBezier(const FVector* ControlPoints, int32 NumPoints, TArray<FVector>& OutPoints);
 
@@ -848,7 +818,6 @@ public:
 	 *
 	 * @param Ar Serialization Archive.
 	 * @param V Vector to serialize.
-	 *
 	 * @return Reference to Archive after serialization.
 	 */
 	friend FArchive& operator<<( FArchive& Ar, FVector& V )
@@ -859,15 +828,12 @@ public:
 	}
 
 	CORE_API bool NetSerialize(FArchive& Ar, class UPackageMap* Map, bool& bOutSuccess);
-
-	/** A zero vector (0,0,0) */
-	static CORE_API const FVector ZeroVector;
-	/** World up vector (0,0,1) */
-	static CORE_API const FVector UpVector;
-	/** Local Unreal forward vector (1,0,0) */
-	static CORE_API const FVector ForwardVector;
 };
 
+
+/* FVector inline functions
+ *****************************************************************************/
+
 /**
  * Multiplies a vector by a scaling factor.
  *
diff --git a/Engine/Source/Runtime/Core/Public/Math/Vector2D.h b/Engine/Source/Runtime/Core/Public/Math/Vector2D.h
--- a/Engine/Source/Runtime/Core/Public/Math/Vector2D.h
+++ b/Engine/Source/Runtime/Core/Public/Math/Vector2D.h
@@ -1,17 +1,37 @@
 // Copyright 1998-2014 Epic Games, Inc. All Rights Reserved.
 
+/*=============================================================================
+	Vector2D.h: Declares the FVector2D class.
+=============================================================================*/
+
 #pragma once
 
+
 /**
  * A 2x1 of FLOATs.
  */
 struct FVector2D 
 {
-	float	X,
-			Y;
+	/** Holds the vector's X-component. */
+	float X;
+
+	/** Holds the vector's Y-component. */
+	float Y;
+
+public:
+
+	/** Global 2D zero vector constant */
+	static CORE_API const FVector2D ZeroVector;
+
+	/** Global 2D unit vector constant */
+	static CORE_API const FVector2D UnitVector;
 
-	/** Constructor */
-	FORCEINLINE FVector2D();
+public:
+
+	/**
+	 * Default constructor (no initialization).
+	 */
+	FORCEINLINE FVector2D( ) { }
 
 	/**
 	 * Constructor
@@ -19,21 +39,21 @@ struct FVector2D
 	 * @param InX X coordinate.
 	 * @param InY Y coordinate.
 	 */
-	FORCEINLINE FVector2D(float InX,float InY);
+	FORCEINLINE FVector2D( float InX,float InY );
 
 	/**
 	 * Constructor
 	 *
 	 * @param InPos Integer point used to set this vector.
 	 */
-	FORCEINLINE FVector2D(FIntPoint InPos);
+	FORCEINLINE FVector2D( FIntPoint InPos );
 
 	/**
 	 * Constructor
 	 *
 	 * @param EForceInit Force init enum
 	 */
-	explicit FORCEINLINE FVector2D(EForceInit);
+	explicit FORCEINLINE FVector2D( EForceInit );
 
 	/**
 	 * Copy Constructor
@@ -42,11 +62,12 @@ struct FVector2D
 	 */
 	explicit FORCEINLINE FVector2D( const FVector& V );
 
+public:
+
 	/**
 	 * Gets the result of adding two vectors together.
 	 *
 	 * @param V The other vector to add to this.
-	 *
 	 * @return The result of adding the vectors together.
 	 */
 	FORCEINLINE FVector2D operator+( const FVector2D& V ) const;
@@ -55,7 +76,6 @@ struct FVector2D
 	 * Gets the result of subtracting a vector from this one.
 	 *
 	 * @param V The other vector to subtract.
-	 *
 	 * @return The result of the subtraction.
 	 */
 	FORCEINLINE FVector2D operator-( const FVector2D& V ) const;
@@ -64,7 +84,6 @@ struct FVector2D
 	 * Gets the result of scaling this vector.
 	 *
 	 * @param Scale What to scale the vector by.
-	 *
 	 * @return The result of scaling this vector.
 	 */
 	FORCEINLINE FVector2D operator*( float Scale ) const;
@@ -73,7 +92,6 @@ struct FVector2D
 	 * Gets the result of dividing this vector.
 	 *
 	 * @param Scale What to divide the vector by.
-	 *
 	 * @return The result of division on this vector.
 	 */
 	FVector2D operator/( float Scale ) const;
@@ -82,7 +100,6 @@ struct FVector2D
 	 * Gets the result of this vector + float A.
 	 *
 	 * @param A		Float to add.
-	 *
 	 * @return		The result of this vector + float A.
 	 */
 	FORCEINLINE FVector2D operator+( float A ) const;
@@ -91,7 +108,6 @@ struct FVector2D
 	 * Gets the result of this vector - float A
 	 *
 	 * @param A		Float to subtract
-	 *
 	 * @return		The result of this vector - float A.
 	 */
 	FORCEINLINE FVector2D operator-( float A ) const;
@@ -100,7 +116,6 @@ struct FVector2D
 	 * Gets the result of multiplying this vector by another.
 	 *
 	 * @param V The other vector to multiply this by.
-	 *
 	 * @return The result of the multiplication.
 	 */
 	FORCEINLINE FVector2D operator*( const FVector2D& V ) const;
@@ -109,7 +124,6 @@ struct FVector2D
 	 * Gets the result of dividing this vector by another.
 	 *
 	 * @param V The other vector to divide this by.
-	 *
 	 * @return The result of the division.
 	 */
 	FVector2D operator/( const FVector2D& V ) const;
@@ -118,7 +132,6 @@ struct FVector2D
 	 * Calculates Dot product of this vector and another.
 	 *
 	 * @param V The other vector.
-	 *
 	 * @return The Dot product.
 	 */
 	FORCEINLINE float operator|( const FVector2D& V) const;
@@ -127,36 +140,16 @@ struct FVector2D
 	 * Calculates Cross product of this vector and another.
 	 *
 	 * @param V The other vector.
-	 *
 	 * @return The Cross product.
 	 */
 	FORCEINLINE float operator^( const FVector2D& V) const;
 
-	/**
-	 * Calculates the Dot product of two vectors.
-	 *
-	 * @param A The first vector.
-	 * @param B The second vector.
-	 *
-	 * @return The Dot product.
-	 */
-	FORCEINLINE static float DotProduct( const FVector2D& A, const FVector2D& B );
-
-	/**
-	 * Squared distance between two 2D points.
-	 *
-	 * @param V1 The first point.
-	 * @param V2 The second point.
-	 *
-	 * @return The squared distance between two 2D points.
-	 */
-	FORCEINLINE static float DistSquared( const FVector2D& V1, const FVector2D& V2 );
+public:
 
 	/**
 	 * Compares this vector against another for equality.
 	 *
 	 * @param V The vector to compare against.
-	 *
 	 * @return true if the two vectors are equal, otherwise false.
 	 */
 	bool operator==( const FVector2D& V ) const;
@@ -165,7 +158,6 @@ struct FVector2D
 	 * Compares this vector against another for inequality.
 	 *
 	 * @param V The vector to compare against.
-	 *
 	 * @return true if the two vectors are not equal, otherwise false.
 	 */
 	bool operator!=( const FVector2D& V ) const;
@@ -174,7 +166,6 @@ struct FVector2D
 	 * Checks whether the components of this vector are less than another.
 	 *
 	 * @param Other The vector to compare against.
-	 *
 	 * @return true if this is the smaller vector, otherwise false.
 	 */
 	bool operator<( const FVector2D& Other ) const;
@@ -183,7 +174,6 @@ struct FVector2D
 	 * Checks whether the components of this vector are greater than another.
 	 *
 	 * @param Other The vector to compare against.
-	 *
 	 * @return true if this is the larger vector, otherwise false.
 	 */
 	bool operator>( const FVector2D& Other ) const;
@@ -192,7 +182,6 @@ struct FVector2D
 	 * Checks whether the components of this vector are less than or equal to another.
 	 *
 	 * @param Other The vector to compare against.
-	 *
 	 * @return true if this vector is less than or equal to the other vector, otherwise false.
 	 */
 	bool operator<=( const FVector2D& Other ) const;
@@ -201,21 +190,10 @@ struct FVector2D
 	 * Checks whether the components of this vector are greater than or equal to another.
 	 *
 	 * @param Other The vector to compare against.
-	 *
 	 * @return true if this vector is greater than or equal to the other vector, otherwise false.
 	 */
 	bool operator>=( const FVector2D& Other ) const;
 
-	/**
-	 * Checks for equality with error-tolerant comparison.
-	 *
-	 * @param V The vector to compare.
-	 * @param Tolerance Error tolerance.
-	 *
-	 * @return true if the vectors are equal within specified tolerance, otherwise false.
-	 */
-	bool Equals(const FVector2D& V, float Tolerance) const;
-
 	/**
 	 * Gets a negated copy of the vector.
 	 *
@@ -227,7 +205,6 @@ struct FVector2D
 	 * Adds another vector to this.
 	 *
 	 * @param V The other vector to add.
-	 *
 	 * @return Copy of the vector after addition.
 	 */
 	FORCEINLINE FVector2D operator+=( const FVector2D& V );
@@ -236,7 +213,6 @@ struct FVector2D
 	 * Subtracts another vector from this.
 	 *
 	 * @param V The other vector to subtract.
-	 *
 	 * @return Copy of the vector after subtraction.
 	 */
 	FORCEINLINE FVector2D operator-=( const FVector2D& V );
@@ -245,7 +221,6 @@ struct FVector2D
 	 * Scales this vector.
 	 *
 	 * @param Scale The scale to multiply vector by.
-	 *
 	 * @return Copy of the vector after scaling.
 	 */
 	FORCEINLINE FVector2D operator*=( float Scale );
@@ -254,7 +229,6 @@ struct FVector2D
 	 * Divides this vector.
 	 *
 	 * @param V What to divide vector by.
-	 *
 	 * @return Copy of the vector after division.
 	 */
 	FVector2D operator/=( float V );
@@ -263,7 +237,6 @@ struct FVector2D
 	 * Multiplies this vector with another vector.
 	 *
 	 * @param V The vector to multiply with.
-	 *
 	 * @return Copy of the vector after multiplication.
 	 */
 	FVector2D operator*=( const FVector2D& V );
@@ -272,7 +245,6 @@ struct FVector2D
 	 * Divides this vector by another vector.
 	 *
 	 * @param V The vector to divide by.
-	 *
 	 * @return Copy of the vector after division.
 	 */
 	FVector2D operator/=( const FVector2D& V );
@@ -281,7 +253,6 @@ struct FVector2D
 	 * Gets specific component of the vector.
 	 *
 	 * @param Index the index of vector component
-	 *
 	 * @return reference to component.
 	 */
     float& operator[]( int32 Index );
@@ -290,11 +261,39 @@ struct FVector2D
 	 * Gets specific component of the vector.
 	 *
 	 * @param Index the index of vector component
-	 *
 	 * @return copy of component value.
 	 */
 	float operator[]( int32 Index ) const;
 
+public:
+
+	/**
+	 * Calculates the Dot product of two vectors.
+	 *
+	 * @param A The first vector.
+	 * @param B The second vector.
+	 * @return The Dot product.
+	 */
+	FORCEINLINE static float DotProduct( const FVector2D& A, const FVector2D& B );
+
+	/**
+	 * Squared distance between two 2D points.
+	 *
+	 * @param V1 The first point.
+	 * @param V2 The second point.
+	 * @return The squared distance between two 2D points.
+	 */
+	FORCEINLINE static float DistSquared( const FVector2D& V1, const FVector2D& V2 );
+
+	/**
+	 * Checks for equality with error-tolerant comparison.
+	 *
+	 * @param V The vector to compare.
+	 * @param Tolerance Error tolerance.
+	 * @return true if the vectors are equal within specified tolerance, otherwise false.
+	 */
+	bool Equals( const FVector2D& V, float Tolerance ) const;
+
 	/**
 	 * Set the values of the vector directly.
 	 *
@@ -308,73 +307,70 @@ struct FVector2D
 	 *
 	 * @return The maximum of the vectors coordinates.
 	 */
-	float GetMax() const;
+	float GetMax( ) const;
 
 	/**
 	 * Get the absolute maximum of the vectors coordinates.
 	 *
 	 * @return The absolute maximum of the vectors coordinates.
 	 */
-	float GetAbsMax() const;
+	float GetAbsMax( ) const;
 
 	/**
 	 * Get the minimum of the vectors coordinates.
 	 *
 	 * @return The minimum of the vectors coordinates.
 	 */
-	float GetMin() const;
+	float GetMin( ) const;
 
 	/**
 	 * Get the length of this vector.
 	 *
 	 * @return The length of this vector.
 	 */
-	float Size() const;
+	float Size( ) const;
 
 	/**
 	 * Get the squared length of this vector.
 	 *
 	 * @return The squared length of this vector.
 	 */
-	float SizeSquared() const;
+	float SizeSquared( ) const;
 
 	/**
 	 * Get a normalized copy of this vector if it is large enough.
 	 *
 	 * @param Tolerance Minimum squared length of vector for normalization.
-	 *
 	 * @return A normalized copy of the vector if safe, (0,0) otherwise.
 	 */
-	FVector2D SafeNormal(float Tolerance=SMALL_NUMBER) const;
+	FVector2D SafeNormal( float Tolerance=SMALL_NUMBER ) const;
 
 	/**
 	 * Normalize this vector if it is large enough, set it to (0,0) otherwise.
 	 *
 	 * @param Tolerance Minimum squared length of vector for normalization.
 	 */
-	void Normalize(float Tolerance=SMALL_NUMBER);
+	void Normalize( float Tolerance=SMALL_NUMBER );
 
 	/**
 	 * Checks whether vector is near to zero within a specified tolerance.
 	 *
 	 * @param Tolerance Error tolerance.
-	 *
 	 * @return true if vector is in tolerance to zero, otherwise false.
 	 */
-	bool IsNearlyZero(float Tolerance=KINDA_SMALL_NUMBER) const;
+	bool IsNearlyZero( float Tolerance=KINDA_SMALL_NUMBER ) const;
 
 	/**
 	 * Checks whether vector is exactly zero.
 	 *
 	 * @return true if vector is exactly zero, otherwise false.
 	 */
-	bool IsZero() const;
+	bool IsZero( ) const;
 
 	/**
 	 * Gets a specific component of the vector.
 	 *
 	 * @param Index The index of the component required.
-	 *
 	 * @return Reference to the specified component.
 	 */
 	float& Component( int32 Index );
@@ -383,7 +379,6 @@ struct FVector2D
 	 * Gets a specific component of the vector.
 	 *
 	 * @param Index The index of the component required.
-	 *
 	 * @return Copy of the specified component.
 	 */
 	float Component( int32 Index ) const;
@@ -393,7 +388,7 @@ struct FVector2D
 	 *
 	 * @return New Int Point from this vector.
 	 */
-	FIntPoint IntPoint() const;
+	FIntPoint IntPoint( ) const;
 
 	/**
 	 * Creates a copy of this vector with both axes clamped to the given range.
@@ -406,14 +401,13 @@ struct FVector2D
 	 *
 	 * @return Text describing the vector.
 	 */
-	FString ToString() const;
+	FString ToString( ) const;
 
 	/**
 	 * Initialize this Vector based on an FString. The String is expected to contain X=, Y=.
 	 * The FVector2D will be bogus when InitFromString returns false.
 	 *
 	 * @param	InSourceString	FString containing the vector values.
-	 *
 	 * @return true if the X,Y values were read successfully; false otherwise.
 	 */
 	bool InitFromString( const FString & InSourceString );
@@ -423,7 +417,6 @@ struct FVector2D
 	 *
 	 * @param Ar Serialization archive.
 	 * @param V Vector being serialized.
-	 *
 	 * @return Reference to Archive after serialization.
 	 */
 	friend FArchive& operator<<( FArchive& Ar, FVector2D& V )
@@ -434,12 +427,12 @@ struct FVector2D
 	}
 
 #if ENABLE_NAN_DIAGNOSTIC
-	FORCEINLINE void DiagnosticCheckNaN() const
+	FORCEINLINE void DiagnosticCheckNaN( ) const
 	{
 		checkf(!ContainsNaN(), TEXT("FVector contains NaN: %s"), *ToString());
 	}
 #else
-	FORCEINLINE void DiagnosticCheckNaN() const {}
+	FORCEINLINE void DiagnosticCheckNaN( ) const {}
 #endif
 
 	/**
@@ -453,17 +446,15 @@ struct FVector2D
 			FMath::IsNaN(Y) || !FMath::IsFinite(Y));
 	}
 
-	CORE_API bool NetSerialize(FArchive& Ar, class UPackageMap* Map, bool& bOutSuccess);
+	CORE_API bool NetSerialize( FArchive& Ar, class UPackageMap* Map, bool& bOutSuccess );
 
 	/** Converts spherical coordinates on the unit sphere into a cartesian unit length vector. */
-	inline FVector SphericalToUnitCartesian() const;
+	inline FVector SphericalToUnitCartesian( ) const;
+};
 
-	/** Global 2D zero vector constant */
-	static CORE_API const FVector2D ZeroVector;
 
-	/** Global 2D unit vector constant */
-	static CORE_API const FVector2D UnitVector;
-};
+/* FVector2D inline functions
+ *****************************************************************************/
 
 FORCEINLINE FVector2D operator*( float Scale, const FVector2D& V )
 {
@@ -472,12 +463,11 @@ FORCEINLINE FVector2D operator*( float Scale, const FVector2D& V )
 
 template <> struct TIsPODType<FVector2D> { enum { Value = true }; };
 
-FORCEINLINE FVector2D::FVector2D()
-{}
 
 FORCEINLINE FVector2D::FVector2D(float InX,float InY)
 	:	X(InX), Y(InY)
-{}
+{ }
+
 
 FORCEINLINE FVector2D::FVector2D(FIntPoint InPos)
 {
@@ -485,130 +475,155 @@ FORCEINLINE FVector2D::FVector2D(FIntPoint InPos)
 	Y = (float)InPos.Y;
 }
 
+
 FORCEINLINE FVector2D::FVector2D(EForceInit)
 	: X(0), Y(0)
 {
 }
 
+
 FORCEINLINE FVector2D FVector2D::operator+( const FVector2D& V ) const
 {
 	return FVector2D( X + V.X, Y + V.Y );
 }
 
+
 FORCEINLINE FVector2D FVector2D::operator-( const FVector2D& V ) const
 {
 	return FVector2D( X - V.X, Y - V.Y );
 }
 
+
 FORCEINLINE FVector2D FVector2D::operator*( float Scale ) const
 {
 	return FVector2D( X * Scale, Y * Scale );
 }
 
+
 FORCEINLINE FVector2D FVector2D::operator/( float Scale ) const
 {
 	const float RScale = 1.f/Scale;
 	return FVector2D( X * RScale, Y * RScale );
 }
 
+
 FORCEINLINE FVector2D FVector2D::operator+( float A ) const
 {
 	return FVector2D( X + A, Y + A );
 }
 
+
 FORCEINLINE FVector2D FVector2D::operator-( float A ) const
 {
 	return FVector2D( X - A, Y - A );
 }
 
+
 FORCEINLINE FVector2D FVector2D::operator*( const FVector2D& V ) const
 {
 	return FVector2D( X * V.X, Y * V.Y );
 }
 
+
 FORCEINLINE FVector2D FVector2D::operator/( const FVector2D& V ) const
 {
 	return FVector2D( X / V.X, Y / V.Y );
 }
 
+
 FORCEINLINE float FVector2D::operator|( const FVector2D& V) const
 {
 	return X*V.X + Y*V.Y;
 }
 
+
 FORCEINLINE float FVector2D::operator^( const FVector2D& V) const
 {
 	return X*V.Y - Y*V.X;
 }
 
+
 FORCEINLINE float FVector2D::DotProduct( const FVector2D& A, const FVector2D& B )
 {
 	return A | B;
 }
 
+
 FORCEINLINE float FVector2D::DistSquared( const FVector2D &V1, const FVector2D &V2 )
 {
 	return FMath::Square(V2.X-V1.X) + FMath::Square(V2.Y-V1.Y);
 }
 
+
 FORCEINLINE bool FVector2D::operator==( const FVector2D& V ) const
 {
 	return X==V.X && Y==V.Y;
 }
 
+
 FORCEINLINE bool FVector2D::operator!=( const FVector2D& V ) const
 {
 	return X!=V.X || Y!=V.Y;
 }
 
+
 FORCEINLINE bool FVector2D::operator<( const FVector2D& Other ) const
 {
 	return X < Other.X && Y < Other.Y;
 }
 
+
 FORCEINLINE bool FVector2D::operator>( const FVector2D& Other ) const
 {
 	return X > Other.X && Y > Other.Y;
 }
 
+
 FORCEINLINE bool FVector2D::operator<=( const FVector2D& Other ) const
 {
 	return X <= Other.X && Y <= Other.Y;
 }
 
+
 FORCEINLINE bool FVector2D::operator>=( const FVector2D& Other ) const
 {
 	return X >= Other.X && Y >= Other.Y;
 }
 
+
 FORCEINLINE bool FVector2D::Equals(const FVector2D& V, float Tolerance) const
 {
 	return FMath::Abs(X-V.X) < Tolerance && FMath::Abs(Y-V.Y) < Tolerance;
 }
 
+
 FORCEINLINE FVector2D FVector2D::operator-() const
 {
 	return FVector2D( -X, -Y );
 }
 
+
 FORCEINLINE FVector2D FVector2D::operator+=( const FVector2D& V )
 {
 	X += V.X; Y += V.Y;
 	return *this;
 }
 
+
 FORCEINLINE FVector2D FVector2D::operator-=( const FVector2D& V )
 {
 	X -= V.X; Y -= V.Y;
 	return *this;
 }
 
+
 FORCEINLINE FVector2D FVector2D::operator*=( float Scale )
 {
 	X *= Scale; Y *= Scale;
 	return *this;
 }
 
+
 FORCEINLINE FVector2D FVector2D::operator/=( float V )
 {
 	const float RV = 1.f/V;
@@ -616,61 +631,72 @@ FORCEINLINE FVector2D FVector2D::operator/=( float V )
 	return *this;
 }
 
+
 FORCEINLINE FVector2D FVector2D::operator*=( const FVector2D& V )
 {
 	X *= V.X; Y *= V.Y;
 	return *this;
 }
 
+
 FORCEINLINE FVector2D FVector2D::operator/=( const FVector2D& V )
 {
 	X /= V.X; Y /= V.Y;
 	return *this;
 }
 
+
 FORCEINLINE float& FVector2D::operator[]( int32 Index )
 {
 	check(Index>=0 && Index<2);
 	return ((Index == 0) ? X : Y);
 }
 
+
 FORCEINLINE float FVector2D::operator[]( int32 Index ) const
 {
 	check(Index>=0 && Index<2);
 	return ((Index == 0) ? X : Y);
 }
 
+
 FORCEINLINE void FVector2D::Set( float InX, float InY )
 {
 	X = InX;
 	Y = InY;
 }
 
+
 FORCEINLINE float FVector2D::GetMax() const
 {
 	return FMath::Max(X,Y);
 }
 
+
 FORCEINLINE float FVector2D::GetAbsMax() const
 {
 	return FMath::Max(FMath::Abs(X),FMath::Abs(Y));
 }
 
+
 FORCEINLINE float FVector2D::GetMin() const
 {
 	return FMath::Min(X,Y);
 }
 
+
 FORCEINLINE float FVector2D::Size() const
 {
 	return FMath::Sqrt( X*X + Y*Y );
 }
 
+
 FORCEINLINE float FVector2D::SizeSquared() const
 {
 	return X*X + Y*Y;
 }
 
+
 FORCEINLINE FVector2D FVector2D::SafeNormal(float Tolerance) const
 {	
 	const float SquareSum = X*X + Y*Y;
@@ -682,6 +708,7 @@ FORCEINLINE FVector2D FVector2D::SafeNormal(float Tolerance) const
 	return FVector2D(0.f, 0.f);
 }
 
+
 FORCEINLINE void FVector2D::Normalize(float Tolerance)
 {
 	const float SquareSum = X*X + Y*Y;
@@ -696,42 +723,50 @@ FORCEINLINE void FVector2D::Normalize(float Tolerance)
 	Y = 0.0f;
 }
 
+
 FORCEINLINE bool FVector2D::IsNearlyZero(float Tolerance) const
 {
 	return	FMath::Abs(X)<Tolerance
 		&&	FMath::Abs(Y)<Tolerance;
 }
 
+
 FORCEINLINE bool FVector2D::IsZero() const
 {
 	return X==0.f && Y==0.f;
 }
 
+
 FORCEINLINE float& FVector2D::Component( int32 Index )
 {
 	return (&X)[Index];
 }
 
+
 FORCEINLINE float FVector2D::Component( int32 Index ) const
 {
 	return (&X)[Index];
 }
 
+
 FORCEINLINE FIntPoint FVector2D::IntPoint() const
 {
 	return FIntPoint( FMath::RoundToInt(X), FMath::RoundToInt(Y) );
 }
 
+
 FORCEINLINE FVector2D FVector2D::ClampAxes( float MinAxisVal, float MaxAxisVal ) const
 {
 	return FVector2D( FMath::Clamp(X, MinAxisVal, MaxAxisVal), FMath::Clamp(Y, MinAxisVal, MaxAxisVal) );
 }
 
+
 FORCEINLINE FString FVector2D::ToString() const
 {
 	return FString::Printf(TEXT("X=%3.3f Y=%3.3f"), X, Y);
 }
 
+
 FORCEINLINE bool FVector2D::InitFromString( const FString & InSourceString )
 {
 	X = Y = 0;
@@ -741,5 +776,3 @@ FORCEINLINE bool FVector2D::InitFromString( const FString & InSourceString )
 
 	return bSuccessful;
 }
-
-
diff --git a/Engine/Source/Runtime/Core/Public/Math/Vector2DHalf.h b/Engine/Source/Runtime/Core/Public/Math/Vector2DHalf.h
--- a/Engine/Source/Runtime/Core/Public/Math/Vector2DHalf.h
+++ b/Engine/Source/Runtime/Core/Public/Math/Vector2DHalf.h
@@ -1,58 +1,80 @@
 // Copyright 1998-2014 Epic Games, Inc. All Rights Reserved.
 
+/*=============================================================================
+	Vector2DHalf.h: Declares the FVector2DHalf class.
+=============================================================================*/
+
 #pragma once
 
+
 /**
-* A 2x1 of 16 bit floats.
-*/
+ * Structure for two dimensional vectors with half floating point precision.
+ */
 struct FVector2DHalf 
 {
-	FFloat16	X,
-				Y;
+	/** Holds the vector's X-component. */
+	FFloat16 X;
+
+	/** Holds the vector's Y-component. */
+	FFloat16 Y;
 
-	/** Default Constructor */
-	FORCEINLINE FVector2DHalf();
+public:
+
+	/**
+	 * Default Constructor (no initialization).
+	 */
+	FORCEINLINE FVector2DHalf( ) { }
 
 	/** Constructor 
 	 *
 	 * InX half float X value
 	 * Iny half float Y value
 	 */
- 	FORCEINLINE FVector2DHalf(const FFloat16& InX,const FFloat16& InY);
+ 	FORCEINLINE FVector2DHalf( const FFloat16& InX,const FFloat16& InY );
 
 	/** Constructor 
 	 *
 	 * InX float X value
 	 * Iny float Y value
 	 */
-	FORCEINLINE FVector2DHalf(float InX,float InY);
+	FORCEINLINE FVector2DHalf( float InX,float InY );
 
 	/** Constructor 
 	 *
 	 * Vector2D float vector
 	 */
-	FORCEINLINE FVector2DHalf(const FVector2D& Vector2D);
+	FORCEINLINE FVector2DHalf( const FVector2D& Vector2D );
+
+public:
+
+	/**
+	 * Assignment operator.
+	 *
+	 * @param Vector2D The value to assign.
+	 */
+ 	FVector2DHalf& operator=( const FVector2D& Vector2D );
 
-	/** assignment operator */
- 	FVector2DHalf& operator=(const FVector2D& Vector2D);
+	/**
+	 * Implicit conversion operator for conversion to FVector2D.
+	 */
+	operator FVector2D( ) const;
+
+public:
 
 	/**
 	 * Get a textual representation of the vector.
 	 *
 	 * @return Text describing the vector.
 	 */
-	FString ToString() const;
+	FString ToString( ) const;
 
-
-	/** Convert from FVector2D to FVector2DHalf. */
-	operator FVector2D() const;
+public:
 
 	/**
 	 * Serializes the FVector2DHalf.
 	 *
 	 * @param Ar Reference to the serialization archive.
 	 * @param V Reference to the FVector2DHalf being serialized.
-	 *
 	 * @return Reference to the Archive after serialization.
 	 */
 	friend FArchive& operator<<( FArchive& Ar, FVector2DHalf& V )
@@ -61,34 +83,40 @@ struct FVector2DHalf
 	}
 };
 
-FORCEINLINE FVector2DHalf::FVector2DHalf()
-{}
 
-FORCEINLINE FVector2DHalf::FVector2DHalf(const FFloat16& InX,const FFloat16& InY)
+/* FVector2DHalf inline functions
+ *****************************************************************************/
+
+FORCEINLINE FVector2DHalf::FVector2DHalf( const FFloat16& InX, const FFloat16& InY )
  	:	X(InX), Y(InY)
-{}
+{ }
 
-FORCEINLINE FVector2DHalf::FVector2DHalf(float InX,float InY)
+FORCEINLINE FVector2DHalf::FVector2DHalf( float InX, float InY )
 	:	X(InX), Y(InY)
-{}
+{ }
 
-FORCEINLINE FVector2DHalf::FVector2DHalf(const FVector2D& Vector2D)
+
+FORCEINLINE FVector2DHalf::FVector2DHalf( const FVector2D& Vector2D )
 	:	X(Vector2D.X), Y(Vector2D.Y)
-{}
+{ }
+
 
-FORCEINLINE FVector2DHalf& FVector2DHalf::operator=(const FVector2D& Vector2D)
+FORCEINLINE FVector2DHalf& FVector2DHalf::operator=( const FVector2D& Vector2D )
 {
  	X = FFloat16(Vector2D.X);
  	Y = FFloat16(Vector2D.Y);
+
 	return *this;
 }
 
-FORCEINLINE FString FVector2DHalf::ToString() const
+
+FORCEINLINE FString FVector2DHalf::ToString( ) const
 {
 	return FString::Printf(TEXT("X=%3.3f Y=%3.3f"), (float)X, (float)Y );
 }
 
-FORCEINLINE FVector2DHalf::operator FVector2D() const
+
+FORCEINLINE FVector2DHalf::operator FVector2D( ) const
 {
 	return FVector2D((float)X,(float)Y);
 }
diff --git a/Engine/Source/Runtime/Core/Public/Math/Vector4.h b/Engine/Source/Runtime/Core/Public/Math/Vector4.h
--- a/Engine/Source/Runtime/Core/Public/Math/Vector4.h
+++ b/Engine/Source/Runtime/Core/Public/Math/Vector4.h
@@ -1,24 +1,32 @@
 // Copyright 1998-2014 Epic Games, Inc. All Rights Reserved.
 
+/*=============================================================================
+	Vector.h: Declares the FVector class.
+=============================================================================*/
+
 #pragma once
 
+
 /**
  * A 4D homogeneous vector, 4x1 FLOATs, 16-byte aligned.
  */
 MS_ALIGN(16) class FVector4
 {
 public:
-	// Variables.
-	float X, Y, Z, W;
 
-#if ENABLE_NAN_DIAGNOSTIC
-	FORCEINLINE void DiagnosticCheckNaN() const
-	{
-		checkf(!ContainsNaN(), TEXT("FVector contains NaN: %s"), *ToString());
-	}
-#else
-	FORCEINLINE void DiagnosticCheckNaN() const {}
-#endif
+	/** Holds the vector's X-component. */
+	float X;
+
+	/** Holds the vector's Y-component. */
+	float Y;
+
+	/** Holds the vector's Z-component. */
+	float Z;
+
+	/** Holds the vector's W-component. */
+	float W;
+
+public:
 
 	/**
 	 * Constructor.
@@ -26,33 +34,41 @@ public:
 	 * @param InVector 3D Vector to set first three components.
 	 * @param InW W Coordinate.
 	 */
-	FVector4(const FVector& InVector,float InW = 1.0f);
+	FVector4( const FVector& InVector, float InW = 1.0f );
 
 	/**
-	 * Constructor.
+	 * Creates and initializes a new vector from a color value.
 	 *
-	 * @param InColour Colour used to set vector.
+	 * @param InColour Color used to set vector.
 	 */
-	FVector4(const FLinearColor& InColor);
+	FVector4( const FLinearColor& InColor );
 
 	/**
-	 * Constructor.
+	 * Creates and initializes a new vector from the specified components.
 	 *
 	 * @param InX X Coordinate.
 	 * @param InY Y Coordinate.
 	 * @param InZ Z Coordinate.
 	 * @param InW W Coordinate.
 	 */
-	explicit FVector4(float InX = 0.0f,float InY = 0.0f,float InZ = 0.0f,float InW = 1.0f);
+	explicit FVector4( float InX = 0.0f, float InY = 0.0f, float InZ = 0.0f, float InW = 1.0f );
 
-	explicit FVector4(FVector2D InXY, FVector2D InZW);
+	/**
+	 * Creates and initializes a new vector from the specified 2D vectors.
+	 *
+	 * @param InXY A 2D vector holding the X- and Y-components.
+	 * @param InZW A 2D vector holding the Z- and W-components.
+	 */
+	explicit FVector4( FVector2D InXY, FVector2D InZW );
 
 	/**
-	 * Constructor.
+	 * Creates and initializes a new vector to zero.
 	 *
 	 * @param EForceInit Force Init Enum.
 	 */
-	explicit FVector4(EForceInit);
+	explicit FVector4( EForceInit );
+
+public:
 
 	/**
 	 * Access a specific component of the vector.
@@ -72,16 +88,6 @@ public:
 	 */
 	FORCEINLINE float operator[]( int32 ComponentIndex ) const;
 
-	/**
-	 * Set all of the vectors coordinates.
-	 *
-	 * @param InX New X Coordinate.
-	 * @param InY New Y Coordinate.
-	 * @param InZ New Z Coordinate.
-	 * @param InW New W Coordinate.
-	 */
-	FORCEINLINE void Set( float InX, float InY, float InZ, float InW );
-
 	// Unary operators.
 
 	/**
@@ -89,7 +95,7 @@ public:
 	 *
 	 * @return A negated copy of the vector.
 	 */
-	FORCEINLINE FVector4 operator-() const;
+	FORCEINLINE FVector4 operator-( ) const;
 
 	/**
 	 * Gets the result of adding a vector to this.
@@ -104,7 +110,6 @@ public:
 	 * Adds another vector to this one.
 	 *
 	 * @param V The other vector to add.
-	 *
 	 * @return Copy of the vector after addition.
 	 */
 	FORCEINLINE FVector4 operator+=( const FVector4& V );
@@ -113,7 +118,6 @@ public:
 	 * Gets the result of subtracting a vector from this.
 	 *
 	 * @param V The vector to subtract.
-	 *
 	 * @return The result of vector subtraction.
 	 */
 	FORCEINLINE FVector4 operator-( const FVector4& V ) const;
@@ -122,7 +126,6 @@ public:
 	 * Gets the result of scaling this vector.
 	 *
 	 * @param Scale The scaling factor.
-	 *
 	 * @return The result of vector scaling.
 	 */
 	FORCEINLINE FVector4 operator*( float Scale ) const;
@@ -131,7 +134,6 @@ public:
 	 * Gets the result of dividing this vector.
 	 *
 	 * @param Scale What to divide by.
-	 *
 	 * @return The result of division.
 	 */
 	FVector4 operator/( float Scale ) const;
@@ -140,7 +142,6 @@ public:
 	 * Gets the result of dividing this vector.
 	 *
 	 * @param V What to divide by.
-	 *
 	 * @return The result of division.
 	 */
 	FVector4 operator/( const FVector4& V ) const;
@@ -149,7 +150,6 @@ public:
 	 * Gets the result of multiplying a vector with this.
 	 *
 	 * @param V The vector to multiply.
-	 *
 	 * @return The result of vector multiplication.
 	 */
 	FVector4 operator*( const FVector4& V ) const;
@@ -158,7 +158,6 @@ public:
 	 * Gets the result of multiplying a vector with another Vector (component wise).
 	 *
 	 * @param V The vector to multiply.
-	 *
 	 * @return The result of vector multiplication.
 	 */
 	FVector4 operator*=( const FVector4& V );
@@ -167,7 +166,6 @@ public:
 	 * Gets the result of dividing a vector with another Vector (component wise).
 	 *
 	 * @param V The vector to divide with.
-	 *
 	 * @return The result of vector multiplication.
 	 */
 	FVector4 operator/=( const FVector4& V );
@@ -176,28 +174,15 @@ public:
 	 * Gets the result of scaling this vector.
 	 *
 	 * @param Scale The scaling factor.
-	 *
 	 * @return The result of vector scaling.
 	 */
 	FVector4 operator*=( float S );
 
-	// Simple functions.
-
-	/**
-	 * Gets a specific component of the vector.
-	 *
-	 * @param Index The index of the component.
-	 *
-	 * @return Reference to the component.
-	 */
-	float& Component( int32 Index );
-
 	/**
 	 * Calculates 3D Dot product of two 4D vectors.
 	 *
 	 * @param V1 The first vector.
 	 * @param V2 The second vector.
-	 *
 	 * @return The 3D Dot product.
 	 */
 	friend FORCEINLINE float Dot3( const FVector4& V1, const FVector4& V2 )
@@ -210,7 +195,6 @@ public:
 	 *
 	 * @param V1 The first vector.
 	 * @param V2 The second vector.
-	 *
 	 * @return The 4D Dot Product.
 	 */
 	friend FORCEINLINE float Dot4( const FVector4& V1, const FVector4& V2 )
@@ -223,7 +207,6 @@ public:
 	 *
 	 * @param Scale The scaling factor.
 	 * @param V The vector to scale.
-	 *
 	 * @return The result of scaling.
 	 */
 	friend FORCEINLINE FVector4 operator*( float Scale, const FVector4& V )
@@ -231,70 +214,71 @@ public:
 		return V.operator*( Scale );
 	}
 
-
-	/**
-	 * Basic == or != operators for FQuat
-	 */	
-
 	/**
 	 * Checks for equality against another vector.
 	 *
 	 * @param V The other vector.
-	 *
 	 * @return true if the two vectors are the same, otherwise false.
 	 */
-	bool operator==(const FVector4& V) const;
+	bool operator==( const FVector4& V ) const;
 
 	/**
 	 * Checks for inequality against another vector.
 	 *
 	 * @param V The other vector.
-	 *
 	 * @return true if the two vectors are different, otherwise false.
 	 */
-	bool operator!=(const FVector4& V) const;
+	bool operator!=( const FVector4& V ) const;
 
 	/**
 	 * Calculate Cross product between this and another vector.
 	 *
 	 * @param V The other vector.
-	 *
 	 * @return The Cross product.
 	 */
 	FVector4 operator^( const FVector4& V ) const;
 
+public:
+
+	// Simple functions.
+
+	/**
+	 * Gets a specific component of the vector.
+	 *
+	 * @param Index The index of the component.
+	 * @return Reference to the component.
+	 */
+	float& Component( int32 Index );
+
 	/**
 	 * Error tolerant comparison.
 	 *
 	 * @param V Vector to compare against.
 	 * @param Tolerance Error Tolerance.
-	 *
 	 * @return true if the two vectors are equal within specified tolerance, otherwise false.
 	 */
-	bool Equals(const FVector4& V, float Tolerance=KINDA_SMALL_NUMBER) const;
+	bool Equals( const FVector4& V, float Tolerance=KINDA_SMALL_NUMBER ) const;
 
 	/**
 	 * Check if the vector is of unit length, with specified tolerance.
 	 *
 	 * @param LengthSquaredTolerance Tolerance against squared length.
-	 *
 	 * @return true if the vector is a unit vector within the specified tolerance.
 	 */
-	bool IsUnit3(float LengthSquaredTolerance = KINDA_SMALL_NUMBER) const;
+	bool IsUnit3( float LengthSquaredTolerance = KINDA_SMALL_NUMBER ) const;
 
 	/**
 	 * Get a textual representation of the vector.
 	 *
 	 * @return Text describing the vector.
 	 */
-	FString ToString() const;
+	FString ToString( ) const;
 
 	/**
 	 * Initialize this Vector based on an FString. The String is expected to contain X=, Y=, Z=, W=.
 	 * The FVector4 will be bogus when InitFromString returns false.
 	 *
-	 * @param	InSourceString	FString containing the vector values.
-	 *
+	 * @param InSourceString	FString containing the vector values.
 	 * @return true if the X,Y,Z values were read successfully; false otherwise.
 	 */
 	bool InitFromString( const FString & InSourceString );
@@ -303,17 +287,16 @@ public:
 	 * Returns a normalized copy of the vector if safe to normalize.
 	 *
 	 * @param Tolerance Minimum squared length of vector for normalization.
-	 *
 	 * @return A normalized copy of the vector or a zero vector.
 	 */
-	FORCEINLINE FVector4 SafeNormal(float Tolerance=SMALL_NUMBER) const;
+	FORCEINLINE FVector4 SafeNormal( float Tolerance=SMALL_NUMBER ) const;
 
 	/**
 	 * Calculates normalized version of vector without checking if it is non-zero.
 	 *
 	 * @return Normalized version of vector.
 	 */
-	FORCEINLINE FVector4 UnsafeNormal3() const;
+	FORCEINLINE FVector4 UnsafeNormal3( ) const;
 
 	/**
 	 * Return the FRotator corresponding to the direction that the vector
@@ -322,27 +305,24 @@ public:
 	 *
 	 * @return The FRotator of the vector's direction.
 	 */
-	CORE_API FRotator Rotation() const;
+	CORE_API FRotator Rotation( ) const;
 
 	/**
-	 * Serializer.
-	 *
-	 * @param Ar The Serialization Archive.
-	 * @param V The vector being serialized.
+	 * Set all of the vectors coordinates.
 	 *
-	 * @return Reference to the Archive after serialization.
+	 * @param InX New X Coordinate.
+	 * @param InY New Y Coordinate.
+	 * @param InZ New Z Coordinate.
+	 * @param InW New W Coordinate.
 	 */
-	friend FArchive& operator<<( FArchive& Ar, FVector4& V )
-	{
-		return Ar << V.X << V.Y << V.Z << V.W;
-	}
+	FORCEINLINE void Set( float InX, float InY, float InZ, float InW );
 
 	/**
 	 * Get the length of this vector not taking W component into account.
 	 *
 	 * @return The length of this vector.
 	 */
-	float Size3() const;
+	float Size3( ) const;
 
 	/**
 	 * Get the squared length of this vector not taking W component into account.
@@ -355,10 +335,10 @@ public:
 	bool ContainsNaN() const;
 
 	/** Utility to check if all of the components of this vector are nearly zero given the tolerance. */
-	bool IsNearlyZero3(float Tolerance=KINDA_SMALL_NUMBER) const;
+	bool IsNearlyZero3( float Tolerance = KINDA_SMALL_NUMBER ) const;
 
 	/** Reflect vector. */
-	FVector4 Reflect3(const FVector4& Normal) const;
+	FVector4 Reflect3( const FVector4& Normal ) const;
 
 	/**
 	 * Find good arbitrary axis vectors to represent U and V axes of a plane,
@@ -366,48 +346,97 @@ public:
 	 */
 	void FindBestAxisVectors3( FVector4& Axis1, FVector4& Axis2 ) const;
 
+#if ENABLE_NAN_DIAGNOSTIC
+	FORCEINLINE void DiagnosticCheckNaN( ) const
+	{
+		checkf(!ContainsNaN(), TEXT("FVector contains NaN: %s"), *ToString());
+	}
+#else
+	FORCEINLINE void DiagnosticCheckNaN( ) const { }
+#endif
+
+public:
+
+	/**
+	 * Serializer.
+	 *
+	 * @param Ar The Serialization Archive.
+	 * @param V The vector being serialized.
+	 * @return Reference to the Archive after serialization.
+	 */
+	friend FArchive& operator<<( FArchive& Ar, FVector4& V )
+	{
+		return Ar << V.X << V.Y << V.Z << V.W;
+	}
+
 } GCC_ALIGN(16);
 
-FORCEINLINE FVector4::FVector4(const FVector& InVector,float InW):
-X(InVector.X), Y(InVector.Y), Z(InVector.Z), W(InW)
+
+/* FVector4 inline functions
+ *****************************************************************************/
+
+FORCEINLINE FVector4::FVector4( const FVector& InVector,float InW )
+	: X(InVector.X)
+	, Y(InVector.Y)
+	, Z(InVector.Z)
+	, W(InW)
 {
 	DiagnosticCheckNaN();
 }
 
-FORCEINLINE FVector4::FVector4(const FLinearColor& InColor):
-X(InColor.R), Y(InColor.G), Z(InColor.B), W(InColor.A)
+
+FORCEINLINE FVector4::FVector4( const FLinearColor& InColor )
+	: X(InColor.R)
+	, Y(InColor.G)
+	, Z(InColor.B)
+	, W(InColor.A)
 {
 	DiagnosticCheckNaN();
 }
 
-FORCEINLINE FVector4::FVector4(float InX,float InY,float InZ,float InW):
-X(InX), Y(InY), Z(InZ), W(InW)
+
+FORCEINLINE FVector4::FVector4( float InX,float InY,float InZ,float InW )
+	: X(InX)
+	, Y(InY)
+	, Z(InZ)
+	, W(InW)
 {
 	DiagnosticCheckNaN();
 }
 
-FORCEINLINE FVector4::FVector4(EForceInit)
-	: X(0.f), Y(0.f), Z(0.f), W(0.f)
+
+FORCEINLINE FVector4::FVector4( EForceInit )
+	: X(0.f)
+	, Y(0.f)
+	, Z(0.f)
+	, W(0.f)
 {
 	DiagnosticCheckNaN();
 }
 
-FORCEINLINE FVector4::FVector4(FVector2D InXY, FVector2D InZW)
-	: X(InXY.X), Y(InXY.Y), Z(InZW.X), W(InZW.Y)
+
+FORCEINLINE FVector4::FVector4( FVector2D InXY, FVector2D InZW )
+	: X(InXY.X)
+	, Y(InXY.Y)
+	, Z(InZW.X)
+	, W(InZW.Y)
 {
 	DiagnosticCheckNaN();
 }
 
+
 FORCEINLINE float & FVector4::operator[]( int32 ComponentIndex )
 {
 	return (&X)[ ComponentIndex ];
 }
 
+
 FORCEINLINE float FVector4::operator[]( int32 ComponentIndex ) const
 {
 	return (&X)[ ComponentIndex ];
 }
 
+
 FORCEINLINE void FVector4::Set( float InX, float InY, float InZ, float InW )
 {
 	X = InX;
@@ -419,16 +448,18 @@ FORCEINLINE void FVector4::Set( float InX, float InY, float InZ, float InW )
 }
 
 
-FORCEINLINE FVector4 FVector4::operator-() const
+FORCEINLINE FVector4 FVector4::operator-( ) const
 {
 	return FVector4( -X, -Y, -Z, -W );
 }
 
+
 FORCEINLINE FVector4 FVector4::operator+( const FVector4& V ) const
 {
 	return FVector4( X + V.X, Y + V.Y, Z + V.Z, W + V.W );
 }
 
+
 FORCEINLINE FVector4 FVector4::operator+=( const FVector4& V )
 {
 	X += V.X; Y += V.Y; Z += V.Z; W += V.W;
@@ -436,63 +467,73 @@ FORCEINLINE FVector4 FVector4::operator+=( const FVector4& V )
 	return *this;
 }
 
+
 FORCEINLINE FVector4 FVector4::operator-( const FVector4& V ) const
 {
 	return FVector4( X - V.X, Y - V.Y, Z - V.Z, W - V.W );
 }
 
+
 FORCEINLINE FVector4 FVector4::operator*( float Scale ) const
 {
 	return FVector4( X * Scale, Y * Scale, Z * Scale, W * Scale );
 }
 
+
 FORCEINLINE FVector4 FVector4::operator/( float Scale ) const
 {
 	const float RScale = 1.f/Scale;
 	return FVector4( X * RScale, Y * RScale, Z * RScale, W * RScale );
 }
 
+
 FORCEINLINE FVector4 FVector4::operator*( const FVector4& V ) const
 {
 	return FVector4( X * V.X, Y * V.Y, Z * V.Z, W * V.W );
 }
 
+
 FORCEINLINE FVector4 FVector4::operator^( const FVector4& V ) const
 {
-	return FVector4
-		(
+	return FVector4(
 		Y * V.Z - Z * V.Y,
 		Z * V.X - X * V.Z,
 		X * V.Y - Y * V.X,
 		0.0f
-		);
+	);
 }
 
+
 FORCEINLINE float& FVector4::Component( int32 Index )
 {
 	return (&X)[Index];
 }
 
-FORCEINLINE bool FVector4::operator==(const FVector4& V) const
+
+FORCEINLINE bool FVector4::operator==( const FVector4& V ) const
 {
 	return ((X == V.X) && (Y == V.Y) && (Z == V.Z) && (W == V.W));
 }
 
-FORCEINLINE bool FVector4::operator!=(const FVector4& V) const
+
+FORCEINLINE bool FVector4::operator!=( const FVector4& V ) const
 {
 	return ((X != V.X) || (Y != V.Y) || (Z != V.Z) || (W != V.W));
 }
 
-FORCEINLINE bool FVector4::Equals(const FVector4& V, float Tolerance) const
+
+FORCEINLINE bool FVector4::Equals( const FVector4& V, float Tolerance ) const
 {
 	return FMath::Abs(X-V.X) < Tolerance && FMath::Abs(Y-V.Y) < Tolerance && FMath::Abs(Z-V.Z) < Tolerance && FMath::Abs(W-V.W) < Tolerance;
 }
 
-FORCEINLINE FString FVector4::ToString() const
+
+FORCEINLINE FString FVector4::ToString( ) const
 {
 	return FString::Printf(TEXT("X=%3.3f Y=%3.3f Z=%3.3f W=%3.3f"), X, Y, Z, W);
 }
 
+
 FORCEINLINE bool FVector4::InitFromString( const FString & InSourceString )
 {
 	X = Y = Z = 0;
@@ -507,7 +548,8 @@ FORCEINLINE bool FVector4::InitFromString( const FString & InSourceString )
 	return bSuccessful;
 }
 
-FORCEINLINE FVector4 FVector4::SafeNormal(float Tolerance) const
+
+FORCEINLINE FVector4 FVector4::SafeNormal( float Tolerance ) const
 {
 	const float SquareSum = X*X + Y*Y + Z*Z;
 	if( SquareSum > Tolerance )
@@ -518,28 +560,33 @@ FORCEINLINE FVector4 FVector4::SafeNormal(float Tolerance) const
 	return FVector4(0.f);
 }
 
-FORCEINLINE FVector4 FVector4::UnsafeNormal3() const
+
+FORCEINLINE FVector4 FVector4::UnsafeNormal3( ) const
 {
 	const float Scale = FMath::InvSqrt(X*X+Y*Y+Z*Z);
 	return FVector4( X*Scale, Y*Scale, Z*Scale, 0.0f );
 }
 
-FORCEINLINE float FVector4::Size3() const
+
+FORCEINLINE float FVector4::Size3( ) const
 {
 	return FMath::Sqrt( X*X + Y*Y + Z*Z );
-
 }
-FORCEINLINE float FVector4::SizeSquared3() const
+
+
+FORCEINLINE float FVector4::SizeSquared3( ) const
 {
 	return X*X + Y*Y + Z*Z;
 }
 
-FORCEINLINE bool FVector4::IsUnit3(float LengthSquaredTolerance) const
+
+FORCEINLINE bool FVector4::IsUnit3( float LengthSquaredTolerance ) const
 {
 	return FMath::Abs(1.0f - SizeSquared3()) < LengthSquaredTolerance;
 }
 
-FORCEINLINE bool FVector4::ContainsNaN() const
+
+FORCEINLINE bool FVector4::ContainsNaN( ) const
 {
 	return (FMath::IsNaN(X) || !FMath::IsFinite(X) || 
 			FMath::IsNaN(Y) || !FMath::IsFinite(Y) ||
@@ -547,7 +594,8 @@ FORCEINLINE bool FVector4::ContainsNaN() const
 			FMath::IsNaN(W) || !FMath::IsFinite(W));
 }
 
-FORCEINLINE bool FVector4::IsNearlyZero3(float Tolerance) const
+
+FORCEINLINE bool FVector4::IsNearlyZero3( float Tolerance ) const
 {
 	return
 			FMath::Abs(X)<Tolerance
@@ -555,11 +603,13 @@ FORCEINLINE bool FVector4::IsNearlyZero3(float Tolerance) const
 		&&	FMath::Abs(Z)<Tolerance;
 }
 
-FORCEINLINE FVector4 FVector4::Reflect3(const FVector4& Normal) const
+
+FORCEINLINE FVector4 FVector4::Reflect3( const FVector4& Normal ) const
 {
 	return 2.0f * Dot3(*this, Normal) * Normal - *this;
 }
 
+
 FORCEINLINE void FVector4::FindBestAxisVectors3( FVector4& Axis1, FVector4& Axis2 ) const
 {
 	const float NX = FMath::Abs(X);
@@ -574,6 +624,7 @@ FORCEINLINE void FVector4::FindBestAxisVectors3( FVector4& Axis1, FVector4& Axis
 	Axis2 = Axis1 ^ *this;
 }
 
+
 FORCEINLINE FVector4 FVector4::operator*=( const FVector4& V )
 {
 	X *= V.X; Y *= V.Y; Z *= V.Z; W *= V.W;
@@ -581,6 +632,7 @@ FORCEINLINE FVector4 FVector4::operator*=( const FVector4& V )
 	return *this;
 }
 
+
 FORCEINLINE FVector4 FVector4::operator/=( const FVector4& V )
 {
 	X /= V.X; Y /= V.Y; Z /= V.Z; W /= V.W;
@@ -588,6 +640,7 @@ FORCEINLINE FVector4 FVector4::operator/=( const FVector4& V )
 	return *this;
 }
 
+
 FORCEINLINE FVector4 FVector4::operator*=( float S )
 {
 	X *= S; Y *= S; Z *= S; W *= S;
@@ -595,7 +648,8 @@ FORCEINLINE FVector4 FVector4::operator*=( float S )
 	return *this;
 }
 
+
 FORCEINLINE FVector4 FVector4::operator/( const FVector4& V ) const
 {
 	return FVector4( X / V.X, Y / V.Y, Z / V.Z, W / V.W );
-}
\ No newline at end of file
+}
diff --git a/Engine/Source/Runtime/Core/Public/Misc/Guid.h b/Engine/Source/Runtime/Core/Public/Misc/Guid.h
--- a/Engine/Source/Runtime/Core/Public/Misc/Guid.h
+++ b/Engine/Source/Runtime/Core/Public/Misc/Guid.h
@@ -79,10 +79,10 @@ public:
 	/**
 	 * Creates and initializes a new GUID from the specified components.
 	 *
-	 * @param InA - The first component.
-	 * @param InB - The second component.
-	 * @param InC - The third component.
-	 * @param InD - The fourth component.
+	 * @param InA The first component.
+	 * @param InB The second component.
+	 * @param InC The third component.
+	 * @param InD The fourth component.
 	 */
 	FGuid( uint32 InA, uint32 InB, uint32 InC, uint32 InD )
 		: A(InA), B(InB), C(InC), D(InD)
@@ -93,9 +93,8 @@ public:
 	/**
 	 * Compares two GUIDs for equality.
 	 *
-	 * @param X - The first GUID to compare.
-	 * @param Y - The second GUID to compare.
-	 *
+	 * @param X The first GUID to compare.
+	 * @param Y The second GUID to compare.
 	 * @return true if the GUIDs are equal, false otherwise.
 	 */
 	friend bool operator==( const FGuid& X, const FGuid& Y )
@@ -106,9 +105,8 @@ public:
 	/**
 	 * Compares two GUIDs for inequality.
 	 *
-	 * @param X - The first GUID to compare.
-	 * @param Y - The second GUID to compare.
-	 *
+	 * @param X The first GUID to compare.
+	 * @param Y The second GUID to compare.
 	 * @return true if the GUIDs are not equal, false otherwise.
 	 */
 	friend bool operator!=( const FGuid& X, const FGuid& Y )
@@ -119,9 +117,8 @@ public:
 	/**
 	 * Compares two GUIDs.
 	 *
-	 * @param X - The first GUID to compare.
-	 * @param Y - The second GUID to compare.
-	 *
+	 * @param X The first GUID to compare.
+	 * @param Y The second GUID to compare.
 	 * @return true if the first GUID is less than the second one.
 	 */
 	friend bool operator<( const FGuid& X, const FGuid& Y )
@@ -135,8 +132,7 @@ public:
 	/**
 	 * Provides access to the GUIDs components.
 	 *
-	 * @param Index - The index of the component to return (0...3).
-	 *
+	 * @param Index The index of the component to return (0...3).
 	 * @return The component.
 	 */
 	uint32& operator[]( int32 Index )
@@ -158,8 +154,7 @@ public:
 	/**
 	 * Provides read-only access to the GUIDs components.
 	 *
-	 * @param Index - The index of the component to return (0...3).
-	 *
+	 * @param Index The index of the component to return (0...3).
 	 * @return The component.
 	 */
 	const uint32& operator[]( int32 Index ) const
@@ -181,8 +176,8 @@ public:
 	/**
 	 * Serializes a GUID from or into an archive.
 	 *
-	 * @param Ar - The archive to serialize from or into.
-	 * @param G - The GUID to serialize.
+	 * @param Ar The archive to serialize from or into.
+	 * @param G The GUID to serialize.
 	 */
 	friend FArchive& operator<<( FArchive& Ar, FGuid& G )
 	{
@@ -194,13 +189,13 @@ public:
 	/**
 	 * Exports the GUIDs value to a string.
 	 *
-	 * @param ValueStr - Will hold the string value.
-	 * @param DefaultValue - The default value.
-	 * @param Parent - Not used.
-	 * @param PortFlags - Not used.
-	 * @param ExportRootScope - Not used.
-	 *
+	 * @param ValueStr Will hold the string value.
+	 * @param DefaultValue The default value.
+	 * @param Parent Not used.
+	 * @param PortFlags Not used.
+	 * @param ExportRootScope Not used.
 	 * @return true on success, false otherwise.
+	 * @see ImportTextItem
 	 */
 	bool ExportTextItem( FString& ValueStr, FGuid const& DefaultValue, UObject* Parent, int32 PortFlags, class UObject* ExportRootScope ) const
 	{
@@ -212,12 +207,12 @@ public:
 	/**
 	 * Imports the GUIDs value from a text buffer.
 	 *
-	 * @param Buffer - The text buffer to import from.
-	 * @param PortFlags - Not used.
-	 * @param Parent - Not used.
-	 * @param ErrorText - The output device for error logging.
-	 *
+	 * @param Buffer The text buffer to import from.
+	 * @param PortFlags Not used.
+	 * @param Parent Not used.
+	 * @param ErrorText The output device for error logging.
 	 * @return true on success, false otherwise.
+	 * @see ExportTextItem
 	 */
 	bool ImportTextItem( const TCHAR*& Buffer, int32 PortFlags, class UObject* Parent, FOutputDevice* ErrorText )
 	{
@@ -242,7 +237,6 @@ public:
 	 * A GUID that has all its components set to zero is considered invalid.
 	 *
 	 * @return true if valid, false otherwise
-	 *
 	 * @see Invalidate
 	 */
 	bool IsValid( ) const
@@ -273,8 +267,7 @@ public:
 	/**
 	 * Converts this GUID to its string representation using the specified format.
 	 *
-	 * @param Format - The string format to use.
-	 *
+	 * @param Format The string format to use.
 	 * @return The string representation.
 	 */
 	CORE_API FString ToString( EGuidFormats::Type Format ) const;
@@ -284,8 +277,7 @@ public:
 	/**
 	 * Calculates the hash for a GUID.
 	 *
-	 * @param Guid - The GUID to calculate the hash for.
-	 *
+	 * @param Guid The GUID to calculate the hash for.
 	 * @return The hash.
 	 */
 	friend uint32 GetTypeHash( const FGuid& Guid )
@@ -305,21 +297,21 @@ public:
 	/**
 	 * Converts a string to a GUID.
 	 *
-	 * @param GuidString - The string to convert.
-	 * @param OutGuid - Will contain the parsed GUID.
-	 *
+	 * @param GuidString The string to convert.
+	 * @param OutGuid Will contain the parsed GUID.
 	 * @return true if the string was converted successfully, false otherwise.
+	 * @see ParseExact, ToString
 	 */
 	static CORE_API bool Parse( const FString& GuidString, FGuid& OutGuid );
 
 	/**
 	 * Converts a string with the specified format to a GUID.
 	 *
-	 * @param GuidString - The string to convert.
-	 * @param Format - The string format to parse.
-	 * @param OutGuid - Will contain the parsed GUID.
-	 *
+	 * @param GuidString The string to convert.
+	 * @param Format The string format to parse.
+	 * @param OutGuid Will contain the parsed GUID.
 	 * @return true if the string was converted successfully, false otherwise.
+	 * @see Parse, ToString
 	 */
 	static CORE_API bool ParseExact( const FString& GuidString, EGuidFormats::Type Format, FGuid& OutGuid );
 
diff --git a/Engine/Source/Runtime/Core/Public/Misc/ITransaction.h b/Engine/Source/Runtime/Core/Public/Misc/ITransaction.h
--- a/Engine/Source/Runtime/Core/Public/Misc/ITransaction.h
+++ b/Engine/Source/Runtime/Core/Public/Misc/ITransaction.h
@@ -6,21 +6,28 @@
 
 #pragma once
 
+
 // Class for handling undo/redo transactions among objects.
-typedef void( *STRUCT_DC )( void* TPtr );						// default construct
-typedef void( *STRUCT_AR )( class FArchive& Ar, void* TPtr );	// serialize
-typedef void( *STRUCT_DTOR )( void* TPtr );						// destruct
+typedef void(*STRUCT_DC)( void* TPtr );						// default construct
+typedef void(*STRUCT_AR)( class FArchive& Ar, void* TPtr );	// serialize
+typedef void(*STRUCT_DTOR)( void* TPtr );					// destruct
 
 
 /**
- * Abstract base class for transaction object annotations
+ * Interface for transaction object annotations.
+ *
+ * Transaction object annotations are used for attaching additional user defined data to a transaction.
+ * This is sometimes useful, because the transaction system only remembers changes that are serializable
+ * on the UObject that a modification was performed on, but it does not see other changes that may have
+ * to be remembered in order to properly restore the object internals.
  */
-class ITransactionObjectAnnotation
-{
-};
+class ITransactionObjectAnnotation { };
 
 /**
  * Interface for transactions.
+ *
+ * Transactions are created each time an UObject is modified, for example in the Editor.
+ * The data stored inside a transaction object can then be used to provide undo/redo functionality.
  */
 class ITransaction
 {
@@ -34,28 +41,28 @@ public:
 	/**
 	 * Saves an array to the transaction.
 	 *
-	 * @param Object - The object that owns the array.
-	 * @param Array - The array to save.
-	 * @param Index - 
-	 * @param Count - 
-	 * @param Oper -
-	 * @param ElementSize -
-	 * @param Serializer -
-	 * @param Destructor -
+	 * @param Object The object that owns the array.
+	 * @param Array The array to save.
+	 * @param Index 
+	 * @param Count 
+	 * @param Oper
+	 * @param ElementSize
+	 * @param Serializer
+	 * @param Destructor
 	 */
 	virtual void SaveArray( UObject* Object, class FScriptArray* Array, int32 Index, int32 Count, int32 Oper, int32 ElementSize, STRUCT_DC DefaultConstructor, STRUCT_AR Serializer, STRUCT_DTOR Destructor ) = 0;
 
 	/**
 	 * Saves an UObject to the transaction.
 	 *
-	 * @param Object - The object to save.
+	 * @param Object The object to save.
 	 */
 	virtual void SaveObject( UObject* Object ) = 0;
 
 	/**
 	 * Sets the transaction's primary object.
 	 *
-	 * @param Object - The primary object to set.
+	 * @param Object The primary object to set.
 	 */
 	virtual void SetPrimaryObject( UObject* Object ) = 0;
 };
diff --git a/Engine/Source/Runtime/Core/Public/Misc/Variant.h b/Engine/Source/Runtime/Core/Public/Misc/Variant.h
--- a/Engine/Source/Runtime/Core/Public/Misc/Variant.h
+++ b/Engine/Source/Runtime/Core/Public/Misc/Variant.h
@@ -137,7 +137,7 @@ public:
 	/**
 	 * Creates and initializes a new instance from a TCHAR string.
 	 *
-	 * @param InString - The initial value.
+	 * @param InString The initial value.
 	 */
 	FVariant( const TCHAR* InString )
 	{
@@ -149,8 +149,7 @@ public:
 	/**
 	 * Copy assignment operator.
 	 *
-	 * @param Other - The value to assign.
-	 *
+	 * @param Other The value to assign.
 	 * @return This instance.
 	 */
 	FVariant& operator=( const FVariant& Other )
@@ -167,9 +166,8 @@ public:
 	/**
 	 * Assignment operator.
 	 *
-	 * @param T - The type of the value to assign.
-	 * @param InValue - The value to assign.
-	 *
+	 * @param T The type of the value to assign.
+	 * @param InValue The value to assign.
 	 * @return This instance.
 	 */
 	template<typename T>
@@ -190,8 +188,7 @@ public:
 	 * serialize any complex types into arrays and then store them in FVariant, you will be
 	 * responsible for ensuring byte ordering if the FVariant gets sent over the network.
 	 *
-	 * @param InArray - The byte array to assign.
-	 *
+	 * @param InArray The byte array to assign.
 	 * @return This instance.
 	 */
 	FVariant& operator=( const TArray<uint8> InArray )
@@ -205,8 +202,7 @@ public:
 	/**
 	 * Assignment operator for TCHAR strings.
 	 *
-	 * @param InString - The value to assign.
-	 *
+	 * @param InString The value to assign.
 	 * @return This instance.
 	 */
 	FVariant& operator=( const TCHAR* InString )
@@ -220,8 +216,7 @@ public:
 	/**
 	 * Implicit conversion operator.
 	 *
-	 * @param T - The type to convert the value to.
-	 *
+	 * @param T The type to convert the value to.
 	 * @return The value converted to the specified type.
 	 */
 	template<typename T>
@@ -235,8 +230,7 @@ public:
 	/**
 	 * Comparison operator for equality.
 	 *
-	 * @param Other - The variant to compare with.
-	 *
+	 * @param Other The variant to compare with.
 	 * @return true if the values are equal, false otherwise.
 	 */
 	bool operator==( const FVariant& Other ) const
@@ -247,8 +241,7 @@ public:
 	/**
 	 * Comparison operator for inequality.
 	 *
-	 * @param Other - The variant to compare with.
-	 *
+	 * @param Other The variant to compare with.
 	 * @return true if the values are not equal, false otherwise.
 	 */
 	bool operator!=( const FVariant& Other ) const
@@ -290,7 +283,6 @@ public:
 	 * GetValue<TArray<uint8>>() instead.
 	 *
 	 * @return Byte array.
-	 *
 	 * @see GetValue
 	 */
 	const TArray<uint8>& GetBytes( ) const
@@ -302,9 +294,7 @@ public:
 	 * Gets the stored value's size (in bytes).
 	 *
 	 * @return Size of the value.
-	 *
-	 * @see GetType
-	 * @see GetValue
+	 * @see GetType, GetValue
 	 */
 	int32 GetSize( ) const
 	{
@@ -315,9 +305,7 @@ public:
 	 * Gets the stored value's type.
 	 *
 	 * @return Type of the value.
-	 *
-	 * @see GetSize
-	 * @see GetValue
+	 * @see GetSize, GetValue
 	 */
 	int32 GetType( ) const
 	{
@@ -331,9 +319,7 @@ public:
 	 * convertible types. The exact type of the value to be extracted must be known.
 	 *
 	 * @return The value.
-	 *
-	 * @see GetSize
-	 * @see GetType
+	 * @see GetSize, GetType
 	 */
 	template<typename T>
 	T GetValue( ) const
@@ -353,9 +339,8 @@ public:
 	/**
 	 * Serializes the given variant type from or into the specified archive.
 	 *
-	 * @param Ar - The archive to serialize from or into.
-	 * @param Variant - The value to serialize.
-	 *
+	 * @param Ar The archive to serialize from or into.
+	 * @param Variant The value to serialize.
 	 * @return The archive.
 	 */
 	friend FArchive& operator<<( FArchive& Ar, FVariant& Variant )
@@ -383,7 +368,6 @@ private:
  * To retrieve any value as an array of serialized bytes, use GetBytes() instead.
  *
  * @return The byte array.
- *
  * @see GetBytes
  */
 template<>
diff --git a/Engine/Source/Runtime/Core/Public/Misc/WildcardString.h b/Engine/Source/Runtime/Core/Public/Misc/WildcardString.h
--- a/Engine/Source/Runtime/Core/Public/Misc/WildcardString.h
+++ b/Engine/Source/Runtime/Core/Public/Misc/WildcardString.h
@@ -18,25 +18,25 @@ public:
 	/**
 	 * Default constructor.
 	 */
-	FWildcardString()
+	FWildcardString( )
 		: FString()
 	{ }
 
 	/**
 	 * Creates and initializes a new instance with the specified pattern.
 	 *
-	 * @param Pattern - The pattern string.
+	 * @param Pattern The pattern string.
 	 */
-	FWildcardString(const FString& Pattern)
+	FWildcardString( const FString& Pattern )
 		: FString(Pattern)
 	{ }
 
 	/**
 	 * Creates and initializes a new instance with the specified pattern.
 	 *
-	 * @param Pattern - The pattern string.
+	 * @param Pattern The pattern string.
 	 */
-	FWildcardString(const TCHAR* Pattern)
+	FWildcardString( const TCHAR* Pattern )
 		: FString(Pattern)
 	{ }
 
@@ -47,7 +47,7 @@ public:
 	 *
 	 * @return true if this string contains wild cards, false otherwise.
 	 */
-	bool ContainsWildcards() const
+	bool ContainsWildcards( ) const
 	{
 		return ContainsWildcards(**this);
 	}
@@ -55,11 +55,10 @@ public:
 	/**
 	 * Matches the given input string to this wild card pattern.
 	 *
-	 * @param Input - The string to match.
-	 *
+	 * @param Input The string to match.
 	 * @return true if the input string matches this pattern, false otherwise.
 	 */
-	bool IsMatch(const TCHAR* Input) const
+	bool IsMatch( const TCHAR* Input ) const
 	{
 		return IsMatch(**this, Input);
 	}
@@ -67,11 +66,10 @@ public:
 	/**
 	 * Matches the given input string to this wild card pattern.
 	 *
-	 * @param Input - The string to match.
-	 *
+	 * @param Input The string to match.
 	 * @return true if the input string matches this pattern, false otherwise.
 	 */
-	bool IsMatch(const FString& Input) const
+	bool IsMatch( const FString& Input ) const
 	{
 		return IsMatch(**this, *Input);
 	}
@@ -81,34 +79,27 @@ public:
 	/**
 	 * Checks whether the specified pattern contains wild card characters.
 	 *
-	 * @param Pattern - The string to check.
-	 *
+	 * @param Pattern The string to check.
 	 * @return true if the string contains wild cards, false otherwise.
 	 */
-	static CORE_API bool ContainsWildcards(const TCHAR* Pattern);
+	static CORE_API bool ContainsWildcards( const TCHAR* Pattern );
 
 	/**
 	 * Non-recursive wild card string matching algorithm.
 	 *
-	 * @param Pattern - The pattern to match.
-	 * @param Input - The input string to check.
+	 * @param Pattern The pattern to match.
+	 * @param Input The input string to check.
 	 */
-	static CORE_API bool IsMatch(const TCHAR* Pattern, const TCHAR* Input);
+	static CORE_API bool IsMatch( const TCHAR* Pattern, const TCHAR* Input );
 
 protected:
 
-	/**
-	 * Holds the string terminator character.
-	 */
+	/** Holds the string terminator character. */
 	static CORE_API const TCHAR EndOfString = TCHAR('\0');
 
-	/**
-	 * Holds the wild card that matches exactly one character (default is '?').
-	 */
+	/** Holds the wild card that matches exactly one character (default is '?'). */
 	static CORE_API const TCHAR ExactWildcard = TCHAR('?');
 
-	/**
-	 * Holds the wild card that matches a sequence of characters (default is '*').
-	 */
+	/** Holds the wild card that matches a sequence of characters (default is '*'). */
 	static CORE_API const TCHAR SequenceWildcard = TCHAR('*');
 };
