Wrap usage of deprecated CGCursorIsVisible function with clang diagnostic ignored so that enabling deprecation warnings won't cause an error. This function was deprecated in Mavericks but there is no published alternative.

diff --git a/Engine/Source/Runtime/Core/Private/Mac/MacCursor.cpp b/Engine/Source/Runtime/Core/Private/Mac/MacCursor.cpp
--- a/Engine/Source/Runtime/Core/Private/Mac/MacCursor.cpp
+++ b/Engine/Source/Runtime/Core/Private/Mac/MacCursor.cpp
@@ -248,6 +248,9 @@ bool FMacCursor::UpdateCursorClipping( FVector2D& CursorPosition )
 
 void FMacCursor::UpdateVisibility()
 {
+	// @TODO: Remove usage of deprecated CGCursorIsVisible function
+#pragma clang diagnostic push
+#pragma clang diagnostic ignored "-Wdeprecated-declarations"
 	if (CurrentCursor && bIsVisible)
 	{
 		// Enable the cursor.
@@ -264,6 +267,7 @@ void FMacCursor::UpdateVisibility()
 			CGDisplayHideCursor(kCGDirectMainDisplay);
 		}
 	}
+#pragma clang diagnostic pop
 }
 
 void FMacCursor::WarpCursor( const int32 X, const int32 Y )
diff --git a/Engine/Source/Runtime/Core/Private/Mac/MacPlatformOutputDevices.cpp b/Engine/Source/Runtime/Core/Private/Mac/MacPlatformOutputDevices.cpp
--- a/Engine/Source/Runtime/Core/Private/Mac/MacPlatformOutputDevices.cpp
+++ b/Engine/Source/Runtime/Core/Private/Mac/MacPlatformOutputDevices.cpp
@@ -127,10 +127,14 @@ void FOutputDeviceMacError::HandleError()
 		GLog->Flush();
 
 		// Unhide the mouse.
+		// @TODO: Remove usage of deprecated CGCursorIsVisible function
+#pragma clang diagnostic push
+#pragma clang diagnostic ignored "-Wdeprecated-declarations"
 		while (!CGCursorIsVisible())
 		{
 			CGDisplayShowCursor(kCGDirectMainDisplay);
 		}
+#pragma clang diagnostic pop
 		// Release capture and allow mouse to freely roam around.
 		CGAssociateMouseAndMouseCursorPosition(true);
 
