Added deprecation macro support for OS X.

diff --git a/Engine/Source/Programs/UnrealBuildTool/Mac/MacToolChain.cs b/Engine/Source/Programs/UnrealBuildTool/Mac/MacToolChain.cs
--- a/Engine/Source/Programs/UnrealBuildTool/Mac/MacToolChain.cs
+++ b/Engine/Source/Programs/UnrealBuildTool/Mac/MacToolChain.cs
@@ -89,7 +89,6 @@ namespace UnrealBuildTool
 			Result += " -Wno-switch-enum";
 			Result += " -Wno-logical-op-parentheses";	// needed for external headers we can't change
 			Result += " -Wno-null-arithmetic";			// needed for external headers we can't change
-			Result += " -Wno-deprecated-declarations";	// needed for wxWidgets
 			Result += " -Wno-return-type-c-linkage";	// needed for PhysX
 			Result += " -Wno-ignored-attributes";		// needed for nvtesslib
 			Result += " -Wno-uninitialized";
diff --git a/Engine/Source/Runtime/Core/Public/HAL/PlatformIncludes.h b/Engine/Source/Runtime/Core/Public/HAL/PlatformIncludes.h
--- a/Engine/Source/Runtime/Core/Public/HAL/PlatformIncludes.h
+++ b/Engine/Source/Runtime/Core/Public/HAL/PlatformIncludes.h
@@ -18,6 +18,8 @@
 // Platform specific compiler pre-setup, must be included before everything else
 #if PLATFORM_WINDOWS
 #include "Windows/WindowsPlatformCompilerPreSetup.h"
+#elif PLATFORM_MAC
+#include "Mac/MacPlatformCompilerPreSetup.h"
 #endif
 
 // Generic versions of the platform abstractions
diff --git a/Engine/Source/Runtime/Core/Public/Mac/MacPlatformCompilerPreSetup.h b/Engine/Source/Runtime/Core/Public/Mac/MacPlatformCompilerPreSetup.h
new file mode 100644
--- /dev/null
+++ b/Engine/Source/Runtime/Core/Public/Mac/MacPlatformCompilerPreSetup.h
@@ -0,0 +1,13 @@
+// Copyright 1998-2014 Epic Games, Inc. All Rights Reserved.
+
+/*==========================================================================================================
+	MacPlatformCompilerPreSetup.h: setup for the compiler that needs to be defined before everything else
+===========================================================================================================*/
+
+#pragma once
+
+#ifndef DISABLE_DEPRECATION
+#define DEPRECATED(VERSION, MESSAGE) __attribute__((deprecated(MESSAGE " Please update your code to the new API before upgrading to the next release, otherwise your project will no longer compile.")))
+#else
+	#define DEPRECATED(VERSION, MESSAGE)
+#endif // DISABLE_DEPRECATION
