Stats - Stats files are now compressed, ratio for regular stats 30%, for raw stats 16%, misc code refactor and improvements

diff --git a/Engine/Source/Developer/Profiler/Private/ProfilerSession.cpp b/Engine/Source/Developer/Profiler/Private/ProfilerSession.cpp
--- a/Engine/Source/Developer/Profiler/Private/ProfilerSession.cpp
+++ b/Engine/Source/Developer/Profiler/Private/ProfilerSession.cpp
@@ -489,6 +489,8 @@ bool FRawProfilerSession::HandleTicker( float DeltaTime )
 
 void FRawProfilerSession::PrepareLoading()
 {
+	SCOPE_LOG_TIME_FUNC();
+
 	const FString Filepath = DataFilepath + FStatConstants::StatsFileRawExtension;
 	const int64 Size = IFileManager::Get().FileSize( *Filepath );
 	if( Size < 4 )
@@ -514,7 +516,7 @@ void FRawProfilerSession::PrepareLoading()
 	StatsThreadStats.MarkAsLoaded();
 
 	TArray<FStatMessage> Messages;
-	if( Stream.Header.bRawStatFile )
+	if( Stream.Header.bRawStatsFile )
 	{
 		// Read metadata.
 		TArray<FStatMessage> MetadataMessages;
@@ -530,23 +532,40 @@ void FRawProfilerSession::PrepareLoading()
 		// Read frames offsets.
 		Stream.ReadFramesOffsets( *FileReader );
 
+		// Buffer used to store the compressed and decompressed data.
+		TArray<uint8> SrcArray;
+		TArray<uint8> DestArray;
+		const bool bHasCompressedData = Stream.Header.HasCompressedData();
+
 		// Temporary fix for setting the correct value for the SecondsPerCycle.
 		{
+			SCOPE_LOG_TIME( TEXT( "Looking for the STAT_SecondsPerCycle" ), nullptr );
 			UE_LOG( LogStats, Log, TEXT( "Looking for the STAT_SecondsPerCycle" ) );
 			FileReader->Seek( Stream.FramesInfo[0].FrameFileOffset );
 
+			if( bHasCompressedData )
+			{
+				// Read the compressed data.
+				FCompressedStatsData UncompressedData( SrcArray, DestArray );
+				*FileReader << UncompressedData;
+			}
+
+			// Select the proper archive.
+			FMemoryReader MemoryReader( DestArray, true );
+			FArchive& Archive = bHasCompressedData ? MemoryReader : *FileReader;
+
 			FStatPacketArray StatPacketArray;
 
 			int64 TargetFrame;
-			*FileReader << TargetFrame;
+			Archive << TargetFrame;
 
 			int32 NumPackets;
-			*FileReader << NumPackets;
+			Archive << NumPackets;
 
 			for( int32 PacketIndex = 0; PacketIndex < NumPackets; PacketIndex++ )
 			{
 				FStatPacket* ToRead = new FStatPacket();
-				Stream.ReadStatPacket( *FileReader, *ToRead, bIsFinalized );
+				Stream.ReadStatPacket( Archive, *ToRead, bIsFinalized );
 				StatPacketArray.Packets.Add( ToRead );
 			}
 
@@ -575,89 +594,120 @@ void FRawProfilerSession::PrepareLoading()
 
 		check( StatMetaData->GetSecondsPerCycle() > 0.0 );
 
-		// Prepare profiler frames.
-		double ElapsedTimeMS = 0;
-		for( int32 FrameIndex = 0; FrameIndex < Stream.FramesInfo.Num(); ++FrameIndex )
 		{
-			const FStatsFrameInfo& StatsFrameInfo = Stream.FramesInfo[FrameIndex];
-			FileReader->Seek( StatsFrameInfo.FrameFileOffset );
-			int64 TargetFrame;
-			*FileReader << TargetFrame;
-
-			const double GameThreadTimeMS = GetMetaData()->ConvertCyclesToMS( (uint32)StatsFrameInfo.ThreadCycles.FindChecked( GameThreadID ) );
+			SCOPE_LOG_TIME( TEXT( "Preparing profiler frames" ), nullptr );
 
-			// Update mini-view, convert from cycles to ms.
-			TMap<uint32, float> ThreadTimesMS;
-			for( auto InnerIt = StatsFrameInfo.ThreadCycles.CreateConstIterator(); InnerIt; ++InnerIt )
+			// Prepare profiler frames.
+			double ElapsedTimeMS = 0;
+			for( int32 FrameIndex = 0; FrameIndex < Stream.FramesInfo.Num(); ++FrameIndex )
 			{
-				const uint32 ThreadID = InnerIt.Key();
-				const bool bCanBeAdded = GameThreadID == ThreadID || GetMetaData()->GetRenderThreadID().Contains( ThreadID );
-				if( bCanBeAdded )
+				const FStatsFrameInfo& StatsFrameInfo = Stream.FramesInfo[FrameIndex];
+				FileReader->Seek( StatsFrameInfo.FrameFileOffset );
+
+				if( bHasCompressedData )
 				{
-					ThreadTimesMS.Add( ThreadID, GetMetaData()->ConvertCyclesToMS( InnerIt.Value() ) );
+					// Read the compressed data.
+					FCompressedStatsData UncompressedData( SrcArray, DestArray );
+					*FileReader << UncompressedData;
 				}
-			}
 
-			// Pass the reference to the stats' metadata.
-			OnAddThreadTime.ExecuteIfBound( FrameIndex, ThreadTimesMS, StatMetaData );
+				// Select the proper archive.
+				FMemoryReader MemoryReader( DestArray, true );
+				FArchive& Archive = bHasCompressedData ? MemoryReader : *FileReader;
+
+				int64 TargetFrame;
+				Archive << TargetFrame;
 
-			// Create a new profiler frame and add it to the stream.
-			ElapsedTimeMS += GameThreadTimeMS;
-			FProfilerFrame* ProfilerFrame = new FProfilerFrame( TargetFrame, GameThreadTimeMS, ElapsedTimeMS );
-			ProfilerFrame->ThreadTimesMS = ThreadTimesMS;
-			ProfilerStream.AddProfilerFrame( TargetFrame, ProfilerFrame );
+				const double GameThreadTimeMS = GetMetaData()->ConvertCyclesToMS( (uint32)StatsFrameInfo.ThreadCycles.FindChecked( GameThreadID ) );
 
-			// Thread cycles can be used to generate low-res thread view.
-			// @TODO yrx 2014-04-23 
+				// Update mini-view, convert from cycles to ms.
+				TMap<uint32, float> ThreadTimesMS;
+				for( auto InnerIt = StatsFrameInfo.ThreadCycles.CreateConstIterator(); InnerIt; ++InnerIt )
+				{
+					const uint32 ThreadID = InnerIt.Key();
+					const bool bCanBeAdded = GameThreadID == ThreadID || GetMetaData()->GetRenderThreadID().Contains( ThreadID );
+					if( bCanBeAdded )
+					{
+						ThreadTimesMS.Add( ThreadID, GetMetaData()->ConvertCyclesToMS( InnerIt.Value() ) );
+					}
+				}
+
+				// Pass the reference to the stats' metadata.
+				OnAddThreadTime.ExecuteIfBound( FrameIndex, ThreadTimesMS, StatMetaData );
+
+				// Create a new profiler frame and add it to the stream.
+				ElapsedTimeMS += GameThreadTimeMS;
+				FProfilerFrame* ProfilerFrame = new FProfilerFrame( TargetFrame, GameThreadTimeMS, ElapsedTimeMS );
+				ProfilerFrame->ThreadTimesMS = ThreadTimesMS;
+				ProfilerStream.AddProfilerFrame( TargetFrame, ProfilerFrame );
+
+				// Thread cycles can be used to generate low-res thread view.
+				// @TODO yrx 2014-04-23 
+			}
+			int64 FrameOffset0 = Stream.FramesInfo[0].FrameFileOffset;
+			FileReader->Seek( FrameOffset0 );
 		}
-		int64 FrameOffset0 = Stream.FramesInfo[0].FrameFileOffset;
-		FileReader->Seek( FrameOffset0 );
 
 #if	1
-		double CycleCounterAdjustmentMS = 0.0f;
-
-		// Read the raw stats messages.
-		for( int32 FrameIndex = 0; FrameIndex < Stream.FramesInfo.Num(); ++FrameIndex )
 		{
-			FProfilerFrame* ProfilerFrame = ProfilerStream.GetProfilerFrame( FrameIndex );
-			FileReader->Seek( Stream.FramesInfo[FrameIndex].FrameFileOffset );
+			SCOPE_LOG_TIME( TEXT( "Processing raw stats frames" ), nullptr );
+
+			double CycleCounterAdjustmentMS = 0.0f;
 
-			UE_LOG( LogStats, Log, TEXT( "Processing raw stats frame: %4i/%4i" ), FrameIndex, Stream.FramesInfo.Num() );
+			// Read the raw stats messages.
+			for( int32 FrameIndex = 0; FrameIndex < Stream.FramesInfo.Num(); ++FrameIndex )
+			{
+				FProfilerFrame* ProfilerFrame = ProfilerStream.GetProfilerFrame( FrameIndex );
+				FileReader->Seek( Stream.FramesInfo[FrameIndex].FrameFileOffset );
 
-			FStatPacketArray Packet;
+				UE_CLOG( FrameIndex % 8 == 0, LogStats, Log, TEXT( "Processing raw stats frame: %4i/%4i" ), FrameIndex, Stream.FramesInfo.Num() );
 
-			int64 TargetFrame;
-			*FileReader << TargetFrame;
+				if( bHasCompressedData )
+				{
+					// Read the compressed data.
+					FCompressedStatsData UncompressedData( SrcArray, DestArray );
+					*FileReader << UncompressedData;
+				}
 
-			int32 NumPackets;
-			*FileReader << NumPackets;
+				// Select the proper archive.
+				FMemoryReader MemoryReader( DestArray, true );
+				FArchive& Archive = bHasCompressedData ? MemoryReader : *FileReader;
 
-			for( int32 PacketIndex = 0; PacketIndex < NumPackets; PacketIndex++ )
-			{
-				FStatPacket* ToRead = new FStatPacket();
-				Stream.ReadStatPacket( *FileReader, *ToRead, bIsFinalized );
-				Packet.Packets.Add( ToRead );
-			}
+				FStatPacketArray Packet;
 
-			ProcessStatPacketArray( Packet, *ProfilerFrame, FrameIndex ); // or ProfilerFrame->TargetFrame
+				int64 TargetFrame;
+				Archive << TargetFrame;
 
-			// Find the first cycle counter for the game thread.
-			if( CycleCounterAdjustmentMS == 0.0f )
-			{
-				CycleCounterAdjustmentMS = ProfilerFrame->Root->CycleCounterStartTimeMS;
+				int32 NumPackets;
+				Archive << NumPackets;
+
+				for( int32 PacketIndex = 0; PacketIndex < NumPackets; PacketIndex++ )
+				{
+					FStatPacket* ToRead = new FStatPacket();
+					Stream.ReadStatPacket( Archive, *ToRead, bIsFinalized );
+					Packet.Packets.Add( ToRead );
+				}
+
+				ProcessStatPacketArray( Packet, *ProfilerFrame, FrameIndex ); // or ProfilerFrame->TargetFrame
+
+				// Find the first cycle counter for the game thread.
+				if( CycleCounterAdjustmentMS == 0.0f )
+				{
+					CycleCounterAdjustmentMS = ProfilerFrame->Root->CycleCounterStartTimeMS;
+				}
+
+				TMap<uint32, int64> ThreadCycles;
+				// Serialize thread cycles.
+				Archive << ThreadCycles;
+
+				// Update thread time and mark profiler frame as valid and ready for use.
+				ProfilerFrame->MarkAsValid();
 			}
 
-			TMap<uint32, int64> ThreadCycles;
-			// Serialize thread cycles.
-			*FileReader << ThreadCycles;
-		
-			// Update thread time and mark profiler frame as valid and ready for use.
-			ProfilerFrame->MarkAsValid();
+			// Adjust all profiler frames.
+			ProfilerStream.AdjustCycleCounters( CycleCounterAdjustmentMS );
 		}
 
-		// Adjust all profiler frames.
-		ProfilerStream.AdjustCycleCounters( CycleCounterAdjustmentMS );
-
 #endif // 0
 	}
 
@@ -765,6 +815,7 @@ void FProfilerStatMetaData::UpdateFromStatsState( const FStatsThreadState& Stats
 				if (ThreadsIt.Value() == StatName)
 				{
 					ThreadID = ThreadsIt.Key();
+					break;
 				}
 			}
 			ThreadIDtoStatID.Add( ThreadID, StatID );
diff --git a/Engine/Source/Developer/ProfilerClient/Private/ProfilerClientManager.cpp b/Engine/Source/Developer/ProfilerClient/Private/ProfilerClientManager.cpp
--- a/Engine/Source/Developer/ProfilerClient/Private/ProfilerClientManager.cpp
+++ b/Engine/Source/Developer/ProfilerClient/Private/ProfilerClientManager.cpp
@@ -272,12 +272,30 @@ void FProfilerClientManager::LoadCapture( const FString& DataFilepath, const FGu
 	}
 
 	// This shouldn't happen.
-	if( LoadConnection->Stream.Header.bRawStatFile )
+	if( LoadConnection->Stream.Header.bRawStatsFile )
 	{
 		delete FileReader;
 		return;
 	}
 
+	const bool bIsFinalized = LoadConnection->Stream.Header.IsFinalized();
+	if( bIsFinalized )
+	{
+		// Read metadata.
+		TArray<FStatMessage> MetadataMessages;
+		LoadConnection->Stream.ReadFNamesAndMetadataMessages( *FileReader, MetadataMessages );
+		LoadConnection->CurrentThreadState.ProcessMetaDataOnly( MetadataMessages );
+
+		// Read frames offsets.
+		LoadConnection->Stream.ReadFramesOffsets( *FileReader );
+		FileReader->Seek( LoadConnection->Stream.FramesInfo[0].FrameFileOffset );
+	}
+
+	if( LoadConnection->Stream.Header.HasCompressedData() )
+	{
+		UE_CLOG( !bIsFinalized, LogProfile, Fatal, TEXT( "Compressed stats file has to be finalized" ) );
+	}
+
 #if PROFILER_THREADED_LOAD
 	LoadTask = new FAsyncTask<FAsyncReadWorker>(LoadConnection, FileReader);
 	LoadTask->StartBackgroundTask();
@@ -1043,65 +1061,137 @@ void FServiceConnection::GenerateCycleGraphs(const FRawStatStackNode& Root, TMap
 	}
 }
 
-bool FServiceConnection::ReadAndConvertStatMessages( FArchive& FileReader, bool bUseInAsync )
+bool FServiceConnection::ReadAndConvertStatMessages( FArchive& Reader, bool bUseInAsync )
 {
 	uint64 ReadMessages = 0;
 
+	// Buffer used to store the compressed and decompressed data.
+	TArray<uint8> SrcData;
+	TArray<uint8> DestData;
+
+	const bool bHasCompressedData = Stream.Header.HasCompressedData();
+	const bool bIsFinalized = Stream.Header.IsFinalized();
+
 	SCOPE_CYCLE_COUNTER( STAT_PC_ReadStatMessages );
-	while( FileReader.Tell() < FileReader.TotalSize() )
+	if( bHasCompressedData )
 	{
-		// read the message
-		FStatMessage Message( Stream.ReadMessage( FileReader ) );
-		ReadMessages++;
-
-		if( Message.NameAndInfo.GetShortName() != TEXT( "Unknown FName" ) )
+		while( Reader.Tell() < Reader.TotalSize() )
 		{
-			if( Message.NameAndInfo.GetField<EStatOperation>() == EStatOperation::SpecialMessageMarker )
+			// Read the compressed data.
+			FCompressedStatsData UncompressedData( SrcData, DestData );
+			Reader << UncompressedData;
+			if( UncompressedData.HasReachedEndOfCompressedData() )
 			{
-				// Simply break the loop.
-				// The profiler supports more advanced handling of this message.
 				return true;
 			}
-			else if( Message.NameAndInfo.GetField<EStatOperation>() == EStatOperation::AdvanceFrameEventGameThread && ReadMessages > 2 )
+
+			FMemoryReader MemoryReader( DestData, true );
+
+			while( MemoryReader.Tell() < MemoryReader.TotalSize() )
 			{
-				AddCollectedStatMessages( Message );
+				// read the message
+				FStatMessage Message( Stream.ReadMessage( MemoryReader, bIsFinalized ) );
+				ReadMessages++;
+
+				if( Message.NameAndInfo.GetShortName() != TEXT( "Unknown FName" ) )
+				{
+					if( Message.NameAndInfo.GetField<EStatOperation>() == EStatOperation::AdvanceFrameEventGameThread && ReadMessages > 2 )
+					{
+						AddCollectedStatMessages( Message );
 
-				// create an old format data frame from the data
-				GenerateProfilerDataFrame();
+						// create an old format data frame from the data
+						GenerateProfilerDataFrame();
 
+						{
+							// add the frame to the work list
+							FScopeLock ScopeLock( &CriticalSection );
+							DataFrames.Add( CurrentData );
+							DataLoadingProgress = (double)Reader.Tell() / (double)Reader.TotalSize();
+						}
+
+						if( bUseInAsync )
+						{
+							if( DataFrames.Num() > FProfilerClientManager::MaxFramesPerTick )
+							{
+								while( DataFrames.Num() )
+								{
+									FPlatformProcess::Sleep( 0.001f );
+								}
+							}
+						}
+					}
+
+					new (Messages)FStatMessage( Message );
+				}
+				else
 				{
-					// add the frame to the work list
-					FScopeLock ScopeLock( &CriticalSection );
-					DataFrames.Add( CurrentData );
-					DataLoadingProgress = (double)FileReader.Tell() / (double)FileReader.TotalSize();
+					break;
 				}
 
-				if( bUseInAsync )
+				if( !bUseInAsync && DataFrames.Num() < FProfilerClientManager::MaxFramesPerTick )
 				{
-					if( DataFrames.Num() > FProfilerClientManager::MaxFramesPerTick )
+					return false;
+				}
+			}
+		}
+	}
+	else
+	{
+		while( Reader.Tell() < Reader.TotalSize() )
+		{
+			// read the message
+			FStatMessage Message( Stream.ReadMessage( Reader, bIsFinalized ) );
+			ReadMessages++;
+
+			if( Message.NameAndInfo.GetShortName() != TEXT( "Unknown FName" ) )
+			{
+				if( Message.NameAndInfo.GetField<EStatOperation>() == EStatOperation::SpecialMessageMarker )
+				{
+					// Simply break the loop.
+					// The profiler supports more advanced handling of this message.
+					return true;
+				}
+				else if( Message.NameAndInfo.GetField<EStatOperation>() == EStatOperation::AdvanceFrameEventGameThread && ReadMessages > 2 )
+				{
+					AddCollectedStatMessages( Message );
+
+					// create an old format data frame from the data
+					GenerateProfilerDataFrame();
+
+					{
+						// add the frame to the work list
+						FScopeLock ScopeLock( &CriticalSection );
+						DataFrames.Add( CurrentData );
+						DataLoadingProgress = (double)Reader.Tell() / (double)Reader.TotalSize();
+					}
+
+					if( bUseInAsync )
 					{
-						while( DataFrames.Num() )
+						if( DataFrames.Num() > FProfilerClientManager::MaxFramesPerTick )
 						{
-							FPlatformProcess::Sleep( 0.001f );
+							while( DataFrames.Num() )
+							{
+								FPlatformProcess::Sleep( 0.001f );
+							}
 						}
 					}
 				}
-			}
 
-			new (Messages) FStatMessage( Message );
-		}
-		else
-		{
-			break;
-		}
+				new (Messages)FStatMessage( Message );
+			}
+			else
+			{
+				break;
+			}
 
-		if( !bUseInAsync && DataFrames.Num() < FProfilerClientManager::MaxFramesPerTick )
-		{
-			return false;
+			if( !bUseInAsync && DataFrames.Num() < FProfilerClientManager::MaxFramesPerTick )
+			{
+				return false;
+			}
 		}
 	}
-
-	if( FileReader.Tell() >= FileReader.TotalSize() )
+	
+	if( Reader.Tell() >= Reader.TotalSize() )
 	{
 		return true;
 	}
diff --git a/Engine/Source/Developer/ProfilerClient/Private/ProfilerClientManager.h b/Engine/Source/Developer/ProfilerClient/Private/ProfilerClientManager.h
--- a/Engine/Source/Developer/ProfilerClient/Private/ProfilerClientManager.h
+++ b/Engine/Source/Developer/ProfilerClient/Private/ProfilerClientManager.h
@@ -103,7 +103,7 @@ struct FServiceConnection
 	 *
 	 * @return true if read a special marker that indicates end of file
 	 */
-	bool ReadAndConvertStatMessages( FArchive& FileReader, bool bUseInAsync );
+	bool ReadAndConvertStatMessages( FArchive& Reader, bool bUseInAsync );
 
 	/** Adds all collected stat messages to the current stats thread state. */
 	void AddCollectedStatMessages( FStatMessage Message );
diff --git a/Engine/Source/Developer/ProfilerService/Private/ProfilerServiceManager.cpp b/Engine/Source/Developer/ProfilerService/Private/ProfilerServiceManager.cpp
--- a/Engine/Source/Developer/ProfilerService/Private/ProfilerServiceManager.cpp
+++ b/Engine/Source/Developer/ProfilerService/Private/ProfilerServiceManager.cpp
@@ -360,10 +360,10 @@ void FProfilerServiceManager::StartCapture()
 	{
 		// @TODO yrx 2014-06-05 Needs to be done on the stats thread via the task graph task.
 		FString Filename = CreateProfileFilename( FStatConstants::StatsFileExtension, true );
-		LastStatFilename = FApp::GetInstanceName() + TEXT("_") + Filename;
+		LastStatsFilename = FApp::GetInstanceName() + TEXT("_") + Filename;
 		TSharedPtr<FStatsWriteFile, ESPMode::ThreadSafe> ArchivePtr = MakeShareable(new FStatsWriteFile());
 		Archive = ArchivePtr;
-		Archive->Start(LastStatFilename);
+		Archive->Start( LastStatsFilename, false );
 		if (!Archive->IsValid())
 		{
 			Archive = nullptr;
@@ -655,10 +655,10 @@ void FProfilerServiceManager::HandleServiceRequestMessage( const FProfilerServic
 	}
 	else if( Message.Request == EProfilerRequestType::PRT_SendLastCapturedFile )
 	{
-		if( LastStatFilename.IsEmpty() == false )
+		if( LastStatsFilename.IsEmpty() == false )
 		{
-			FileTransferRunnable->EnqueueFileToSend( LastStatFilename, Context->GetSender(), InstanceId );
-			LastStatFilename.Empty();
+			FileTransferRunnable->EnqueueFileToSend( LastStatsFilename, Context->GetSender(), InstanceId );
+			LastStatsFilename.Empty();
 		}
 	}
 }
@@ -693,7 +693,7 @@ void FProfilerServiceManager::HandleServiceSubscribeMessage( const FProfilerServ
 		FClientData Data;
 		Data.Active = true;
 		Data.Preview = false;
-		Data.StatsWriteFile.WriteHeader();
+		Data.StatsWriteFile.WriteHeader( false );
 
 		// add to the client list
 		ClientData.Add( MsgAddress, Data );
diff --git a/Engine/Source/Developer/ProfilerService/Private/ProfilerServiceManager.h b/Engine/Source/Developer/ProfilerService/Private/ProfilerServiceManager.h
--- a/Engine/Source/Developer/ProfilerService/Private/ProfilerServiceManager.h
+++ b/Engine/Source/Developer/ProfilerService/Private/ProfilerServiceManager.h
@@ -163,7 +163,7 @@ private:
 	class FFileTransferRunnable* FileTransferRunnable;
 
 	/** Filename of last capture file. */
-	FString LastStatFilename;
+	FString LastStatsFilename;
 
 	/** Stat meta data */
 	FStatMetaData MetaData;
diff --git a/Engine/Source/Runtime/Core/Private/Stats/StatsFile.cpp b/Engine/Source/Runtime/Core/Private/Stats/StatsFile.cpp
--- a/Engine/Source/Runtime/Core/Private/Stats/StatsFile.cpp
+++ b/Engine/Source/Runtime/Core/Private/Stats/StatsFile.cpp
@@ -14,6 +14,65 @@
 DECLARE_CYCLE_STAT( TEXT( "Stream File" ), STAT_StreamFile, STATGROUP_StatSystem );
 DECLARE_CYCLE_STAT( TEXT( "Wait For Write" ), STAT_StreamFileWaitForWrite, STATGROUP_StatSystem );
 
+/*-----------------------------------------------------------------------------
+	FAsyncWriteWorker
+-----------------------------------------------------------------------------*/
+
+/**
+*	Helper class used to save the capture stats data via the background thread.
+*	CAUTION!! Can exist only one instance at the same time. Synchronized via EnsureCompletion.
+*/
+class FAsyncWriteWorker : public FNonAbandonableTask
+{
+public:
+	/**
+	 *	Pointer to the instance of the stats write file.
+	 *	Generally speaking accessing this pointer by a different thread is not thread-safe.
+	 *	But in this specific case it is.
+	 *	@see SendTask
+	 */
+	FStatsWriteFile* Outer;
+
+	/** Data for the file. Moved via Exchange. */
+	TArray<uint8> Data;
+
+	/** Thread cycles for the last frame. Moved via Exchange. */
+	TMap<uint32, int64> ThreadCycles;
+
+
+	/** Constructor. */
+	FAsyncWriteWorker( FStatsWriteFile* InStatsWriteFile )
+		: Outer( InStatsWriteFile )
+	{
+		Exchange( Data, InStatsWriteFile->OutData );
+		Exchange( ThreadCycles, InStatsWriteFile->ThreadCycles );
+	}
+
+	/** Write compressed data to the file. */
+	void DoWork()
+	{
+		check( Data.Num() );
+		FArchive& Ar = *Outer->File;
+
+		// Seek to the end of the file.
+		Ar.Seek( Ar.TotalSize() );
+		const int64 FrameFileOffset = Ar.Tell();
+
+		FCompressedStatsData CompressedData( Data, Outer->CompressedData );
+		Ar << CompressedData;
+
+		Outer->FramesInfo.Add( FStatsFrameInfo( FrameFileOffset, ThreadCycles ) );
+	}
+
+	/**
+	* @return	the name to display in external event viewers
+	*/
+	static const TCHAR *Name()
+	{
+		return TEXT( "FAsyncStatsWriteWorker" );
+	}
+};
+
 /*-----------------------------------------------------------------------------
 	FStatsThreadState
 -----------------------------------------------------------------------------*/
@@ -53,7 +112,7 @@ FStatsThreadState::FStatsThreadState( FString const& Filename )
 	check( bIsFinalized );
 	
 	TArray<FStatMessage> Messages;
-	if( Stream.Header.bRawStatFile )
+	if( Stream.Header.bRawStatsFile )
 	{
 		const int64 CurrentFilePos = FileReader->Tell();
 
@@ -168,13 +227,15 @@ void FStatsThreadState::AddMessages( TArray<FStatMessage>& InMessages )
 -----------------------------------------------------------------------------*/
 
 FStatsWriteFile::FStatsWriteFile()
-	: Filepos( 0 )
-	, File( nullptr )
+	: File( nullptr )
 	, AsyncTask( nullptr )
-{}
+{
+	// Reserve 1MB.
+	CompressedData.Reserve( EStatsFileConstants::MAX_COMPRESSED_SIZE );
+}
 
 
-void FStatsWriteFile::Start( FString const& InFilename )
+void FStatsWriteFile::Start( FString const& InFilename, bool bIsRawStatsFile )
 {
 	const FString PathName = *(FPaths::ProfilingDir() + TEXT( "UnrealStats/" ));
 
@@ -193,8 +254,7 @@ void FStatsWriteFile::Start( FString const& InFilename )
 	{
 		ArchiveFilename = Filename;
 		AddNewFrameDelegate();
-		WriteHeader();
-		SendTask();
+		WriteHeader( bIsRawStatsFile );
 		StatsMasterEnableAdd();
 	}
 }
@@ -207,7 +267,7 @@ void FStatsWriteFile::Stop()
 		RemoveNewFrameDelegate();
 		SendTask();
 		SendTask();
-		Finalize( File );
+		Finalize();
 
 		File->Close();
 		delete File;
@@ -218,24 +278,24 @@ void FStatsWriteFile::Stop()
 	}
 }
 
-void FStatsWriteFile::WriteHeader()
+void FStatsWriteFile::WriteHeader( bool bIsRawStatsFile )
 {
-	FMemoryWriter Ar( OutData, false, true );
+	FMemoryWriter MemoryWriter( OutData, false, true );
+	FArchive& Ar = File ? *File : MemoryWriter;
 
 	uint32 Magic = EStatMagicWithHeader::MAGIC;
 	// Serialize magic value.
 	Ar << Magic;
 
 	// Serialize dummy header, overwritten in Finalize.
-	Header.Version = EStatMagicWithHeader::VERSION_2;
+	Header.Version = EStatMagicWithHeader::VERSION_4;
 	Header.PlatformName = FPlatformProperties::PlatformName();
-	Header.bRawStatFile = false;
+	Header.bRawStatsFile = bIsRawStatsFile;
 	Ar << Header;
 
 	// Serialize metadata.
 	WriteMetadata( Ar );
-
-	Filepos += Ar.Tell();
+	Ar.Flush();
 }
 
 void FStatsWriteFile::WriteFrame( int64 TargetFrame, bool bNeedFullMetadata /*= false*/ )
@@ -256,7 +316,6 @@ void FStatsWriteFile::WriteFrame( int64 TargetFrame, bool bNeedFullMetadata /*=
 	}
 
 	// Get cycles for all threads, so we can use that data to generate the mini-view.
-	TMap<uint32, int64> ThreadCycles;
 	for( auto It = Stats.Threads.CreateConstIterator(); It; ++It )
 	{
 		const int64 Cycles = Stats.GetFastThreadFrameTime( TargetFrame, It.Key() );
@@ -265,19 +324,14 @@ void FStatsWriteFile::WriteFrame( int64 TargetFrame, bool bNeedFullMetadata /*=
 
 	// Serialize thread cycles. Disabled for now.
 	//Ar << ThreadCycles;
-
-	FramesInfo.Add( FStatsFrameInfo( Filepos, ThreadCycles ) );
-	Filepos += Ar.Tell() - PrevArPos;
 }
 
-void FStatsWriteFile::Finalize( FArchive* InFile )
+void FStatsWriteFile::Finalize()
 {
-	FArchive& Ar = *InFile;
+	FArchive& Ar = *File;
 
-	// Write special marker message.
-	static FStatNameAndInfo Adv(NAME_None, "", "", TEXT(""), EStatDataType::ST_int64, true, false);
-	FStatMessage SpecialMessage( Adv.GetEncodedName(), EStatOperation::SpecialMessageMarker, 0LL, false );
-	WriteMessage( Ar, SpecialMessage );
+	// Write dummy compression size, so we can detect the end of the file.
+	FCompressedStatsData::WriteEndOfCompressedData( Ar );
 
 	// Real header, written at start of the file, but written out right before we close the file.
 
@@ -341,7 +395,7 @@ void FStatsWriteFile::NewFrame( int64 TargetFrame )
 	{
 		MAX_NUM_RAWFRAMES = 120,
 	};
-	if( Header.bRawStatFile )
+	if( Header.bRawStatsFile )
 	{
 		if( FCommandStatsFile::FirstFrame == -1 )
 		{
@@ -356,13 +410,10 @@ void FStatsWriteFile::NewFrame( int64 TargetFrame )
 	}
 
 	WriteFrame( TargetFrame );
-	if( OutData.Num() > 1024 * 1024 )
-	{
-		SendTask();
-	}
+	SendTask();
 }
 
-void FStatsWriteFile::SendTask( int64 InDataOffset /*= -1*/ )
+void FStatsWriteFile::SendTask()
 {
 	if( AsyncTask )
 	{
@@ -373,8 +424,9 @@ void FStatsWriteFile::SendTask( int64 InDataOffset /*= -1*/ )
 	}
 	if( OutData.Num() )
 	{
-		AsyncTask = new FAsyncTask<FAsyncWriteWorker>( File, &OutData, InDataOffset );
+		AsyncTask = new FAsyncTask<FAsyncWriteWorker>( this );
 		check( !OutData.Num() );
+		check( !ThreadCycles.Num() );
 		AsyncTask->StartBackgroundTask();
 	}
 }
@@ -400,33 +452,9 @@ void FStatsWriteFile::RemoveNewFrameDelegate()
 	FRawStatsWriteFile
 -----------------------------------------------------------------------------*/
 
-void FRawStatsWriteFile::WriteHeader()
-{
-	FMemoryWriter Ar( OutData, false, true );
-
-	uint32 Magic = EStatMagicWithHeader::MAGIC;
-	// Serialize magic value.
-	Ar << Magic;
-
-	// Serialize dummy header, overwritten in Finalize.
-	Header.Version = EStatMagicWithHeader::VERSION_2;
-	Header.PlatformName = FPlatformProperties::PlatformName();
-	Header.bRawStatFile = true;
-	Ar << Header;
-
-	// Serialize metadata messages.
-	FStatsThreadState const& Stats = FStatsThreadState::GetLocalState();
-	int32 NumMetadataMessages = Stats.ShortNameToLongName.Num();
-	Ar << NumMetadataMessages;
-	WriteMetadata( Ar );
-
-	Filepos += Ar.Tell();
-}
-
 void FRawStatsWriteFile::WriteFrame( int64 TargetFrame, bool bNeedFullMetadata /*= false */)
 {
 	FMemoryWriter Ar( OutData, false, true );
-	const int64 PrevArPos = Ar.Tell();
 	const FStatsThreadState& Stats = FStatsThreadState::GetLocalState();
 
 	check( Stats.IsFrameValid( TargetFrame ) );
@@ -444,7 +472,6 @@ void FRawStatsWriteFile::WriteFrame( int64 TargetFrame, bool bNeedFullMetadata /
 	}
 
 	// Get cycles for all threads, so we can use that data to generate the mini-view.
-	TMap<uint32, int64> ThreadCycles;
 	for( auto It = Stats.Threads.CreateConstIterator(); It; ++It )
 	{
 		const int64 Cycles = Stats.GetFastThreadFrameTime( TargetFrame, It.Key() );
@@ -453,9 +480,6 @@ void FRawStatsWriteFile::WriteFrame( int64 TargetFrame, bool bNeedFullMetadata /
 	
 	// Serialize thread cycles.
 	Ar << ThreadCycles;
-
-	FramesInfo.Add( FStatsFrameInfo( Filepos, ThreadCycles ) );
-	Filepos += Ar.Tell() - PrevArPos;
 }
 
 /*-----------------------------------------------------------------------------
@@ -464,43 +488,43 @@ void FRawStatsWriteFile::WriteFrame( int64 TargetFrame, bool bNeedFullMetadata /
 
 FString FCommandStatsFile::LastFileSaved;
 int64 FCommandStatsFile::FirstFrame = -1;
-TSharedPtr<FStatsWriteFile, ESPMode::ThreadSafe> FCommandStatsFile::CurrentStatFile = nullptr;
+TSharedPtr<FStatsWriteFile, ESPMode::ThreadSafe> FCommandStatsFile::CurrentStatsFile = nullptr;
 
 void FCommandStatsFile::Start( const TCHAR* Cmd )
 {
-	CurrentStatFile = NULL;
+	CurrentStatsFile = NULL;
 	FString File;
 	FParse::Token( Cmd, File, false );
 	TSharedPtr<FStatsWriteFile, ESPMode::ThreadSafe> StatFile = MakeShareable( new FStatsWriteFile() );
-	CurrentStatFile = StatFile;
-	CurrentStatFile->Start( File );
-	if( !CurrentStatFile->IsValid() )
+	CurrentStatsFile = StatFile;
+	CurrentStatsFile->Start( File, false );
+	if( !CurrentStatsFile->IsValid() )
 	{
-		CurrentStatFile = nullptr;
+		CurrentStatsFile = nullptr;
 	}
 }
 
 void FCommandStatsFile::StartRaw( const TCHAR* Cmd )
 {
-	CurrentStatFile = NULL;
+	CurrentStatsFile = NULL;
 	FString File;
 	FParse::Token( Cmd, File, false );
 	TSharedPtr<FStatsWriteFile, ESPMode::ThreadSafe> StatFile = MakeShareable( new FRawStatsWriteFile() );
-	CurrentStatFile = StatFile;
-	CurrentStatFile->Start( File );
-	if( !CurrentStatFile->IsValid() )
+	CurrentStatsFile = StatFile;
+	CurrentStatsFile->Start( File, true );
+	if( !CurrentStatsFile->IsValid() )
 	{
-		CurrentStatFile = nullptr;
+		CurrentStatsFile = nullptr;
 	}
 }
 
 void FCommandStatsFile::Stop()
 {
-	if( CurrentStatFile.IsValid() )
+	if( CurrentStatsFile.IsValid() )
 	{
-		CurrentStatFile->Stop();
+		CurrentStatsFile->Stop();
 	}
-	CurrentStatFile = nullptr;
+	CurrentStatsFile = nullptr;
 }
 
 #endif // STATS
diff --git a/Engine/Source/Runtime/Core/Public/Stats/StatsFile.h b/Engine/Source/Runtime/Core/Public/Stats/StatsFile.h
--- a/Engine/Source/Runtime/Core/Public/Stats/StatsFile.h
+++ b/Engine/Source/Runtime/Core/Public/Stats/StatsFile.h
@@ -8,6 +8,8 @@
 
 #if	STATS
 
+class FAsyncWriteWorker;
+
 /**
 * Magic numbers for stats streams, this is for the first version.
 */
@@ -33,9 +35,151 @@ namespace EStatMagicWithHeader
 		MAGIC_SWAPPED = 0x47382910,
 		VERSION_2 = 2,
 		VERSION_3 = 3,
+
+		/**
+		 *	Added support for compressing the stats data, now each frame is compressed.
+		 */
+		VERSION_4 = 4,
+		HAS_COMPRESSED_DATA_VER = VERSION_4
 	};
 }
 
+struct EStatsFileConstants
+{
+	enum
+	{
+		/**
+		 *	Maximum size of the data that can be compressed.
+		 *	Should be ok for almost all cases, except the boot stats.
+		 */
+		MAX_COMPRESSED_SIZE = 1024 * 1024,
+
+		/** Header reserved for the compression internals. */
+		DUMMY_HEADER_SIZE = 1024,
+
+		/** Indicates the end of the compressed data. */
+		END_OF_COMPRESSED_DATA = 0xE0F0DA4A,
+
+		/** Indicates that the compression is disabled for the data. */
+		NO_COMPRESSION = 0,
+	};
+};
+
+/*-----------------------------------------------------------------------------
+	FCompressedStatsData
+-----------------------------------------------------------------------------*/
+
+/** Helper struct used to operate on the compressed data. */
+struct FCompressedStatsData
+{
+	/**
+	 * Initialization constructor 
+	 *	
+	 * @param SrcData - uncompressed data if saving, compressed data if loading
+	 * @param DestData - compressed data if saving, uncompressed data if loading
+	 *
+	 */
+	FCompressedStatsData( TArray<uint8>& InSrcData, TArray<uint8>& InDestData )
+		: SrcData( InSrcData )
+		, DestData( InDestData )
+		, bEndOfCompressedData( false )
+	{}
+
+	/**
+	 * Writes a special data to mark the end of the compressed data.
+	 */
+	static void WriteEndOfCompressedData( FArchive& Writer )
+	{
+		int32 Marker = EStatsFileConstants::END_OF_COMPRESSED_DATA;
+		check( Writer.IsSaving() );
+		Writer << Marker << Marker;
+	}
+
+protected:
+	/** Serialization operator. */
+	friend FArchive& operator << (FArchive& Ar, FCompressedStatsData& Data)
+	{
+		if( Ar.IsSaving() )
+		{
+			Data.WriteCompressed( Ar );
+		}
+		else if( Ar.IsLoading() )
+		{
+			Data.ReadCompressed( Ar );
+		}
+		else
+		{
+			check( 0 );
+		}
+		return Ar;
+	}
+
+	/** Compress the data and writes to the archive. */
+	void WriteCompressed( FArchive& Writer )
+	{
+		int32 UncompressedSize = SrcData.Num();
+		if( UncompressedSize > EStatsFileConstants::MAX_COMPRESSED_SIZE - EStatsFileConstants::DUMMY_HEADER_SIZE )
+		{
+			int32 DisabledCompressionSize = EStatsFileConstants::NO_COMPRESSION;
+			Writer << DisabledCompressionSize << UncompressedSize;
+			Writer.Serialize( SrcData.GetData(), UncompressedSize );
+		}
+		else
+		{
+			DestData.Reserve( EStatsFileConstants::MAX_COMPRESSED_SIZE );
+			int32 CompressedSize = DestData.GetAllocatedSize();
+
+			const bool bResult = FCompression::CompressMemory( COMPRESS_ZLIB, DestData.GetData(), CompressedSize, SrcData.GetData(), UncompressedSize );
+			check( bResult );
+			Writer << CompressedSize << UncompressedSize;
+			Writer.Serialize( DestData.GetData(), CompressedSize );
+		}
+	}
+
+	/** Reads the data and decompresses it. */
+	void ReadCompressed( FArchive& Reader )
+	{
+		int32 CompressedSize = 0;
+		int32 UncompressedSize = 0;
+		Reader << CompressedSize << UncompressedSize;
+
+		if( CompressedSize == EStatsFileConstants::END_OF_COMPRESSED_DATA && UncompressedSize == EStatsFileConstants::END_OF_COMPRESSED_DATA )
+		{
+			bEndOfCompressedData = true;
+		}
+		else
+		{
+			SrcData.Reset( CompressedSize );
+			DestData.Reset( UncompressedSize );
+			SrcData.AddUninitialized( CompressedSize );
+			DestData.AddUninitialized( UncompressedSize );
+
+			Reader.Serialize( SrcData.GetData(), CompressedSize );
+			const bool bResult = FCompression::UncompressMemory( COMPRESS_ZLIB, DestData.GetData(), UncompressedSize, SrcData.GetData(), CompressedSize );
+			check( bResult );
+		}
+	}
+
+public:
+	/**
+	 * @return true if we reached the end of the compressed data.
+	 */
+	const bool HasReachedEndOfCompressedData() const
+	{
+		return bEndOfCompressedData;
+	}
+	
+protected:
+	/** Reference to the source data. */
+	TArray<uint8>& SrcData;
+
+	/** Reference to the destination data. */
+	TArray<uint8>& DestData;
+
+	/** Set to true if we reached the end of the compressed data. */
+	bool bEndOfCompressedData;
+};
+
 /*-----------------------------------------------------------------------------
 	Stats file writing functionality
 -----------------------------------------------------------------------------*/
@@ -51,7 +195,7 @@ struct FStatsStreamHeader
 		, NumFNames( 0 )
 		, MetadataMessagesOffset( 0 )
 		, NumMetadataMessages( 0 )
-		, bRawStatFile( false )
+		, bRawStatsFile( false )
 	{}
 
 	/**
@@ -79,7 +223,7 @@ struct FStatsStreamHeader
 	uint64	NumMetadataMessages;
 
 	/** Whether this stats file uses raw data, required for thread view. */
-	bool bRawStatFile;
+	bool bRawStatsFile;
 
 	/** Whether this stats file has been finalized, so we have the whole data and can load file a unordered access. */
 	bool IsFinalized() const
@@ -87,6 +231,12 @@ struct FStatsStreamHeader
 		return NumMetadataMessages > 0 && MetadataMessagesOffset > 0 && FrameTableOffset > 0;
 	}
 
+	/** Whether this stats file contains compressed data and needs the special serializer to read that data. */
+	bool HasCompressedData() const
+	{
+		return Version >= EStatMagicWithHeader::HAS_COMPRESSED_DATA_VER;
+	}
+
 	/** Serialization operator. */
 	friend FArchive& operator << (FArchive& Ar, FStatsStreamHeader& Header)
 	{
@@ -109,7 +259,7 @@ struct FStatsStreamHeader
 		Ar << Header.MetadataMessagesOffset
 			<< Header.NumMetadataMessages;
 
-		Ar << Header.bRawStatFile;
+		Ar << Header.bRawStatsFile;
 
 		return Ar;
 	}
@@ -146,55 +296,11 @@ struct FStatsFrameInfo
 	TMap<uint32, int64> ThreadCycles;
 };
 
-/**
- *	Helper class used to save the capture stats data via the background thread.
- *	CAUTION!! Can exist only one instance at the same time.
- */
-class FAsyncWriteWorker : public FNonAbandonableTask
-{
-public:
-	/** File To write to **/
-	FArchive *File;
-	/** Data for the file **/
-	TArray<uint8> Data;
-	/** Offset of the data, -1 means append at the end of file. */
-	int64 DataOffset;
-
-	/** Constructor. */
-	FAsyncWriteWorker( FArchive *InFile, TArray<uint8>* InData, int64 InDataOffset = -1 )
-		: File( InFile )
-		, DataOffset( InDataOffset )
-	{
-		Exchange( Data, *InData );
-	}
-
-	/** Write the file  */
-	void DoWork()
-	{
-		check( Data.Num() );
-		if( DataOffset != -1 )
-		{
-			File->Seek( DataOffset );
-		}
-		else
-		{
-			// Seek to the end of the file.
-			File->Seek( File->TotalSize() );
-		}
-		File->Serialize( Data.GetTypedData(), Data.Num() );
-	}
-
-	/** 
-	* @return	the name to display in external event viewers
-	*/
-	static const TCHAR *Name()
-	{
-		return TEXT( "FAsyncStatsWriteWorker" );
-	}
-};
 
 struct CORE_API FStatsWriteFile : public TSharedFromThis<FStatsWriteFile, ESPMode::ThreadSafe>
 {
+	friend class FAsyncWriteWorker;
+
 protected:
 	/** Stats stream header. */
 	FStatsStreamHeader Header;
@@ -208,8 +314,8 @@ protected:
 	/** Data to write through the async task. **/
 	TArray<uint8> OutData;
 
-	/** Current position in the file, emulated to calculate the offsets. */
-	int64 Filepos;
+	/** Thread cycles for the last frame. */
+	TMap<uint32, int64> ThreadCycles;
 
 	/** Filename of the archive that we are writing to. */
 	FString ArchiveFilename;
@@ -220,6 +326,9 @@ protected:
 	/** Async task used to offload saving the capture data. */
 	FAsyncTask<FAsyncWriteWorker>* AsyncTask;
 
+	/** Buffer to store the compressed data, used by the FAsyncWriteWorker. */
+	TArray<uint8> CompressedData;
+
 public:
 	/** Constructor. **/
 	FStatsWriteFile();
@@ -229,7 +338,7 @@ public:
 	{}
 
 	/** Starts writing the stats data into the specified file. */
-	void Start( FString const& InFilename );
+	void Start( FString const& InFilename, bool bIsRawStatsFile );
 
 	/** Stops writing the stats data. */
 	void Stop();
@@ -246,23 +355,37 @@ public:
 		OutData.Reset();
 	}
 
-	/** Writes magic value, dummy header and initial metadata. */
-	virtual void WriteHeader();
+	/**
+	*	Writes magic value, dummy header and initial metadata.
+	*	Called from the stats thread.
+	*/
+	void WriteHeader( bool bIsRawStatsFile );
 
-	/** Grabs a frame from the local FStatsThreadState and adds it to the output **/
+	/**
+	 *	Grabs a frame from the local FStatsThreadState and adds it to the output.
+	 *	Called from the stats thread, but the data is saved using the the FAsyncWriteWorker. 
+	 */
 	virtual void WriteFrame( int64 TargetFrame, bool bNeedFullMetadata = false );
 
 protected:
-	/** Finalizes writing to the file, called directly form the same thread. */
-	void Finalize( FArchive* File );
+	/**
+	 *	Finalizes writing to the file.
+	 *	Called from the stats thread.
+	 */
+	void Finalize();
 
 	void NewFrame( int64 TargetFrame );
-	void SendTask( int64 InDataOffset = -1 );
+
+	/** Sends the data to the file via async task. */
+	void SendTask();
 
 	void AddNewFrameDelegate();
 	void RemoveNewFrameDelegate();
 
-	/** Writes metadata messages into the stream. */
+	/**
+	 *	Writes metadata messages into the stream. 
+	 *	Called from the stats thread.
+	 */
 	void WriteMetadata( FArchive& Ar )
 	{
 		const FStatsThreadState& Stats = FStatsThreadState::GetLocalState();
@@ -328,10 +451,11 @@ public:
 	{}
 
 protected:
-	/** Writes magic value, dummy header and initial metadata. */
-	virtual void WriteHeader() OVERRIDE;
 
-	/** Grabs a frame of raw message from the local FStatsThreadState and adds it to the output. */
+	/**
+	*	Grabs a frame from the local FStatsThreadState and adds it to the output.
+	*	Called from the stats thread, but the data is saved using the the FAsyncWriteWorker.
+	*/
 	virtual void WriteFrame( int64 TargetFrame, bool bNeedFullMetadata = false ) OVERRIDE;
 
 	/** Write a stat packed into the specified archive. */
@@ -571,7 +695,7 @@ struct FCommandStatsFile
 	static int64 FirstFrame;
 
 	/** Stats file that is currently being saved. */
-	static TSharedPtr<FStatsWriteFile, ESPMode::ThreadSafe> CurrentStatFile;
+	static TSharedPtr<FStatsWriteFile, ESPMode::ThreadSafe> CurrentStatsFile;
 
 	/** Stat StartFile. */
 	static void Start( const TCHAR* Cmd );
