Allow OEM keys to be used as FKey

diff --git a/Engine/Source/Runtime/Core/Private/Windows/WindowsPlatformMisc.cpp b/Engine/Source/Runtime/Core/Private/Windows/WindowsPlatformMisc.cpp
--- a/Engine/Source/Runtime/Core/Private/Windows/WindowsPlatformMisc.cpp
+++ b/Engine/Source/Runtime/Core/Private/Windows/WindowsPlatformMisc.cpp
@@ -866,10 +866,44 @@ uint32 FWindowsPlatformMisc::GetKeyMap( uint16* KeyCodes, FString* KeyNames, uin
 		ADDKEYMAP( VK_RMENU, TEXT("RightAlt") );
 		ADDKEYMAP( VK_LWIN, TEXT("LeftCommand") );
 		ADDKEYMAP( VK_RWIN, TEXT("RightCommand") );
+
+		TMap<uint16, uint16> ScanToVKMap;
+#define MAP_OEM_VK_TO_SCAN(KeyCode) { uint16 CharCode = MapVirtualKey(KeyCode,2); if (CharCode != 0) { ScanToVKMap.Add(CharCode,KeyCode); } }
+		MAP_OEM_VK_TO_SCAN(VK_OEM_1);
+		MAP_OEM_VK_TO_SCAN(VK_OEM_2);
+		MAP_OEM_VK_TO_SCAN(VK_OEM_3);
+		MAP_OEM_VK_TO_SCAN(VK_OEM_4);
+		MAP_OEM_VK_TO_SCAN(VK_OEM_5);
+		MAP_OEM_VK_TO_SCAN(VK_OEM_6);
+		MAP_OEM_VK_TO_SCAN(VK_OEM_7);
+		MAP_OEM_VK_TO_SCAN(VK_OEM_8);
+		MAP_OEM_VK_TO_SCAN(VK_OEM_PLUS);
+		MAP_OEM_VK_TO_SCAN(VK_OEM_COMMA);
+		MAP_OEM_VK_TO_SCAN(VK_OEM_MINUS);
+		MAP_OEM_VK_TO_SCAN(VK_OEM_PERIOD);
+		MAP_OEM_VK_TO_SCAN(VK_OEM_102);
+#undef  MAP_OEM_VK_TO_SCAN
+
+		static const uint32 MAX_KEY_MAPPINGS(256);
+		uint16 CharCodes[MAX_KEY_MAPPINGS];
+		FString CharKeyNames[MAX_KEY_MAPPINGS];
+		const int32 CharMappings = GetCharKeyMap(CharCodes, CharKeyNames, MAX_KEY_MAPPINGS);
+
+		for (int32 MappingIndex = 0; MappingIndex < CharMappings; ++MappingIndex)
+		{
+			ScanToVKMap.Remove(CharCodes[MappingIndex]);
+		}
+
+		for (auto It(ScanToVKMap.CreateConstIterator()); It; ++It)
+		{
+			ADDKEYMAP(It.Value(), FString::Chr(It.Key()));
+		}
 	}
 
 	check(NumMappings < MaxMappings);
 	return NumMappings;
+
+#undef ADDKEYMAP
 }
 
 void FWindowsPlatformMisc::SetUTF8Output()
diff --git a/Engine/Source/Runtime/InputCore/Private/InputCoreTypes.cpp b/Engine/Source/Runtime/InputCore/Private/InputCoreTypes.cpp
--- a/Engine/Source/Runtime/InputCore/Private/InputCoreTypes.cpp
+++ b/Engine/Source/Runtime/InputCore/Private/InputCoreTypes.cpp
@@ -432,6 +432,9 @@ void EKeys::Initialize()
 	AddKey(FKeyDetails(EKeys::Global_Pause, LOCTEXT("Global_Pause", "Global Pause"), FKeyDetails::GamepadKey));
 	AddKey(FKeyDetails(EKeys::Global_Play, LOCTEXT("Global_Play", "Global Play"), FKeyDetails::GamepadKey));
 	AddKey(FKeyDetails(EKeys::Global_Back, LOCTEXT("Global_Back", "Global Back"), FKeyDetails::GamepadKey));
+
+	// Initialize the input key manager.  This will cause any additional OEM keys to get added
+	FInputKeyManager::Get();
 }
 
 void EKeys::AddKey(const FKeyDetails& KeyDetails)
@@ -767,12 +770,14 @@ void FInputKeyManager::InitKeyMappings()
 
 	for (uint32 Idx=0; Idx<KeyMapSize; ++Idx)
 	{
-		const FKey Key(*KeyNames[Idx]);
+		FKey Key(*KeyNames[Idx]);
 		
-		if (ensureMsgf(Key.IsValid(), TEXT("Failed to get key for name %s"), *KeyNames[Idx]))
+		if (!Key.IsValid())
 		{
-			KeyMapVirtualToEnum.Add(KeyCodes[Idx], Key);
+			EKeys::AddKey(FKeyDetails(Key, Key.GetDisplayName()));
 		}
+
+		KeyMapVirtualToEnum.Add(KeyCodes[Idx], Key);
 	}
 	for (uint32 Idx=0; Idx<CharKeyMapSize; ++Idx)
 	{
