Check for FString::Replace being passed a null replacement string.

#ttp 337751: [GitHub] 208 : FString Replace() Crash Fix

#add Added an assert for when the crash, but will assert instead.

diff --git a/Engine/Source/Runtime/Core/Private/Containers/String.cpp b/Engine/Source/Runtime/Core/Private/Containers/String.cpp
--- a/Engine/Source/Runtime/Core/Private/Containers/String.cpp
+++ b/Engine/Source/Runtime/Core/Private/Containers/String.cpp
@@ -811,6 +811,9 @@ int32 FString::ParseIntoArrayWS( TArray<FString>* InArray, const TCHAR* pchExtra
 
 FString FString::Replace(const TCHAR* From, const TCHAR* To, ESearchCase::Type SearchCase) const
 {
+	// Previous code used to accidentally accept a NULL replacement string - this is no longer accepted.
+	check(To);
+
 	if (IsEmpty() || !From || !*From)
 	{
 		return *this;
