fixed inconsistent line breaks; renamed NULL to nullptr

diff --git a/Engine/Source/Runtime/Core/Public/Templates/SharedPointer.h b/Engine/Source/Runtime/Core/Public/Templates/SharedPointer.h
--- a/Engine/Source/Runtime/Core/Public/Templates/SharedPointer.h
+++ b/Engine/Source/Runtime/Core/Public/Templates/SharedPointer.h
@@ -51,10 +51,10 @@
  *
  *
  *	Tips:
- *		- Use TSharedRef instead of TSharedPtr whenever possible -- it can never be NULL!
+ *		- Use TSharedRef instead of TSharedPtr whenever possible -- it can never be nullptr!
  *		- You can call TSharedPtr::Reset() to release a reference to your object (and potentially deallocate) 
  *		- Use the MakeShareable() helper function to implicitly convert to TSharedRefs or TSharedPtrs
- *		- You can never reset a TSharedRef or assign it to NULL, but you can assign it a new object
+ *		- You can never reset a TSharedRef or assign it to nullptr, but you can assign it a new object
  *		- Shared pointers assume ownership of objects -- no need to call delete yourself!
  *		- Usually you should "operator new" when passing a C++ pointer to a new shared pointer
  *		- Use TSharedRef or TSharedPtr when passing smart pointers as function parameters, not TWeakPtr
@@ -87,7 +87,7 @@
  *		- No exceptions are allowed (all related features have been omitted)
  *		- Custom allocators and custom delete functions are not supported yet
  *		- Our implementation supports non-nullable smart pointers (TSharedRef)
- *		- Several other new features added, such as MakeShareable and NULL assignment
+ *		- Several other new features added, such as MakeShareable and nullptr assignment
  *
  *
  *	Why did we write our own Unreal shared pointer instead of using available alternatives?
@@ -97,7 +97,7 @@
  *		- Can work seamlessly with other Unreal containers and types
  *		- Better control over platform specifics, including threading and optimizations
  *		- We want thread-safety features to be optional (for performance)
- *		- We've added our own improvements (MakeShareable, assign to NULL, etc.)
+ *		- We've added our own improvements (MakeShareable, assign to nullptr, etc.)
  *		- Exceptions were not needed nor desired in our implementation
  *		- We wanted more control over performance (inlining, memory, use of virtuals, etc.)
  *		- Potentially easier to debug (liberal code comments, etc.)
@@ -105,7 +105,6 @@
  *
  */
 
-
 // SharedPointerInternals.h contains the implementation of reference counting structures we need
 #include "SharedPointerInternals.h"
 
@@ -122,15 +121,12 @@ FORCEINLINE TSharedRef< CastToType, Mode > StaticCastSharedRef( TSharedRef< Cast
 }
 
 
-/**
- * This tricky code is used to prevent the use of TSharedPtrs and TSharedRefs with UObjects
- */
 class UObjectBase;
 
 /**
- * TSharedRef is a non-nullable, non-intrusive reference-counted authoritative object reference.   This
- * shared reference  will be conditionally thread-safe when the optional Mode template argument is set
- * to ThreadSafe.
+ * TSharedRef is a non-nullable, non-intrusive reference-counted authoritative object reference.
+ *
+ * This shared reference will be conditionally thread-safe when the optional Mode template argument is set to ThreadSafe.
  */
 // NOTE: TSharedRef is an Unreal extension to standard smart pointer feature set
 template< class ObjectType, ESPMode::Type Mode >
@@ -138,54 +134,52 @@ class TSharedRef
 {
 	// TSharedRefs with UObjects are illegal.
 	checkAtCompileTime((!CanConvertPointerFromTo<ObjectType, UObjectBase>::Result), You_Cannot_Use_TSharedRefs_With_UObjects);
+
 public:
 
 	// NOTE: TSharedRef has no default constructor as it does not support empty references.  You must
 	//		 initialize your TSharedRef to a valid object at construction time.
 
-
 	/**
-	 * Constructs a shared reference that owns the specified object.  Must not be NULL.
+	 * Constructs a shared reference that owns the specified object.  Must not be nullptr.
 	 *
 	 * @param  InObject  Object this shared reference to retain a reference to
 	 */
 	template< class OtherType >
 	FORCEINLINE explicit TSharedRef( OtherType* InObject )
-		: Object( InObject ),
-		  SharedReferenceCount( InObject )
+		: Object( InObject )
+		, SharedReferenceCount( InObject )
 	{
-		// If the following assert goes off, it means a TSharedRef was initialized from a NULL object pointer.
-		// Shared references must never be NULL, so either pass a valid object or consider using TSharedPtr instead.
-		check( InObject != NULL );
+		// If the following assert goes off, it means a TSharedRef was initialized from a nullptr object pointer.
+		// Shared references must never be nullptr, so either pass a valid object or consider using TSharedPtr instead.
+		check( InObject != nullptr );
 
 		// If the object happens to be derived from TSharedFromThis, the following method
 		// will prime the object with a weak pointer to itself.
 		SharedPointerInternals::EnableSharedFromThis( this, InObject, InObject );
 	}
 
-
 	/**
 	 * Constructs a shared reference using a proxy reference to a raw pointer. (See MakeShareable())
-	 * Must not be NULL.
+	 * Must not be nullptr.
 	 *
 	 * @param  InRawPtrProxy  Proxy raw pointer that contains the object that the new shared reference will reference
 	 */
 	// NOTE: The following is an Unreal extension to standard shared_ptr behavior
 	template< class OtherType >
 	FORCEINLINE TSharedRef( SharedPointerInternals::FRawPtrProxy< OtherType > const& InRawPtrProxy )
-		: Object( InRawPtrProxy.Object ),
-		  SharedReferenceCount( InRawPtrProxy.Object )
+		: Object( InRawPtrProxy.Object )
+		, SharedReferenceCount( InRawPtrProxy.Object )
 	{
-		// If the following assert goes off, it means a TSharedRef was initialized from a NULL object pointer.
-		// Shared references must never be NULL, so either pass a valid object or consider using TSharedPtr instead.
-		check( InRawPtrProxy.Object != NULL );
+		// If the following assert goes off, it means a TSharedRef was initialized from a nullptr object pointer.
+		// Shared references must never be nullptr, so either pass a valid object or consider using TSharedPtr instead.
+		check( InRawPtrProxy.Object != nullptr );
 
 		// If the object happens to be derived from TSharedFromThis, the following method
 		// will prime the object with a weak pointer to itself.
 		SharedPointerInternals::EnableSharedFromThis( this, InRawPtrProxy.Object, InRawPtrProxy.Object );
 	}
 
-
 	/**
 	 * Constructs a shared reference as a reference to an existing shared reference's object.
 	 * This constructor is needed so that we can implicitly upcast to base classes.
@@ -194,11 +188,9 @@ public:
 	 */
 	template< class OtherType >
 	FORCEINLINE TSharedRef( TSharedRef< OtherType, Mode > const& InSharedRef )
-		: Object( InSharedRef.Object ),
-		  SharedReferenceCount( InSharedRef.SharedReferenceCount )
-	{
-	}
-
+		: Object( InSharedRef.Object )
+		, SharedReferenceCount( InSharedRef.SharedReferenceCount )
+	{ }
 
 	/**
 	 * Special constructor used internally to statically cast one shared reference type to another.  You
@@ -210,11 +202,9 @@ public:
 	 */
 	template< class OtherType >
 	FORCEINLINE TSharedRef( TSharedRef< OtherType, Mode > const& InSharedRef, SharedPointerInternals::FStaticCastTag )
-		: Object( static_cast< ObjectType* >( InSharedRef.Object ) ),
-		  SharedReferenceCount( InSharedRef.SharedReferenceCount )
-	{
-	}
-	
+		: Object( static_cast< ObjectType* >( InSharedRef.Object ) )
+		, SharedReferenceCount( InSharedRef.SharedReferenceCount )
+	{ }
   
 	/**
 	 * Special constructor used internally to cast a 'const' shared reference a 'mutable' reference.  You
@@ -226,12 +216,10 @@ public:
 	 */
 	template< class OtherType >
 	FORCEINLINE TSharedRef( TSharedRef< OtherType, Mode > const& InSharedRef, SharedPointerInternals::FConstCastTag )
-		: Object( const_cast< ObjectType* >( InSharedRef.Object ) ),
-		  SharedReferenceCount( InSharedRef.SharedReferenceCount )
-	{
-	}
-	
-  
+		: Object( const_cast< ObjectType* >( InSharedRef.Object ) )
+		, SharedReferenceCount( InSharedRef.SharedReferenceCount )
+	{ }
+	  
 	/**
 	 * Special constructor used internally to create a shared reference from an existing shared reference,
 	 * while using the specified object reference instead of the incoming shared reference's object
@@ -242,18 +230,18 @@ public:
 	 */
 	template< class OtherType >
 	FORCEINLINE TSharedRef( TSharedRef< OtherType, Mode > const& OtherSharedRef, ObjectType* InObject )
-		: Object( InObject ),
-		  SharedReferenceCount( OtherSharedRef.SharedReferenceCount )
-	{
-	}
+		: Object( InObject )
+		, SharedReferenceCount( OtherSharedRef.SharedReferenceCount )
+	{ }
+
 	FORCEINLINE TSharedRef( TSharedRef const& InSharedRef )
-		: Object( InSharedRef.Object ),
-		  SharedReferenceCount( InSharedRef.SharedReferenceCount )
-	{
-	}
+		: Object( InSharedRef.Object )
+		, SharedReferenceCount( InSharedRef.SharedReferenceCount )
+	{ }
+
 	FORCEINLINE TSharedRef( TSharedRef&& InSharedRef )
-		: Object( InSharedRef.Object ),
-		  SharedReferenceCount( InSharedRef.SharedReferenceCount )
+		: Object( InSharedRef.Object )
+		, SharedReferenceCount( InSharedRef.SharedReferenceCount )
 	{
 		// We're intentionally not moving here, because we don't want to leave InSharedRef in a
 		// null state, because that breaks the class invariant.  But we provide a move constructor
@@ -273,6 +261,7 @@ public:
 		Object = InSharedRef.Object;
 		return *this;
 	}
+
 	FORCEINLINE TSharedRef& operator=( TSharedRef&& InSharedRef )
 	{
 		FMemory::Memswap(this, &InSharedRef, sizeof(TSharedRef));
@@ -282,7 +271,7 @@ public:
 	/**
 	 * Assignment operator replaces this shared reference with the specified shared reference.  The object
 	 * currently referenced by this shared reference will no longer be referenced and will be deleted if
-	 * there are no other referencers.  Must not be NULL.
+	 * there are no other referencers.  Must not be nullptr.
 	 *
 	 * @param  InRawPtrProxy  Proxy object used to assign the object (see MakeShareable helper function)
 	 */
@@ -290,16 +279,14 @@ public:
 	template< class OtherType >
 	FORCEINLINE TSharedRef& operator=( SharedPointerInternals::FRawPtrProxy< OtherType > const& InRawPtrProxy )
 	{
-		// If the following assert goes off, it means a TSharedRef was initialized from a NULL object pointer.
-		// Shared references must never be NULL, so either pass a valid object or consider using TSharedPtr instead.
-		check( InRawPtrProxy.Object != NULL );
+		// If the following assert goes off, it means a TSharedRef was initialized from a nullptr object pointer.
+		// Shared references must never be nullptr, so either pass a valid object or consider using TSharedPtr instead.
+		check( InRawPtrProxy.Object != nullptr );
 
 		*this = TSharedRef< ObjectType, Mode >( InRawPtrProxy.Object );
 		return *this;
 	}
 
-
-
 	/**
 	 * Returns a C++ reference to the object this shared reference is referencing
 	 *
@@ -307,12 +294,11 @@ public:
 	 */
 	FORCEINLINE ObjectType& Get() const
 	{
-		// Should never be NULL as TSharedRef is never nullable
+		// Should never be nullptr as TSharedRef is never nullable
 		checkSlow( IsValid() );
 		return *Object;
 	}
 
-
 	/**
 	 * Dereference operator returns a reference to the object this shared pointer points to
 	 *
@@ -320,12 +306,11 @@ public:
 	 */
 	FORCEINLINE ObjectType& operator*() const
 	{
-		// Should never be NULL as TSharedRef is never nullable
+		// Should never be nullptr as TSharedRef is never nullable
 		checkSlow( IsValid() );
 		return *Object;
 	}
 
-
 	/**
 	 * Arrow operator returns a pointer to this shared reference's object
 	 *
@@ -333,12 +318,11 @@ public:
 	 */
 	FORCEINLINE ObjectType* operator->() const
 	{
-		// Should never be NULL as TSharedRef is never nullable
+		// Should never be nullptr as TSharedRef is never nullable
 		checkSlow( IsValid() );
 		return Object;
 	}
 
-
 	/**
 	 * Returns the number of shared references to this object (including this reference.)
 	 * IMPORTANT: Not necessarily fast!  Should only be used for debugging purposes!
@@ -350,7 +334,6 @@ public:
 		return SharedReferenceCount.GetSharedReferenceCount();
 	}
 
-
 	/**
 	 * Returns true if this is the only shared reference to this object.  Note that there may be
 	 * outstanding weak references left.
@@ -363,7 +346,6 @@ public:
 		return SharedReferenceCount.IsUnique();
 	}
 
-
 private:
 
 	/**
@@ -374,22 +356,22 @@ private:
 	 */
 	template< class OtherType >
 	FORCEINLINE explicit TSharedRef( TSharedPtr< OtherType, Mode > const& InSharedPtr )
-		: Object( InSharedPtr.Object ),
-		  SharedReferenceCount( InSharedPtr.SharedReferenceCount )
+		: Object( InSharedPtr.Object )
+		, SharedReferenceCount( InSharedPtr.SharedReferenceCount )
 	{
-		// If this assert goes off, it means a shared reference was created from a shared pointer that was NULL.
+		// If this assert goes off, it means a shared reference was created from a shared pointer that was nullptr.
 		// Shared references are never allowed to be null.  Consider using TSharedPtr instead.
 		check( IsValid() );
 	}
 
 	template< class OtherType >
 	FORCEINLINE explicit TSharedRef( TSharedPtr< OtherType, Mode >&& InSharedPtr )
-		: Object( InSharedPtr.Object ),
-		  SharedReferenceCount( MoveTemp(InSharedPtr.SharedReferenceCount) )
+		: Object( InSharedPtr.Object )
+		, SharedReferenceCount( MoveTemp(InSharedPtr.SharedReferenceCount) )
 	{
-		InSharedPtr.Object = NULL;
+		InSharedPtr.Object = nullptr;
 
-		// If this assert goes off, it means a shared reference was created from a shared pointer that was NULL.
+		// If this assert goes off, it means a shared reference was created from a shared pointer that was nullptr.
 		// Shared references are never allowed to be null.  Consider using TSharedPtr instead.
 		check( IsValid() );
 	}
@@ -402,10 +384,9 @@ private:
 	 */
 	FORCEINLINE const bool IsValid() const
 	{
-		return Object != NULL;
+		return Object != nullptr;
 	}
 
-
 	/**
 	 * Computes a hash code for this object
 	 *
@@ -418,7 +399,6 @@ private:
 		return ::PointerHash( InSharedRef.Object );
 	}
 
-
 	// We declare ourselves as a friend (templated using OtherType) so we can access members as needed
     template< class OtherType, ESPMode::Type OtherMode > friend class TSharedRef;
 
@@ -426,18 +406,17 @@ private:
     template< class OtherType, ESPMode::Type OtherMode > friend class TSharedPtr;
     template< class OtherType, ESPMode::Type OtherMode > friend class TWeakPtr;
 
-
 private:
 
-	/** The object we're holding a reference to.  Can be NULL. */
+	/** The object we're holding a reference to.  Can be nullptr. */
 	ObjectType* Object;
 
 	/** Interface to the reference counter for this object.  Note that the actual reference
 		controller object is shared by all shared and weak pointers that refer to the object */
 	SharedPointerInternals::FSharedReferencer< Mode > SharedReferenceCount;
-
 };
 
+
 /**
  * TSharedPtr is a non-intrusive reference-counted authoritative object pointer.  This shared pointer
  * will be conditionally thread-safe when the optional Mode template argument is set to ThreadSafe.
@@ -447,36 +426,34 @@ class TSharedPtr
 {
 	// TSharedPtrs with UObjects are illegal.
 	checkAtCompileTime((!CanConvertPointerFromTo<ObjectType, UObjectBase>::Result), You_Cannot_Use_TSharedPtrs_With_UObjects);
+
 public:
 
 	/**
 	 * Constructs an empty shared pointer
 	 */
 	// NOTE: FNullTag parameter is an Unreal extension to standard shared_ptr behavior
-	FORCEINLINE TSharedPtr( SharedPointerInternals::FNullTag* = NULL )
-		: Object( NULL ),
-		  SharedReferenceCount()
-	{
-	}
-
+	FORCEINLINE TSharedPtr( SharedPointerInternals::FNullTag* = nullptr )
+		: Object( nullptr )
+		, SharedReferenceCount()
+	{ }
 
 	/**
-	 * Constructs a shared pointer that owns the specified object.  Note that passing NULL here will
-	 * still create a tracked reference to a NULL pointer. (Consistent with std::shared_ptr)
+	 * Constructs a shared pointer that owns the specified object.  Note that passing nullptr here will
+	 * still create a tracked reference to a nullptr pointer. (Consistent with std::shared_ptr)
 	 *
 	 * @param  InObject  Object this shared pointer to retain a reference to
 	 */
 	template< class OtherType >
 	FORCEINLINE explicit TSharedPtr( OtherType* InObject )
-		: Object( InObject ),
-		  SharedReferenceCount( InObject )
+		: Object( InObject )
+		, SharedReferenceCount( InObject )
 	{
 		// If the object happens to be derived from TSharedFromThis, the following method
 		// will prime the object with a weak pointer to itself.
 		SharedPointerInternals::EnableSharedFromThis( this, InObject, InObject );
 	}
 
-
 	/**
 	 * Constructs a shared pointer using a proxy reference to a raw pointer. (See MakeShareable())
 	 *
@@ -485,15 +462,14 @@ public:
 	// NOTE: The following is an Unreal extension to standard shared_ptr behavior
 	template< class OtherType >
 	FORCEINLINE TSharedPtr( SharedPointerInternals::FRawPtrProxy< OtherType > const& InRawPtrProxy )
-		: Object( InRawPtrProxy.Object ),
-		  SharedReferenceCount( InRawPtrProxy.Object )
+		: Object( InRawPtrProxy.Object )
+		, SharedReferenceCount( InRawPtrProxy.Object )
 	{
 		// If the object happens to be derived from TSharedFromThis, the following method
 		// will prime the object with a weak pointer to itself.
 		SharedPointerInternals::EnableSharedFromThis( this, InRawPtrProxy.Object, InRawPtrProxy.Object );
 	}
 
-
 	/**
 	 * Constructs a shared pointer as a shared reference to an existing shared pointer's object.
 	 * This constructor is needed so that we can implicitly upcast to base classes.
@@ -502,23 +478,22 @@ public:
 	 */
 	template< class OtherType >
 	FORCEINLINE TSharedPtr( TSharedPtr< OtherType, Mode > const& InSharedPtr )
-		: Object( InSharedPtr.Object ),
-		  SharedReferenceCount( InSharedPtr.SharedReferenceCount )
-	{
-	}
+		: Object( InSharedPtr.Object )
+		, SharedReferenceCount( InSharedPtr.SharedReferenceCount )
+	{ }
+
 	FORCEINLINE TSharedPtr( TSharedPtr const& InSharedPtr )
-		: Object( InSharedPtr.Object ),
-		  SharedReferenceCount( InSharedPtr.SharedReferenceCount )
-	{
-	}
+		: Object( InSharedPtr.Object )
+		, SharedReferenceCount( InSharedPtr.SharedReferenceCount )
+	{ }
+
 	FORCEINLINE TSharedPtr( TSharedPtr&& InSharedPtr )
-		: Object( InSharedPtr.Object ),
-		  SharedReferenceCount( MoveTemp(InSharedPtr.SharedReferenceCount) )
+		: Object( InSharedPtr.Object )
+		, SharedReferenceCount( MoveTemp(InSharedPtr.SharedReferenceCount) )
 	{
-		InSharedPtr.Object = NULL;
+		InSharedPtr.Object = nullptr;
 	}
 
-
 	/**
 	 * Implicitly converts a shared reference to a shared pointer, adding a reference to the object.
 	 * NOTE: We allow an implicit conversion from TSharedRef to TSharedPtr because it's always a safe conversion.
@@ -528,14 +503,13 @@ public:
 	// NOTE: The following is an Unreal extension to standard shared_ptr behavior
 	template< class OtherType >
 	FORCEINLINE TSharedPtr( TSharedRef< OtherType, Mode > const& InSharedRef )
-		: Object( InSharedRef.Object ),
-		  SharedReferenceCount( InSharedRef.SharedReferenceCount )
+		: Object( InSharedRef.Object )
+		, SharedReferenceCount( InSharedRef.SharedReferenceCount )
 	{
 		// There is no rvalue overload of this constructor, because 'stealing' the pointer from a
 		// TSharedRef would leave it as null, which would invalidate its invariant.
 	}
 
-
 	/**
 	 * Special constructor used internally to statically cast one shared pointer type to another.  You
 	 * should never call this constructor directly.  Instead, use the StaticCastSharedPtr() function.
@@ -546,12 +520,10 @@ public:
 	 */
 	template< class OtherType >
 	FORCEINLINE TSharedPtr( TSharedPtr< OtherType, Mode > const& InSharedPtr, SharedPointerInternals::FStaticCastTag )
-		: Object( static_cast< ObjectType* >( InSharedPtr.Object ) ),
-		  SharedReferenceCount( InSharedPtr.SharedReferenceCount )
-	{
-	}
+		: Object( static_cast< ObjectType* >( InSharedPtr.Object ) )
+		, SharedReferenceCount( InSharedPtr.SharedReferenceCount )
+	{ }
 	
-  
 	/**
 	 * Special constructor used internally to cast a 'const' shared pointer a 'mutable' pointer.  You
 	 * should never call this constructor directly.  Instead, use the ConstCastSharedPtr() function.
@@ -562,11 +534,9 @@ public:
 	 */
 	template< class OtherType >
 	FORCEINLINE TSharedPtr( TSharedPtr< OtherType, Mode > const& InSharedPtr, SharedPointerInternals::FConstCastTag )
-		: Object( const_cast< ObjectType* >( InSharedPtr.Object ) ),
-		  SharedReferenceCount( InSharedPtr.SharedReferenceCount )
-	{
-	}
-	
+		: Object( const_cast< ObjectType* >( InSharedPtr.Object ) )
+		, SharedReferenceCount( InSharedPtr.SharedReferenceCount )
+	{ }
   
 	/**
 	 * Special constructor used internally to create a shared pointer from an existing shared pointer,
@@ -578,14 +548,12 @@ public:
 	 */
 	template< class OtherType >
 	FORCEINLINE TSharedPtr( TSharedPtr< OtherType, Mode > const& OtherSharedPtr, ObjectType* InObject )
-		: Object( InObject ),
-		  SharedReferenceCount( OtherSharedPtr.SharedReferenceCount )
-	{
-	}
-
+		: Object( InObject )
+		, SharedReferenceCount( OtherSharedPtr.SharedReferenceCount )
+	{ }
 
 	/**
-	 * Assignment to a NULL pointer.  The object currently referenced by this shared pointer will no longer be
+	 * Assignment to a nullptr pointer.  The object currently referenced by this shared pointer will no longer be
 	 * referenced and will be deleted if there are no other referencers.
 	 */
 	// NOTE: The following is an Unreal extension to standard shared_ptr behavior
@@ -595,7 +563,6 @@ public:
 		return *this;
 	}
 
-
 	/**
 	 * Assignment operator replaces this shared pointer with the specified shared pointer.  The object
 	 * currently referenced by this shared pointer will no longer be referenced and will be deleted if
@@ -609,12 +576,13 @@ public:
 		Object = InSharedPtr.Object;
 		return *this;
 	}
+
 	FORCEINLINE TSharedPtr& operator=( TSharedPtr&& InSharedPtr )
 	{
 		if (this != &InSharedPtr)
 		{
 			Object = InSharedPtr.Object;
-			InSharedPtr.Object = NULL;
+			InSharedPtr.Object = nullptr;
 			SharedReferenceCount = MoveTemp(InSharedPtr.SharedReferenceCount);
 		}
 		return *this;
@@ -635,7 +603,6 @@ public:
 		return *this;
 	}
 
-
 	/**
 	 * Converts a shared pointer to a shared reference.  The pointer *must* be valid or an assertion will trigger.
 	 *
@@ -644,24 +611,22 @@ public:
 	// NOTE: The following is an Unreal extension to standard shared_ptr behavior
 	FORCEINLINE TSharedRef< ObjectType, Mode > ToSharedRef() const
 	{
- 		// If this assert goes off, it means a shared reference was created from a shared pointer that was NULL.
+ 		// If this assert goes off, it means a shared reference was created from a shared pointer that was nullptr.
  		// Shared references are never allowed to be null.  Consider using TSharedPtr instead.
 		check( IsValid() );
 		return TSharedRef< ObjectType, Mode >( *this );
 	}
 
-
 	/**
-	 * Returns the object referenced by this pointer, or NULL if no object is reference
+	 * Returns the object referenced by this pointer, or nullptr if no object is reference
 	 *
-	 * @return  The object owned by this shared pointer, or NULL
+	 * @return  The object owned by this shared pointer, or nullptr
 	 */
 	FORCEINLINE ObjectType* Get() const
 	{
 		return Object;
 	}
 
-
 	/**
 	 * Checks to see if this shared pointer is actually pointing to an object
 	 *
@@ -669,10 +634,9 @@ public:
 	 */
 	FORCEINLINE const bool IsValid() const
 	{
-		return Object != NULL;
+		return Object != nullptr;
 	}
 
-
 	/**
 	 * Dereference operator returns a reference to the object this shared pointer points to
 	 *
@@ -684,7 +648,6 @@ public:
 		return *Object;
 	}
 
-
 	/**
 	 * Arrow operator returns a pointer to the object this shared pointer references
 	 *
@@ -696,7 +659,6 @@ public:
 		return Object;
 	}
 
-
 	/**
 	 * Resets this shared pointer, removing a reference to the object.  If there are no other shared
 	 * references to the object then it will be destroyed.
@@ -705,7 +667,6 @@ public:
 	{
  		*this = TSharedPtr< ObjectType, Mode >();
 	}
-	
 
 	/**
 	 * Returns the number of shared references to this object (including this reference.)
@@ -718,7 +679,6 @@ public:
 		return SharedReferenceCount.GetSharedReferenceCount();
 	}
 
-
 	/**
 	 * Returns true if this is the only shared reference to this object.  Note that there may be
 	 * outstanding weak references left.
@@ -731,7 +691,6 @@ public:
 		return SharedReferenceCount.IsUnique();
 	}
 
-
 private:
 
 	/**
@@ -745,8 +704,8 @@ private:
 	 */
 	template< class OtherType >
 	FORCEINLINE explicit TSharedPtr( TWeakPtr< OtherType, Mode > const& InWeakPtr )
-		: Object( NULL ),
-		  SharedReferenceCount( InWeakPtr.WeakReferenceCount )
+		: Object( nullptr )
+		, SharedReferenceCount( InWeakPtr.WeakReferenceCount )
 	{
 		// Check that the shared reference was created from the weak reference successfully.  We'll only
 		// cache a pointer to the object if we have a valid shared reference.
@@ -756,7 +715,6 @@ private:
 		}
 	}
 
-
 	/**
 	 * Computes a hash code for this object
 	 *
@@ -769,7 +727,6 @@ private:
 		return ::PointerHash( InSharedPtr.Object );
 	}
 
-
 	// We declare ourselves as a friend (templated using OtherType) so we can access members as needed
     template< class OtherType, ESPMode::Type OtherMode > friend class TSharedPtr;
 
@@ -778,10 +735,9 @@ private:
     template< class OtherType, ESPMode::Type OtherMode > friend class TWeakPtr;
     template< class OtherType, ESPMode::Type OtherMode > friend class TSharedFromThis;
 
-
 private:
 
-	/** The object we're holding a reference to.  Can be NULL. */
+	/** The object we're holding a reference to.  Can be nullptr. */
 	ObjectType* Object;
 
 	/** Interface to the reference counter for this object.  Note that the actual reference
@@ -790,8 +746,8 @@ private:
 
 };
 
-template<class ObjectType, ESPMode::Type Mode> struct TIsZeroConstructType<TSharedPtr<ObjectType, Mode>> { enum { Value = true }; };
 
+template<class ObjectType, ESPMode::Type Mode> struct TIsZeroConstructType<TSharedPtr<ObjectType, Mode>> { enum { Value = true }; };
 
 
 /**
@@ -801,17 +757,14 @@ template<class ObjectType, ESPMode::Type Mode> struct TIsZeroConstructType<TShar
 template< class ObjectType, ESPMode::Type Mode >
 class TWeakPtr
 {
-
 public:
 
 	/** Constructs an empty TWeakPtr */
 	// NOTE: FNullTag parameter is an Unreal extension to standard shared_ptr behavior
-	FORCEINLINE TWeakPtr( SharedPointerInternals::FNullTag* = NULL )
-		: Object( NULL ),
-		  WeakReferenceCount()
-	{
-	}
-
+	FORCEINLINE TWeakPtr( SharedPointerInternals::FNullTag* = nullptr )
+		: Object( nullptr )
+		, WeakReferenceCount()
+	{ }
 
 	/**
 	 * Constructs a weak pointer from a shared reference
@@ -821,11 +774,9 @@ public:
 	// NOTE: The following is an Unreal extension to standard shared_ptr behavior
 	template< class OtherType >
 	FORCEINLINE TWeakPtr( TSharedRef< OtherType, Mode > const& InSharedRef )
-		: Object( InSharedRef.Object ),
-		  WeakReferenceCount( InSharedRef.SharedReferenceCount )
-	{
-	}
-
+		: Object( InSharedRef.Object )
+		, WeakReferenceCount( InSharedRef.SharedReferenceCount )
+	{ }
 
 	/**
 	 * Constructs a weak pointer from a shared pointer
@@ -834,11 +785,9 @@ public:
 	 */
 	template< class OtherType >
 	FORCEINLINE TWeakPtr( TSharedPtr< OtherType, Mode > const& InSharedPtr )
-		: Object( InSharedPtr.Object ),
-		  WeakReferenceCount( InSharedPtr.SharedReferenceCount )
-	{
-	}
-
+		: Object( InSharedPtr.Object )
+		, WeakReferenceCount( InSharedPtr.SharedReferenceCount )
+	{ }
 
 	/**
 	 * Constructs a weak pointer from a weak pointer of another type.
@@ -848,32 +797,32 @@ public:
 	 */
 	template< class OtherType >
 	FORCEINLINE TWeakPtr( TWeakPtr< OtherType, Mode > const& InWeakPtr )
-		: Object( InWeakPtr.Object ),
-		  WeakReferenceCount( InWeakPtr.WeakReferenceCount )
-	{
-	}
+		: Object( InWeakPtr.Object )
+		, WeakReferenceCount( InWeakPtr.WeakReferenceCount )
+	{ }
+
 	template< class OtherType >
 	FORCEINLINE TWeakPtr( TWeakPtr< OtherType, Mode >&& InWeakPtr )
-		: Object            ( InWeakPtr.Object ),
-		  WeakReferenceCount( MoveTemp(InWeakPtr.WeakReferenceCount) )
+		: Object( InWeakPtr.Object )
+		, WeakReferenceCount( MoveTemp(InWeakPtr.WeakReferenceCount) )
 	{
-		InWeakPtr.Object = NULL;
+		InWeakPtr.Object = nullptr;
 	}
+
 	FORCEINLINE TWeakPtr( TWeakPtr const& InWeakPtr )
-		: Object            ( InWeakPtr.Object ),
-		  WeakReferenceCount( InWeakPtr.WeakReferenceCount )
-	{
-	}
+		: Object( InWeakPtr.Object )
+		, WeakReferenceCount( InWeakPtr.WeakReferenceCount )
+	{ }
+
 	FORCEINLINE TWeakPtr( TWeakPtr&& InWeakPtr )
-		: Object            ( InWeakPtr.Object ),
-		  WeakReferenceCount( MoveTemp(InWeakPtr.WeakReferenceCount) )
+		: Object( InWeakPtr.Object )
+		, WeakReferenceCount( MoveTemp(InWeakPtr.WeakReferenceCount) )
 	{
-		InWeakPtr.Object = NULL;
+		InWeakPtr.Object = nullptr;
 	}
 
-
 	/**
-	 * Assignment to a NULL pointer.  Clears this weak pointer's reference.
+	 * Assignment to a nullptr pointer.  Clears this weak pointer's reference.
 	 */
 	// NOTE: The following is an Unreal extension to standard shared_ptr behavior
 	FORCEINLINE TWeakPtr& operator=( SharedPointerInternals::FNullTag* )
@@ -882,7 +831,6 @@ public:
 		return *this;
 	}
 
-
 	/**
 	 * Assignment operator adds a weak reference to the object referenced by the specified weak pointer
 	 *
@@ -894,18 +842,18 @@ public:
 		WeakReferenceCount = InWeakPtr.WeakReferenceCount;
 		return *this;
 	}
+
 	FORCEINLINE TWeakPtr& operator=( TWeakPtr&& InWeakPtr )
 	{
 		if (this != &InWeakPtr)
 		{
 			Object             = InWeakPtr.Object;
-			InWeakPtr.Object   = NULL;
+			InWeakPtr.Object   = nullptr;
 			WeakReferenceCount = MoveTemp(InWeakPtr.WeakReferenceCount);
 		}
 		return *this;
 	}
 
-
 	/**
 	 * Assignment operator adds a weak reference to the object referenced by the specified weak pointer.
 	 * This assignment operator is intended to allow derived-to-base conversions.
@@ -919,16 +867,16 @@ public:
 		WeakReferenceCount = InWeakPtr.WeakReferenceCount;
 		return *this;
 	}
+
 	template <typename OtherType>
 	FORCEINLINE TWeakPtr& operator=( TWeakPtr<OtherType, Mode>&& InWeakPtr )
 	{
 		Object             = InWeakPtr.Object;
-		InWeakPtr.Object   = NULL;
+		InWeakPtr.Object   = nullptr;
 		WeakReferenceCount = MoveTemp(InWeakPtr.WeakReferenceCount);
 		return *this;
 	}
 
-
 	/**
 	 * Assignment operator sets this weak pointer from a shared reference
 	 *
@@ -943,7 +891,6 @@ public:
 		return *this;
 	}
 
-
 	/**
 	 * Assignment operator sets this weak pointer from a shared pointer
 	 *
@@ -957,7 +904,6 @@ public:
 		return *this;
 	}
 
-
 	/**
 	 * Converts this weak pointer to a shared pointer that you can use to access the object (if it
 	 * hasn't expired yet.)  Remember, if there are no more shared references to the object, the
@@ -971,7 +917,6 @@ public:
 		return TSharedPtr< ObjectType, Mode >( *this );
 	}
 
-
 	/**
 	 * Checks to see if this weak pointer actually has a valid reference to an object
 	 *
@@ -979,10 +924,9 @@ public:
 	 */
 	FORCEINLINE const bool IsValid() const
 	{
-		return Object != NULL && WeakReferenceCount.IsValid();
+		return Object != nullptr && WeakReferenceCount.IsValid();
 	}
 
-
 	/**
 	 * Resets this weak pointer, removing a weak reference to the object.  If there are no other shared
 	 * or weak references to the object, then the tracking object will be destroyed.
@@ -992,7 +936,6 @@ public:
 		*this = TWeakPtr< ObjectType, Mode >();
 	}
 
-
 	/**
 	 * Returns true if the object this weak pointer points to is the same as the specified object pointer.
 	 */
@@ -1001,10 +944,8 @@ public:
 		return Pin().Get() == InOtherPtr;
 	}
 
-
 private:
 	
-
 	/**
 	 * Computes a hash code for this object
 	 *
@@ -1017,17 +958,15 @@ private:
 		return ::PointerHash( InWeakPtr.Object );
 	}
 
-
 	// We declare ourselves as a friend (templated using OtherType) so we can access members as needed
     template< class OtherType, ESPMode::Type OtherMode > friend class TWeakPtr;
 
 	// Declare ourselves as a friend of TSharedPtr so we can access members as needed
     template< class OtherType, ESPMode::Type OtherMode > friend class TSharedPtr;
 
-
 private:
 
-	/** The object we have a weak reference to.  Can be NULL.  Also, it's important to note that because
+	/** The object we have a weak reference to.  Can be nullptr.  Also, it's important to note that because
 	    this is a weak reference, the object this pointer points to may have already been destroyed. */
 	ObjectType* Object;
 
@@ -1036,6 +975,7 @@ private:
 	SharedPointerInternals::FWeakReferencer< Mode > WeakReferenceCount;
 };
 
+
 template<class T, ESPMode::Type Mode> struct TIsWeakPointerType<TWeakPtr<T, Mode> > { enum { Value = true }; };
 template<class T, ESPMode::Type Mode> struct TIsZeroConstructType<TWeakPtr<T, Mode> > { enum { Value = true }; };
 
@@ -1047,7 +987,6 @@ template<class T, ESPMode::Type Mode> struct TIsZeroConstructType<TWeakPtr<T, Mo
 template< class ObjectType, ESPMode::Type Mode >
 class TSharedFromThis
 {
-
 public:
 
 	/**
@@ -1077,7 +1016,6 @@ public:
 		return SharedThis.ToSharedRef();
 	}
 
-
 	/**
 	 * Provides access to a shared reference to this object (const.)  Note that is only valid to call
 	 * this after a shared reference (or shared pointer) to the object has already been created.
@@ -1105,7 +1043,6 @@ public:
 		return SharedThis.ToSharedRef();
 	}
 
-
 protected:
 
 	/**
@@ -1122,7 +1059,6 @@ protected:
 		return StaticCastSharedRef< OtherType >( ThisPtr->AsShared() );
 	}
 
-
 	/**
 	 * Provides access to a shared reference to an object, given the object's 'this' pointer. Uses
 	 * the 'this' pointer to derive the object's actual type, then casts and returns an appropriately
@@ -1137,7 +1073,6 @@ protected:
 		return StaticCastSharedRef< OtherType const >( ThisPtr->AsShared() );
 	}
 
-
 public:		// @todo: Ideally this would be private, but template sharing problems prevent it
 
 	/**
@@ -1177,18 +1112,13 @@ public:		// @todo: Ideally this would be private, but template sharing problems
 		return WeakThis.IsValid();
 	}
 
-
 protected:
 
 	/** Hidden stub constructor */
-	TSharedFromThis()
-	{
-	}
+	TSharedFromThis() { }
 
 	/** Hidden stub copy constructor */
-	TSharedFromThis( TSharedFromThis const& )
-	{
-	}
+	TSharedFromThis( TSharedFromThis const& ) { }
 
 	/** Hidden stub assignment operator */
 	FORCEINLINE TSharedFromThis& operator=( TSharedFromThis const& )
@@ -1197,21 +1127,16 @@ protected:
 	}
 
 	/** Hidden destructor */
-	~TSharedFromThis()
-	{
-	}
-
+	~TSharedFromThis() { }
 
 private:
 
 	/** Weak reference to ourselves.  If we're destroyed then this weak pointer reference will be destructed
 	    with ourselves.  Note this is declared mutable only so that UpdateWeakReferenceInternal() can update it. */
 	mutable TWeakPtr< ObjectType, Mode > WeakThis;	
-
 };
 
 
-
 /**
  * Global equality operator for TSharedRef
  *
@@ -1392,7 +1317,6 @@ FORCEINLINE TSharedPtr< CastToType, Mode > StaticCastSharedPtr( TSharedPtr< Cast
 }
 
 
-
 /**
  * Casts a 'const' shared reference to 'mutable' shared reference. (const_cast)
  *
@@ -1405,7 +1329,6 @@ FORCEINLINE TSharedRef< CastToType, Mode > ConstCastSharedRef( TSharedRef< CastF
 }
 
 
-
 /**
  * Casts a 'const' shared pointer to 'mutable' shared pointer. (const_cast)
  *
@@ -1431,7 +1354,6 @@ FORCEINLINE SharedPointerInternals::FRawPtrProxy< ObjectType > MakeShareable( Ob
 }
 
 
-
 /**
  * Given a TArray of TWeakPtr's, will remove any invalid pointers.
  * @param  PointerArray  The pointer array to prune invalid pointers out of
@@ -1451,7 +1373,6 @@ FORCEINLINE void CleanupPointerArray(TArray< TWeakPtr<Type> >& PointerArray)
 }
 
 
-
 /**
  * Given a TMap of TWeakPtr's, will remove any invalid pointers. Not the most efficient.
  * @param  PointerMap  The pointer map to prune invalid pointers out of
@@ -1472,9 +1393,5 @@ FORCEINLINE void CleanupPointerMap(TMap< TWeakPtr<KeyType>, ValueType >& Pointer
 }
 
 
-
 // Shared pointer testing
 #include "SharedPointerTesting.h"
-
-
-
