Relaxed check in FMatrix::Inverse() to a non fatal ensure, since it's non fatal in shipping and test builds.

diff --git a/Engine/Source/Runtime/Core/Public/Math/UnrealMatrix.h b/Engine/Source/Runtime/Core/Public/Math/UnrealMatrix.h
--- a/Engine/Source/Runtime/Core/Public/Math/UnrealMatrix.h
+++ b/Engine/Source/Runtime/Core/Public/Math/UnrealMatrix.h
@@ -258,13 +258,14 @@ inline float FMatrix::RotDeterminant() const
 inline FMatrix FMatrix::Inverse() const
 {
 	// If we're in non final release, then make sure we're not creating NaNs
-#if 1 && !(UE_BUILD_SHIPPING || UE_BUILD_TEST)
+#if !(UE_BUILD_SHIPPING || UE_BUILD_TEST)
 	// Check for zero scale matrix to invert
 	if(	GetScaledAxis( EAxis::X ).IsNearlyZero(SMALL_NUMBER) && 
 		GetScaledAxis( EAxis::Y ).IsNearlyZero(SMALL_NUMBER) && 
 		GetScaledAxis( EAxis::Z ).IsNearlyZero(SMALL_NUMBER) ) 
 	{
-		UE_LOG(LogUnrealMath, Fatal, TEXT("FMatrix::Inverse(), trying to invert a NIL matrix, this results in NaNs! Use InverseSafe() instead."));
+		UE_LOG(LogUnrealMath, Error, TEXT("FMatrix::Inverse(), trying to invert a NIL matrix, this results in NaNs! Use InverseSafe() instead."));
+		ensureMsgf(false, TEXT("FMatrix::Inverse(), trying to invert a NIL matrix, this results in NaNs! Use InverseSafe() instead."));
 	}
 #endif
 	FMatrix Result;
