Core - Added FNoopStruct, may be used to avoid cluttering code with ifdefs

diff --git a/Engine/Source/Runtime/Core/Public/Templates/UnrealTemplate.h b/Engine/Source/Runtime/Core/Public/Templates/UnrealTemplate.h
--- a/Engine/Source/Runtime/Core/Public/Templates/UnrealTemplate.h
+++ b/Engine/Source/Runtime/Core/Public/Templates/UnrealTemplate.h
@@ -468,6 +468,16 @@ struct TForceInitAtBoot
 	}
 };
 
+/** Used to avoid cluttering code with ifdefs. */
+struct FNoopStruct
+{
+	FNoopStruct()
+	{}
+
+	~FNoopStruct()
+	{}
+};
+
 /**
  * Copies the cv-qualifiers from one type to another, e.g.:
  *
