Core - Added LockFreeList using 128-bit CAS operation (disabled by default, may be enabled via setting USE_LOCKFREELIST_128 to 1)

diff --git a/Engine/Source/Runtime/Core/Private/Containers/LockFreeList.cpp b/Engine/Source/Runtime/Core/Private/Containers/LockFreeList.cpp
--- a/Engine/Source/Runtime/Core/Private/Containers/LockFreeList.cpp
+++ b/Engine/Source/Runtime/Core/Private/Containers/LockFreeList.cpp
@@ -6,7 +6,56 @@
 
 DEFINE_LOG_CATEGORY(LogLockFreeList);
 
+FThreadSafeCounter FLockFreeListStats::NumOperations;
+FThreadSafeCounter FLockFreeListStats::Cycles;
 
+#if USE_LOCKFREELIST_128
+
+/*-----------------------------------------------------------------------------
+	FLockFreeVoidPointerListBase128
+-----------------------------------------------------------------------------*/
+
+FLockFreeVoidPointerListBase128::FLinkAllocator& FLockFreeVoidPointerListBase128::FLinkAllocator::Get()
+{
+	static FLinkAllocator TheAllocator;
+	return TheAllocator;
+}
+
+FLockFreeVoidPointerListBase128::FLinkAllocator::FLinkAllocator()
+	: SpecialClosedLink( FLargePointer( new FLink(), FLockFreeListConstants::SpecialClosedLink ) ) 
+{
+	if( !FPlatformAtomics::CanUseCompareExchange128() )
+	{
+		// This is a fatal error.
+		FPlatformMisc::MessageBoxExt( EAppMsgType::Ok, TEXT( "CPU does not support Compare and Exchange 128-bit operation. Unreal Engine 4 will exit now." ), TEXT( "Unsupported processor" ) );
+		FPlatformMisc::RequestExit( true );
+		UE_LOG( LogHAL, Fatal, TEXT( "CPU does not support Compare and Exchange 128-bit operation" ) );
+	}
+}
+	
+
+FLockFreeVoidPointerListBase128::FLinkAllocator::~FLinkAllocator()
+{
+ 	// - Deliberately leak to allow LockFree to be used during shutdown
+#if	0
+ 	check(SpecialClosedLink.Pointer); // double free?
+ 	check(!NumUsedLinks.GetValue());
+ 	while( FLink* ToDelete = FLink::Unlink(FreeLinks) )
+ 	{
+ 		delete ToDelete;
+ 		NumFreeLinks.Decrement();
+ 	}
+ 	check(!NumFreeLinks.GetValue());
+ 	delete SpecialClosedLink.Pointer;
+ 	SpecialClosedLink.Pointer = 0;
+#endif // 0
+}
+
+#else
+
+/*-----------------------------------------------------------------------------
+	FLockFreeVoidPointerListBase
+-----------------------------------------------------------------------------*/
 
 FLockFreeVoidPointerListBase::FLinkAllocator& FLockFreeVoidPointerListBase::FLinkAllocator::Get()
 {
@@ -23,16 +72,20 @@ FLockFreeVoidPointerListBase::FLinkAllocator::FLinkAllocator()
 
 FLockFreeVoidPointerListBase::FLinkAllocator::~FLinkAllocator()
 {
-	/** - Deliberately leak to allow LockFree to be used during shutdown
-	* check(SpecialClosedLink); // double free?
-	* check(!NumUsedLinks.GetValue());
-	* while (FLink* ToDelete = FLink::Unlink(&FreeLinks))
-	* {
-	* 	delete ToDelete;
-	* 	NumFreeLinks.Decrement();
-	* }
-	* check(!NumFreeLinks.GetValue());
-	* delete SpecialClosedLink;
-	* SpecialClosedLink = 0;
-	*/
+	// - Deliberately leak to allow LockFree to be used during shutdown
+#if	0
+	check(SpecialClosedLink); // double free?
+	check(!NumUsedLinks.GetValue());
+	while (FLink* ToDelete = FLink::Unlink(&FreeLinks))
+	{
+		delete ToDelete;
+		NumFreeLinks.Decrement();
+	}
+	check(!NumFreeLinks.GetValue());
+	delete SpecialClosedLink;
+	SpecialClosedLink = 0;
+#endif // 0
 }
+
+#endif //USE_LOCKFREELIST_128
+
diff --git a/Engine/Source/Runtime/Core/Public/Containers/LockFreeList.h b/Engine/Source/Runtime/Core/Public/Containers/LockFreeList.h
--- a/Engine/Source/Runtime/Core/Public/Containers/LockFreeList.h
+++ b/Engine/Source/Runtime/Core/Public/Containers/LockFreeList.h
@@ -6,820 +6,72 @@
 
 #pragma once
 
-DECLARE_LOG_CATEGORY_EXTERN(LogLockFreeList, Log, All);
+CORE_API DECLARE_LOG_CATEGORY_EXTERN(LogLockFreeList, Log, All);
 
 // what level of checking to perform...normally checkLockFreePointerList but could be ensure or check
-#define checkLockFreePointerList checkSlow
+#define checkLockFreePointerList checkSlow /*check*/
 
-#define CHECK_NON_CONCURRENT_ASSUMPTIONS (0)
+#define MONITOR_LINK_ALLOCATION (1)
 
-/** Fake Thread safe counter, used to avoid cluttering code with ifdefs when counters are only used for debugging */
-class FNoopCounter
+#if	MONITOR_LINK_ALLOCATION==1
+	typedef FThreadSafeCounter	FLockFreeListCounter;
+#else
+	typedef FNoopCounter		FLockFreeListCounter;
+#endif // MONITOR_LINK_ALLOCATION==1
+
+struct CORE_API FLockFreeListStats
 {
-public:
-	FNoopCounter()
-	{
-	}
-	FNoopCounter( const FNoopCounter& Other )
-	{
-	}
-	FNoopCounter( int32 Value )
-	{
-	}
-	int32 Increment()
-	{
-		return 0;
-	}
-	int32 Add( int32 Amount )
-	{
-		return 0;
-	}
-	int32 Decrement()
-	{
-		return 0;
-	}
-	int32 Subtract( int32 Amount )
-	{
-		return 0;
-	}
-	int32 Set( int32 Value )
-	{
-		return 0;
-	}
-	int32 Reset()
-	{
-		return 0;
-	}
-	int32 GetValue() const
-	{
-		return 0;
-	}
+	/** Number of lock free list operations like Push, Pop, PopAll etc. */
+	static FThreadSafeCounter NumOperations;
+	/** Cycles spent in lock free list operations like Push, Pop, PopAll etc. */
+	static FThreadSafeCounter Cycles;
 };
 
-#define MONITOR_LINK_ALLOCATION (0)
+#define MONITOR_LOCKFREELIST_PERFORMANCE (1)
 
-/** 
-* Base class for a lock free list of pointers 
-* This class has several different features. Not all methods are supported in all situations.
-* The intention is that derived classes use private inheritance to build a data structure only supports a small meaningful set of operations.
-*
-* A key feature is that lists can be "closed". A closed list cannot be pushed to or popped from. Closure is atomic with push and pop.
-* All operations are threadsafe, except that the constructor is not reentrant.
-**/
-class FLockFreeVoidPointerListBase
+/** Helper struct used to test the performance of the lock free list operations. */
+struct FLockFreeListPerfCounter
 {
-public:
-	/** Constructor - sets the list to empty and initializes the allocator if it hasn't already been intialized **/
-	FLockFreeVoidPointerListBase()
-		: Head(NULL)
-	{
-		FLinkAllocator::Get(); // construct allocator singleton if it isn't already
-	}
-	/** Destructor - checks that the list is either empty or closed. Lists should not be destroyed in any other state. **/
-	~FLockFreeVoidPointerListBase()
-	{
-		// leak these at shutdown
-		//checkLockFreePointerList(Head == NULL || Head == FLinkAllocator::Get().ClosedLink()); // we do not allow destruction except in the empty or closed state
-	}
+	const uint32 StartCycles;
 
-	/**	
-	 *	Push an item onto the head of the list
-	 *	@param NewItem, the new item to push on the list, cannot be NULL
-	 *	CAUTION: This method should not be used unless the list is known to not be closed.
-	**/
-	void Push(void *NewItem)
-	{
-		FLink* NewLink = FLinkAllocator::Get().AllocateLink(NewItem);
-		NewLink->Link(&Head);
-		checkLockFreePointerList(NewLink->Next != FLinkAllocator::Get().ClosedLink()); // it is not permissible to call push on a Closed queue
-	}
-
-	/**	
-	 *	Push an item onto the head of the list, unless the list is closed
-	 *	@param NewItem, the new item to push on the list, cannot be NULL
-	 *	@return true if the item was pushed on the list, false if the list was closed.
-	**/
-	bool PushIfNotClosed(void *NewItem)
-	{
-		FLink* NewLink = FLinkAllocator::Get().AllocateLink(NewItem);
-		bool bSuccess = NewLink->LinkIfHeadNotEqual(&Head, FLinkAllocator::Get().ClosedLink());
-		if (!bSuccess)
-		{
-			NewLink->Dispose();
-		}
-		return bSuccess;
-	}
-
-	/**	
-	 *	Push an item onto the head of the list, opening it first if necessary
-	 *	@param NewItem, the new item to push on the list, cannot be NULL
-	 *	@return true if the list needed to be opened first, false if the list was not closed before our push
-	**/
-	bool ReopenIfClosedAndPush(void *NewItem)
-	{
-		FLink* NewLink = FLinkAllocator::Get().AllocateLink(NewItem);
-		while (1)
-		{
-			if (NewLink->LinkIfHeadNotEqual(&Head, FLinkAllocator::Get().ClosedLink()))
-			{
-				return false;
-			}
-			if (NewLink->ReplaceHeadIfHeadEqual(&Head, FLinkAllocator::Get().ClosedLink()))
-			{
-				return true;
-			}
-		}
-	}
+	FLockFreeListPerfCounter()
+		: StartCycles( FPlatformTime::Cycles() )
+	{}
 
-	/**	
-	 *	Pop an item from the list or return NULL if the list is empty
-	 *	@return The popped item, if any
-	 *	CAUTION: This method should not be used unless the list is known to not be closed.
-	**/
-	void* Pop()
+	~FLockFreeListPerfCounter()
 	{
-		FLink* Link = FLink::Unlink(&Head);
-		if (!Link)
-		{
-			return NULL;
-		}
-		checkLockFreePointerList(Link != FLinkAllocator::Get().ClosedLink()); // it is not permissible to call pop on a Closed queue
-		void* Return = Link->Item;
-		Link->Dispose();
-		return Return;
+		const uint32 EndCycles = FPlatformTime::Cycles();
+		const int32 Delta = int32( (uint32)EndCycles - (uint32)StartCycles );
+		FLockFreeListStats::Cycles.Add( Delta );
+		FLockFreeListStats::NumOperations.Increment();
 	}
+};
 
-	/**	
-	 *	Pop an item from the list or return NULL if the list is empty or closed
-	 *	@return The new item, if any
-	**/
-	void* PopIfNotClosed()
-	{
-		FLink* Link = FLink::Unlink(&Head, FLinkAllocator::Get().ClosedLink());
-		if (!Link)
-		{
-			return NULL;
-		}
-		checkLockFreePointerList(Link != FLinkAllocator::Get().ClosedLink()); // internal error, this should not be possible
-		void* Return = Link->Item;
-		Link->Dispose();
-		return Return;
-	}
-
-	/**	
-	 *	Close the list if it is empty
-	 *	@return true if this call actively closed the list
-	**/
-	bool CloseIfEmpty()
-	{
-		if (FLinkAllocator::Get().ClosedLink()->ReplaceHeadIfHeadEqual(&Head,NULL))
-		{
-			return true;
-		}
-		return false;
-	}
-
-	/**	
-	 *	Pop all items from the list 
-	 *	@param Output The array to hold the returned items. Must be empty.
-	 *	CAUTION: This method should not be used unless the list is known to not be closed.
-	**/
-	template <class ARRAYTYPE, typename ElementType>
-	void PopAll(ARRAYTYPE& Output)
-	{
-		checkLockFreePointerList(!Output.Num());
-		FLink* Link = FLink::ReplaceList(&Head);
-		while (Link)
-		{
-			checkLockFreePointerList(Link != FLinkAllocator::Get().ClosedLink()); // it is not permissible to call PopAll on a Closed queue
-			checkLockFreePointerList(Link->Item); // we don't allow null items
-			Output.Add((ElementType)(Link->Item));
-			FLink* NextLink = Link->Next;
-			Link->Dispose();
-			Link = NextLink;
-		}
-	}
-
-	/**	
-	 *	Pop all items from the list and atomically close it.
-	 *	@param Output The array to hold the returned items. Must be empty.
-	**/
-	template <class ARRAYTYPE, typename ElementType>
-	void PopAllAndClose(ARRAYTYPE& Output)
-	{
-		checkLockFreePointerList(!Output.Num());
-		FLink* Link = FLink::ReplaceList(&Head, FLinkAllocator::Get().ClosedLink());
-		while (Link)
-		{
-			checkLockFreePointerList(Link != FLinkAllocator::Get().ClosedLink()); // we should not pop off the closed link
-			checkLockFreePointerList(Link->Item); // we don't allow null items
-			Output.Add((ElementType)(Link->Item));
-			FLink* NextLink = Link->Next;
-			Link->Dispose();
-			Link = NextLink;
-		}
-	}
-
-	/**	
-	 *	Check if the list is closed
-	 *	@return true if the list is closed.
-	 *	CAUTION: This methods safety depends on external assumptions. For example, if another thread could open or close the list at any time, the return value is no better than a best guess.
-	 *	As typically used, the list is only closed once it is lifetime, so a return of true here means it is closed forever.
-	**/
-	bool IsClosed() const  // caution, if this returns false, that does not mean the list is open!
-	{
-		FPlatformMisc::MemoryBarrier();
-		return Head == FLinkAllocator::Get().ClosedLink();
-	}
-
-	/**	
-	 *	Check if the list is empty
-	 *	@return true if the list is empty.
-	 *	CAUTION: This methods safety depends on external assumptions. For example, if another thread could add to the list at any time, the return value is no better than a best guess.
-	 *	As typically used, the list is not being access concurrently when this is called.
-	**/
-	bool IsEmpty() const  
-	{
-		FPlatformMisc::MemoryBarrier();
-		return Head == NULL;
-	}
-
-#if !(UE_BUILD_SHIPPING || UE_BUILD_TEST)
-	/**
-	 * Checks if item pointers look 'ok'. Used for debugging.
-	 */
-	bool CheckPointers() const
-	{
-		const int32 AlignmentCheck = MIN_ALIGNMENT - 1;
-		FPlatformMisc::MemoryBarrier();
-		for (FLink* Link = Head; Link; Link = Link->Next)
-		{
-			void* ItemPtr = Link->Item;
-			if (!ItemPtr || (UPTRINT)ItemPtr < 0x100)
-			{
-				return false;
-			}	
-			if ((UPTRINT)ItemPtr & AlignmentCheck)
-			{
-				return false;
-			}
-		}
-		return true;
-	}
-#endif
-
-	/**	Frees all of the links from the recycling queues.**/
-	static void FlushMemoryForShutdown();
-
-private:
-	/**	Internal single-linked list link structure. **/
-	struct FLink
-	{
-		/** Next item in the list or NULL if we are at the end of the list **/
-		FLink*				Next;
-		/** 
-		 * Pointer to the item this list holds. The is useful to the external users of the system. Must not be NULL for ordinary links. 
-		 * Also part of the solution to the ABA problem. NULL Items are always on the recycling list and NON-null items are not.
-		**/
-		void*				Item;
-		/** Part of the solution to the ABA problem. Links cannot be recycled until no threads hold pointers to them. When the lock count drops to zero, the link is recycled. **/
-		FThreadSafeCounter	LockCount;
-		/** Part of the solution to the ABA problem. Links that are free still need locking, but we need to make sure that links are only added to the free list when they are transitioning to the free list. **/
-		FThreadSafeCounter	MarkedForDeath;
-
-		/** Construcor, everything is intialized to zero. **/
-		FLink()
-			: Next(NULL)
-			, Item(NULL)
-		{
-		}
-
-		/**	
-		 *	Link this node into the head of the given list.
-		 *	@param HeadPointer; the head of the list
-		 *	CAUTION: Not checked here, but linking into a closed list with the routine would accidently and incorrectly open the list
-		**/
-		void Link(FLink** HeadPointer)
-		{
-			CheckNotMarkedForDeath();
-			while (1)
-			{
-				FLink*  LocalHeadPointer = LockLink(HeadPointer);
-				CheckNotMarkedForDeath();
-				SetNextPointer(LocalHeadPointer);
-				FLink* ValueWas = (FLink*)FPlatformAtomics::InterlockedCompareExchangePointer((void**)HeadPointer,this,LocalHeadPointer);
-				if (ValueWas == LocalHeadPointer)
-				{
-					if (LocalHeadPointer)
-					{
-						LocalHeadPointer->Unlock();
-					}
-					break;
-				}
-				SetNextPointer(NULL);
-				if (LocalHeadPointer)
-				{
-					LocalHeadPointer->Unlock();
-				}
-			}
-		}
-
-		/**	
-		 *	Link this node into the head of the given list, unless the head is a special, given node
-		 *	@param HeadPointer; the head of the list
-		 *	@param SpecialClosedLink; special link that is never recycled that indicates a closed list.
-		 *	@return true if this node was linked to the head.
-		**/
-		bool LinkIfHeadNotEqual(FLink** HeadPointer, FLink* SpecialClosedLink)
-		{
-			checkLockFreePointerList(!Next); // this should not be set yet
-			CheckNotMarkedForDeath();
-			while (1)
-			{
-				if (*HeadPointer == SpecialClosedLink)
-				{
-					CheckNotMarkedForDeath();
-					return false;
-				}
-				FLink*  LocalHeadPointer = LockLink(HeadPointer, SpecialClosedLink);
-				CheckNotMarkedForDeath();
-				// here if the list turns out to be locked now, the compare and swap will fail anyway
-				SetNextPointer(LocalHeadPointer);
-				FLink* ValueWas = (FLink*)FPlatformAtomics::InterlockedCompareExchangePointer((void**)HeadPointer,this,LocalHeadPointer);
-				if (ValueWas == LocalHeadPointer)
-				{
-					if (LocalHeadPointer)
-					{
-						LocalHeadPointer->Unlock();
-					}
-					break;
-				}
-				SetNextPointer(NULL);
-				if (LocalHeadPointer)
-				{
-					LocalHeadPointer->Unlock();
-				}
-			}
-			return true;
-		}
-		/**	
-		 *	If the head is a particular value, then replace it with this node. This is the primitive that is used to close and open lists
-		 *	@param HeadPointer; the head of the list
-		 *	@param SpecialTestLink; special link that is never recycled that indicates a closed list. Or NULL in the case that we are using this to open the list
-		 *	@return true if this node was linked to the head.
-		 *	CAUTION: The test link must NULL or a special link that is never recycled. The ABA problem is assumed to not occur here.
-		**/
-		bool ReplaceHeadIfHeadEqual(FLink** HeadPointer, FLink* SpecialTestLink)
-		{
-			CheckNotMarkedForDeath();
-			while (1)
-			{
-				if (*HeadPointer != SpecialTestLink)
-				{
-					return false;
-				}
-				CheckNotMarkedForDeath();
-				if (FPlatformAtomics::InterlockedCompareExchangePointer((void**)HeadPointer,this,SpecialTestLink) == SpecialTestLink)
-				{
-					break;
-				}
-			}
-			return true;
-		}
-		/**	
-		 *	Pop an item off the list, unless the list is empty or the head is a special, given node
-		 *	@param HeadPointer; the head of the list
-		 *	@param SpecialClosedLink; special link that is never recycled that indicates a closed list. Can be NULL for lists that are known to not be closed
-		 *	@return The link that was popped off, or NULL if the list was empty or closed.
-		 *	CAUTION: Not checked here, but if the list is closed and SpecialClosedLink is NULL, you will pop off the closed link, and that usually isn't what you want
-		**/
-		static FLink* Unlink(FLink** HeadPointer, FLink* SpecialClosedLink = NULL)
-		{
-			while (1)
-			{
-				FLink*  LocalHeadPointer = LockLink(HeadPointer, SpecialClosedLink);
-				if (!LocalHeadPointer)
-				{
-					break;
-				}
-				FLink* NextLink = LocalHeadPointer->Next;
-				FLink* ValueWas = (FLink*)FPlatformAtomics::InterlockedCompareExchangePointer((void**)HeadPointer,NextLink,LocalHeadPointer);
-				if (ValueWas == LocalHeadPointer)
-				{
-					checkLockFreePointerList(NextLink == LocalHeadPointer->Next);
-					LocalHeadPointer->SetNextPointer(NULL); // this is redundant, but we are not given the caller ownership of the other items
-					LocalHeadPointer->Unlock(true); // this cannot possibly result in a destruction of the link because we have taken sole possession of it
-					LocalHeadPointer->CheckNotMarkedForDeath(); // we are the exclusive owners at this point
-					return LocalHeadPointer;
-				}
-				LocalHeadPointer->Unlock();
-			}
-			return NULL;
-		}
-		/**	
-		 *	Replace the list with another list.h Use to either acquire all of the items or acquire all of the items and close the list.
-		 *	@param HeadPointer; the head of the list
-		 *	@param NewHeadLink; Head of the new list to replace this list with
-		 *	@return The original list before we replaced it.
-		 *	CAUTION: Not checked here, but if the list is closed this is probably not what you want. 
-		**/
-		static FLink* ReplaceList(FLink** HeadPointer, FLink* NewHeadLink = NULL)
-		{
-			while (1)
-			{
-				FLink*  LocalHeadPointer = LockLink(HeadPointer);
-				if (!LocalHeadPointer && !NewHeadLink)
-				{
-					// we are asking to replace NULL with NULL, this does not require an atomic (or any) operation
-					break;
-				}
-				checkLockFreePointerList(LocalHeadPointer != NewHeadLink); // replacing NULL with NULL is ok, but otherwise we cannot be replacing something with itself or we lose determination of who did what.
-				FLink* ValueWas = (FLink*)FPlatformAtomics::InterlockedCompareExchangePointer((void**)HeadPointer,NewHeadLink,LocalHeadPointer);
-				if (ValueWas == LocalHeadPointer)
-				{
-					if (LocalHeadPointer)
-					{
-						LocalHeadPointer->Unlock(true); // this cannot possibly result in a destruction of the link because we have taken sole possession of it
-						LocalHeadPointer->CheckNotMarkedForDeath(); // we are the exclusive owners at this point
-					}
-					return LocalHeadPointer;
-				}
-				if (LocalHeadPointer)
-				{
-					LocalHeadPointer->Unlock();
-				}
-			}
-			return NULL;  
-		}
-
-		/**	
-		 *	Internal function to safely grab the head pointer and increase the ref count of a link to avoid the ABA problem
-		 *	@param HeadPointer; the head of the list
-		 *	@param SpecialClosedLink; special link that is never recycled that indicates a closed list. Can be NULL for lists that are known to not be closed, or if you want to operate on the closed link.
-		 *	@return Pointer to the head link, with the ref count increased. Or NULL if the head is NULL or SpecialClosedLink
-		 *	CAUTION: Not checked here, but if the list is closed this is probably not what you want. 
-		**/
-		static FLink* LockLink(FLink** HeadPointer, FLink* SpecialClosedLink = NULL)
-		{
-			while (1)
-			{
-				FLink*  LocalHeadPointer = *HeadPointer;
-				if (!LocalHeadPointer || LocalHeadPointer == SpecialClosedLink)
-				{
-					return NULL;
-				}
-				LocalHeadPointer->LockCount.Increment();
-				FPlatformMisc::MemoryBarrier();
-				if (*HeadPointer == LocalHeadPointer)
-				{
-					return LocalHeadPointer;
-				}
-				LocalHeadPointer->Unlock();
-			}
-		}
-
-		void SetNextPointer(FLink* NewNext)
-		{
-			FLink* LocalNextPointer = Next;
-			checkLockFreePointerList(LocalNextPointer != NewNext || !NewNext); // should not be the same, except both can be NULL
-			checkLockFreePointerList(!LocalNextPointer || !NewNext); // one (or both) of them should be NULL
-#if CHECK_NON_CONCURRENT_ASSUMPTIONS
-			FLink* ValueWas = (FLink*)FPlatformAtomics::InterlockedCompareExchangePointer((void**)&Next, NewNext, LocalNextPointer); // not known to need to be interlocked
-			checkLockFreePointerList(ValueWas == LocalNextPointer); // this should always work because there should never be contention on setting next
-#else
-			Next = NewNext;
-#endif
-			FPlatformMisc::MemoryBarrier();
-			if (NewNext)
-			{
-				CheckNotMarkedForDeath(); // it shouldn't be possible for the node to die while we are setting a non-null next pointer
-			}
-		}
-
-		// verify that we are not marked for destruction
-		void CheckNotMarkedForDeath()
-		{
-			checkLockFreePointerList(MarkedForDeath.GetValue() == 0); // we don't allow null items
-		}
-
-		/**	
-		 *	Mark the node for disposal and unlock it
-		**/
-		void Dispose()
-		{
-			CheckNotMarkedForDeath();
-
-#if CHECK_NON_CONCURRENT_ASSUMPTIONS
-			{
-				void *CurrentItem = Item;
-				checkLockFreePointerList(CurrentItem); // we don't allow null items
-				void* ValueWas = FPlatformAtomics::InterlockedCompareExchangePointer(&Item,NULL,CurrentItem); // do not need an interlocked operation here
-				checkLockFreePointerList(ValueWas == CurrentItem); // there should be no concurrency here
-			}
-
-			{
-				FLink* LocalNextPointer = Next;
-				FLink* ValueWas = (FLink*)FPlatformAtomics::InterlockedCompareExchangePointer((void**)&Next, NULL, LocalNextPointer); // not known to need to be interlocked
-				checkLockFreePointerList(ValueWas == LocalNextPointer); // this should always work because there should never be contention on setting next
-			}
-#else
-			Item = NULL;
-			Next = NULL;
-#endif
-			FPlatformMisc::MemoryBarrier(); // memory barrier in here to make sure item and next are cleared before the link is recycled
-
-			int32 Death = MarkedForDeath.Increment();
-			checkLockFreePointerList(Death == 1); // there should be no concurrency here
-			Unlock();
-		}
-
-		/**	
-		 *	Internal function to unlock a link and recycle it if the ref count drops to zero and it hasn't already been recycled.
-		 *	This solves the ABA problem by delaying recycling until all pointers to this link are gone.
-		 *	@param bShouldNeverCauseAFree - if true, then this should never result in a free if it does, there is a run time error
-		**/
-		void Unlock(bool bShouldNeverCauseAFree = false);
-	};
-
-	/**	
-	 *	Class to allocate and recycle links.
-	**/
-	class FLinkAllocator
-	{
-	public:
-
-		/**	
-		 *	Return a new link.
-		 *	@param NewItem "item" pointer of the new node.
-		 *	@return New node, ready for use
-		**/
-		FLink* AllocateLink(void *NewItem)
-		{
-			checkLockFreePointerList(NewItem); // we don't allow null items
-#if MONITOR_LINK_ALLOCATION
-			if (NumUsedLinks.Increment() % 8192 == 1)
-			{
-				UE_LOG(LogLockFreeList, Log, TEXT("Number of links %d"),NumUsedLinks.GetValue());
-			}
-#endif
-			NumUsedLinks.Increment();
-
-			FLink*  NewLink = NULL;
-			NewLink = FLink::Unlink(&FreeLinks);
-			if (NewLink)
-			{
-				NumFreeLinks.Decrement();
-			}
-			else
-			{
-#if MONITOR_LINK_ALLOCATION
-				if (NumAllocatedLinks.Increment() % 10 == 1)
-				{
-					UE_LOG(LogLockFreeList, Log, TEXT("Number of allocated links %d"),NumAllocatedLinks.GetValue());
-				}
-#endif
-				NumAllocatedLinks.Increment();
-
-				NewLink = new FLink();
-				NewLink->LockCount.Increment();
-			}
-			checkLockFreePointerList(!NewLink->Item);
-			checkLockFreePointerList(!NewLink->Next);
-			checkLockFreePointerList(NewLink->LockCount.GetValue() >= 1);
-#if CHECK_NON_CONCURRENT_ASSUMPTIONS
-			void* ValueWas = FPlatformAtomics::InterlockedCompareExchangePointer(&NewLink->Item, NewItem, NULL); // in theory this doesn't need to be interlocked
-			checkLockFreePointerList(ValueWas == NULL);
-#else
-			NewLink->Item = NewItem;
-#endif
-			FPlatformMisc::MemoryBarrier();
-			return NewLink;
-		}
-		/**	
-		 *	Make a link available for recycling. 
-		 *	@param Link link to recycle
-		 *	CAUTION: Do not call this directly, it should only be called when the reference count drop to zero.
-		**/
-		void FreeLink(FLink* Link)
-		{
-			checkLockFreePointerList(Link != ClosedLink());  // very bad to recycle the special link
-			NumUsedLinks.Decrement();
-			Link->LockCount.Increment();
-			FPlatformMisc::MemoryBarrier();
-			Link->Link(&FreeLinks);
-			NumFreeLinks.Increment();
-		}
-		/**	
-		 *	Return a pointer to the special closed link
-		 *	@return pointer to special closed link
-		**/
-		FLink* ClosedLink()
-		{
-			checkLockFreePointerList(NULL != SpecialClosedLink);
-			return SpecialClosedLink;
-		}
-		/**	
-		 *	Singleton access
-		 *	@return the singleton for the link allocator
-		**/
-		static CORE_API FLinkAllocator& Get();
-
-	private:
-		/**	
-		 *	Constructor zeros the free link list and creates the special closed link
-		 *	@return the singleton for the link allocator
-		**/
-		FLinkAllocator();
-
-		/**	
-		 *	Destructor, should only be called when there are no outstanding links. 
-		 *	Frees the elements of the free list and frees the special link
-		 *	@return the singleton for the link allocator
-		**/
-		~FLinkAllocator();
-
-		/** Head to a list of free links that can be used for new allocations. **/
-		FLink*				FreeLinks;
 
-		/** Pointer to the special closed link. It should never be recycled. **/
-		FLink*				SpecialClosedLink;
-#if MONITOR_LINK_ALLOCATION
-		/** Total number of links outstanding and not in the free list **/
-		FThreadSafeCounter	NumUsedLinks; 
-		/** Total number of links in the free list **/
-		FThreadSafeCounter	NumFreeLinks;
-		/** Total number of links allocated **/
-		FThreadSafeCounter	NumAllocatedLinks; 
+#if	MONITOR_LOCKFREELIST_PERFORMANCE==1
+	typedef FLockFreeListPerfCounter	FLockFreeListPerfCounter;
 #else
-		/** If we aren't monitoring allocation, these do nothing **/
-		FNoopCounter		NumUsedLinks; 
-		FNoopCounter		NumFreeLinks;
-		FNoopCounter		NumAllocatedLinks; 
-#endif
-	};
-
-	/** Head of the list **/
-	MS_ALIGN(8) FLink*	Head;
-
-private:
-	FLockFreeVoidPointerListBase(const FLockFreeVoidPointerListBase&);
-	FLockFreeVoidPointerListBase& operator=(const FLockFreeVoidPointerListBase&);
-};
+	typedef FNoopStruct					FLockFreeListPerfCounter;
+#endif // MONITOR_LOCKFREELIST_PERFORMANCE==1
 
-/**	
-	*	Internal function to unlock a link and recycle it if the ref count drops to zero and it hasn't already been recycled.
-	*	This solves the ABA problem by delaying recycling until all pointers to this link are gone.
-	*	@param bShouldNeverCauseAFree - if true, then this should never result in a free if it does, there is a run time error
-**/
-inline void FLockFreeVoidPointerListBase::FLink::Unlock(bool bShouldNeverCauseAFree)
+/** Enumerates lock free list alignments. */
+namespace ELockFreeAligment
 {
-	FPlatformMisc::MemoryBarrier();
-	checkLockFreePointerList(LockCount.GetValue() > 0);
-	if (LockCount.Decrement() == 0)
+	enum
 	{
-		checkLockFreePointerList(MarkedForDeath.GetValue() < 2);
-		if (MarkedForDeath.Reset())
-		{
-			checkLockFreePointerList(!bShouldNeverCauseAFree); 
-			FLockFreeVoidPointerListBase::FLinkAllocator::Get().FreeLink(this);
-		}
-	}
+		LF64bit = 8,
+		LF128bit = 16,
+	};
 }
 
+/** Whether to use the lock free list based on 128-bit atomics. */
+#define USE_LOCKFREELIST_128 (0/*PLATFORM_HAS_128BIT_ATOMICS*/)
 
-/** 
- *	Wrapper around FLockFreeVoidPointerListBase
- *	Simple lock free list (stack) that does not permit closure.
- *	T is the type of the pointer the list will contain.
-**/
-template <class T>
-class TLockFreePointerList : private FLockFreeVoidPointerListBase 
-{
-public:
-	/**	
-	 *	Push an item onto the head of the list
-	 *	@param NewItem, the new item to push on the list, cannot be NULL
-	**/
-	void Push(T *NewItem)
-	{
-		FLockFreeVoidPointerListBase::Push(NewItem);
-	}
-
-	/**	
-	 *	Pop an item from the list or return NULL if the list is empty
-	 *	@return The popped item, if any
-	**/
-	T* Pop()
-	{
-		return (T*)FLockFreeVoidPointerListBase::Pop();
-	}
-
-	/**	
-	 *	Pop all items from the list 
-	 *	@param Output The array to hold the returned items. Must be empty.
-	**/
-	void PopAll(TArray<T *>& Output)
-	{
-		FLockFreeVoidPointerListBase::PopAll<TArray<T *>, T * >(Output);
-	}	
-
-	/**	
-	 *	Check if the list is empty
-	 *	@return true if the list is empty.
-	 *	CAUTION: This methods safety depends on external assumptions. For example, if another thread could add to the list at any time, the return value is no better than a best guess.
-	 *	As typically used, the list is not being access concurrently when this is called.
-	**/
-	bool IsEmpty() const  
-	{
-		return FLockFreeVoidPointerListBase::IsEmpty();
-	}
-
-};
-
-template <class T>
-class TClosableLockFreePointerList : private FLockFreeVoidPointerListBase 
-{
-public:
-	/**	
-	 *	Push an item onto the head of the list, unless the list is closed
-	 *	@param NewItem, the new item to push on the list, cannot be NULL
-	 *	@return true if the item was pushed on the list, false if the list was closed.
-	**/
-	bool PushIfNotClosed(T *NewItem)
-	{
-		return FLockFreeVoidPointerListBase::PushIfNotClosed(NewItem);
-	}
-
-	/**	
-	 *	Pop all items from the list and atomically close it.
-	 *	@param Output The array to hold the returned items. Must be empty.
-	**/
-	void PopAllAndClose(TArray<T *>& Output)
-	{
-		FLockFreeVoidPointerListBase::PopAllAndClose<TArray<T *>, T * >(Output);
-	}	
-
-	/**	
-	 *	Check if the list is closed
-	 *	@return true if the list is closed.
-	**/
-	bool IsClosed() const
-	{
-		return FLockFreeVoidPointerListBase::IsClosed();
-	}
-
-#if !(UE_BUILD_SHIPPING || UE_BUILD_TEST)
-	/**
-	 * Checks if the pointers look 'ok'. Used for debugging.
-	 */
-	bool CheckPointers() const
-	{
-		return FLockFreeVoidPointerListBase::CheckPointers();
-	}
+#if USE_LOCKFREELIST_128==1
+#include "LockFreeVoidPointerListBase128.h"
+#else
+#include "LockFreeVoidPointerListBase.h"
 #endif
-};
-
-template <class T>
-class TReopenableLockFreePointerList : private FLockFreeVoidPointerListBase 
-{
-public:
-	/**	
-	 *	Push an item onto the head of the list, opening it first if necessary
-	 *	@param NewItem, the new item to push on the list, cannot be NULL
-	 *	@return true if the list needed to be opened first, false if the list was not closed before our push
-	**/
-	bool ReopenIfClosedAndPush(T *NewItem)
-	{
-		bool bWasReopenedByMe = FLockFreeVoidPointerListBase::ReopenIfClosedAndPush(NewItem);
-		return bWasReopenedByMe;
-	}
-
-	/**	
-	 *	Pop an item from the list or return NULL if the list is empty or closed
-	 *	@return The new item, if any
-	**/
-	T* PopIfNotClosed()
-	{
-		return (T*)FLockFreeVoidPointerListBase::PopIfNotClosed();
-	}
-
-	/**	
-	 *	Pop all items from the list 
-	 *	@param Output The array to hold the returned items. Must be empty.
-	 *	CAUTION: This method should not be used unless the list is known to not be closed.
-	**/
-	void PopAll(TArray<T *>& Output)
-	{
-		FLockFreeVoidPointerListBase::PopAll<TArray<T *>, T * >(Output);
-	}	
-	/**	
-	 *	Close the list if it is empty
-	 *	@return true if this call actively closed the list
-	**/
-	bool CloseIfEmpty()
-	{
-		return FLockFreeVoidPointerListBase::CloseIfEmpty();
-	}
-};
 
+#include "LockFreeListImpl.h"
 
diff --git a/Engine/Source/Runtime/Core/Public/Containers/LockFreeListImpl.h b/Engine/Source/Runtime/Core/Public/Containers/LockFreeListImpl.h
new file mode 100644
--- /dev/null
+++ b/Engine/Source/Runtime/Core/Public/Containers/LockFreeListImpl.h
@@ -0,0 +1,195 @@
+// Copyright 1998-2013 Epic Games, Inc. All Rights Reserved.
+
+/*=============================================================================
+	LockFreeListImpl.h: Implements various lock free linked list containers.
+=============================================================================*/
+
+#pragma once
+
+/** 
+ *	Wrapper around FLockFreeVoidPointerListBase
+ *	Simple lock free list (stack) that does not permit closure.
+ *	T is the type of the pointer the list will contain.
+**/
+template <class T>
+class TLockFreePointerList : private FLockFreeVoidPointerListGeneric 
+{
+public:
+	/**	
+	 *	Push an item onto the head of the list
+	 *	@param NewItem, the new item to push on the list, cannot be NULL
+	**/
+	void Push(T *NewItem)
+	{
+		FLockFreeListPerfCounter LockFreePerfCounter;
+		FLockFreeVoidPointerListGeneric::Push(NewItem);
+	}
+
+	/**	
+	 *	Pop an item from the list or return NULL if the list is empty
+	 *	@return The popped item, if any
+	**/
+	T* Pop()
+	{
+		FLockFreeListPerfCounter LockFreePerfCounter;
+		return (T*)FLockFreeVoidPointerListGeneric::Pop();
+	}
+
+	/**	
+	 *	Pop all items from the list 
+	 *	@param Output The array to hold the returned items. Must be empty.
+	**/
+	void PopAll(TArray<T *>& Output)
+	{
+		FLockFreeListPerfCounter LockFreePerfCounter;
+		FLockFreeVoidPointerListGeneric::PopAll<TArray<T *>, T * >(Output);
+	}	
+
+	/**	
+	 *	Check if the list is empty
+	 *	@return true if the list is empty.
+	 *	CAUTION: This methods safety depends on external assumptions. For example, if another thread could add to the list at any time, the return value is no better than a best guess.
+	 *	As typically used, the list is not being access concurrently when this is called.
+	**/
+	bool IsEmpty() const  
+	{
+		return FLockFreeVoidPointerListGeneric::IsEmpty();
+	}
+
+#if	!(UE_BUILD_SHIPPING || UE_BUILD_TEST)
+	/**
+	 * Checks if the pointers look 'ok'. Used for debugging.
+	 */
+	bool CheckPointers() const
+	{
+		return FLockFreeVoidPointerListGeneric::CheckPointers();
+	}
+
+	/**
+	 * @return number of elements via iterating through this list, used for debugging.
+	 */
+	int32 NumVerified() const
+	{
+		return FLockFreeVoidPointerListGeneric::NumVerified();
+	}
+#endif // !(UE_BUILD_SHIPPING || UE_BUILD_TEST)
+
+};
+
+template <class T>
+class TClosableLockFreePointerList : private FLockFreeVoidPointerListGeneric 
+{
+public:
+	/**	
+	 *	Push an item onto the head of the list, unless the list is closed
+	 *	@param NewItem, the new item to push on the list, cannot be NULL
+	 *	@return true if the item was pushed on the list, false if the list was closed.
+	**/
+	bool PushIfNotClosed(T *NewItem)
+	{
+		FLockFreeListPerfCounter LockFreePerfCounter;
+		return FLockFreeVoidPointerListGeneric::PushIfNotClosed(NewItem);
+	}
+
+	/**	
+	 *	Pop all items from the list and atomically close it.
+	 *	@param Output The array to hold the returned items. Must be empty.
+	**/
+	void PopAllAndClose(TArray<T *>& Output)
+	{
+		FLockFreeListPerfCounter LockFreePerfCounter;
+		FLockFreeVoidPointerListGeneric::PopAllAndClose<TArray<T *>, T * >(Output);
+	}	
+
+	/**	
+	 *	Check if the list is closed
+	 *	@return true if the list is closed.
+	**/
+	bool IsClosed() const
+	{
+		return FLockFreeVoidPointerListGeneric::IsClosed();
+	}
+
+#if	!(UE_BUILD_SHIPPING || UE_BUILD_TEST)
+	/**
+	 * Checks if the pointers look 'ok'. Used for debugging.
+	 */
+	bool CheckPointers() const
+	{
+		return FLockFreeVoidPointerListGeneric::CheckPointers();
+	}
+
+	/**
+	 * @return number of elements via iterating through this list, used for debugging.
+	 */
+	int32 NumVerified() const
+	{
+		return FLockFreeVoidPointerListGeneric::NumVerified();
+	}
+#endif // !(UE_BUILD_SHIPPING || UE_BUILD_TEST)
+};
+
+template <class T>
+class TReopenableLockFreePointerList : private FLockFreeVoidPointerListGeneric
+{
+public:
+	/**	
+	 *	Push an item onto the head of the list, opening it first if necessary
+	 *	@param NewItem, the new item to push on the list, cannot be NULL
+	 *	@return true if the list needed to be opened first, false if the list was not closed before our push
+	**/
+	bool ReopenIfClosedAndPush(T *NewItem)
+	{
+		FLockFreeListPerfCounter LockFreePerfCounter;
+		bool bWasReopenedByMe = FLockFreeVoidPointerListGeneric::ReopenIfClosedAndPush(NewItem);
+		return bWasReopenedByMe;
+	}
+
+	/**	
+	 *	Pop an item from the list or return NULL if the list is empty or closed
+	 *	@return The new item, if any
+	**/
+	T* PopIfNotClosed()
+	{
+		FLockFreeListPerfCounter LockFreePerfCounter;
+		return (T*)FLockFreeVoidPointerListGeneric::PopIfNotClosed();
+	}
+
+	/**	
+	 *	Pop all items from the list 
+	 *	@param Output The array to hold the returned items. Must be empty.
+	 *	CAUTION: This method should not be used unless the list is known to not be closed.
+	**/
+	void PopAll(TArray<T *>& Output)
+	{
+		FLockFreeListPerfCounter LockFreePerfCounter;
+		FLockFreeVoidPointerListGeneric::PopAll<TArray<T *>, T * >(Output);
+	}	
+	/**	
+	 *	Close the list if it is empty
+	 *	@return true if this call actively closed the list
+	**/
+	bool CloseIfEmpty()
+	{
+		FLockFreeListPerfCounter LockFreePerfCounter;
+		return FLockFreeVoidPointerListGeneric::CloseIfEmpty();
+	}
+
+#if	!(UE_BUILD_SHIPPING || UE_BUILD_TEST)
+	/**
+	 * Checks if the pointers look 'ok'. Used for debugging.
+	 */
+	bool CheckPointers() const
+	{
+		return FLockFreeVoidPointerListGeneric::CheckPointers();
+	}
+
+	/**
+	 * @return number of elements via iterating through this list, used for debugging.
+	 */
+	int32 NumVerified() const
+	{
+		return FLockFreeVoidPointerListGeneric::NumVerified();
+	}
+#endif // !(UE_BUILD_SHIPPING || UE_BUILD_TEST)
+};
diff --git a/Engine/Source/Runtime/Core/Public/Containers/LockFreeVoidPointerListBase.h b/Engine/Source/Runtime/Core/Public/Containers/LockFreeVoidPointerListBase.h
new file mode 100644
--- /dev/null
+++ b/Engine/Source/Runtime/Core/Public/Containers/LockFreeVoidPointerListBase.h
@@ -0,0 +1,636 @@
+// Copyright 1998-2013 Epic Games, Inc. All Rights Reserved.
+
+/*=============================================================================
+	LockFreeVoidPointerListBase.h: A lock free single linked list
+=============================================================================*/
+
+#pragma once
+
+#define CHECK_NON_CONCURRENT_ASSUMPTIONS (0)
+
+/** 
+ * Base class for a lock free list of pointers 
+ * This class has several different features. Not all methods are supported in all situations.
+ * The intention is that derived classes use private inheritance to build a data structure only supports a small meaningful set of operations.
+ *
+ * A key feature is that lists can be "closed". A closed list cannot be pushed to or popped from. Closure is atomic with push and pop.
+ * All operations are threadsafe, except that the constructor is not reentrant.
+ */
+class FLockFreeVoidPointerListBase : public FNoncopyable
+{
+public:
+	/** Constructor - sets the list to empty and initializes the allocator if it hasn't already been initialized. */
+	FLockFreeVoidPointerListBase()
+		: Head(nullptr)
+	{
+		FLinkAllocator::Get(); // construct allocator singleton if it isn't already
+	}
+	/** Destructor - checks that the list is either empty or closed. Lists should not be destroyed in any other state. */
+	~FLockFreeVoidPointerListBase()
+	{
+		// leak these at shutdown
+		//checkLockFreePointerList(Head == nullptr || Head == FLinkAllocator::Get().ClosedLink()); // we do not allow destruction except in the empty or closed state
+	}
+
+	/**	
+	 *	Push an item onto the head of the list
+	 *	@param NewItem, the new item to push on the list, cannot be nullptr
+	 *	CAUTION: This method should not be used unless the list is known to not be closed.
+	 */
+	void Push(void *NewItem)
+	{
+		FLink* NewLink = FLinkAllocator::Get().AllocateLink(NewItem);
+		NewLink->Link(&Head);
+		checkLockFreePointerList(NewLink->Next != FLinkAllocator::Get().ClosedLink()); // it is not permissible to call push on a Closed queue
+	}
+
+	/**	
+	 *	Push an item onto the head of the list, unless the list is closed
+	 *	@param NewItem, the new item to push on the list, cannot be nullptr
+	 *	@return true if the item was pushed on the list, false if the list was closed.
+	 */
+	bool PushIfNotClosed(void *NewItem)
+	{
+		FLink* NewLink = FLinkAllocator::Get().AllocateLink(NewItem);
+		bool bSuccess = NewLink->LinkIfHeadNotEqual(&Head, FLinkAllocator::Get().ClosedLink());
+		if (!bSuccess)
+		{
+			NewLink->Dispose();
+		}
+		return bSuccess;
+	}
+
+	/**	
+	 *	Push an item onto the head of the list, opening it first if necessary
+	 *	@param NewItem, the new item to push on the list, cannot be nullptr
+	 *	@return true if the list needed to be opened first, false if the list was not closed before our push
+	 */
+	bool ReopenIfClosedAndPush(void *NewItem)
+	{
+		FLink* NewLink = FLinkAllocator::Get().AllocateLink(NewItem);
+		while (1)
+		{
+			if (NewLink->LinkIfHeadNotEqual(&Head, FLinkAllocator::Get().ClosedLink()))
+			{
+				return false;
+			}
+			if (NewLink->ReplaceHeadIfHeadEqual(&Head, FLinkAllocator::Get().ClosedLink()))
+			{
+				return true;
+			}
+		}
+	}
+
+	/**	
+	 *	Pop an item from the list or return nullptr if the list is empty
+	 *	@return The popped item, if any
+	 *	CAUTION: This method should not be used unless the list is known to not be closed.
+	 */
+	void* Pop()
+	{
+		FLink* Link = FLink::Unlink(&Head);
+		if (!Link)
+		{
+			return nullptr;
+		}
+		checkLockFreePointerList(Link != FLinkAllocator::Get().ClosedLink()); // it is not permissible to call pop on a Closed queue
+		void* Return = Link->Item;
+		Link->Dispose();
+		return Return;
+	}
+
+	/**	
+	 *	Pop an item from the list or return nullptr if the list is empty or closed
+	 *	@return The new item, if any
+	 */
+	void* PopIfNotClosed()
+	{
+		FLink* Link = FLink::Unlink(&Head, FLinkAllocator::Get().ClosedLink());
+		if (!Link)
+		{
+			return nullptr;
+		}
+		checkLockFreePointerList(Link != FLinkAllocator::Get().ClosedLink()); // internal error, this should not be possible
+		void* Return = Link->Item;
+		Link->Dispose();
+		return Return;
+	}
+
+	/**	
+	 *	Close the list if it is empty
+	 *	@return true if this call actively closed the list
+	 */
+	bool CloseIfEmpty()
+	{
+		if (FLinkAllocator::Get().ClosedLink()->ReplaceHeadIfHeadEqual(&Head,nullptr))
+		{
+			return true;
+		}
+		return false;
+	}
+
+	/**	
+	 *	Pop all items from the list 
+	 *	@param Output The array to hold the returned items. Must be empty.
+	 *	CAUTION: This method should not be used unless the list is known to not be closed.
+	 */
+	template <class ARRAYTYPE, typename ElementType>
+	void PopAll(ARRAYTYPE& Output)
+	{
+		checkLockFreePointerList(!Output.Num());
+		FLink* Link = FLink::ReplaceList(&Head);
+		while (Link)
+		{
+			checkLockFreePointerList(Link != FLinkAllocator::Get().ClosedLink()); // it is not permissible to call PopAll on a Closed queue
+			checkLockFreePointerList(Link->Item); // we don't allow nullptr items
+			Output.Add((ElementType)(Link->Item));
+			FLink* NextLink = Link->Next;
+			Link->Dispose();
+			Link = NextLink;
+		}
+	}
+
+	/**	
+	 *	Pop all items from the list and atomically close it.
+	 *	@param Output The array to hold the returned items. Must be empty.
+	 */
+	template <class ARRAYTYPE, typename ElementType>
+	void PopAllAndClose(ARRAYTYPE& Output)
+	{
+		checkLockFreePointerList(!Output.Num());
+		FLink* Link = FLink::ReplaceList(&Head, FLinkAllocator::Get().ClosedLink());
+		while (Link)
+		{
+			checkLockFreePointerList(Link != FLinkAllocator::Get().ClosedLink()); // we should not pop off the closed link
+			checkLockFreePointerList(Link->Item); // we don't allow nullptr items
+			Output.Add((ElementType)(Link->Item));
+			FLink* NextLink = Link->Next;
+			Link->Dispose();
+			Link = NextLink;
+		}
+	}
+
+	/**	
+	 *	Check if the list is closed
+	 *	@return true if the list is closed.
+	 *	CAUTION: This methods safety depends on external assumptions. For example, if another thread could open or close the list at any time, the return value is no better than a best guess.
+	 *	As typically used, the list is only closed once it is lifetime, so a return of true here means it is closed forever.
+	 */
+	bool IsClosed() const  // caution, if this returns false, that does not mean the list is open!
+	{
+		FPlatformMisc::MemoryBarrier();
+		return Head == FLinkAllocator::Get().ClosedLink();
+	}
+
+	/**	
+	 *	Check if the list is empty
+	 *	@return true if the list is empty.
+	 *	CAUTION: This methods safety depends on external assumptions. For example, if another thread could add to the list at any time, the return value is no better than a best guess.
+	 *	As typically used, the list is not being access concurrently when this is called.
+	 */
+	bool IsEmpty() const  
+	{
+		FPlatformMisc::MemoryBarrier();
+		return Head == nullptr;
+	}
+
+#if !(UE_BUILD_SHIPPING || UE_BUILD_TEST)
+	/**
+	 * Checks if item pointers look 'ok'. Used for debugging.
+	 */
+	bool CheckPointers() const
+	{
+		const int32 AlignmentCheck = ELockFreeAligment::LF64bit - 1;
+		FPlatformMisc::MemoryBarrier();
+		for (FLink* Link = Head; Link; Link = Link->Next)
+		{
+			void* ItemPtr = Link->Item;
+			if (!ItemPtr || (UPTRINT)ItemPtr < 0x100)
+			{
+				return false;
+			}	
+			if ((UPTRINT)ItemPtr & AlignmentCheck)
+			{
+				return false;
+			}
+		}
+		return true;
+	}
+
+	/**
+	 * @return number of elements via iterating through this list, used for debugging.
+	 */
+	int32 NumVerified() const
+	{
+		FPlatformMisc::MemoryBarrier();
+
+		int32 NumElements = 0;	
+		for (FLink* Link = Head; Link; Link = Link->Next)
+		{
+			NumElements++;
+		}
+		return NumElements;
+	}
+#endif
+
+private:
+	/**	Internal single-linked list link structure. */
+	struct FLink
+	{
+		/** Next item in the list or nullptr if we are at the end of the list */
+		FLink*				Next;
+		/** 
+		 * Pointer to the item this list holds. The is useful to the external users of the system. Must not be nullptr for ordinary links. 
+		 * Also part of the solution to the ABA problem. nullptr Items are always on the recycling list and NON-nullptr items are not.
+		 */
+		void*				Item;
+		/** Part of the solution to the ABA problem. Links cannot be recycled until no threads hold pointers to them. When the lock count drops to zero, the link is recycled. */
+		FThreadSafeCounter	LockCount;
+		/** Part of the solution to the ABA problem. Links that are free still need locking, but we need to make sure that links are only added to the free list when they are transitioning to the free list. */
+		FThreadSafeCounter	MarkedForDeath;
+
+		/** Construcor, everything is intialized to zero. */
+		FLink()
+			: Next(nullptr)
+			, Item(nullptr)
+		{
+		}
+
+		/**	
+		 *	Link this node into the head of the given list.
+		 *	@param HeadPointer; the head of the list
+		 *	CAUTION: Not checked here, but linking into a closed list with the routine would accidently and incorrectly open the list
+		 */
+		void Link(FLink** HeadPointer)
+		{
+			CheckNotMarkedForDeath();
+			while (1)
+			{
+				FLink*  LocalHeadPointer = LockLink(HeadPointer);
+				CheckNotMarkedForDeath();
+				SetNextPointer(LocalHeadPointer);
+				FLink* ValueWas = (FLink*)FPlatformAtomics::InterlockedCompareExchangePointer((void**)HeadPointer,this,LocalHeadPointer);
+				if (ValueWas == LocalHeadPointer)
+				{
+					if (LocalHeadPointer)
+					{
+						LocalHeadPointer->Unlock();
+					}
+					break;
+				}
+				SetNextPointer(nullptr);
+				if (LocalHeadPointer)
+				{
+					LocalHeadPointer->Unlock();
+				}
+			}
+		}
+
+		/**	
+		 *	Link this node into the head of the given list, unless the head is a special, given node
+		 *	@param HeadPointer; the head of the list
+		 *	@param SpecialClosedLink; special link that is never recycled that indicates a closed list.
+		 *	@return true if this node was linked to the head.
+		 */
+		bool LinkIfHeadNotEqual(FLink** HeadPointer, FLink* SpecialClosedLink)
+		{
+
+			checkLockFreePointerList(!Next); // this should not be set yet
+			CheckNotMarkedForDeath();
+			while (1)
+			{
+				if (*HeadPointer == SpecialClosedLink)
+				{
+					CheckNotMarkedForDeath();
+					return false;
+				}
+				FLink*  LocalHeadPointer = LockLink(HeadPointer, SpecialClosedLink);
+				CheckNotMarkedForDeath();
+				// here if the list turns out to be locked now, the compare and swap will fail anyway
+				SetNextPointer(LocalHeadPointer);
+				FLink* ValueWas = (FLink*)FPlatformAtomics::InterlockedCompareExchangePointer((void**)HeadPointer,this,LocalHeadPointer);
+				if (ValueWas == LocalHeadPointer)
+				{
+					if (LocalHeadPointer)
+					{
+						LocalHeadPointer->Unlock();
+					}
+					break;
+				}
+				SetNextPointer(nullptr);
+				if (LocalHeadPointer)
+				{
+					LocalHeadPointer->Unlock();
+				}
+			}
+			return true;
+		}
+		/**	
+		 *	If the head is a particular value, then replace it with this node. This is the primitive that is used to close and open lists
+		 *	@param HeadPointer; the head of the list
+		 *	@param SpecialTestLink; special link that is never recycled that indicates a closed list. Or nullptr in the case that we are using this to open the list
+		 *	@return true if this node was linked to the head.
+		 *	CAUTION: The test link must nullptr or a special link that is never recycled. The ABA problem is assumed to not occur here.
+		 */
+		bool ReplaceHeadIfHeadEqual(FLink** HeadPointer, FLink* SpecialTestLink)
+		{
+			CheckNotMarkedForDeath();
+			while (1)
+			{
+				if (*HeadPointer != SpecialTestLink)
+				{
+					return false;
+				}
+				CheckNotMarkedForDeath();
+				if (FPlatformAtomics::InterlockedCompareExchangePointer((void**)HeadPointer,this,SpecialTestLink) == SpecialTestLink)
+				{
+					break;
+				}
+			}
+			return true;
+		}
+		/**	
+		 *	Pop an item off the list, unless the list is empty or the head is a special, given node
+		 *	@param HeadPointer; the head of the list
+		 *	@param SpecialClosedLink; special link that is never recycled that indicates a closed list. Can be nullptr for lists that are known to not be closed
+		 *	@return The link that was popped off, or nullptr if the list was empty or closed.
+		 *	CAUTION: Not checked here, but if the list is closed and SpecialClosedLink is nullptr, you will pop off the closed link, and that usually isn't what you want
+		 */
+		static FLink* Unlink(FLink** HeadPointer, FLink* SpecialClosedLink = nullptr)
+		{
+			while (1)
+			{
+				FLink*  LocalHeadPointer = LockLink(HeadPointer, SpecialClosedLink);
+				if (!LocalHeadPointer)
+				{
+					break;
+				}
+				FLink* NextLink = LocalHeadPointer->Next;
+				FLink* ValueWas = (FLink*)FPlatformAtomics::InterlockedCompareExchangePointer((void**)HeadPointer,NextLink,LocalHeadPointer);
+				if (ValueWas == LocalHeadPointer)
+				{
+					checkLockFreePointerList(NextLink == LocalHeadPointer->Next);
+					LocalHeadPointer->SetNextPointer(nullptr); // this is redundant, but we are not given the caller ownership of the other items
+					LocalHeadPointer->Unlock(true); // this cannot possibly result in a destruction of the link because we have taken sole possession of it
+					LocalHeadPointer->CheckNotMarkedForDeath(); // we are the exclusive owners at this point
+					return LocalHeadPointer;
+				}
+				LocalHeadPointer->Unlock();
+			}
+			return nullptr;
+		}
+		/**	
+		 *	Replace the list with another list.h Use to either acquire all of the items or acquire all of the items and close the list.
+		 *	@param HeadPointer; the head of the list
+		 *	@param NewHeadLink; Head of the new list to replace this list with
+		 *	@return The original list before we replaced it.
+		 *	CAUTION: Not checked here, but if the list is closed this is probably not what you want. 
+		 */
+		static FLink* ReplaceList(FLink** HeadPointer, FLink* NewHeadLink = nullptr)
+		{
+			while (1)
+			{
+				FLink*  LocalHeadPointer = LockLink(HeadPointer);
+				if (!LocalHeadPointer && !NewHeadLink)
+				{
+					// we are asking to replace nullptr with nullptr, this does not require an atomic (or any) operation
+					break;
+				}
+				checkLockFreePointerList(LocalHeadPointer != NewHeadLink); // replacing nullptr with nullptr is ok, but otherwise we cannot be replacing something with itself or we lose determination of who did what.
+				FLink* ValueWas = (FLink*)FPlatformAtomics::InterlockedCompareExchangePointer((void**)HeadPointer,NewHeadLink,LocalHeadPointer);
+				if (ValueWas == LocalHeadPointer)
+				{
+					if (LocalHeadPointer)
+					{
+						LocalHeadPointer->Unlock(true); // this cannot possibly result in a destruction of the link because we have taken sole possession of it
+						LocalHeadPointer->CheckNotMarkedForDeath(); // we are the exclusive owners at this point
+					}
+					return LocalHeadPointer;
+				}
+				if (LocalHeadPointer)
+				{
+					LocalHeadPointer->Unlock();
+				}
+			}
+			return nullptr;  
+		}
+
+		/**	
+		 *	Internal function to safely grab the head pointer and increase the ref count of a link to avoid the ABA problem
+		 *	@param HeadPointer; the head of the list
+		 *	@param SpecialClosedLink; special link that is never recycled that indicates a closed list. Can be nullptr for lists that are known to not be closed, or if you want to operate on the closed link.
+		 *	@return Pointer to the head link, with the ref count increased. Or nullptr if the head is nullptr or SpecialClosedLink
+		 *	CAUTION: Not checked here, but if the list is closed this is probably not what you want. 
+		 */
+		static FLink* LockLink(FLink** HeadPointer, FLink* SpecialClosedLink = nullptr)
+		{
+			while (1)
+			{
+				FLink*  LocalHeadPointer = *HeadPointer;
+				if (!LocalHeadPointer || LocalHeadPointer == SpecialClosedLink)
+				{
+					return nullptr;
+				}
+				LocalHeadPointer->LockCount.Increment();
+				FPlatformMisc::MemoryBarrier();
+				if (*HeadPointer == LocalHeadPointer)
+				{
+					return LocalHeadPointer;
+				}
+				LocalHeadPointer->Unlock();
+			}
+		}
+
+		void SetNextPointer(FLink* NewNext)
+		{
+			FLink* LocalNextPointer = Next;
+			checkLockFreePointerList(LocalNextPointer != NewNext || !NewNext); // should not be the same, except both can be nullptr
+			checkLockFreePointerList(!LocalNextPointer || !NewNext); // one (or both) of them should be nullptr
+#if CHECK_NON_CONCURRENT_ASSUMPTIONS
+			FLink* ValueWas = (FLink*)FPlatformAtomics::InterlockedCompareExchangePointer((void**)&Next, NewNext, LocalNextPointer); // not known to need to be interlocked
+			checkLockFreePointerList(ValueWas == LocalNextPointer); // this should always work because there should never be contention on setting next
+#else
+			Next = NewNext;
+#endif
+			FPlatformMisc::MemoryBarrier();
+			if (NewNext)
+			{
+				CheckNotMarkedForDeath(); // it shouldn't be possible for the node to die while we are setting a non-nullptr next pointer
+			}
+		}
+
+		// verify that we are not marked for destruction
+		void CheckNotMarkedForDeath()
+		{
+			checkLockFreePointerList(MarkedForDeath.GetValue() == 0); // we don't allow nullptr items
+		}
+
+		/**	
+		 *	Mark the node for disposal and unlock it
+		 */
+		void Dispose()
+		{
+			CheckNotMarkedForDeath();
+
+#if CHECK_NON_CONCURRENT_ASSUMPTIONS
+			{
+				void *CurrentItem = Item;
+				checkLockFreePointerList(CurrentItem); // we don't allow nullptr items
+				void* ValueWas = FPlatformAtomics::InterlockedCompareExchangePointer(&Item,nullptr,CurrentItem); // do not need an interlocked operation here
+				checkLockFreePointerList(ValueWas == CurrentItem); // there should be no concurrency here
+			}
+
+			{
+				FLink* LocalNextPointer = Next;
+				FLink* ValueWas = (FLink*)FPlatformAtomics::InterlockedCompareExchangePointer((void**)&Next, nullptr, LocalNextPointer); // not known to need to be interlocked
+				checkLockFreePointerList(ValueWas == LocalNextPointer); // this should always work because there should never be contention on setting next
+			}
+#else
+			Item = nullptr;
+			Next = nullptr;
+#endif
+			FPlatformMisc::MemoryBarrier(); // memory barrier in here to make sure item and next are cleared before the link is recycled
+
+			int32 Death = MarkedForDeath.Increment();
+			checkLockFreePointerList(Death == 1); // there should be no concurrency here
+			Unlock();
+		}
+
+		/**	
+		 *	Internal function to unlock a link and recycle it if the ref count drops to zero and it hasn't already been recycled.
+		 *	This solves the ABA problem by delaying recycling until all pointers to this link are gone.
+		 *	@param bShouldNeverCauseAFree - if true, then this should never result in a free if it does, there is a run time error
+		 */
+		void Unlock(bool bShouldNeverCauseAFree = false);
+	};
+
+	/**	
+	 *	Class to allocate and recycle links.
+	 */
+	class FLinkAllocator
+	{
+	public:
+
+		/**	
+		 *	Return a new link.
+		 *	@param NewItem "item" pointer of the new node.
+		 *	@return New node, ready for use
+		 */
+		FLink* AllocateLink(void *NewItem)
+		{
+			checkLockFreePointerList(NewItem); // we don't allow nullptr items
+			if (NumUsedLinks.Increment() % 8192 == 1)
+			{
+				UE_CLOG(0/*MONITOR_LINK_ALLOCATION*/,LogLockFreeList, Log, TEXT("Number of links %d"),NumUsedLinks.GetValue());
+			}
+
+			FLink*  NewLink = nullptr;
+			NewLink = FLink::Unlink(&FreeLinks);
+			if (NewLink)
+			{
+				NumFreeLinks.Decrement();
+			}
+			else
+			{
+				if (NumAllocatedLinks.Increment() % 10 == 1)
+				{
+					UE_CLOG(0/*MONITOR_LINK_ALLOCATION*/,LogLockFreeList, Log, TEXT("Number of allocated links %d"),NumAllocatedLinks.GetValue());
+				}
+
+				NewLink = new FLink();
+				NewLink->LockCount.Increment();
+			}
+			checkLockFreePointerList(!NewLink->Item);
+			checkLockFreePointerList(!NewLink->Next);
+			checkLockFreePointerList(NewLink->LockCount.GetValue() >= 1);
+#if CHECK_NON_CONCURRENT_ASSUMPTIONS
+			void* ValueWas = FPlatformAtomics::InterlockedCompareExchangePointer(&NewLink->Item, NewItem, nullptr); // in theory this doesn't need to be interlocked
+			checkLockFreePointerList(ValueWas == nullptr);
+#else
+			NewLink->Item = NewItem;
+#endif
+			FPlatformMisc::MemoryBarrier();
+			return NewLink;
+		}
+		/**	
+		 *	Make a link available for recycling. 
+		 *	@param Link link to recycle
+		 *	CAUTION: Do not call this directly, it should only be called when the reference count drop to zero.
+		 */
+		void FreeLink(FLink* Link)
+		{
+			checkLockFreePointerList(Link != ClosedLink());  // very bad to recycle the special link
+			NumUsedLinks.Decrement();
+			Link->LockCount.Increment();
+			FPlatformMisc::MemoryBarrier();
+			Link->Link(&FreeLinks);
+			NumFreeLinks.Increment();
+		}
+		/**	
+		 *	Return a pointer to the special closed link
+		 *	@return pointer to special closed link
+		 */
+		FLink* ClosedLink()
+		{
+			checkLockFreePointerList(nullptr != SpecialClosedLink);
+			return SpecialClosedLink;
+		}
+		/**	
+		 *	Singleton access
+		 *	@return the singleton for the link allocator
+		 */
+		static CORE_API FLinkAllocator& Get();
+
+	private:
+		/**	
+		 *	Constructor zeros the free link list and creates the special closed link
+		 *	@return the singleton for the link allocator
+		 */
+		FLinkAllocator();
+
+		/**	
+		 *	Destructor, should only be called when there are no outstanding links. 
+		 *	Frees the elements of the free list and frees the special link
+		 *	@return the singleton for the link allocator
+		 */
+		~FLinkAllocator();
+
+		/** Head to a list of free links that can be used for new allocations. */
+		FLink*				FreeLinks;
+
+		/** Pointer to the special closed link. It should never be recycled. */
+		FLink*				SpecialClosedLink;
+
+		/** Total number of links outstanding and not in the free list */
+		FLockFreeListCounter	NumUsedLinks; 
+		/** Total number of links in the free list */
+		FLockFreeListCounter	NumFreeLinks;
+		/** Total number of links allocated */
+		FLockFreeListCounter	NumAllocatedLinks; 
+	};
+
+	/** Head of the list */
+	MS_ALIGN(8) FLink*	Head;
+};
+
+/**	
+ *	Internal function to unlock a link and recycle it if the ref count drops to zero and it hasn't already been recycled.
+ *	This solves the ABA problem by delaying recycling until all pointers to this link are gone.
+ *	@param bShouldNeverCauseAFree - if true, then this should never result in a free if it does, there is a run time error
+ */
+inline void FLockFreeVoidPointerListBase::FLink::Unlock(bool bShouldNeverCauseAFree)
+{
+	FPlatformMisc::MemoryBarrier();
+	checkLockFreePointerList(LockCount.GetValue() > 0);
+	if (LockCount.Decrement() == 0)
+	{
+		checkLockFreePointerList(MarkedForDeath.GetValue() < 2);
+		if (MarkedForDeath.Reset())
+		{
+			checkLockFreePointerList(!bShouldNeverCauseAFree); 
+			FLockFreeVoidPointerListBase::FLinkAllocator::Get().FreeLink(this);
+		}
+	}
+}
+
+typedef FLockFreeVoidPointerListBase	FLockFreeVoidPointerListGeneric;
diff --git a/Engine/Source/Runtime/Core/Public/Containers/LockFreeVoidPointerListBase128.h b/Engine/Source/Runtime/Core/Public/Containers/LockFreeVoidPointerListBase128.h
new file mode 100644
--- /dev/null
+++ b/Engine/Source/Runtime/Core/Public/Containers/LockFreeVoidPointerListBase128.h
@@ -0,0 +1,632 @@
+// Copyright 1998-2013 Epic Games, Inc. All Rights Reserved.
+
+/*=============================================================================
+	LockFreeVoidPointerListBase128.h: A lock free single linked list
+	Version using 128-bit CAS operation
+=============================================================================*/
+
+#pragma once
+
+/** Defines initial values used by the lock free linked list. */
+struct FLockFreeListConstants
+{
+	/**
+	 *	Initial ABA counter, this value is greater than 'lower half' of virtual address space in 64-bit addressing.
+	 *	It means that you will never get this address in the pointer obtained by new function, mostly for debugging purpose.
+	 *	Set to 200000000000000 for easier reading and debugging.
+	 */
+	static const uint64 InitialABACounter = 0x0000B5E620F48000; //      200000000000000 > 48-bit
+
+	/**
+	 *	ABA counter for the special closed link, this value is less than 'higher half' of virtual address space in 64-bit addressing.
+	 *	It means that you will never get this address in the pointer obtained by new function, mostly for debugging purpose.
+	 */
+	static const uint64 SpecialClosedLink = 0xFFFF700000000000; // 18446603336221196288 > 48-bit
+
+	/**
+	 * @return next ABA counter value.
+	 */
+	FORCEINLINE static const uint64 GetNextABACounter()
+	{
+		static uint64 ABACounter = FLockFreeListConstants::InitialABACounter;
+		return FPlatformAtomics::InterlockedIncrement( (volatile int64*)&ABACounter );
+	}
+};
+
+/*-----------------------------------------------------------------------------
+	FLockFreeVoidPointerListBase128
+-----------------------------------------------------------------------------*/
+
+MS_ALIGN(16)
+class FLockFreeVoidPointerListBase128 : public FNoncopyable
+{
+	struct FLargePointer;
+public:
+
+	/** Constructor - sets the list to empty and initializes the allocator if it hasn't already been initialized. */
+	FLockFreeVoidPointerListBase128()
+	{
+		FLinkAllocator::Get(); // construct allocator singleton if it isn't already
+	}
+
+	/** Destructor - checks that the list is either empty or closed. Lists should not be destroyed in any other state. */
+	~FLockFreeVoidPointerListBase128()
+	{
+		// leak these at shutdown
+		//checkLockFreePointerList(Head.Pointer == NULL || Head.Pointer == FLinkAllocator::Get().ClosedLink()); // we do not allow destruction except in the empty or closed state
+	}
+
+	/**	
+	 *	Push an item onto the head of the list
+	 *	@param NewItem, the new item to push on the list, cannot be NULL
+	 *	CAUTION: This method should not be used unless the list is known to not be closed.
+	 */
+	void Push(void *NewItem)
+	{
+		FLink* NewLink = FLinkAllocator::Get().AllocateLink(NewItem);
+		NewLink->Link( Head );
+		NumElements.Increment();
+		checkLockFreePointerList(NewLink->Next.Pointer != FLinkAllocator::Get().ClosedLink().Pointer); // it is not permissible to call push on a Closed queue
+	}
+
+	/**	
+	 *	Push an item onto the head of the list, unless the list is closed
+	 *	@param NewItem, the new item to push on the list, cannot be NULL
+	 *	@return true if the item was pushed on the list, false if the list was closed.
+	 */
+	bool PushIfNotClosed(void *NewItem)
+	{
+		FLink* NewLink = FLinkAllocator::Get().AllocateLink(NewItem);
+		bool bSuccess = NewLink->LinkIfHeadNotEqual( Head, FLinkAllocator::Get().ClosedLink() );
+		if( !bSuccess )
+		{
+			NewLink->Dispose();
+		}
+		else
+		{
+			NumElements.Increment();
+		}
+		return bSuccess;
+	}
+
+	/**	
+	 *	Push an item onto the head of the list, opening it first if necessary
+	 *	@param NewItem, the new item to push on the list, cannot be NULL
+	 *	@return true if the list needed to be opened first, false if the list was not closed before our push
+	 */
+	bool ReopenIfClosedAndPush(void *NewItem)
+	{
+		FLink* NewLink = FLinkAllocator::Get().AllocateLink(NewItem);
+		while (1)
+		{
+			if( NewLink->LinkIfHeadNotEqual( Head, FLinkAllocator::Get().ClosedLink() ) )
+			{
+				NumElements.Increment();
+				return false;
+			}
+			if( NewLink->ReplaceHeadWithThisIfHeadIsClosed( Head ) )
+			{
+				NumElements.Increment();
+				return true;
+			}
+		}
+	}
+	
+	/**	
+	 *	Pop an item from the list or return NULL if the list is empty
+	 *	@return The popped item, if any
+	 *	CAUTION: This method should not be used unless the list is known to not be closed.
+	 */
+	void* Pop()
+	{
+		FLink* Link = FLink::Unlink( Head );
+		if( !Link )
+		{
+			return nullptr;
+		}
+		NumElements.Decrement();
+		checkLockFreePointerList(Link != FLinkAllocator::Get().ClosedLink().Pointer); // it is not permissible to call pop on a Closed queue
+		void* Return = Link->Item;
+		Link->Dispose();
+		return Return;
+	}
+
+	/**	
+	 *	Pop an item from the list or return NULL if the list is empty or closed
+	 *	@return The new item, if any
+	 */
+	void* PopIfNotClosed()
+	{
+		FLink* Link = FLink::Unlink( Head, FLinkAllocator::Get().ClosedLink() );
+		if( !Link )
+		{
+			return nullptr;
+		}
+		NumElements.Decrement();
+		checkLockFreePointerList(Link != FLinkAllocator::Get().ClosedLink().Pointer); // internal error, this should not be possible
+		void* Return = Link->Item;
+		Link->Dispose();
+		return Return;
+	}
+
+	/**	
+	 *	Close the list if it is empty
+	 *	@return true if this call actively closed the list
+	 */
+	bool CloseIfEmpty()
+	{
+		if( FLink::ReplaceHeadWithClosedIfHeadIsEmpty( Head ) )
+		{
+			return true;
+		}
+		return false;
+	}
+
+	/**	
+	 *	Pop all items from the list 
+	 *	@param Output The array to hold the returned items. Must be empty.
+	 *	CAUTION: This method should not be used unless the list is known to not be closed.
+	 */
+	template <class ARRAYTYPE, typename ElementType>
+	void PopAll(ARRAYTYPE& Output)
+	{
+		checkLockFreePointerList(!Output.Num());
+		FLink* Link = FLink::ReplaceList( Head );
+		NumElements.Set( 0 );
+		while (Link)
+		{
+			checkLockFreePointerList(Link != FLinkAllocator::Get().ClosedLink().Pointer); // it is not permissible to call PopAll on a Closed queue
+			checkLockFreePointerList(Link->Item); // we don't allow null items
+			Output.Add((ElementType)(Link->Item));
+			FLink* NextLink = Link->Next.Pointer;
+			Link->Dispose();
+			Link = NextLink;
+		}
+	}
+
+	/**	
+	 *	Pop all items from the list and atomically close it.
+	 *	@param Output The array to hold the returned items. Must be empty.
+	 */
+	template <class ARRAYTYPE, typename ElementType>
+	void PopAllAndClose(ARRAYTYPE& Output)
+	{
+		checkLockFreePointerList(!Output.Num());
+		FLink* Link = FLink::ReplaceList( Head, FLinkAllocator::Get().ClosedLink() );
+		NumElements.Set( 0 );
+		while (Link)
+		{
+			checkLockFreePointerList(Link != FLinkAllocator::Get().ClosedLink().Pointer); // we should not pop off the closed link
+			checkLockFreePointerList(Link->Item); // we don't allow null items
+			Output.Add((ElementType)(Link->Item));
+			FLink* NextLink = Link->Next.Pointer;
+			Link->Dispose();
+			Link = NextLink;
+		}
+	}
+
+	/**	
+	 *	Check if the list is closed
+	 *	@return true if the list is closed.
+	 *	CAUTION: This methods safety depends on external assumptions. For example, if another thread could open or close the list at any time, the return value is no better than a best guess.
+	 *	As typically used, the list is only closed once it is lifetime, so a return of true here means it is closed forever.
+	 */
+	bool IsClosed() const  // caution, if this returns false, that does not mean the list is open!
+	{
+		FPlatformMisc::MemoryBarrier();
+		return Head.Pointer == FLinkAllocator::Get().ClosedLink().Pointer && Head.ABACounter == FLinkAllocator::Get().ClosedLink().ABACounter;
+	}
+
+	/**	
+	 *	Check if the list is empty
+	 *	@return true if the list is empty.
+	 *	CAUTION: This methods safety depends on external assumptions. For example, if another thread could add to the list at any time, the return value is no better than a best guess.
+	 *	As typically used, the list is not being access concurrently when this is called.
+	 */
+	bool IsEmpty() const  
+	{
+		FPlatformMisc::MemoryBarrier();
+		return Head.Pointer == nullptr && Head.ABACounter == 0;
+	}
+
+#if !(UE_BUILD_SHIPPING || UE_BUILD_TEST)
+	/**
+	 * Checks if item pointers look 'ok'. Used for debugging.
+	 */
+	bool CheckPointers() const
+	{
+		const int32 AlignmentCheck = ELockFreeAligment::LF128bit - 1;
+		FPlatformMisc::MemoryBarrier();
+
+		for( FLink* Link = Head.Pointer; Link; Link = Link->Next.Pointer )
+		{
+			void* ItemPtr = Link->Item;
+			if (!ItemPtr || (UPTRINT)ItemPtr < 0x100 || (UPTRINT)ItemPtr > FLockFreeListConstants::InitialABACounter )
+			{
+				return false;
+			}	
+			if ((UPTRINT)ItemPtr & AlignmentCheck)
+			{
+				return false;
+			}
+
+			if ((UPTRINT)Link & AlignmentCheck)
+			{
+				return false;
+			}
+
+			if( Link->Next.ABACounter != 0 && Link->Next.ABACounter < FLockFreeListConstants::InitialABACounter )
+			{
+				return false;
+			}
+		}
+
+		return true;
+	}
+
+	/**
+	 * @return number of elements via iterating through this list, used for debugging.
+	 */
+	int32 NumVerified() const
+	{
+		FPlatformMisc::MemoryBarrier();
+
+		int32 NumElements = 0;
+		for( FLink* Link = Head.Pointer; Link; Link = Link->Next.Pointer )
+		{
+			NumElements++;
+		}
+		return NumElements;
+	}
+
+#endif //!(UE_BUILD_SHIPPING || UE_BUILD_TEST)
+
+private:
+	struct FLink;
+
+	/** Fakes 128-bit pointers to solve the ABA problem. */
+	MS_ALIGN(16)
+	struct FLargePointer
+	{
+		/** Default constructor, initializes to zero. */
+		FLargePointer()
+			: Pointer( nullptr )
+			, ABACounter( 0 )
+		{
+		}
+
+		/**
+		 *	Initialization constructor, only used to initialize the SpecialClosedLink
+		 */
+		FLargePointer( FLink* InPointer, uint64 InABACounter )
+			: Pointer( InPointer )
+			, ABACounter( InABACounter )
+		{
+		}
+
+		/**
+		 *	Initialization constructor, only used to get the current value of the head.
+		 *	Read atomically 8-bytes and sets ABACounter to 0, which is invalid.
+		 */
+		FLargePointer( FLink* InPointer  )
+			: Pointer( InPointer )
+			, ABACounter( 0 )
+		{
+		}
+
+		/** Copy constructor, only valid if copied from constant. */
+		FLargePointer( const FLargePointer& Other )
+		{
+			Pointer = Other.Pointer;
+			ABACounter = Other.ABACounter;
+		}
+
+		/** Holds pointer to the element. */
+		FLink* Pointer;
+		/** Holds ABA counter. */
+		uint64 ABACounter;
+	}
+	GCC_ALIGN(16);
+
+	/**
+	 *	Internal single-linked list link structure.
+	 *	32 bytes, 8 bytes wasted, not used
+	 */
+	MS_ALIGN(16) 
+	struct FLink : public FNoncopyable
+	{
+		/** 16 bytes. 'Pointer' to the next element. */
+		FLargePointer Next;
+
+		/** 8 bytes. Pointer to the user data. */
+		void* Item;
+
+		FLink()
+			: Item( nullptr )
+		{}
+
+		FLargePointer GetSafeNewLinkPointer()
+		{
+			return FLargePointer( this, FLockFreeListConstants::GetNextABACounter() );	
+		}
+
+		/**	
+		 *	Link this node into the head of the given list.
+		 *	@param Head; the head of the list
+		 *	CAUTION: Not checked here, but linking into a closed list with the routine would accidentally and incorrectly open the list
+		 */
+		void Link(FLargePointer& HeadPointer)
+		{
+			const FLargePointer NewLinkPointer = GetSafeNewLinkPointer();
+			FLargePointer LocalHeadPointer     = FLargePointer(HeadPointer.Pointer);
+
+			do
+			{
+				Next = LocalHeadPointer;
+			}
+			while( FPlatformAtomics::InterlockedCompareExchange128( (volatile FInt128*)&HeadPointer, (FInt128&)NewLinkPointer, (FInt128*)&LocalHeadPointer ) == false );
+		}
+
+		/**	
+		 *	Link this node into the head of the given list, unless the head is a special, given node
+		 *	@param HeadPointer; the head of the list
+		 *	@param SpecialClosedLink; special link that is never recycled that indicates a closed list.
+		 *	@return true if this node was linked to the head.
+		 */
+		bool LinkIfHeadNotEqual( FLargePointer& HeadPointer, const FLargePointer& SpecialClosedLink )
+		{
+			checkLockFreePointerList(!Next.Pointer); // this should not be set yet
+
+			const FLargePointer NewLinkPointer = GetSafeNewLinkPointer();
+			FLargePointer LocalHeadPointer     = FLargePointer(HeadPointer.Pointer);
+
+			do
+			{
+				if( LocalHeadPointer.Pointer==SpecialClosedLink.Pointer )
+				{
+					Next = FLargePointer();
+					return false;
+				}
+
+				Next = LocalHeadPointer;
+			}
+			while( FPlatformAtomics::InterlockedCompareExchange128( (volatile FInt128*)&HeadPointer, (FInt128&)NewLinkPointer, (FInt128*)&LocalHeadPointer ) == false );
+
+			return true;
+		}
+
+		/**	
+		 *	If the head is a closed list, then replace it with this node. This is the primitive that is used to open lists
+		 *	@param HeadPointer; the head of the list
+		 *	@return true if head was replaced with this node
+		 *	CAUTION: The test link must a special link that is never recycled. The ABA problem is assumed to not occur here.
+		 */
+		bool ReplaceHeadWithThisIfHeadIsClosed( const FLargePointer& HeadPointer )
+		{
+			const FLargePointer NewLinkPointer    = GetSafeNewLinkPointer();
+			const FLargePointer ClosedLinkPointer = FLinkAllocator::Get().ClosedLink();
+			FLargePointer LocalHeadPointer        = FLinkAllocator::Get().ClosedLink();
+
+			do
+			{
+				// This will be the same on the first iteration, then LocalHeadPointer will be updated with the latest HeadPointer.
+				if( LocalHeadPointer.Pointer != ClosedLinkPointer.Pointer )
+				{
+					return false;
+				}
+			}
+			while( FPlatformAtomics::InterlockedCompareExchange128( (volatile FInt128*)&HeadPointer, (FInt128&)NewLinkPointer, (FInt128*)&LocalHeadPointer ) == false );
+
+			return true;
+		}
+
+		/**	
+		 *	If the head is empty, then replace it with a closed list. This is the primitive that is used to close lists
+		 *	@param HeadPointer; the head of the list
+		 *	@return true if head was replaced with a closed list
+		 */
+		static bool ReplaceHeadWithClosedIfHeadIsEmpty( FLargePointer& HeadPointer )
+		{
+			const FLargePointer LocalClosedLink = FLinkAllocator::Get().ClosedLink();
+			const FLargePointer LocalEmptyLink  = FLargePointer();
+			FLargePointer LocalHeadPointer      = FLargePointer();
+
+			do
+			{
+				// This will be the same on the first iteration, then LocalHeadPointer will be updated with the latest HeadPointer.
+				if( LocalHeadPointer.Pointer != LocalEmptyLink.Pointer )
+				{
+					return false;
+				}
+			}
+			while( FPlatformAtomics::InterlockedCompareExchange128( (volatile FInt128*)&HeadPointer, (FInt128&)LocalClosedLink, (FInt128*)&LocalHeadPointer ) == false );
+
+			return true;
+		}
+
+		/**	
+		 *	Pop an item off the list, unless the list is empty or the head is a special, given node
+		 *	@param HeadPointer; the head of the list
+		 *	@param SpecialClosedLink; special link that is never recycled that indicates a closed list. Can be NULL for lists that are known to not be closed
+		 *	@return The link that was popped off, or NULL if the list was empty or closed.
+		 *	CAUTION: Not checked here, but if the list is closed and SpecialClosedLink is NULL, you will pop off the closed link, and that usually isn't what you want
+		 */
+		static FLink* Unlink( FLargePointer& HeadPointer, const FLargePointer& SpecialClosedLink = FLargePointer() )
+		{
+			FLargePointer LocalHeadPointer = FLargePointer(HeadPointer.Pointer);
+			FLargePointer NextPointer;
+
+			do
+			{
+				if( LocalHeadPointer.Pointer==nullptr )
+				{
+					return nullptr;
+				}
+				if( LocalHeadPointer.Pointer==SpecialClosedLink.Pointer )
+				{
+					return nullptr;
+				}
+
+				//NextPointer = FLargePointer( LocalHeadPointer.Pointer->Next.Pointer, FLockFreeListConstants::GetNextABACounter() );	
+				NextPointer = LocalHeadPointer.Pointer->Next;
+			}
+			while( FPlatformAtomics::InterlockedCompareExchange128( (volatile FInt128*)&HeadPointer, (FInt128&)NextPointer, (FInt128*)&LocalHeadPointer ) == false );
+			checkLockFreePointerList(NextPointer.Pointer == LocalHeadPointer.Pointer->Next.Pointer);
+
+			return LocalHeadPointer.Pointer;
+		}
+
+		/**	
+		 *	Replace the list with another list. Use to either acquire all of the items or acquire all of the items and close the list.
+		 *	@param HeadPointer; the head of the list
+		 *	@param NewHeadLink; Head of the new list to replace this list with
+		 *	@return The original list before we replaced it.
+		 *	CAUTION: Not checked here, but if the list is closed this is probably not what you want. 
+		 */
+		static FLink* ReplaceList( FLargePointer& HeadPointer, const FLargePointer& NewHeadLink = FLargePointer() )
+		{
+			FLargePointer LocalHeadPointer         = FLargePointer(HeadPointer.Pointer);
+			const FLargePointer NewHeadLinkPointer = NewHeadLink;
+
+			do 
+			{
+				if (LocalHeadPointer.Pointer==nullptr && NewHeadLink.Pointer==nullptr)
+				{
+					// we are asking to replace NULL with NULL, this does not require an atomic (or any) operation
+					return nullptr;
+				}
+
+				checkLockFreePointerList(LocalHeadPointer.Pointer != NewHeadLink.Pointer); // replacing NULL with NULL is ok, but otherwise we cannot be replacing something with itself or we lose determination of who did what.
+			} 
+			while( FPlatformAtomics::InterlockedCompareExchange128( (volatile FInt128*)&HeadPointer, (FInt128&)NewHeadLinkPointer, (FInt128*)&LocalHeadPointer ) == false );
+
+			return LocalHeadPointer.Pointer;
+		}
+
+		/**
+		 *	Disposes this node.
+		 */
+		void Dispose();
+	}
+	GCC_ALIGN(16);
+
+	/**
+	 *	Class to allocate and recycle links. 
+	 */
+	MS_ALIGN(16) 
+	class FLinkAllocator
+	{
+	public:
+		/**	
+		 *	Return a new link.
+		 *	@param NewItem "item" pointer of the new node.
+		 *	@return New node, ready for use
+		 */
+		FLink* AllocateLink( void *NewItem )
+		{
+			checkLockFreePointerList(NewItem); // we don't allow null items
+			if (NumUsedLinks.Increment() % 8192 == 1)
+			{
+				UE_CLOG(0/*MONITOR_LINK_ALLOCATION*/,LogLockFreeList, Log, TEXT("Number of links %d"),NumUsedLinks.GetValue());
+			}
+	
+			FLink* NewLink = FLink::Unlink( FreeLinks );
+			if( NewLink )
+			{
+				NewLink->Next = FLargePointer();
+				NewLink->Item = nullptr;
+				NumFreeLinks.Decrement();
+			}
+			else
+			{
+				if (NumAllocatedLinks.Increment() % 10 == 1)
+				{
+					UE_CLOG(0/*MONITOR_LINK_ALLOCATION*/,LogLockFreeList, Log, TEXT("Number of allocated links %d"),NumAllocatedLinks.GetValue());
+				}
+
+				NewLink = new FLink();
+			}
+	
+			checkLockFreePointerList(!NewLink->Item);
+			checkLockFreePointerList(!NewLink->Next.Pointer);
+			NewLink->Item = NewItem;
+			FPlatformMisc::MemoryBarrier();
+			return NewLink;
+		}	
+		/**	
+		 *	Make a link available for recycling. 
+		 *	@param Link link to recycle
+		 *	CAUTION: Do not call this directly, it should only be called when the reference count drop to zero.
+		 */
+		void FreeLink( FLink* Link )
+		{
+			checkLockFreePointerList(Link != ClosedLink().Pointer);  // very bad to recycle the special link
+			NumUsedLinks.Decrement();
+			FPlatformMisc::MemoryBarrier();
+			Link->Link(FreeLinks);
+			NumFreeLinks.Increment();
+		}
+		/**
+		 *	@return a pointer to special closed link 
+		 */
+		const FLargePointer& ClosedLink() const
+		{
+			checkLockFreePointerList(NULL != SpecialClosedLink.Pointer);
+			return SpecialClosedLink;
+		}
+		/**	
+		 *	Singleton access
+		 *	@return the singleton for the link allocator
+		 */
+		static CORE_API FLinkAllocator& Get();
+	
+	private:
+		/**	
+		 *	Constructor zeros the free link list and creates the special closed link
+		 *	@return the singleton for the link allocator
+		 */
+		FLinkAllocator();
+	
+		/**	
+		 *	Destructor, should only be called when there are no outstanding links. 
+		 *	Frees the elements of the free list and frees the special link
+		 *	@return the singleton for the link allocator
+		 */
+		~FLinkAllocator();
+	
+		/** Head to a list of free links that can be used for new allocations. **/
+		FLargePointer	FreeLinks;
+	
+		/** Pointer to the special closed link. It should never be recycled. **/
+		FLargePointer	SpecialClosedLink;
+	
+		/** Total number of links outstanding and not in the free list **/
+		FLockFreeListCounter	NumUsedLinks; 
+		/** Total number of links in the free list **/
+		FLockFreeListCounter	NumFreeLinks;
+		/** Total number of links allocated **/
+		FLockFreeListCounter	NumAllocatedLinks; 
+	}
+	GCC_ALIGN(16);
+
+	/** Head of the list. */
+	FLargePointer Head;
+
+	/**
+	 *	Current number of elements, due to the high contention this value may not be valid. 
+	 *	Used for debugging.
+	 */
+	FLockFreeListCounter NumElements;
+};
+
+inline void FLockFreeVoidPointerListBase128::FLink::Dispose()
+{
+	Next = FLockFreeVoidPointerListBase128::FLargePointer();
+	Item = nullptr;
+
+	FPlatformMisc::MemoryBarrier(); // memory barrier in here to make sure item and next are cleared before the link is recycled
+
+	FLockFreeVoidPointerListBase128::FLinkAllocator::Get().FreeLink( this );
+}
+
+typedef FLockFreeVoidPointerListBase128	FLockFreeVoidPointerListGeneric;
diff --git a/Engine/Source/Runtime/Core/Public/HAL/ThreadingBase.h b/Engine/Source/Runtime/Core/Public/HAL/ThreadingBase.h
--- a/Engine/Source/Runtime/Core/Public/HAL/ThreadingBase.h
+++ b/Engine/Source/Runtime/Core/Public/HAL/ThreadingBase.h
@@ -761,6 +761,48 @@ private:
 	volatile int32 Counter;
 };
 
+/** Fake Thread safe counter, used to avoid cluttering code with ifdefs when counters are only used for debugging. */
+class FNoopCounter
+{
+public:
+	FNoopCounter()
+	{
+	}
+	FNoopCounter( const FNoopCounter& Other )
+	{
+	}
+	FNoopCounter( int32 Value )
+	{
+	}
+	int32 Increment()
+	{
+		return 0;
+	}
+	int32 Add( int32 Amount )
+	{
+		return 0;
+	}
+	int32 Decrement()
+	{
+		return 0;
+	}
+	int32 Subtract( int32 Amount )
+	{
+		return 0;
+	}
+	int32 Set( int32 Value )
+	{
+		return 0;
+	}
+	int32 Reset()
+	{
+		return 0;
+	}
+	int32 GetValue() const
+	{
+		return 0;
+	}
+};
 
 /**
  * Implements a scope lock.
