Core - Disabled debug code in the lock free list

diff --git a/Engine/Source/Runtime/Core/Public/Containers/LockFreeList.h b/Engine/Source/Runtime/Core/Public/Containers/LockFreeList.h
--- a/Engine/Source/Runtime/Core/Public/Containers/LockFreeList.h
+++ b/Engine/Source/Runtime/Core/Public/Containers/LockFreeList.h
@@ -11,7 +11,7 @@ CORE_API DECLARE_LOG_CATEGORY_EXTERN(LogLockFreeList, Log, All);
 // what level of checking to perform...normally checkLockFreePointerList but could be ensure or check
 #define checkLockFreePointerList checkSlow /*check*/
 
-#define MONITOR_LINK_ALLOCATION (1)
+#define MONITOR_LINK_ALLOCATION (0)
 
 #if	MONITOR_LINK_ALLOCATION==1
 	typedef FThreadSafeCounter	FLockFreeListCounter;
@@ -27,7 +27,7 @@ struct CORE_API FLockFreeListStats
 	static FThreadSafeCounter Cycles;
 };
 
-#define MONITOR_LOCKFREELIST_PERFORMANCE (1)
+#define MONITOR_LOCKFREELIST_PERFORMANCE (0)
 
 /** Helper struct used to test the performance of the lock free list operations. */
 struct FLockFreeListPerfCounter
