Core - Compilation fix

diff --git a/Engine/Source/Runtime/Core/Public/Containers/LockFreeList.h b/Engine/Source/Runtime/Core/Public/Containers/LockFreeList.h
--- a/Engine/Source/Runtime/Core/Public/Containers/LockFreeList.h
+++ b/Engine/Source/Runtime/Core/Public/Containers/LockFreeList.h
@@ -30,15 +30,15 @@ struct CORE_API FLockFreeListStats
 #define MONITOR_LOCKFREELIST_PERFORMANCE (0)
 
 /** Helper struct used to test the performance of the lock free list operations. */
-struct FLockFreeListPerfCounter
+struct FThreadSafePerfCounter
 {
 	const uint32 StartCycles;
 
-	FLockFreeListPerfCounter()
+	FThreadSafePerfCounter()
 		: StartCycles( FPlatformTime::Cycles() )
 	{}
 
-	~FLockFreeListPerfCounter()
+	~FThreadSafePerfCounter()
 	{
 		const uint32 EndCycles = FPlatformTime::Cycles();
 		const int32 Delta = int32( (uint32)EndCycles - (uint32)StartCycles );
@@ -49,9 +49,9 @@ struct FLockFreeListPerfCounter
 
 
 #if	MONITOR_LOCKFREELIST_PERFORMANCE==1
-	typedef FLockFreeListPerfCounter	FLockFreeListPerfCounter;
+	typedef FThreadSafePerfCounter	FLockFreeListPerfCounter;
 #else
-	typedef FNoopStruct					FLockFreeListPerfCounter;
+	typedef FNoopStruct				FLockFreeListPerfCounter;
 #endif // MONITOR_LOCKFREELIST_PERFORMANCE==1
 
 /** Enumerates lock free list alignments. */
