Core - Added a missing file

diff --git a/Engine/Source/Runtime/Core/Private/GenericPlatform/GenericPlatformAtomics.cpp b/Engine/Source/Runtime/Core/Private/GenericPlatform/GenericPlatformAtomics.cpp
new file mode 100644
--- /dev/null
+++ b/Engine/Source/Runtime/Core/Private/GenericPlatform/GenericPlatformAtomics.cpp
@@ -0,0 +1,10 @@
+// Copyright 1998-2014 Epic Games, Inc. All Rights Reserved.
+
+/*=============================================================================
+	GenericPlatformAtomics.cpp: Generic implementations of atomic functions
+=============================================================================*/
+
+#include "CorePrivate.h"
+
+checkAtCompileTime( sizeof( FInt128 ) == 16, FInt128_sizeof_must_be_16 );
+checkAtCompileTime( ALIGNOF( FInt128 ) == 16, FInt128_alignment_must_be_16 );
\ No newline at end of file
