Fixed some bad uses of 'override'

diff --git a/Engine/Source/Runtime/Core/Private/WinRT/WinRTFile.cpp b/Engine/Source/Runtime/Core/Private/WinRT/WinRTFile.cpp
--- a/Engine/Source/Runtime/Core/Private/WinRT/WinRTFile.cpp
+++ b/Engine/Source/Runtime/Core/Private/WinRT/WinRTFile.cpp
@@ -365,24 +365,6 @@ public:
 		return NULL;
 	}
 	
-	virtual IFileHandle* OpenAppend(const TCHAR* Filename) OVERRIDE
-	{
-		DWORD  Access    = GENERIC_WRITE;
-		DWORD  WinFlags  = 0;
-		DWORD  Create    = OPEN_ALWAYS;
-		HANDLE Handle	 = CreateFile2(*NormalizeFilename(Filename), Access, WinFlags, Create, NULL);
-		if(Handle != INVALID_HANDLE_VALUE)
-		{
-			LARGE_INTEGER li;
-			li.QuadPart = 0;
-			if (SetFilePointerEx(Handle, li, &li, FILE_END) == TRUE)
-			{
-				return new FFileHandleWinRT(Handle);
-			}
-		}
-		return NULL;
-	}
-	
 	virtual bool DirectoryExists(const TCHAR* Directory) OVERRIDE
 	{
 		WIN32_FILE_ATTRIBUTE_DATA AttribData;
