Delegates: Cleanup pass on delegate binding implementations

- replaced OVERRIDE with override
- replaced NULL with nullptr
- removed redundant code documentation

diff --git a/Engine/Source/Runtime/Core/Public/Delegates/DelegateBase.h b/Engine/Source/Runtime/Core/Public/Delegates/DelegateBase.h
--- a/Engine/Source/Runtime/Core/Public/Delegates/DelegateBase.h
+++ b/Engine/Source/Runtime/Core/Public/Delegates/DelegateBase.h
@@ -83,7 +83,6 @@ protected:
 	 * Gets the delegate instance.
 	 *
 	 * @return The delegate instance.
-	 *
 	 * @see SetDelegateInstance
 	 */
 	IDelegateInstance* GetDelegateInstance( ) const
@@ -95,7 +94,6 @@ protected:
 	 * Sets the delegate instance.
 	 *
 	 * @param InDelegateInstance The delegate instance to set.
-	 *
 	 * @see GetDelegateInstance.
 	 */
 	inline void SetDelegateInstance( IDelegateInstance* InDelegateInstance )
diff --git a/Engine/Source/Runtime/Core/Public/Delegates/DelegateInstanceInterfaceImpl.inl b/Engine/Source/Runtime/Core/Public/Delegates/DelegateInstanceInterfaceImpl.inl
--- a/Engine/Source/Runtime/Core/Public/Delegates/DelegateInstanceInterfaceImpl.inl
+++ b/Engine/Source/Runtime/Core/Public/Delegates/DelegateInstanceInterfaceImpl.inl
@@ -55,7 +55,6 @@ public:
 	 * even if the delegate objects themselves are different.  Also, the delegate types *must* be compatible.
 	 *
 	 * @param  InOtherDelegate
-	 *
 	 * @return  True if delegates match
 	 */
 	virtual bool IsSameFunction( const DELEGATE_INSTANCE_INTERFACE_CLASS< FUNC_TEMPLATE_ARGS >& InOtherDelegate ) const = 0;
diff --git a/Engine/Source/Runtime/Core/Public/Delegates/DelegateInstancesImpl.inl b/Engine/Source/Runtime/Core/Public/Delegates/DelegateInstancesImpl.inl
--- a/Engine/Source/Runtime/Core/Public/Delegates/DelegateInstancesImpl.inl
+++ b/Engine/Source/Runtime/Core/Public/Delegates/DelegateInstancesImpl.inl
@@ -1,26 +1,29 @@
 // Copyright 1998-2014 Epic Games, Inc. All Rights Reserved.
 
 /*================================================================================
-	DelegateInstancesImpl.inl: Inline implementation of delegate instance classes
-================================================================================*/
+	DelegateInstancesImpl.inl: Inline implementation of delegate bindings.
+
+	This file in re-included in DelegateCombinations.h for EVERY supported delegate signature.
+	Every combination of parameter count, return value presence or other function modifier will
+	include this file to generate a delegate interface type and implementation type for that signature.
 
+	The types declared in this file are for internal use only. 
+================================================================================*/
 
-// Only designed to be included directly by Delegate.h
-#if !defined( __Delegate_h__ ) || !defined( FUNC_INCLUDING_INLINE_IMPL )
+#if !defined(__Delegate_h__) || !defined(FUNC_INCLUDING_INLINE_IMPL)
 	#error "This inline header must only be included by Delegate.h"
 #endif
 
 
-// NOTE: This file in re-included for EVERY delegate signature that we support.  That is, every combination
-//		 of parameter count, return value presence or other function modifier will include this file to
-//		 generate a delegate interface type and implementation type for that signature.
+/* Class names for supported delegate bindings.
+ *****************************************************************************/
 
-#define DELEGATE_INSTANCE_INTERFACE_CLASS FUNC_COMBINE( IBaseDelegateInstance_, FUNC_SUFFIX )
-#define SP_METHOD_DELEGATE_INSTANCE_CLASS FUNC_COMBINE( TBaseSPMethodDelegateInstance_, FUNC_COMBINE( FUNC_PAYLOAD_SUFFIX, FUNC_CONST_SUFFIX ) )
-#define RAW_METHOD_DELEGATE_INSTANCE_CLASS FUNC_COMBINE( TBaseRawMethodDelegateInstance_, FUNC_COMBINE( FUNC_PAYLOAD_SUFFIX, FUNC_CONST_SUFFIX ) )
-#define UOBJECT_METHOD_DELEGATE_INSTANCE_CLASS FUNC_COMBINE( TBaseUObjectMethodDelegateInstance_, FUNC_COMBINE( FUNC_PAYLOAD_SUFFIX, FUNC_CONST_SUFFIX ) )
-#define STATIC_DELEGATE_INSTANCE_CLASS FUNC_COMBINE( TBaseStaticDelegateInstance_, FUNC_COMBINE( FUNC_PAYLOAD_SUFFIX, FUNC_CONST_SUFFIX ) )
-#define UFUNCTION_DELEGATE_INSTANCE_CLASS FUNC_COMBINE( TBaseUFunctionDelegateInstance_, FUNC_COMBINE( FUNC_PAYLOAD_SUFFIX, FUNC_CONST_SUFFIX ) )
+#define DELEGATE_INSTANCE_INTERFACE_CLASS FUNC_COMBINE(IBaseDelegateInstance_, FUNC_SUFFIX)
+#define SP_METHOD_DELEGATE_INSTANCE_CLASS FUNC_COMBINE(TBaseSPMethodDelegateInstance_, FUNC_COMBINE(FUNC_PAYLOAD_SUFFIX, FUNC_CONST_SUFFIX))
+#define RAW_METHOD_DELEGATE_INSTANCE_CLASS FUNC_COMBINE(TBaseRawMethodDelegateInstance_, FUNC_COMBINE(FUNC_PAYLOAD_SUFFIX, FUNC_CONST_SUFFIX))
+#define UOBJECT_METHOD_DELEGATE_INSTANCE_CLASS FUNC_COMBINE(TBaseUObjectMethodDelegateInstance_, FUNC_COMBINE(FUNC_PAYLOAD_SUFFIX, FUNC_CONST_SUFFIX))
+#define STATIC_DELEGATE_INSTANCE_CLASS FUNC_COMBINE(TBaseStaticDelegateInstance_, FUNC_COMBINE(FUNC_PAYLOAD_SUFFIX, FUNC_CONST_SUFFIX))
+#define UFUNCTION_DELEGATE_INSTANCE_CLASS FUNC_COMBINE(TBaseUFunctionDelegateInstance_, FUNC_COMBINE(FUNC_PAYLOAD_SUFFIX, FUNC_CONST_SUFFIX))
 
 #if FUNC_IS_CONST
 	#define DELEGATE_QUALIFIER const
@@ -28,808 +31,711 @@
 	#define DELEGATE_QUALIFIER
 #endif
 
+
+/* Macros for function parameter and delegate payload lists
+ *****************************************************************************/
+
 #if FUNC_HAS_PAYLOAD
-	#define DELEGATE_COMMA_PAYLOAD_LIST             , FUNC_PAYLOAD_LIST
-	#define DELEGATE_COMMA_PAYLOAD_PASSTHRU         , FUNC_PAYLOAD_PASSTHRU
-	#define DELEGATE_COMMA_PAYLOAD_INITIALIZER_LIST , FUNC_PAYLOAD_INITIALIZER_LIST
-	#define DELEGATE_COMMA_PAYLOAD_PASSIN           , FUNC_PAYLOAD_PASSIN
+	#define DELEGATE_COMMA_PAYLOAD_LIST							, FUNC_PAYLOAD_LIST
+	#define DELEGATE_COMMA_PAYLOAD_PASSTHRU						, FUNC_PAYLOAD_PASSTHRU
+	#define DELEGATE_COMMA_PAYLOAD_INITIALIZER_LIST				, FUNC_PAYLOAD_INITIALIZER_LIST
+	#define DELEGATE_COMMA_PAYLOAD_PASSIN						, FUNC_PAYLOAD_PASSIN
 	#if FUNC_HAS_PARAMS
-		#define DELEGATE_PARAM_PASSTHRU_COMMA_PAYLOAD_PASSIN FUNC_PARAM_PASSTHRU, FUNC_PAYLOAD_PASSIN
-		#define DELEGATE_PARAM_LIST_COMMA_PAYLOAD_LIST       FUNC_PARAM_LIST, FUNC_PAYLOAD_LIST
-		#define DELEGATE_PARMS_COLON_INITIALIZER_LIST	     : FUNC_PARAM_INITIALIZER_LIST, FUNC_PAYLOAD_INITIALIZER_LIST
+		#define DELEGATE_PARAM_PASSTHRU_COMMA_PAYLOAD_PASSIN	FUNC_PARAM_PASSTHRU, FUNC_PAYLOAD_PASSIN
+		#define DELEGATE_PARAM_LIST_COMMA_PAYLOAD_LIST			FUNC_PARAM_LIST, FUNC_PAYLOAD_LIST
+		#define DELEGATE_PARMS_COLON_INITIALIZER_LIST			: FUNC_PARAM_INITIALIZER_LIST, FUNC_PAYLOAD_INITIALIZER_LIST
 	#else
-		#define DELEGATE_PARAM_PASSTHRU_COMMA_PAYLOAD_PASSIN FUNC_PAYLOAD_PASSIN
-		#define DELEGATE_PARAM_LIST_COMMA_PAYLOAD_LIST       FUNC_PAYLOAD_LIST
-		#define DELEGATE_PARMS_COLON_INITIALIZER_LIST	     : FUNC_PAYLOAD_INITIALIZER_LIST
+		#define DELEGATE_PARAM_PASSTHRU_COMMA_PAYLOAD_PASSIN	FUNC_PAYLOAD_PASSIN
+		#define DELEGATE_PARAM_LIST_COMMA_PAYLOAD_LIST			FUNC_PAYLOAD_LIST
+		#define DELEGATE_PARMS_COLON_INITIALIZER_LIST			: FUNC_PAYLOAD_INITIALIZER_LIST
 	#endif
 #else
 	#define DELEGATE_COMMA_PAYLOAD_LIST
 	#define DELEGATE_COMMA_PAYLOAD_PASSTHRU
 	#define DELEGATE_COMMA_PAYLOAD_INITIALIZER_LIST
 	#define DELEGATE_COMMA_PAYLOAD_PASSIN
-	#define DELEGATE_PARAM_PASSTHRU_COMMA_PAYLOAD_PASSIN FUNC_PARAM_PASSTHRU
-	#define DELEGATE_PARAM_LIST_COMMA_PAYLOAD_LIST       FUNC_PARAM_LIST
+	#define DELEGATE_PARAM_PASSTHRU_COMMA_PAYLOAD_PASSIN		FUNC_PARAM_PASSTHRU
+	#define DELEGATE_PARAM_LIST_COMMA_PAYLOAD_LIST				FUNC_PARAM_LIST
 	#if FUNC_HAS_PARAMS
-		#define DELEGATE_PARMS_COLON_INITIALIZER_LIST	 : FUNC_PARAM_INITIALIZER_LIST
+		#define DELEGATE_PARMS_COLON_INITIALIZER_LIST			: FUNC_PARAM_INITIALIZER_LIST
 	#else
 		#define DELEGATE_PARMS_COLON_INITIALIZER_LIST
 	#endif
 #endif
 
 
-template< class UserClass, FUNC_PAYLOAD_TEMPLATE_DECL_TYPENAME > class UOBJECT_METHOD_DELEGATE_INSTANCE_CLASS;
-template< class UserClass, FUNC_PAYLOAD_TEMPLATE_DECL_TYPENAME > class RAW_METHOD_DELEGATE_INSTANCE_CLASS;
+/* Delegate binding types
+ *****************************************************************************/
+
+template<class UserClass, FUNC_PAYLOAD_TEMPLATE_DECL_TYPENAME> class UOBJECT_METHOD_DELEGATE_INSTANCE_CLASS;
+template<class UserClass, FUNC_PAYLOAD_TEMPLATE_DECL_TYPENAME> class RAW_METHOD_DELEGATE_INSTANCE_CLASS;
 
 
 /**
- * Shared pointer-based delegate instance class for methods.  For internal use only.
+ * Implements a delegate binding for shared pointer member functions.
  */
-template< class UserClass, FUNC_PAYLOAD_TEMPLATE_DECL_TYPENAME, ESPMode::Type SPMode >
+template<class UserClass, FUNC_PAYLOAD_TEMPLATE_DECL_TYPENAME, ESPMode::Type SPMode>
 class SP_METHOD_DELEGATE_INSTANCE_CLASS
-	: public DELEGATE_INSTANCE_INTERFACE_CLASS< FUNC_TEMPLATE_ARGS >
+	: public DELEGATE_INSTANCE_INTERFACE_CLASS<FUNC_TEMPLATE_ARGS>
 {
-	/**
-	 * Internal, inlined and non-virtual version of GetRawUserObject interface.
-	 */
-	FORCEINLINE const void* GetRawUserObjectInternal() const
-	{
-		return UserObject.Pin().Get();
-	}
+public:
+
+	/** Type definition for member function pointers. */
+	typedef RetValType (UserClass::*FMethodPtr)( DELEGATE_PARAM_LIST_COMMA_PAYLOAD_LIST ) DELEGATE_QUALIFIER;
 
 	/**
-	 * Internal, inlined and non-virtual version of GetRawMethod interface.
+	 * Creates and initializes a new instance.
+	 *
+	 * @param InUserObject A shared reference to an arbitrary object (templated) that hosts the member function.
+	 * @param InMethodPtr C++ member function pointer for the method to bind.
 	 */
-	FORCEINLINE const void* GetRawMethodPtrInternal() const
+	SP_METHOD_DELEGATE_INSTANCE_CLASS( const TSharedPtr<UserClass, SPMode>& InUserObject, FMethodPtr InMethodPtr DELEGATE_COMMA_PAYLOAD_LIST )
+		: UserObject(InUserObject)
+		, MethodPtr(InMethodPtr)
+		DELEGATE_COMMA_PAYLOAD_INITIALIZER_LIST
 	{
-		return *(const void**)&MethodPtr;
+		// NOTE: Shared pointer delegates are allowed to have a null incoming object pointer.  Weak pointers can expire,
+		//       an it is possible for a copy of a delegate instance to end up with a null pointer.
+		checkSlow(MethodPtr != nullptr);
 	}
 
 public:
 
-	/** Pointer-to-member typedef */
-	typedef RetValType ( UserClass::*FMethodPtr )( DELEGATE_PARAM_LIST_COMMA_PAYLOAD_LIST ) DELEGATE_QUALIFIER;
+	// Begin IDelegateInstance interface
 
-	/**
-	 * Static helper function that returns a newly-created delegate object for
-	 * the specified user class and class method
-	 *
-	 * @param  InUserObjectRef  Shared reference to the user's object that contains the class method
-	 * @param  InFunc  Member function pointer to your class method
-	 *
-	 * @return  Returns an instance of the new delegate
-	 */
-	FORCEINLINE static DELEGATE_INSTANCE_INTERFACE_CLASS< FUNC_TEMPLATE_ARGS >* Create( const TSharedPtr< UserClass, SPMode >& InUserObjectRef, FMethodPtr InFunc DELEGATE_COMMA_PAYLOAD_LIST )
+	virtual FName GetFunctionName( ) const override
 	{
-		return new SP_METHOD_DELEGATE_INSTANCE_CLASS< UserClass, FUNC_PAYLOAD_TEMPLATE_ARGS, SPMode >( InUserObjectRef, InFunc DELEGATE_COMMA_PAYLOAD_PASSTHRU );
+		return NAME_None;
 	}
 
-	/**
-	 * Static helper function that returns a newly-created delegate object for
-	 * the specified user class and class method.  Requires that the supplied object derives from TSharedFromThis.
-	 *
-	 * @param  InUserObject  The user's object that contains the class method.  Must derive from TSharedFromThis.
-	 * @param  InFunc  Member function pointer to your class method
-	 *
-	 * @return  Returns an instance of the new delegate
-	 */
-	FORCEINLINE static DELEGATE_INSTANCE_INTERFACE_CLASS< FUNC_TEMPLATE_ARGS >* Create( UserClass* InUserObject, FMethodPtr InFunc DELEGATE_COMMA_PAYLOAD_LIST )
+	virtual const void* GetRawMethodPtr( ) const override
 	{
-		/* We expect the incoming InUserObject to derived from TSharedFromThis */
-		TSharedRef< UserClass > UserObjectRef( StaticCastSharedRef< UserClass >( InUserObject->AsShared() ) );
-		return Create( UserObjectRef, InFunc DELEGATE_COMMA_PAYLOAD_PASSTHRU );
+		return GetRawMethodPtrInternal();
 	}
 
-	/**
-	 * Constructor
-	 *
-	 * @param  InUserObject  A shared reference to an arbitrary object (templated) that hosts the member function
-	 * @param  InMethodPtr  C++ member function pointer for the method to bind
-	 */
-	SP_METHOD_DELEGATE_INSTANCE_CLASS( const TSharedPtr< UserClass, SPMode >& InUserObject, FMethodPtr InMethodPtr DELEGATE_COMMA_PAYLOAD_LIST )
-		: UserObject( InUserObject ),
-		  MethodPtr( InMethodPtr )
-		  DELEGATE_COMMA_PAYLOAD_INITIALIZER_LIST
+	virtual const void* GetRawUserObject( ) const override
 	{
-		// NOTE: Shared pointer delegates are allowed to have a null incoming object pointer.  Weak pointers can expire,
-		//       an it is possible for a copy of a delegate instance to end up with a null pointer.
-		checkSlow( MethodPtr != NULL );
+		return GetRawUserObjectInternal();
 	}
 
-	/**
-	 * Returns the type of delegate instance
-	 *
-	 * @return  Delegate instance type
-	 */
-	virtual EDelegateInstanceType::Type GetType() const OVERRIDE
+	virtual EDelegateInstanceType::Type GetType( ) const override
 	{
 		return SPMode == ESPMode::ThreadSafe ? EDelegateInstanceType::ThreadSafeSharedPointerMethod : EDelegateInstanceType::SharedPointerMethod;
 	}
 
-	/**
-	 * Creates a copy of the delegate instance
-	 *
-	 * @return	The newly created copy
-	 */
-	virtual DELEGATE_INSTANCE_INTERFACE_CLASS< FUNC_TEMPLATE_ARGS >* CreateCopy() OVERRIDE
+	virtual bool HasSameObject( const void* InUserObject ) const override
 	{
-		return Create( UserObject.Pin(), MethodPtr DELEGATE_COMMA_PAYLOAD_PASSIN );
+		return UserObject.HasSameObject(InUserObject);
 	}
 
-	/**
-	 * Returns true if this delegate points to exactly the same object and method as the specified delegate,
-	 * even if the delegate objects themselves are different.  Also, the delegate types *must* be compatible.
-	 *
-	 * @param  InOtherDelegate
-	 *
-	 * @return  True if delegates match
-	 */
-	virtual bool IsSameFunction( const DELEGATE_INSTANCE_INTERFACE_CLASS< FUNC_TEMPLATE_ARGS >& InOtherDelegate ) const OVERRIDE
+	virtual bool IsSafeToExecute( ) const override
 	{
-		// NOTE: Payload data is not currently considered when comparing delegate instances.  See the comment in
-		// multi-cast delegate's Remove() method for more information.
-
-		if ( InOtherDelegate.GetType() == EDelegateInstanceType::SharedPointerMethod || 
-			   InOtherDelegate.GetType() == EDelegateInstanceType::ThreadSafeSharedPointerMethod ||
-				 InOtherDelegate.GetType() == EDelegateInstanceType::RawMethod ||
-				 InOtherDelegate.GetType() == EDelegateInstanceType::UObjectMethod )
-		{
-			return GetRawMethodPtrInternal() == InOtherDelegate.GetRawMethodPtr() && UserObject.HasSameObject( InOtherDelegate.GetRawUserObject() );
-		}
-
-		return false;
+		return UserObject.IsValid();
 	}
 
-	/**
-	 * Returns true if this delegate is bound to the specified UserObject,
-	 *
-	 * @param  InUserObject
-	 *
-	 * @return  True if delegate is bound to the specified UserObject
-	 */
-	virtual bool HasSameObject( const void* InUserObject ) const OVERRIDE
-	{
-		return UserObject.HasSameObject( InUserObject );
-	}
+	// End IDelegateInstance interface
 
-	virtual FName GetFunctionName() const OVERRIDE
-	{
-		return NAME_None;
-	}
+public:
 
-	virtual const void* GetRawUserObject() const OVERRIDE
-	{
-		return GetRawUserObjectInternal();
-	}
+	// Begin DELEGATE_INSTANCE_INTERFACE_CLASS interface
 
-	virtual const void* GetRawMethodPtr() const OVERRIDE
+	virtual DELEGATE_INSTANCE_INTERFACE_CLASS<FUNC_TEMPLATE_ARGS>* CreateCopy() override
 	{
-		return GetRawMethodPtrInternal();
+		return Create(UserObject.Pin(), MethodPtr DELEGATE_COMMA_PAYLOAD_PASSIN);
 	}
 
-	/**
-	 * Execute the delegate.  If the function pointer is not valid, an error will occur.
-	 */
-	virtual RetValType Execute( FUNC_PARAM_LIST ) const OVERRIDE
+	virtual RetValType Execute( FUNC_PARAM_LIST ) const override
 	{
 		// Verify that the user object is still valid.  We only have a weak reference to it.
-		TSharedPtr< UserClass, SPMode > SharedUserObject( UserObject.Pin() );
-		checkSlow( SharedUserObject.IsValid() );
+		TSharedPtr<UserClass, SPMode> SharedUserObject( UserObject.Pin());
+		checkSlow(SharedUserObject.IsValid());
 
 		// Safely remove const to work around a compiler issue with instantiating template permutations for 
 		// overloaded functions that take a function pointer typedef as a member of a templated class.  In
 		// all cases where this code is actually invoked, the UserClass will already be a const pointer.
 		typedef typename TRemoveConst<UserClass>::Type MutableUserClass;
-		MutableUserClass* MutableUserObject = const_cast<MutableUserClass*>( SharedUserObject.Get() );
+		MutableUserClass* MutableUserObject = const_cast<MutableUserClass*>(SharedUserObject.Get());
 
 		// Call the member function on the user's object.  And yes, this is the correct C++ syntax for calling a
 		// pointer-to-member function.
-		checkSlow( MethodPtr != NULL );
+		checkSlow(MethodPtr != nullptr);
 
-		return ( MutableUserObject->*MethodPtr )( DELEGATE_PARAM_PASSTHRU_COMMA_PAYLOAD_PASSIN );
-	}
-
-	/**
-	 * Checks to see if the user object bound to this delegate is still valid
-	 *
-	 * @return  True if the object is still valid and it's safe to execute the function call
-	 */
-	virtual bool IsSafeToExecute() const OVERRIDE
-	{
-		return UserObject.IsValid();
+		return (MutableUserObject->*MethodPtr)(DELEGATE_PARAM_PASSTHRU_COMMA_PAYLOAD_PASSIN);
 	}
 
 #if FUNC_IS_VOID
-	/**
-	 * Execute the delegate, but only if the function pointer is still valid
-	 *
-	 * @return  Returns true if the function was executed
-	 */
-	virtual bool ExecuteIfSafe( FUNC_PARAM_LIST ) const OVERRIDE
+	virtual bool ExecuteIfSafe( FUNC_PARAM_LIST ) const override
 	{
 		// Verify that the user object is still valid.  We only have a weak reference to it.
-		TSharedPtr< UserClass, SPMode > SharedUserObject( UserObject.Pin() );
-		if( SharedUserObject.IsValid() )
+		TSharedPtr<UserClass, SPMode> SharedUserObject(UserObject.Pin());
+		if (SharedUserObject.IsValid())
 		{
-			Execute( FUNC_PARAM_PASSTHRU );
+			Execute(FUNC_PARAM_PASSTHRU);
+
 			return true;
 		}
+
 		return false;
 	}
 #endif
 
-private:
-
-	// Declare ourselves as a friend so we can access other template permutations in IsSameFunction()
-	template< class UserClassNoShadow, FUNC_PAYLOAD_TEMPLATE_DECL_NO_SHADOW, ESPMode::Type SPModeNoShadow > friend class SP_METHOD_DELEGATE_INSTANCE_CLASS;
+	virtual bool IsSameFunction( const DELEGATE_INSTANCE_INTERFACE_CLASS<FUNC_TEMPLATE_ARGS>& InOtherDelegate ) const override
+	{
+		// NOTE: Payload data is not currently considered when comparing delegate instances.
+		// See the comment in multi-cast delegate's Remove() method for more information.
+		if ((InOtherDelegate.GetType() == EDelegateInstanceType::SharedPointerMethod) || 
+			(InOtherDelegate.GetType() == EDelegateInstanceType::ThreadSafeSharedPointerMethod) ||
+			(InOtherDelegate.GetType() == EDelegateInstanceType::RawMethod) ||
+			(InOtherDelegate.GetType() == EDelegateInstanceType::UObjectMethod))
+		{
+			return (GetRawMethodPtrInternal() == InOtherDelegate.GetRawMethodPtr() && UserObject.HasSameObject(InOtherDelegate.GetRawUserObject()));
+		}
 
-	// Declare other pointer-based delegates as a friend so IsSameFunction() can compare members
-	template< class UserClassNoShadow, FUNC_PAYLOAD_TEMPLATE_DECL_NO_SHADOW > friend class UOBJECT_METHOD_DELEGATE_INSTANCE_CLASS;
-	template< class UserClassNoShadow, FUNC_PAYLOAD_TEMPLATE_DECL_NO_SHADOW > friend class RAW_METHOD_DELEGATE_INSTANCE_CLASS;
+		return false;
+	}
 
-	/** Weak reference to an instance of the user's class which contains a method we would like to call. */
-	TWeakPtr< UserClass, SPMode > UserObject;
+	// End DELEGATE_INSTANCE_INTERFACE_CLASS interface
 
-	/** C++ member function pointer */
-	FMethodPtr MethodPtr;
+public:
 
-	/** Payload member variables (if any) */
-	FUNC_PAYLOAD_MEMBERS
-};
+	/**
+	 * Creates a new shared pointer delegate binding for the given user object and method pointer.
+	 *
+	 * @param InUserObjectRef Shared reference to the user's object that contains the class method.
+	 * @param InFunc Member function pointer to your class method.
+	 * @return The new delegate.
+	 */
+	FORCEINLINE static DELEGATE_INSTANCE_INTERFACE_CLASS<FUNC_TEMPLATE_ARGS>* Create( const TSharedPtr<UserClass, SPMode>& InUserObjectRef, FMethodPtr InFunc DELEGATE_COMMA_PAYLOAD_LIST )
+	{
+		return new SP_METHOD_DELEGATE_INSTANCE_CLASS<UserClass, FUNC_PAYLOAD_TEMPLATE_ARGS, SPMode>(InUserObjectRef, InFunc DELEGATE_COMMA_PAYLOAD_PASSTHRU);
+	}
 
+	/**
+	 * Creates a new shared pointer delegate binding for the given user object and method pointer.
+	 *
+	 * This overload requires that the supplied object derives from TSharedFromThis.
+	 *
+	 * @param InUserObject  The user's object that contains the class method.  Must derive from TSharedFromThis.
+	 * @param InFunc  Member function pointer to your class method.
+	 * @return The new delegate.
+	 */
+	FORCEINLINE static DELEGATE_INSTANCE_INTERFACE_CLASS<FUNC_TEMPLATE_ARGS>* Create( UserClass* InUserObject, FMethodPtr InFunc DELEGATE_COMMA_PAYLOAD_LIST )
+	{
+		// We expect the incoming InUserObject to derived from TSharedFromThis.
+		TSharedRef<UserClass> UserObjectRef(StaticCastSharedRef<UserClass>(InUserObject->AsShared()));
+		return Create(UserObjectRef, InFunc DELEGATE_COMMA_PAYLOAD_PASSTHRU);
+	}
 
-/**
- * Delegate type implementation class for raw C++ methods.  For internal use only.
- */
-template< class UserClass, FUNC_PAYLOAD_TEMPLATE_DECL_TYPENAME >
-class RAW_METHOD_DELEGATE_INSTANCE_CLASS
-	: public DELEGATE_INSTANCE_INTERFACE_CLASS< FUNC_TEMPLATE_ARGS >
-{
-	// Raw method delegates to UObjects are illegal.
-	checkAtCompileTime((!CanConvertPointerFromTo<UserClass, UObjectBase>::Result), You_Cannot_Use_Raw_Method_Delegates_With_UObjects);
+protected:
 
 	/**
 	 * Internal, inlined and non-virtual version of GetRawUserObject interface.
 	 */
-	FORCEINLINE const void* GetRawUserObjectInternal() const
+	FORCEINLINE const void* GetRawUserObjectInternal( ) const
 	{
-		return UserObject;
+		return UserObject.Pin().Get();
 	}
 
 	/**
-	 * Internal, inlined and non-virtual version of GetRawMethodPtr interface.
+	 * Internal, inlined and non-virtual version of GetRawMethod interface.
 	 */
-	FORCEINLINE const void* GetRawMethodPtrInternal() const
+	FORCEINLINE const void* GetRawMethodPtrInternal( ) const
 	{
 		return *(const void**)&MethodPtr;
 	}
 
+private:
+
+	// Declare ourselves as a friend so we can access other template permutations in IsSameFunction().
+	template<class UserClassNoShadow, FUNC_PAYLOAD_TEMPLATE_DECL_NO_SHADOW, ESPMode::Type SPModeNoShadow> friend class SP_METHOD_DELEGATE_INSTANCE_CLASS;
+
+	// Declare other pointer-based delegates as a friend so IsSameFunction() can compare members.
+	template<class UserClassNoShadow, FUNC_PAYLOAD_TEMPLATE_DECL_NO_SHADOW> friend class UOBJECT_METHOD_DELEGATE_INSTANCE_CLASS;
+	template<class UserClassNoShadow, FUNC_PAYLOAD_TEMPLATE_DECL_NO_SHADOW> friend class RAW_METHOD_DELEGATE_INSTANCE_CLASS;
+
+	// Weak reference to an instance of the user's class which contains a method we would like to call.
+	TWeakPtr<UserClass, SPMode> UserObject;
+
+	// C++ member function pointer.
+	FMethodPtr MethodPtr;
+
+	// Payload member variables, if any.
+	FUNC_PAYLOAD_MEMBERS
+};
+
+
+/**
+ * Implements a delegate binding for C++ member functions.
+ */
+template<class UserClass, FUNC_PAYLOAD_TEMPLATE_DECL_TYPENAME>
+class RAW_METHOD_DELEGATE_INSTANCE_CLASS
+	: public DELEGATE_INSTANCE_INTERFACE_CLASS<FUNC_TEMPLATE_ARGS>
+{
+	checkAtCompileTime((!CanConvertPointerFromTo<UserClass, UObjectBase>::Result), You_Cannot_Use_Raw_Method_Delegates_With_UObjects);
+
 public:
 
 	/** Pointer-to-member typedef */
 	typedef RetValType ( UserClass::*FMethodPtr )( DELEGATE_PARAM_LIST_COMMA_PAYLOAD_LIST ) DELEGATE_QUALIFIER;
 
-
-	/**
-	 * Static helper function that returns a delegate object for the specified user class and class method
-	 *
-	 * @param  InUserObject  User's object that contains the class method
-	 * @param  InFunc  Member function pointer to your class method
-	 *
-	 * @return  Returns an interface to the new delegate
-	 */
-	FORCEINLINE static DELEGATE_INSTANCE_INTERFACE_CLASS< FUNC_TEMPLATE_ARGS >* Create( UserClass* InUserObject, FMethodPtr InFunc DELEGATE_COMMA_PAYLOAD_LIST )
-	{
-		return new RAW_METHOD_DELEGATE_INSTANCE_CLASS< UserClass, FUNC_PAYLOAD_TEMPLATE_ARGS >( InUserObject, InFunc DELEGATE_COMMA_PAYLOAD_PASSTHRU );
-	}
-
 	/**
-	 * Constructor
+	 * Creates and initializes a new instance.
 	 *
-	 * @param  InUserObject  An arbitrary object (templated) that hosts the member function
-	 * @param  InMethodPtr  C++ member function pointer for the method to bind
+	 * @param InUserObject An arbitrary object (templated) that hosts the member function.
+	 * @param InMethodPtr C++ member function pointer for the method to bind.
 	 */
 	RAW_METHOD_DELEGATE_INSTANCE_CLASS( UserClass* InUserObject, FMethodPtr InMethodPtr DELEGATE_COMMA_PAYLOAD_LIST )
-		: UserObject( InUserObject ),
-		  MethodPtr( InMethodPtr )
-		  DELEGATE_COMMA_PAYLOAD_INITIALIZER_LIST
+		: UserObject(InUserObject)
+		, MethodPtr(InMethodPtr)
+		DELEGATE_COMMA_PAYLOAD_INITIALIZER_LIST
 	{
 		// Non-expirable delegates must always have a non-null object pointer on creation (otherwise they could never execute.)
-		check( InUserObject != NULL && MethodPtr != NULL );
+		check(InUserObject != nullptr && MethodPtr != nullptr);
 	}
 
-	/**
-	 * Returns the type of delegate instance
-	 *
-	 * @return  Delegate instance type
-	 */
-	virtual EDelegateInstanceType::Type GetType() const OVERRIDE
+public:
+
+	// Begin IDelegateInstance interface
+
+	virtual FName GetFunctionName( ) const override
 	{
-		return EDelegateInstanceType::RawMethod;
+		return NAME_None;
 	}
 
-	/**
-	 * Creates a copy of the delegate instance
-	 *
-	 * @return	The newly created copy
-	 */
-	virtual DELEGATE_INSTANCE_INTERFACE_CLASS< FUNC_TEMPLATE_ARGS >* CreateCopy() OVERRIDE
+	virtual const void* GetRawMethodPtr( ) const override
 	{
-		return Create( UserObject, MethodPtr DELEGATE_COMMA_PAYLOAD_PASSIN );
+		return GetRawMethodPtrInternal();
 	}
 
-	/**
-	 * Returns true if this delegate points to exactly the same object and method as the specified delegate,
-	 * even if the delegate objects themselves are different.  Also, the delegate types *must* be compatible.
-	 *
-	 * @param  InOtherDelegate
-	 *
-	 * @return  True if delegates match
-	 */
-	virtual bool IsSameFunction( const DELEGATE_INSTANCE_INTERFACE_CLASS< FUNC_TEMPLATE_ARGS >& InOtherDelegate ) const OVERRIDE
+	virtual const void* GetRawUserObject( ) const override
 	{
-		// NOTE: Payload data is not currently considered when comparing delegate instances.  See the comment in
-		// multi-cast delegate's Remove() method for more information.
-
-		if ( InOtherDelegate.GetType() == EDelegateInstanceType::RawMethod || 
-			   InOtherDelegate.GetType() == EDelegateInstanceType::UObjectMethod ||
-				 InOtherDelegate.GetType() == EDelegateInstanceType::SharedPointerMethod ||
-				 InOtherDelegate.GetType() == EDelegateInstanceType::ThreadSafeSharedPointerMethod )
-		{
-			return GetRawMethodPtrInternal() == InOtherDelegate.GetRawMethodPtr() && UserObject == InOtherDelegate.GetRawUserObject();
-		}
-
-		return false;
+		return GetRawUserObjectInternal();
 	}
 
-	/**
-	 * Returns true if this delegate is bound to the specified UserObject,
-	 *
-	 * @param  InUserObject
-	 *
-	 * @return  True if delegate is bound to the specified UserObject
-	 */
-	virtual bool HasSameObject( const void* InUserObject ) const OVERRIDE
+	virtual EDelegateInstanceType::Type GetType( ) const override
 	{
-		return UserObject == InUserObject;
+		return EDelegateInstanceType::RawMethod;
 	}
 
-	virtual FName GetFunctionName() const OVERRIDE
+	virtual bool HasSameObject( const void* InUserObject ) const override
 	{
-		return NAME_None;
+		return UserObject == InUserObject;
 	}
 
-	virtual const void* GetRawUserObject() const OVERRIDE
+	virtual bool IsSafeToExecute( ) const override
 	{
-		return GetRawUserObjectInternal();
+		// We never know whether or not it is safe to deference a C++ pointer, but we have to
+		// trust the user in this case.  Prefer using a shared-pointer based delegate type instead!
+		return true;
 	}
 
-	virtual const void* GetRawMethodPtr() const OVERRIDE
+	// End IDelegateInstance interface
+
+public:
+
+	// Begin DELEGATE_INSTANCE_INTERFACE_CLASS interface
+
+	virtual DELEGATE_INSTANCE_INTERFACE_CLASS<FUNC_TEMPLATE_ARGS>* CreateCopy( ) override
 	{
-		return GetRawMethodPtrInternal();
+		return Create(UserObject, MethodPtr DELEGATE_COMMA_PAYLOAD_PASSIN);
 	}
 
-	/**
-	 * Execute the delegate.  If the function pointer is not valid, an error will occur.
-	 */
-	virtual RetValType Execute( FUNC_PARAM_LIST ) const OVERRIDE
+	virtual RetValType Execute( FUNC_PARAM_LIST ) const override
 	{
 		// Safely remove const to work around a compiler issue with instantiating template permutations for 
 		// overloaded functions that take a function pointer typedef as a member of a templated class.  In
 		// all cases where this code is actually invoked, the UserClass will already be a const pointer.
 		typedef typename TRemoveConst<UserClass>::Type MutableUserClass;
-		MutableUserClass* MutableUserObject = const_cast<MutableUserClass*>( UserObject );
+		MutableUserClass* MutableUserObject = const_cast<MutableUserClass*>(UserObject);
 
 		// Call the member function on the user's object.  And yes, this is the correct C++ syntax for calling a
 		// pointer-to-member function.
-		checkSlow( MethodPtr != NULL );
+		checkSlow(MethodPtr != nullptr);
 
-		return ( MutableUserObject->*MethodPtr )( DELEGATE_PARAM_PASSTHRU_COMMA_PAYLOAD_PASSIN );
+		return (MutableUserObject->*MethodPtr)(DELEGATE_PARAM_PASSTHRU_COMMA_PAYLOAD_PASSIN);
 	}
 
-	/**
-	 * Checks to see if the user object bound to this delegate is still valid
-	 *
-	 * @return  True if the object is still valid and it's safe to execute the function call
-	 */
-	virtual bool IsSafeToExecute() const OVERRIDE
+#if FUNC_IS_VOID
+	virtual bool ExecuteIfSafe( FUNC_PARAM_LIST ) const override
 	{
 		// We never know whether or not it is safe to deference a C++ pointer, but we have to
 		// trust the user in this case.  Prefer using a shared-pointer based delegate type instead!
+		Execute(FUNC_PARAM_PASSTHRU);
+
 		return true;
 	}
+#endif
+
+	virtual bool IsSameFunction( const DELEGATE_INSTANCE_INTERFACE_CLASS<FUNC_TEMPLATE_ARGS>& InOtherDelegate ) const override
+	{
+		// NOTE: Payload data is not currently considered when comparing delegate instances.
+		// See the comment in multi-cast delegate's Remove() method for more information.
+		if ((InOtherDelegate.GetType() == EDelegateInstanceType::RawMethod) || 
+			(InOtherDelegate.GetType() == EDelegateInstanceType::UObjectMethod) ||
+			(InOtherDelegate.GetType() == EDelegateInstanceType::SharedPointerMethod) ||
+			(InOtherDelegate.GetType() == EDelegateInstanceType::ThreadSafeSharedPointerMethod))
+		{
+			return (GetRawMethodPtrInternal() == InOtherDelegate.GetRawMethodPtr()) && (UserObject == InOtherDelegate.GetRawUserObject());
+		}
+
+		return false;
+	}
+
+public:
 
-#if FUNC_IS_VOID
 	/**
-	 * Execute the delegate, but only if the function pointer is still valid
+	 * Creates a new raw method delegate binding for the given user object and function pointer.
 	 *
-	 * @return  Returns true if the function was executed
+	 * @param InUserObject User's object that contains the class method.
+	 * @param InFunc Member function pointer to your class method.
+	 * @return The new delegate.
 	 */
-	virtual bool ExecuteIfSafe( FUNC_PARAM_LIST ) const OVERRIDE
+	FORCEINLINE static DELEGATE_INSTANCE_INTERFACE_CLASS<FUNC_TEMPLATE_ARGS>* Create( UserClass* InUserObject, FMethodPtr InFunc DELEGATE_COMMA_PAYLOAD_LIST )
 	{
-		// We never know whether or not it is safe to deference a C++ pointer, but we have to
-		// trust the user in this case.  Prefer using a shared-pointer based delegate type instead!
-		Execute( FUNC_PARAM_PASSTHRU );
-		return true;
+		return new RAW_METHOD_DELEGATE_INSTANCE_CLASS<UserClass, FUNC_PAYLOAD_TEMPLATE_ARGS>(InUserObject, InFunc DELEGATE_COMMA_PAYLOAD_PASSTHRU);
+	}
+
+protected:
+
+	/**
+	 * Internal, inlined and non-virtual version of GetRawUserObject interface.
+	 */
+	FORCEINLINE const void* GetRawUserObjectInternal( ) const
+	{
+		return UserObject;
+	}
+
+	/**
+	 * Internal, inlined and non-virtual version of GetRawMethodPtr interface.
+	 */
+	FORCEINLINE const void* GetRawMethodPtrInternal( ) const
+	{
+		return *(const void**)&MethodPtr;
 	}
-#endif
 
 private:
 
 	// Declare other pointer-based delegates as a friend so IsSameFunction() can compare members
-	template< class UserClassNoShadow, FUNC_PAYLOAD_TEMPLATE_DECL_NO_SHADOW, ESPMode::Type SPModeNoShadow > friend class SP_METHOD_DELEGATE_INSTANCE_CLASS;
-	template< class UserClassNoShadow, FUNC_PAYLOAD_TEMPLATE_DECL_NO_SHADOW > friend class UOBJECT_METHOD_DELEGATE_INSTANCE_CLASS;
+	template<class UserClassNoShadow, FUNC_PAYLOAD_TEMPLATE_DECL_NO_SHADOW, ESPMode::Type SPModeNoShadow> friend class SP_METHOD_DELEGATE_INSTANCE_CLASS;
+	template<class UserClassNoShadow, FUNC_PAYLOAD_TEMPLATE_DECL_NO_SHADOW> friend class UOBJECT_METHOD_DELEGATE_INSTANCE_CLASS;
 
-	/** Pointer to the user's class which contains a method we would like to call. */
+	// Pointer to the user's class which contains a method we would like to call.
 	UserClass* UserObject;
 
-	/** C++ member function pointer */
+	// C++ member function pointer.
 	FMethodPtr MethodPtr;
 
-	/** Payload member variables (if any) */
+	// Payload member variables (if any).
 	FUNC_PAYLOAD_MEMBERS
 };
 
 
 /**
- * Delegate type implementation class for UObject methods.  For internal use only.
+ * Implements a delegate binding for UObject methods.
  */
-template< class UserClass, FUNC_PAYLOAD_TEMPLATE_DECL_TYPENAME >
+template<class UserClass, FUNC_PAYLOAD_TEMPLATE_DECL_TYPENAME>
 class UOBJECT_METHOD_DELEGATE_INSTANCE_CLASS
-	: public DELEGATE_INSTANCE_INTERFACE_CLASS< FUNC_TEMPLATE_ARGS >
+	: public DELEGATE_INSTANCE_INTERFACE_CLASS<FUNC_TEMPLATE_ARGS>
 {
-		// UObject delegates can't be used with non-UObject pointers.
 	checkAtCompileTime((CanConvertPointerFromTo<UserClass, UObjectBase>::Result), You_Cannot_Use_UObject_Method_Delegates_With_Raw_Pointers);
 
-	/**
-	 * Internal, inlined and non-virtual version of GetRawUserObject interface.
-	 */
-	FORCEINLINE const void* GetRawUserObjectInternal() const
-	{
-		return UserObject.Get();
-	}
-
-	/**
-	 * Internal, inlined and non-virtual version of GetRawMethodPtr interface.
-	 */
-	FORCEINLINE const void* GetRawMethodPtrInternal() const
-	{
-		return *(const void**)&MethodPtr;
-	}
-
 public:
 
-	/** Pointer-to-member typedef */
-	typedef RetValType ( UserClass::*FMethodPtr )( DELEGATE_PARAM_LIST_COMMA_PAYLOAD_LIST ) DELEGATE_QUALIFIER;
-
+	// Type definition for member function pointers. */
+	typedef RetValType (UserClass::*FMethodPtr)( DELEGATE_PARAM_LIST_COMMA_PAYLOAD_LIST ) DELEGATE_QUALIFIER;
 
 	/**
-	 * Static helper function that returns a delegate object for the specified user class and class method
-	 *
-	 * @param  InUserObject  User's object that contains the class method
-	 * @param  InFunc  Member function pointer to your class method
-	 *
-	 * @return  Returns an interface to the new delegate
-	 */
-	FORCEINLINE static DELEGATE_INSTANCE_INTERFACE_CLASS< FUNC_TEMPLATE_ARGS >* Create( UserClass* InUserObject, FMethodPtr InFunc DELEGATE_COMMA_PAYLOAD_LIST )
-	{
-		return new UOBJECT_METHOD_DELEGATE_INSTANCE_CLASS< UserClass, FUNC_PAYLOAD_TEMPLATE_ARGS >( InUserObject, InFunc DELEGATE_COMMA_PAYLOAD_PASSTHRU );
-	}
-
-	/**
-	 * Constructor
+	 * Creates and initializes a new instance.
 	 *
-	 * @param  InUserObject  An arbitrary object (templated) that hosts the member function
-	 * @param  InMethodPtr  C++ member function pointer for the method to bind
+	 * @param InUserObject An arbitrary object (templated) that hosts the member function.
+	 * @param InMethodPtr C++ member function pointer for the method to bind.
 	 */
 	UOBJECT_METHOD_DELEGATE_INSTANCE_CLASS( UserClass* InUserObject, FMethodPtr InMethodPtr DELEGATE_COMMA_PAYLOAD_LIST )
-		: UserObject( InUserObject ),
-		  MethodPtr( InMethodPtr )
-		  DELEGATE_COMMA_PAYLOAD_INITIALIZER_LIST
+		: UserObject(InUserObject)
+		, MethodPtr(InMethodPtr)
+		DELEGATE_COMMA_PAYLOAD_INITIALIZER_LIST
 	{
 		// NOTE: UObject delegates are allowed to have a null incoming object pointer.  UObject weak pointers can expire,
 		//       an it is possible for a copy of a delegate instance to end up with a null pointer.
-		checkSlow( MethodPtr != NULL );
+		checkSlow(MethodPtr != nullptr);
 	}
 
-	/**
-	 * Returns the type of delegate instance
-	 *
-	 * @return  Delegate instance type
-	 */
-	virtual EDelegateInstanceType::Type GetType() const OVERRIDE
+public:
+
+	// Begin IDelegateInstance interface
+
+	virtual FName GetFunctionName( ) const override
 	{
-		return EDelegateInstanceType::UObjectMethod;
+		return NAME_None;
 	}
 
-	/**
-	 * Creates a copy of the delegate instance
-	 *
-	 * @return	The newly created copy
-	 */
-	virtual DELEGATE_INSTANCE_INTERFACE_CLASS< FUNC_TEMPLATE_ARGS >* CreateCopy() OVERRIDE
+	virtual const void* GetRawMethodPtr( ) const override
 	{
-		return Create( UserObject.Get(), MethodPtr DELEGATE_COMMA_PAYLOAD_PASSIN );
+		return GetRawMethodPtrInternal();
 	}
 
-	/**
-	 * Returns true if this delegate points to exactly the same object and method as the specified delegate,
-	 * even if the delegate objects themselves are different.  Also, the delegate types *must* be compatible.
-	 *
-	 * @param  InOtherDelegate
-	 *
-	 * @return  True if delegates match
-	 */
-	virtual bool IsSameFunction( const DELEGATE_INSTANCE_INTERFACE_CLASS< FUNC_TEMPLATE_ARGS >& InOtherDelegate ) const OVERRIDE
+	virtual const void* GetRawUserObject( ) const override
 	{
-		// NOTE: Payload data is not currently considered when comparing delegate instances.  See the comment in
-		// multi-cast delegate's Remove() method for more information.
-
-		if ( InOtherDelegate.GetType() == EDelegateInstanceType::UObjectMethod || 
-			   InOtherDelegate.GetType() == EDelegateInstanceType::RawMethod ||
-				 InOtherDelegate.GetType() == EDelegateInstanceType::SharedPointerMethod ||
-				 InOtherDelegate.GetType() == EDelegateInstanceType::ThreadSafeSharedPointerMethod )
-		{
-			return GetRawMethodPtrInternal() == InOtherDelegate.GetRawMethodPtr() && UserObject.Get() == InOtherDelegate.GetRawUserObject();
-		}
+		return GetRawUserObjectInternal();
+	}
 
-		return false;
+	virtual EDelegateInstanceType::Type GetType( ) const override
+	{
+		return EDelegateInstanceType::UObjectMethod;
 	}
 
-	/**
-	 * Returns true if this delegate is bound to the specified UserObject,
-	 *
-	 * @param  InUserObject
-	 *
-	 * @return  True if delegate is bound to the specified UserObject
-	 */
-	virtual bool HasSameObject( const void* InUserObject ) const OVERRIDE
+	virtual bool HasSameObject( const void* InUserObject ) const override
 	{
-		return UserObject.Get() == InUserObject;
+		return (UserObject.Get() == InUserObject);
 	}
 
-	virtual FName GetFunctionName() const OVERRIDE
+	virtual bool IsCompactable( ) const override
 	{
-		return NAME_None;
+		return !UserObject.Get(true);
 	}
 
-	virtual const void* GetRawUserObject() const OVERRIDE
+	virtual bool IsSafeToExecute( ) const override
 	{
-		return GetRawUserObjectInternal();
+		return !!UserObject.Get();
 	}
 
-	virtual const void* GetRawMethodPtr() const OVERRIDE
+	// End IDelegateInstance interface
+
+public:
+
+	// Begin DELEGATE_INSTANCE_INTERFACE_CLASS interface
+
+	virtual DELEGATE_INSTANCE_INTERFACE_CLASS<FUNC_TEMPLATE_ARGS>* CreateCopy( ) override
 	{
-		return GetRawMethodPtrInternal();
+		return Create(UserObject.Get(), MethodPtr DELEGATE_COMMA_PAYLOAD_PASSIN);
 	}
 
-	/**
-	 * Execute the delegate.  If the function pointer is not valid, an error will occur.
-	 */
-	virtual RetValType Execute( FUNC_PARAM_LIST ) const OVERRIDE
+	virtual RetValType Execute( FUNC_PARAM_LIST ) const override
 	{
 		// Verify that the user object is still valid.  We only have a weak reference to it.
-		checkSlow( UserObject.IsValid() );
+		checkSlow(UserObject.IsValid());
 
 		// Safely remove const to work around a compiler issue with instantiating template permutations for 
 		// overloaded functions that take a function pointer typedef as a member of a templated class.  In
 		// all cases where this code is actually invoked, the UserClass will already be a const pointer.
 		typedef typename TRemoveConst<UserClass>::Type MutableUserClass;
-		MutableUserClass* MutableUserObject = const_cast<MutableUserClass*>( UserObject.Get() );
+		MutableUserClass* MutableUserObject = const_cast<MutableUserClass*>(UserObject.Get());
 
 		// Call the member function on the user's object.  And yes, this is the correct C++ syntax for calling a
 		// pointer-to-member function.
-		checkSlow( MethodPtr != NULL );
+		checkSlow(MethodPtr != nullptr);
+
+		return (MutableUserObject->*MethodPtr)(DELEGATE_PARAM_PASSTHRU_COMMA_PAYLOAD_PASSIN);
+	}
+
+#if FUNC_IS_VOID
+	virtual bool ExecuteIfSafe( FUNC_PARAM_LIST ) const override
+	{
+		// Verify that the user object is still valid.  We only have a weak reference to it.
+		UserClass* ActualUserObject = UserObject.Get();
+		if (ActualUserObject != nullptr)
+		{
+			Execute(FUNC_PARAM_PASSTHRU);
+
+			return true;
+		}
+		return false;
+	}
+#endif
 
-		return ( MutableUserObject->*MethodPtr )( DELEGATE_PARAM_PASSTHRU_COMMA_PAYLOAD_PASSIN );
+	virtual bool IsSameFunction( const DELEGATE_INSTANCE_INTERFACE_CLASS<FUNC_TEMPLATE_ARGS>& InOtherDelegate ) const override
+	{
+		// NOTE: Payload data is not currently considered when comparing delegate instances.
+		// See the comment in multi-cast delegate's Remove() method for more information.
+		if ((InOtherDelegate.GetType() == EDelegateInstanceType::UObjectMethod) || 
+			(InOtherDelegate.GetType() == EDelegateInstanceType::RawMethod) ||
+			(InOtherDelegate.GetType() == EDelegateInstanceType::SharedPointerMethod) ||
+			(InOtherDelegate.GetType() == EDelegateInstanceType::ThreadSafeSharedPointerMethod))
+		{
+			return (GetRawMethodPtrInternal() == InOtherDelegate.GetRawMethodPtr()) && (UserObject.Get() == InOtherDelegate.GetRawUserObject());
+		}
+
+		return false;
 	}
 
+public:
+
 	/**
-	 * Checks to see if the user object bound to this delegate is still valid
+	 * Creates a new UObject delegate binding for the given user object and method pointer.
 	 *
-	 * @return  True if the object is still valid and it's safe to execute the function call
+	 * @param InUserObject User's object that contains the class method.
+	 * @param InFunc Member function pointer to your class method.
+	 * @return The new delegate.
 	 */
-	virtual bool IsSafeToExecute() const OVERRIDE
+	FORCEINLINE static DELEGATE_INSTANCE_INTERFACE_CLASS<FUNC_TEMPLATE_ARGS>* Create( UserClass* InUserObject, FMethodPtr InFunc DELEGATE_COMMA_PAYLOAD_LIST )
 	{
-		return !!UserObject.Get();
+		return new UOBJECT_METHOD_DELEGATE_INSTANCE_CLASS<UserClass, FUNC_PAYLOAD_TEMPLATE_ARGS>(InUserObject, InFunc DELEGATE_COMMA_PAYLOAD_PASSTHRU);
 	}
 
-	virtual bool IsCompactable() const OVERRIDE
+protected:
+
+	/**
+	 * Internal, inlined and non-virtual version of GetRawUserObject interface.
+	 */
+	FORCEINLINE const void* GetRawUserObjectInternal( ) const
 	{
-		return !UserObject.Get(true);
+		return UserObject.Get();
 	}
 
-#if FUNC_IS_VOID
 	/**
-	 * Execute the delegate, but only if the function pointer is still valid
-	 *
-	 * @return  Returns true if the function was executed
+	 * Internal, inlined and non-virtual version of GetRawMethodPtr interface.
 	 */
-	virtual bool ExecuteIfSafe( FUNC_PARAM_LIST ) const OVERRIDE
+	FORCEINLINE const void* GetRawMethodPtrInternal( ) const
 	{
-		// Verify that the user object is still valid.  We only have a weak reference to it.
-		UserClass* ActualUserObject = UserObject.Get();
-		if( ActualUserObject != NULL )
-		{
-			Execute( FUNC_PARAM_PASSTHRU );
-			return true;
-		}
-		return false;
+		return *(const void**)&MethodPtr;
 	}
-#endif
 
 private:
 
 	// Declare other pointer-based delegates as a friend so IsSameFunction() can compare members
-	template< class UserClassNoShadow, FUNC_PAYLOAD_TEMPLATE_DECL_NO_SHADOW, ESPMode::Type SPModeNoShadow > friend class SP_METHOD_DELEGATE_INSTANCE_CLASS;
-	template< class UserClassNoShadow, FUNC_PAYLOAD_TEMPLATE_DECL_NO_SHADOW > friend class RAW_METHOD_DELEGATE_INSTANCE_CLASS;
+	template<class UserClassNoShadow, FUNC_PAYLOAD_TEMPLATE_DECL_NO_SHADOW,ESPMode::Type SPModeNoShadow> friend class SP_METHOD_DELEGATE_INSTANCE_CLASS;
+	template<class UserClassNoShadow, FUNC_PAYLOAD_TEMPLATE_DECL_NO_SHADOW> friend class RAW_METHOD_DELEGATE_INSTANCE_CLASS;
 
-	/** Pointer to the user's class which contains a method we would like to call. */
-	TWeakObjectPtr< UserClass > UserObject;
+	// Pointer to the user's class which contains a method we would like to call.
+	TWeakObjectPtr<UserClass> UserObject;
 
-	/** C++ member function pointer */
+	// C++ member function pointer.
 	FMethodPtr MethodPtr;
 
-	/** Payload member variables (if any) */
+	// Payload member variables (if any).
 	FUNC_PAYLOAD_MEMBERS
 };
 
 
 /**
- * Delegate type implementation class for global functions.  For internal use only.
+ * Implements a delegate binding for regular C++ functions.
  */
-template< FUNC_PAYLOAD_TEMPLATE_DECL_TYPENAME >
+template<FUNC_PAYLOAD_TEMPLATE_DECL_TYPENAME>
 class STATIC_DELEGATE_INSTANCE_CLASS
-	: public DELEGATE_INSTANCE_INTERFACE_CLASS< FUNC_TEMPLATE_ARGS >
+	: public DELEGATE_INSTANCE_INTERFACE_CLASS<FUNC_TEMPLATE_ARGS>
 {
 public:
 
-	/** Static function pointer typedef */
-	// NOTE: Static C++ functions can never be const, so 'const' is always omitted here.  That means there is
-	//       usually an extra static delegate class created with a 'const' name, even though it doesn't use const.
-	typedef RetValType ( *FFuncPtr )( DELEGATE_PARAM_LIST_COMMA_PAYLOAD_LIST );
-
 	/**
-	 * Static helper function that returns a newly-created delegate object for a
-	 * static/global function pointer.
-	 *
-	 * @param  InFunc  Static function pointer
+	 * Type definition for static function pointers.
 	 *
-	 * @return  Returns a new instance of the delegate
+	 * NOTE: Static C++ functions can never be const, so 'const' is always omitted here. That means there is
+	 * usually an extra static delegate class created with a 'const' name, even though it doesn't use const.
 	 */
-	FORCEINLINE static DELEGATE_INSTANCE_INTERFACE_CLASS< FUNC_TEMPLATE_ARGS >* Create( FFuncPtr InFunc DELEGATE_COMMA_PAYLOAD_LIST )
-	{
-		return new STATIC_DELEGATE_INSTANCE_CLASS< FUNC_PAYLOAD_TEMPLATE_ARGS >( InFunc DELEGATE_COMMA_PAYLOAD_PASSTHRU );
-	}
+	typedef RetValType (*FFuncPtr)( DELEGATE_PARAM_LIST_COMMA_PAYLOAD_LIST );
 
 	/**
-	 * Constructor
+	 * Creates and initializes a new instance.
 	 *
-	 * @param  InStaticFuncPtr  C++ function pointer
+	 * @param InStaticFuncPtr C++ function pointer.
 	 */
 	STATIC_DELEGATE_INSTANCE_CLASS( FFuncPtr InStaticFuncPtr DELEGATE_COMMA_PAYLOAD_LIST )
-		: StaticFuncPtr( InStaticFuncPtr )
-		  DELEGATE_COMMA_PAYLOAD_INITIALIZER_LIST
+		: StaticFuncPtr(InStaticFuncPtr)
+		DELEGATE_COMMA_PAYLOAD_INITIALIZER_LIST
 	{
-		check( StaticFuncPtr != NULL );
+		check(StaticFuncPtr != nullptr);
 	}
 
-	/**
-	 * Returns the type of delegate instance
-	 *
-	 * @return  Delegate instance type
-	 */
-	virtual EDelegateInstanceType::Type GetType() const OVERRIDE
+public:
+
+	// Begin IDelegateInstance interface
+
+	virtual FName GetFunctionName( ) const override
 	{
-		return EDelegateInstanceType::Raw;
+		return NAME_None;
 	}
 
-	/**
-	 * Creates a copy of the delegate instance
-	 *
-	 * @return	The newly created copy
-	 */
-	virtual DELEGATE_INSTANCE_INTERFACE_CLASS< FUNC_TEMPLATE_ARGS >* CreateCopy() OVERRIDE
+	virtual const void* GetRawMethodPtr( ) const override
 	{
-		return Create( StaticFuncPtr DELEGATE_COMMA_PAYLOAD_PASSIN );
+		return *(const void**)&StaticFuncPtr;
 	}
 
-	/**
-	 * Returns true if this delegate points to exactly the same object and method as the specified delegate,
-	 * even if the delegate objects themselves are different.  However, the delegate types *must* match.
-	 *
-	 * @param  InOtherDelegate
-	 *
-	 * @return  True if delegates match
-	 */
-	virtual bool IsSameFunction( const DELEGATE_INSTANCE_INTERFACE_CLASS< FUNC_TEMPLATE_ARGS >& InOtherDelegate ) const OVERRIDE
+	virtual const void* GetRawUserObject( ) const override
 	{
-		// NOTE: Payload data is not currently considered when comparing delegate instances.  See the comment in
-		// multi-cast delegate's Remove() method for more information.
-
-		if( InOtherDelegate.GetType() == EDelegateInstanceType::Raw )
-		{
-			// Downcast to our delegate type and compare
-			const STATIC_DELEGATE_INSTANCE_CLASS< FUNC_PAYLOAD_TEMPLATE_ARGS >& OtherStaticDelegate =
-				static_cast< const STATIC_DELEGATE_INSTANCE_CLASS< FUNC_PAYLOAD_TEMPLATE_ARGS >& >( InOtherDelegate );
-			if( StaticFuncPtr == OtherStaticDelegate.StaticFuncPtr )
-			{
-				return true;
-			}
-		}
+		return nullptr;
+	}
 
-		return false;
+	virtual EDelegateInstanceType::Type GetType( ) const override
+	{
+		return EDelegateInstanceType::Raw;
 	}
 
-	/**
-	 * Returns true if this delegate is bound to the specified UserObject,
-	 *
-	 * @param  InUserObject
-	 *
-	 * @return  True if delegate is bound to the specified UserObject
-	 */
-	virtual bool HasSameObject( const void* UserObject ) const OVERRIDE
+	virtual bool HasSameObject( const void* UserObject ) const override
 	{
 		// Raw Delegates aren't bound to an object so they can never match
 		return false;
 	}
 
-	virtual FName GetFunctionName() const OVERRIDE
+	virtual bool IsSafeToExecute( ) const override
 	{
-		return NAME_None;
+		// Static functions are always safe to execute!
+		return true;
 	}
 
-	virtual const void* GetRawUserObject() const OVERRIDE
+	// End IDelegateInstance interface
+
+public:
+
+	// Begin DELEGATE_INSTANCE_INTERFACE_CLASS interface
+
+	virtual DELEGATE_INSTANCE_INTERFACE_CLASS<FUNC_TEMPLATE_ARGS>* CreateCopy( ) override
 	{
-		return NULL;
+		return Create(StaticFuncPtr DELEGATE_COMMA_PAYLOAD_PASSIN);
 	}
 
-	virtual const void* GetRawMethodPtr() const OVERRIDE
+	virtual RetValType Execute( FUNC_PARAM_LIST ) const override
 	{
-		return *(const void**)&StaticFuncPtr;
+		// Call the member function on the user's object.  And yes, this is the correct C++ syntax for calling a
+		// pointer-to-member function.
+		checkSlow(StaticFuncPtr != nullptr);
+
+		return StaticFuncPtr(DELEGATE_PARAM_PASSTHRU_COMMA_PAYLOAD_PASSIN);
 	}
 
-	/**
-	 * Checks to see if the user object bound to this delegate is still valid
-	 *
-	 * @return  True if the object is still valid and it's safe to execute the function call
-	 */
-	virtual bool IsSafeToExecute() const OVERRIDE
+#if FUNC_IS_VOID
+	virtual bool ExecuteIfSafe( FUNC_PARAM_LIST ) const override
 	{
-		// Static functions are always safe to execute!
+		Execute(FUNC_PARAM_PASSTHRU);
+
 		return true;
 	}
+#endif
 
-	/**
-	 * Execute the delegate.  If the function pointer is not valid, an error will occur.
-	 */
-	virtual RetValType Execute( FUNC_PARAM_LIST ) const OVERRIDE
+	virtual bool IsSameFunction( const DELEGATE_INSTANCE_INTERFACE_CLASS<FUNC_TEMPLATE_ARGS>& InOtherDelegate ) const override
 	{
-		// Call the member function on the user's object.  And yes, this is the correct C++ syntax for calling a
-		// pointer-to-member function.
-		checkSlow( StaticFuncPtr != NULL );
+		// NOTE: Payload data is not currently considered when comparing delegate instances.
+		// See the comment in multi-cast delegate's Remove() method for more information.
+		if (InOtherDelegate.GetType() == EDelegateInstanceType::Raw)
+		{
+			// Downcast to our delegate type and compare
+			const STATIC_DELEGATE_INSTANCE_CLASS<FUNC_PAYLOAD_TEMPLATE_ARGS>& OtherStaticDelegate =
+				static_cast<const STATIC_DELEGATE_INSTANCE_CLASS<FUNC_PAYLOAD_TEMPLATE_ARGS>&>(InOtherDelegate);
+
+			return (StaticFuncPtr == OtherStaticDelegate.StaticFuncPtr);
+		}
 
-		return StaticFuncPtr( DELEGATE_PARAM_PASSTHRU_COMMA_PAYLOAD_PASSIN );
+		return false;
 	}
 
-#if FUNC_IS_VOID
+public:
+
 	/**
-	 * Execute the delegate, but only if the function pointer is still valid
+	 * Creates a new static function delegate binding for the given function pointer.
 	 *
-	 * @return  Returns true if the function was executed
+	 * @param InFunc Static function pointer.
+	 * @return The new delegate.
 	 */
-	virtual bool ExecuteIfSafe( FUNC_PARAM_LIST ) const OVERRIDE
+	FORCEINLINE static DELEGATE_INSTANCE_INTERFACE_CLASS<FUNC_TEMPLATE_ARGS>* Create( FFuncPtr InFunc DELEGATE_COMMA_PAYLOAD_LIST )
 	{
-		Execute( FUNC_PARAM_PASSTHRU );
-		return true;
+		return new STATIC_DELEGATE_INSTANCE_CLASS<FUNC_PAYLOAD_TEMPLATE_ARGS>(InFunc DELEGATE_COMMA_PAYLOAD_PASSTHRU);
 	}
-#endif
 
 private:
 
-	/** C++ function pointer */
+	// C++ function pointer.
 	FFuncPtr StaticFuncPtr;
 
-	/** Payload member variables (if any) */
+	// Payload member variables, if any.
 	FUNC_PAYLOAD_MEMBERS
 };
 
@@ -837,7 +743,7 @@ private:
 #if !FUNC_IS_CONST
 
 /**
- * Delegate type implementation class for UFunctions.  For internal use only.
+ * Implements a delegate binding for UFunctions.
  *
  * @params UserClass Must be an UObject derived class.
  */
@@ -888,37 +794,37 @@ public:
 
 	// Begin IDelegateInstance interface
 
-	virtual FName GetFunctionName( ) const OVERRIDE
+	virtual FName GetFunctionName( ) const override
 	{
 		return FunctionName;
 	}
 
-	virtual const void* GetRawMethodPtr( ) const OVERRIDE
+	virtual const void* GetRawMethodPtr( ) const override
 	{
 		return nullptr;
 	}
 
-	virtual const void* GetRawUserObject( ) const OVERRIDE
+	virtual const void* GetRawUserObject( ) const override
 	{
 		return UserObjectPtr.Get();
 	}
 
-	virtual EDelegateInstanceType::Type GetType() const OVERRIDE
+	virtual EDelegateInstanceType::Type GetType() const override
 	{
 		return EDelegateInstanceType::UFunction;
 	}
 
-	virtual bool HasSameObject( const void* InUserObject ) const OVERRIDE
+	virtual bool HasSameObject( const void* InUserObject ) const override
 	{
 		return (UserObjectPtr.Get() == InUserObject);
 	}
 
-	virtual bool IsCompactable( ) const OVERRIDE
+	virtual bool IsCompactable( ) const override
 	{
 		return !UserObjectPtr.Get(true);
 	}
 
-	virtual bool IsSafeToExecute( ) const OVERRIDE
+	virtual bool IsSafeToExecute( ) const override
 	{
 		return UserObjectPtr.IsValid();
 	}
@@ -929,12 +835,12 @@ public:
 
 	// Begin DELEGATE_INSTANCE_INTERFACE_CLASS interface
 
-	virtual DELEGATE_INSTANCE_INTERFACE_CLASS<FUNC_TEMPLATE_ARGS>* CreateCopy() OVERRIDE
+	virtual DELEGATE_INSTANCE_INTERFACE_CLASS<FUNC_TEMPLATE_ARGS>* CreateCopy( ) override
 	{
 		return Create(UserObjectPtr.Get(), GetFunctionName() DELEGATE_COMMA_PAYLOAD_PASSIN);
 	}
 
-	virtual RetValType Execute( FUNC_PARAM_LIST ) const OVERRIDE
+	virtual RetValType Execute( FUNC_PARAM_LIST ) const override
 	{
 		checkSlow(IsSafeToExecute());
 
@@ -947,7 +853,7 @@ public:
 	}
 
 #if FUNC_IS_VOID
-	virtual bool ExecuteIfSafe( FUNC_PARAM_LIST ) const OVERRIDE
+	virtual bool ExecuteIfSafe( FUNC_PARAM_LIST ) const override
 	{
 		if (IsSafeToExecute())
 		{
@@ -960,8 +866,10 @@ public:
 	}
 #endif
 
-	virtual bool IsSameFunction( const DELEGATE_INSTANCE_INTERFACE_CLASS<FUNC_TEMPLATE_ARGS>& Other ) const OVERRIDE
+	virtual bool IsSameFunction( const DELEGATE_INSTANCE_INTERFACE_CLASS<FUNC_TEMPLATE_ARGS>& Other ) const override
 	{
+		// NOTE: Payload data is not currently considered when comparing delegate instances.
+		// See the comment in multi-cast delegate's Remove() method for more information.
 		return ((Other.GetType() == EDelegateInstanceType::UFunction) &&
 				(Other.GetRawUserObject() == GetRawUserObject()) &&
 				(Other.GetFunctionName() == GetFunctionName()));
@@ -972,12 +880,11 @@ public:
 public:
 
 	/**
-	 * Static helper function that returns a delegate object for the specified user class and class method.
+	 * Creates a new UFunction delegate binding for the given user object and function name.
 	 *
 	 * @param InObject The user object to call the function on.
 	 * @param InFunctionName The name of the function call.
-	 *
-	 * @return  Returns an interface to the new delegate
+	 * @return The new delegate.
 	 */
 	FORCEINLINE static DELEGATE_INSTANCE_INTERFACE_CLASS<FUNC_TEMPLATE_ARGS>* Create( UserClass* InUserObject, const FName& InFunctionName DELEGATE_COMMA_PAYLOAD_LIST )
 	{
