Fixed FText::TrimPrecedingAndTrailing

TTP# 338106 - BLOCKER: EDITOR: MAIN: Search in Content Browser does not filter beyond first letter entered

I'd recently changed this to use FText::IsWhitespace rather than using the FString functions (which use FChar::IsWhitespace, which isn't culture aware), but had messed up the call to FString::Mid which trimmed the string (I was passing a start point and an end point, but it needs a start point and a length).

PendingReview Max.Preussner

diff --git a/Engine/Source/Runtime/Core/Private/Internationalization/Text.cpp b/Engine/Source/Runtime/Core/Private/Internationalization/Text.cpp
--- a/Engine/Source/Runtime/Core/Private/Internationalization/Text.cpp
+++ b/Engine/Source/Runtime/Core/Private/Internationalization/Text.cpp
@@ -389,7 +389,8 @@ FText FText::TrimPrecedingAndTrailing( const FText& InText )
 			--EndPos;
 		}
 
-		TrimmedString = TrimmedString.Mid( StartPos, TrimmedString.Len() - EndPos );
+		const int32 Len = (EndPos + 1) - StartPos;
+		TrimmedString = TrimmedString.Mid( StartPos, Len );
 	}
 
 	FText NewText = FText( MoveTemp( TrimmedString ) );
