TArray::FilterByPredicate added

#github 218
#ttp 338155: [GitHub] 218 : Filter by predicate support for TArray

https://github.com/EpicGames/UnrealEngine/pull/218

diff --git a/Engine/Source/Runtime/Core/Public/Containers/Array.h b/Engine/Source/Runtime/Core/Public/Containers/Array.h
--- a/Engine/Source/Runtime/Core/Public/Containers/Array.h
+++ b/Engine/Source/Runtime/Core/Public/Containers/Array.h
@@ -734,6 +734,24 @@ public:
 
 		return NULL;
 	}
+	/**
+	 * Filters the elements in the array based on a predicate functor. 
+	 * @param  Pred The functor to apply to each element.
+	 * @return      TArray with the same type as this object which contains the subset of elements for which the functor returns true.
+	 */
+	template <typename Predicate>
+	TArray<ElementType> FilterByPredicate(Predicate Pred)
+	{
+		TArray<ElementType> FilterResults;
+		for (ElementType* RESTRICT Data = GetTypedData(), *RESTRICT DataEnd = Data + ArrayNum; Data != DataEnd; ++Data)
+		{
+			if (Pred(*Data))
+			{
+				FilterResults.Add(*Data);
+			}
+		}
+		return FilterResults;
+	}
 	/**
 	 * @return		true if found
 	 */
