Fixed non-unity Mac build issues.

diff --git a/Engine/Source/Runtime/Core/Private/Internationalization/ICUText.h b/Engine/Source/Runtime/Core/Private/Internationalization/ICUText.h
--- a/Engine/Source/Runtime/Core/Private/Internationalization/ICUText.h
+++ b/Engine/Source/Runtime/Core/Private/Internationalization/ICUText.h
@@ -4,6 +4,7 @@
 
 #include "ICUUtilities.h"
 #include "ICUCulture.h"
+#include "TextHistory.h"
 #include "unicode/utypes.h"
 #include <unicode/numfmt.h>
 #include "unicode/fmtable.h"
@@ -67,4 +68,4 @@ FText FText::AsPercentTemplate(T1 Val, const FNumberFormattingOptions* const Opt
 	FText ReturnText = FText::CreateNumericalText(NativeString);
 	ReturnText.History = MakeShareable(new FTextHistory_AsPercent(Val, Options, TargetCulture));
 	return ReturnText;
-}
\ No newline at end of file
+}
