Remove logging for uproject file scanning, print a list on completion for default engine instead. Prevents polluting all the project upgrade error dialogs.

diff --git a/Engine/Source/Runtime/Core/Private/Misc/UProjectInfo.cpp b/Engine/Source/Runtime/Core/Private/Misc/UProjectInfo.cpp
--- a/Engine/Source/Runtime/Core/Private/Misc/UProjectInfo.cpp
+++ b/Engine/Source/Runtime/Core/Private/Misc/UProjectInfo.cpp
@@ -9,8 +9,6 @@ DEFINE_LOG_CATEGORY(LogUProjectInfo);
 
 FUProjectDictionary::FUProjectDictionary(const FString& InRootDir)
 {
-	double StartTime = FPlatformTime::Seconds();
-
 	FString RootDir = InRootDir;
 	FPaths::NormalizeDirectoryName(RootDir);
 
@@ -24,7 +22,6 @@ FUProjectDictionary::FUProjectDictionary(const FString& InRootDir)
 	for (int32 FileIndex = 0; FileIndex < UProjectDirsFiles.Num(); ++FileIndex)
 	{
 		FString ProjDirsFile = RootDir / UProjectDirsFiles[FileIndex];
-		UE_LOG(LogUProjectInfo, Log, TEXT("Found uprojectdirs file %s"), *ProjDirsFile);
 
 		TArray<FString> FileStrings;
 		if (FFileHelper::LoadANSITextFileToStrings(*ProjDirsFile, &IFileManager::Get(), FileStrings))
@@ -49,12 +46,9 @@ FUProjectDictionary::FUProjectDictionary(const FString& InRootDir)
 		}
 	}
 
-	UE_LOG(LogUProjectInfo, Log, TEXT("Found %d directories to search"), DirectoriesToSearch.Num());
-
 	for (int32 DirIdx = 0; DirIdx < DirectoriesToSearch.Num(); DirIdx++)
 	{
 		FString DirToSearch = DirectoriesToSearch[DirIdx];
-		UE_LOG(LogUProjectInfo, Log, TEXT("\tSearching %s"), *DirToSearch);
 
 		// To specific the <UE4> folder itself, './' is used... strip it out if present
 		DirToSearch = DirToSearch.Replace(TEXT("/./"), TEXT("/"));
@@ -63,17 +57,14 @@ FUProjectDictionary::FUProjectDictionary(const FString& InRootDir)
 		TArray<FString> SubDirectories;
 		SearchPath = DirToSearch / FString(TEXT("*"));
 		IFileManager::Get().FindFiles(SubDirectories, *SearchPath, false, true);
-		UE_LOG(LogUProjectInfo, Log, TEXT("\t\tFound %2d subdirectories"), SubDirectories.Num());
 		
 		for (int32 SubIdx = 0; SubIdx < SubDirectories.Num(); SubIdx++)
 		{
 			FString SubDir = DirToSearch / SubDirectories[SubIdx];
-			UE_LOG(LogUProjectInfo, Log, TEXT("\t\t\tSubdir %s"), *SubDir);
 
 			TArray<FString> SubDirFiles;
 			SearchPath = SubDir / TEXT("*.uproject");
 			IFileManager::Get().FindFiles(SubDirFiles, *SearchPath, true, false);
-			UE_LOG(LogUProjectInfo, Log, TEXT("\t\t\t\tFound %2d uproject files"), SubDirFiles.Num());
 
 			for (int32 UProjIdx = 0; UProjIdx < SubDirFiles.Num(); UProjIdx++)
 			{
@@ -85,9 +76,6 @@ FUProjectDictionary::FUProjectDictionary(const FString& InRootDir)
 			}
 		}
 	}
-
-	double TotalProjectInfoTime = FPlatformTime::Seconds() - StartTime;
-	UE_LOG(LogUProjectInfo, Log, TEXT("FillProjectInfo took %5.4f seconds"), TotalProjectInfoTime);
 }
 
 bool FUProjectDictionary::IsForeignProject(const FString& InProjectFileName) const
@@ -128,5 +116,19 @@ TArray<FString> FUProjectDictionary::GetProjectPaths() const
 FUProjectDictionary& FUProjectDictionary::GetDefault()
 {
 	static FUProjectDictionary DefaultDictionary(FPaths::RootDir());
+
+#if !NO_LOGGING
+	static bool bHaveLoggedProjects = false;
+	if(!bHaveLoggedProjects)
+	{
+		UE_LOG(LogUProjectInfo, Log, TEXT("Found projects:"));
+		for(TMap<FString, FString>::TConstIterator Iter(DefaultDictionary.ShortProjectNameDictionary); Iter; ++Iter)
+		{
+			UE_LOG(LogUProjectInfo, Log, TEXT("    %s: \"%s\""), *Iter.Key(), *Iter.Value());
+		}
+		bHaveLoggedProjects = true;
+	}
+#endif
+
 	return DefaultDictionary;
 }
