Core: Easier debugging of file delete operations

diff --git a/Engine/Source/Runtime/Core/Private/HAL/FileManagerGeneric.cpp b/Engine/Source/Runtime/Core/Private/HAL/FileManagerGeneric.cpp
--- a/Engine/Source/Runtime/Core/Private/HAL/FileManagerGeneric.cpp
+++ b/Engine/Source/Runtime/Core/Private/HAL/FileManagerGeneric.cpp
@@ -263,7 +263,7 @@ bool FFileManagerGeneric::Delete( const TCHAR* Filename, bool RequireExists, boo
 		{
 			if (!Quiet)
 			{
-				UE_LOG( LogFileManager, Warning, TEXT( "Error deleting file: %s" ), Filename );
+				UE_LOG( LogFileManager, Warning, TEXT( "Error deleting file: %s (Error Code %i)" ), Filename, FPlatformMisc::GetLastError() );
 			}
 			return false;
 		}
diff --git a/Engine/Source/Runtime/Core/Private/Windows/WindowsPlatformFile.cpp b/Engine/Source/Runtime/Core/Private/Windows/WindowsPlatformFile.cpp
--- a/Engine/Source/Runtime/Core/Private/Windows/WindowsPlatformFile.cpp
+++ b/Engine/Source/Runtime/Core/Private/Windows/WindowsPlatformFile.cpp
@@ -159,7 +159,8 @@ public:
 	}
 	virtual bool DeleteFile(const TCHAR* Filename) OVERRIDE
 	{
-		return !!DeleteFileW(*NormalizeFilename(Filename));
+		const FString NormalizedFilename = NormalizeFilename(Filename);
+		return !!DeleteFileW(*NormalizedFilename);
 	}
 	virtual bool IsReadOnly(const TCHAR* Filename) OVERRIDE
 	{
