Core - Added a new define PLATFORM_USES_ANSI_STRING_FOR_EXTERNAL_PROFILING, specify which version of string should be used for external profiling like PIX, optimization to avoid expensive ANSI_TO_TCHAR/TCHAR_TO_ANSI conversion

diff --git a/Engine/Source/Runtime/Core/Public/HAL/Platform.h b/Engine/Source/Runtime/Core/Public/HAL/Platform.h
--- a/Engine/Source/Runtime/Core/Public/HAL/Platform.h
+++ b/Engine/Source/Runtime/Core/Public/HAL/Platform.h
@@ -235,6 +235,10 @@
 	#define PLATFORM_HAS_128BIT_ATOMICS 0
 #endif
 
+#ifndef PLATFORM_USES_ANSI_STRING_FOR_EXTERNAL_PROFILING
+	#define PLATFORM_USES_ANSI_STRING_FOR_EXTERNAL_PROFILING 1
+#endif
+
 //These are deprecated old defines that we want to make sure are not used
 #define CONSOLE (#)
 #define MOBILE (#)
diff --git a/Engine/Source/Runtime/Core/Public/Windows/WIndowsPlatform.h b/Engine/Source/Runtime/Core/Public/Windows/WIndowsPlatform.h
--- a/Engine/Source/Runtime/Core/Public/Windows/WIndowsPlatform.h
+++ b/Engine/Source/Runtime/Core/Public/Windows/WIndowsPlatform.h
@@ -33,48 +33,49 @@ typedef FWindowsPlatformTypes FPlatformTypes;
 #if defined( _MANAGED ) || defined( _M_CEE )
 	#define PLATFORM_COMPILER_COMMON_LANGUAGE_RUNTIME_COMPILATION	1
 #endif
-#define PLATFORM_CAN_SUPPORT_EDITORONLY_DATA	1
+#define PLATFORM_CAN_SUPPORT_EDITORONLY_DATA				1
 
 // Base defines, defaults are commented out
 
-#define PLATFORM_LITTLE_ENDIAN						1
+#define PLATFORM_LITTLE_ENDIAN								1
 #if PLATFORM_WINDOWS_CLANG
 	// @todo Clang Clang emits a compile error on Windows when __try/__except blocks are used (missing functionality in Clang.)
 	//             After Clang adds proper support for SEH exceptions, we should delete this definition and all use cases of it!
-	#define PLATFORM_SEH_EXCEPTIONS_DISABLED			1
+	#define PLATFORM_SEH_EXCEPTIONS_DISABLED				1
 
 	// @todo clang: This is needed because of a bug in Clang with RTTI and exceptions (even when RTTI is disabled!)  This manifests as a 
 	// compile error whenever try/catch/throw is used.  See http://llvm.org/bugs/show_bug.cgi?id=17403.  This code should be removed 
 	// after the bug is fixed in Clang!
-	#define PLATFORM_EXCEPTIONS_DISABLED				1
+	#define PLATFORM_EXCEPTIONS_DISABLED					1
 #endif
 
-#define PLATFORM_SUPPORTS_PRAGMA_PACK				1
+#define PLATFORM_SUPPORTS_PRAGMA_PACK						1
 #if PLATFORM_WINDOWS_CLANG
-	#define PLATFORM_ENABLE_VECTORINTRINSICS		0
+	#define PLATFORM_ENABLE_VECTORINTRINSICS				0
 #else
-	#define PLATFORM_ENABLE_VECTORINTRINSICS		1
+	#define PLATFORM_ENABLE_VECTORINTRINSICS				1
 #endif
-//#define PLATFORM_USE_LS_SPEC_FOR_WIDECHAR			1
-//#define PLATFORM_USE_SYSTEM_VSWPRINTF				1
-//#define PLATFORM_TCHAR_IS_4_BYTES					0
-#define PLATFORM_HAS_vsnprintf						0
-#define PLATFORM_HAS_BSD_TIME						0
-#define PLATFORM_USE_PTHREADS						0
-#define PLATFORM_MAX_FILEPATH_LENGTH				MAX_PATH
-#define PLATFORM_HAS_BSD_SOCKET_FEATURE_WINSOCKETS	1
-#define PLATFORM_USES_MICROSOFT_LIBC_FUNCTIONS		1
-#define PLATFORM_SUPPORTS_TBB						1
-#define PLATFORM_SUPPORTS_NAMED_PIPES				1
-#define PLATFORM_COMPILER_HAS_DEFAULTED_FUNCTIONS	0
-#define PLATFORM_COMPILER_HAS_VARIADIC_TEMPLATES	0
-#define PLATFORM_COMPILER_HAS_EXPLICIT_OPERATORS	0
-#define PLATFORM_COMPILER_HAS_TCHAR_WMAIN			1
-
-#define PLATFORM_HAS_128BIT_ATOMICS					(!HACK_HEADER_GENERATOR && PLATFORM_64BITS)
+//#define PLATFORM_USE_LS_SPEC_FOR_WIDECHAR					1
+//#define PLATFORM_USE_SYSTEM_VSWPRINTF						1
+//#define PLATFORM_TCHAR_IS_4_BYTES							0
+#define PLATFORM_HAS_vsnprintf								0
+#define PLATFORM_HAS_BSD_TIME								0
+#define PLATFORM_USE_PTHREADS								0
+#define PLATFORM_MAX_FILEPATH_LENGTH						MAX_PATH
+#define PLATFORM_HAS_BSD_SOCKET_FEATURE_WINSOCKETS			1
+#define PLATFORM_USES_MICROSOFT_LIBC_FUNCTIONS				1
+#define PLATFORM_SUPPORTS_TBB								1
+#define PLATFORM_SUPPORTS_NAMED_PIPES						1
+#define PLATFORM_COMPILER_HAS_DEFAULTED_FUNCTIONS			0
+#define PLATFORM_COMPILER_HAS_VARIADIC_TEMPLATES			0
+#define PLATFORM_COMPILER_HAS_EXPLICIT_OPERATORS			0
+#define PLATFORM_COMPILER_HAS_TCHAR_WMAIN					1
+
+#define PLATFORM_HAS_128BIT_ATOMICS							(!HACK_HEADER_GENERATOR && PLATFORM_64BITS)
+#define PLATFORM_USES_ANSI_STRING_FOR_EXTERNAL_PROFILING	0
 
 #if _MSC_VER < 1700
-	#define PLATFORM_COMPILER_HAS_RANGED_FOR_LOOP 0
+	#define PLATFORM_COMPILER_HAS_RANGED_FOR_LOOP			0
 #endif
 
 
