Stats - Fixed an issue with reading uncompressed chunks

diff --git a/Engine/Source/Runtime/Core/Public/Stats/StatsFile.h b/Engine/Source/Runtime/Core/Public/Stats/StatsFile.h
--- a/Engine/Source/Runtime/Core/Public/Stats/StatsFile.h
+++ b/Engine/Source/Runtime/Core/Public/Stats/StatsFile.h
@@ -147,6 +147,13 @@ protected:
 		{
 			bEndOfCompressedData = true;
 		}
+		// This chunk is not compressed.
+		else if( CompressedSize == 0 )
+		{
+			DestData.Reset( UncompressedSize );
+			DestData.AddUninitialized( UncompressedSize );
+			Reader.Serialize( DestData.GetData(), UncompressedSize );
+		}
 		else
 		{
 			SrcData.Reset( CompressedSize );
