Core: Change TUnion::SetSubtype to return a typed pointer to the storage

diff --git a/Engine/Source/Runtime/Core/Public/Containers/Union.h b/Engine/Source/Runtime/Core/Public/Containers/Union.h
--- a/Engine/Source/Runtime/Core/Public/Containers/Union.h
+++ b/Engine/Source/Runtime/Core/Public/Containers/Union.h
@@ -183,7 +183,7 @@ public:
 
 	/** Replaces the value of the union with a value of the given subtype. */
 	template<typename Subtype>
-	void SetSubtype(typename TCallTraits<Subtype>::ParamType NewValue)
+	Subtype* SetSubtype(typename TCallTraits<Subtype>::ParamType NewValue)
 	{
 		int32 SubtypeIndex;
 		Subtype* SubtypeValuePointer;
@@ -194,6 +194,7 @@ public:
 		new(SubtypeValuePointer) Subtype(NewValue);
 
 		CurrentSubtypeIndex = SubtypeIndex;
+		return SubtypeValuePointer;
 	}
 
 	/** Sets the union's value to NULL. */
