Added TRemoveCV template

diff --git a/Engine/Source/Runtime/Core/Public/Templates/UnrealTemplate.h b/Engine/Source/Runtime/Core/Public/Templates/UnrealTemplate.h
--- a/Engine/Source/Runtime/Core/Public/Templates/UnrealTemplate.h
+++ b/Engine/Source/Runtime/Core/Public/Templates/UnrealTemplate.h
@@ -428,6 +428,16 @@ FORCEINLINE T StaticCast(ArgType&& Arg)
 	return static_cast<T>(Arg);
 }
 
+/*
+ * TRemoveCV<type> will remove any const/volatile qualifiers from a type.
+ * (based on std::remove_cv<>
+ * note: won't remove the const from "const int*", as the pointer is not const
+ */
+template <typename T> struct TRemoveCV                   { typedef T Type; };
+template <typename T> struct TRemoveCV<const T>          { typedef T Type; };
+template <typename T> struct TRemoveCV<volatile T>       { typedef T Type; };
+template <typename T> struct TRemoveCV<const volatile T> { typedef T Type; };
+
 /**
  * TRValueToLValueReference converts any rvalue reference type into the equivalent lvalue reference, otherwise returns the same type.
  */
