Fixed issue with the automation test completing unsuccessfully.

#ttp 338292 - CRITICAL: Regression: AUTOMATION: TEXT test failed with "FormattedTestLayer2 did not rebuild to correctly in LEET!"

diff --git a/Engine/Source/Runtime/Core/Private/Tests/Internationalization/TextTest.cpp b/Engine/Source/Runtime/Core/Private/Tests/Internationalization/TextTest.cpp
--- a/Engine/Source/Runtime/Core/Private/Tests/Internationalization/TextTest.cpp
+++ b/Engine/Source/Runtime/Core/Private/Tests/Internationalization/TextTest.cpp
@@ -429,6 +429,8 @@ bool FTextTest::RunTest (const FString& Parameters)
 			FDateTime DateTimeInfo(2080, 8, 20, 9, 33, 22);
 			FText AsDateTimeTest1 = FText::AsDateTime(DateTimeInfo);
 
+			// FormattedTestLayer2 must be updated when adding or removing from this block. Further, below, 
+			// verifying the LEET translated string must be changed to reflect what the new string looks like.
 			FFormatNamedArguments ArgsLayer2;
 			ArgsLayer2.Add("NamedLayer1", FormattedTest1);
 			ArgsLayer2.Add("OrderedLayer1", FormattedTestOrdered1);
@@ -437,7 +439,7 @@ bool FTextTest::RunTest (const FString& Parameters)
 			ArgsLayer2.Add("DateTime", AsDateTimeTest1);
 			ArgsLayer2.Add("Percent", AsPercentTest1);
 			ArgsLayer2.Add("Currency", AsCurrencyTest1);
-			FText FormattedTestLayer2 = FText::Format(LOCTEXT("RebuildTextLayer2", "{NamedLayer1} | {OrderedLayer1} | {FTextNumber} | {Number} | {Date} | {Time} | {DateTime} | {Percent} | {Currency}"), ArgsLayer2);
+			FText FormattedTestLayer2 = FText::Format(LOCTEXT("RebuildTextLayer2", "{NamedLayer1} | {OrderedLayer1} | {FTextNumber} | {Number} | {DateTime} | {Percent} | {Currency}"), ArgsLayer2);
 
 			{
 				// Serialize the full, bulky FText that is a composite of most of the other FTextHistories.
@@ -452,8 +454,11 @@ bool FTextTest::RunTest (const FString& Parameters)
 				// Swap to "LEET" culture to check if rebuilding works (verify the whole)
 				I18N.SetCurrentCulture("LEET");
 
+				// When changes are made to FormattedTestLayer2, please pull out the newly translated LEET string and update the below if-statement to keep the test passing!
+				FString LEETTranslatedString = FormattedTestLayer2.ToString();
+
 				// Convert the baked string into an FText, which will be leetified, then compare it to the rebuilt FText
-				if(FormattedTestLayer2.ToString() != TEXT("‡««L0r3m» \"L0r3m 1p$um\" «Ip$um»» | ««L0r3m» \"L0r3m 1p$um\" «Ip$um»» | «5.5421» | «5010.89221» | «1944 Month4 13» | «02:25:44» | «1944 Month7 14 22:05:06» | «92%» | «¤ 100.25»‡"))
+				if(LEETTranslatedString != TEXT("‡««L0r3m» \"L0r3m 1p$um\" «Ip$um»» | ««L0r3m» \"L0r3m 1p$um\" «Ip$um»» | «5.5421» | «5010.89221» | «1944 Month7 14 22:05:06» | «92%» | «¤ 100.25»‡"))
 				{
 					AddError( TEXT("FormattedTestLayer2 did not rebuild to correctly in LEET!") );
 				}
