Delegates: Fixed multicast delegates not making a proper deep copy on copy/assignment

diff --git a/Engine/Source/Runtime/Core/Public/Delegates/DelegateSignatureImpl.inl b/Engine/Source/Runtime/Core/Public/Delegates/DelegateSignatureImpl.inl
--- a/Engine/Source/Runtime/Core/Public/Delegates/DelegateSignatureImpl.inl
+++ b/Engine/Source/Runtime/Core/Public/Delegates/DelegateSignatureImpl.inl
@@ -1768,6 +1768,43 @@ protected:
 	 */
 	inline BASE_MULTICAST_DELEGATE_CLASS( ) { }
 
+	/**
+	 * Hidden copy constructor (for proper deep copies).
+	 *
+	 * @param Other The multicast delegate to copy from.
+	 */
+	BASE_MULTICAST_DELEGATE_CLASS( const BASE_MULTICAST_DELEGATE_CLASS& Other )
+	{
+		*this = Other;
+	}
+
+	/**
+	 * Hidden assignment operator (for proper deep copies).
+	 *
+	 * @param Other The delegate to assign from.
+	 * @return This instance.
+	 */
+	BASE_MULTICAST_DELEGATE_CLASS& operator=( const BASE_MULTICAST_DELEGATE_CLASS& Other )
+	{
+		if (&Other != this)
+		{
+			Clear();
+
+			for (IDelegateInstance* DelegateInstance : Other.GetInvocationList())
+			{
+				// this down-cast is OK! allows for managing invocation list in the base class without requiring virtual functions
+				TDelegateInstanceInterface* OtherInstance = (TDelegateInstanceInterface*)DelegateInstance;
+
+				if (OtherInstance != nullptr)
+				{
+					AddInternal(OtherInstance->CreateCopy());
+				}			
+			}
+		}
+
+		return *this;
+	}
+
 protected:
 
 	/**
@@ -1869,6 +1906,34 @@ protected:
 template< FUNC_TEMPLATE_DECL_TYPENAME >
 class MULTICAST_DELEGATE_CLASS : public BASE_MULTICAST_DELEGATE_CLASS< FUNC_TEMPLATE_DECL >
 {
+public:
+
+	/**
+	 * Default constructor.
+	 */
+	MULTICAST_DELEGATE_CLASS( ) { }
+
+	/**
+	 * Copy constructor (for proper deep copies).
+	 *
+	 * @param Other The multicast delegate to copy from.
+	 */
+	MULTICAST_DELEGATE_CLASS( const MULTICAST_DELEGATE_CLASS& Other )
+		: Super(Other)
+	{ }
+
+	/**
+	 * Assignment operator (for proper deep copies).
+	 *
+	 * @param Other The delegate to assign from.
+	 * @return This instance.
+	 */
+	MULTICAST_DELEGATE_CLASS& operator=( const MULTICAST_DELEGATE_CLASS& Other )
+	{
+		Super::operator=(Other);
+		return *this;
+	}
+
 public:
 
 	/**
