Core - Fixed malloc profiler not working under the non-game projects like UFE

diff --git a/Engine/Source/Runtime/Core/Private/ProfilingDebugging/MallocProfiler.cpp b/Engine/Source/Runtime/Core/Private/ProfilingDebugging/MallocProfiler.cpp
--- a/Engine/Source/Runtime/Core/Private/ProfilingDebugging/MallocProfiler.cpp
+++ b/Engine/Source/Runtime/Core/Private/ProfilingDebugging/MallocProfiler.cpp
@@ -1032,10 +1032,10 @@ FMallocProfilerBufferedFileWriter::~FMallocProfilerBufferedFileWriter()
 void FMallocProfilerBufferedFileWriter::Serialize( void* V, int64 Length )
 {
 #if (ALLOW_DEBUG_FILES && !HACK_HEADER_GENERATOR)
-	// Copy to buffered memory array if GGameName hasn't been set up yet.
+	// Copy to buffered memory array if GConfig hasn't been set up yet.
 	// This isn't the best solution, but due to complexity of the engine initialization order is the safest.
-	const bool bHasGameName = FPaths::IsProjectFilePathSet() && FApp::HasGameName();
-	if( !bHasGameName )
+	const bool bIsIniReady = GConfig && GConfig->IsReadyForUse();
+	if( !bIsIniReady )
 	{
 		const int32 Index = BufferedData.AddUninitialized( Length );
 		FMemory::Memcpy( &BufferedData[Index], V, Length );
