Stats - Removed unused code

diff --git a/Engine/Source/Runtime/Core/Private/Stats/StatsData.cpp b/Engine/Source/Runtime/Core/Private/Stats/StatsData.cpp
--- a/Engine/Source/Runtime/Core/Private/Stats/StatsData.cpp
+++ b/Engine/Source/Runtime/Core/Private/Stats/StatsData.cpp
@@ -28,18 +28,6 @@ const FString FStatConstants::StatsFileRawExtension = TEXT( ".ue4statsraw" );
 
 const FString FStatConstants::ThreadNameMarker = TEXT( "Thread_" );
 
-/**
-* Magic numbers for stats streams
-*/
-namespace EStatMagic
-{
-	enum Type
-	{
-		MAGIC=0x7E1B83C1,
-		MAGIC_SWAPPED=0xC1831B7E
-	};
-}
-
 /*-----------------------------------------------------------------------------
 	FRawStatStackNode
 -----------------------------------------------------------------------------*/
