Core: Fix shadow warnings, etc... for TUnion on clang-based platforms

diff --git a/Engine/Source/Runtime/Core/Public/Containers/Union.h b/Engine/Source/Runtime/Core/Public/Containers/Union.h
--- a/Engine/Source/Runtime/Core/Public/Containers/Union.h
+++ b/Engine/Source/Runtime/Core/Public/Containers/Union.h
@@ -43,7 +43,7 @@ public:
  * This differs from C union types by being type-safe, and supporting non-trivial data types as subtypes.
  * Since a value for the union must be of a single subtype, the union stores potential values of different subtypes in overlapped memory, and keeps track of which one is currently valid.
  */
-template<typename A,typename B = FNull,typename C = FNull,typename D = FNull,typename E = FNull,typename F = FNull>
+template<typename TypeA,typename TypeB = FNull,typename TypeC = FNull,typename TypeD = FNull,typename TypeE = FNull,typename TypeF = FNull>
 class TUnion
 {
 public:
@@ -54,45 +54,45 @@ public:
 	{}
 
 	/** Initialization constructor. */
-	explicit TUnion(typename TCallTraits<A>::ParamType InValue,TDisambiguater<0> Disambiguater = TDisambiguater<0>())
+	explicit TUnion(typename TCallTraits<TypeA>::ParamType InValue, TDisambiguater<0> Disambiguater = TDisambiguater<0>())
 	:	CurrentSubtypeIndex(-1)
 	{
-		SetSubtype<A>(InValue);
+		SetSubtype<TypeA>(InValue);
 	}
 	
 	/** Initialization constructor. */
-	explicit TUnion(typename TCallTraits<B>::ParamType InValue,TDisambiguater<1> Disambiguater = TDisambiguater<1>())
+	explicit TUnion(typename TCallTraits<TypeB>::ParamType InValue, TDisambiguater<1> Disambiguater = TDisambiguater<1>())
 	:	CurrentSubtypeIndex(-1)
 	{
-		SetSubtype<B>(InValue);
+		SetSubtype<TypeB>(InValue);
 	}
 	
 	/** Initialization constructor. */
-	explicit TUnion(typename TCallTraits<C>::ParamType InValue,TDisambiguater<2> Disambiguater = TDisambiguater<2>())
+	explicit TUnion(typename TCallTraits<TypeC>::ParamType InValue, TDisambiguater<2> Disambiguater = TDisambiguater<2>())
 	:	CurrentSubtypeIndex(-1)
 	{
-		SetSubtype<C>(InValue);
+		SetSubtype<TypeC>(InValue);
 	}
 	
 	/** Initialization constructor. */
-	explicit TUnion(typename TCallTraits<D>::ParamType InValue,TDisambiguater<3> Disambiguater = TDisambiguater<3>())
+	explicit TUnion(typename TCallTraits<TypeD>::ParamType InValue, TDisambiguater<3> Disambiguater = TDisambiguater<3>())
 	:	CurrentSubtypeIndex(-1)
 	{
-		SetSubtype<D>(InValue);
+		SetSubtype<TypeD>(InValue);
 	}
 	
 	/** Initialization constructor. */
-	explicit TUnion(typename TCallTraits<E>::ParamType InValue,TDisambiguater<4> Disambiguater = TDisambiguater<4>())
+	explicit TUnion(typename TCallTraits<TypeE>::ParamType InValue, TDisambiguater<4> Disambiguater = TDisambiguater<4>())
 	:	CurrentSubtypeIndex(-1)
 	{
-		SetSubtype<E>(InValue);
+		SetSubtype<TypeE>(InValue);
 	}
 	
 	/** Initialization constructor. */
-	explicit TUnion(typename TCallTraits<F>::ParamType InValue,TDisambiguater<5> Disambiguater = TDisambiguater<5>())
+	explicit TUnion(typename TCallTraits<TypeF>::ParamType InValue, TDisambiguater<5> Disambiguater = TDisambiguater<5>())
 	:	CurrentSubtypeIndex(-1)
 	{
-		SetSubtype<F>(InValue);
+		SetSubtype<TypeF>(InValue);
 	}
 
 	/** Copy constructor. */
@@ -103,13 +103,13 @@ public:
 		switch(Other.CurrentSubtypeIndex)
 		{
 		case -1: break;
-		case 0: SetSubtype<A>(Other.GetSubtype<A>()); break;
-		case 1: SetSubtype<B>(Other.GetSubtype<B>()); break;
-		case 2: SetSubtype<C>(Other.GetSubtype<C>()); break;
-		case 3: SetSubtype<D>(Other.GetSubtype<D>()); break;
-		case 4: SetSubtype<E>(Other.GetSubtype<E>()); break;
-		case 5: SetSubtype<F>(Other.GetSubtype<F>()); break;
-		default: ThrowError(TEXT("Unrecognized TUnion subtype")); break;
+		case 0: SetSubtype<TypeA>(Other.GetSubtype<TypeA>()); break;
+		case 1: SetSubtype<TypeB>(Other.GetSubtype<TypeB>()); break;
+		case 2: SetSubtype<TypeC>(Other.GetSubtype<TypeC>()); break;
+		case 3: SetSubtype<TypeD>(Other.GetSubtype<TypeD>()); break;
+		case 4: SetSubtype<TypeE>(Other.GetSubtype<TypeE>()); break;
+		case 5: SetSubtype<TypeF>(Other.GetSubtype<TypeF>()); break;
+		default: FatalErrorUndefinedSubtype(); break;
 		};
 	}
 
@@ -202,12 +202,12 @@ public:
 	{
 		switch(CurrentSubtypeIndex)
 		{
-		case 0: ResetSubtype<A>(); break;
-		case 1: ResetSubtype<B>(); break;
-		case 2: ResetSubtype<C>(); break;
-		case 3: ResetSubtype<D>(); break;
-		case 4: ResetSubtype<E>(); break;
-		case 5: ResetSubtype<F>(); break;
+		case 0: ResetSubtype<TypeA>(); break;
+		case 1: ResetSubtype<TypeB>(); break;
+		case 2: ResetSubtype<TypeC>(); break;
+		case 3: ResetSubtype<TypeD>(); break;
+		case 4: ResetSubtype<TypeE>(); break;
+		case 5: ResetSubtype<TypeF>(); break;
 		};
 	}
 
@@ -218,13 +218,13 @@ public:
 
 		switch(Union.CurrentSubtypeIndex)
 		{
-		case 0: Result ^= GetTypeHash(Union.GetSubtype<A>()); break;
-		case 1: Result ^= GetTypeHash(Union.GetSubtype<B>()); break;
-		case 2: Result ^= GetTypeHash(Union.GetSubtype<C>()); break;
-		case 3: Result ^= GetTypeHash(Union.GetSubtype<D>()); break;
-		case 4: Result ^= GetTypeHash(Union.GetSubtype<E>()); break;
-		case 5: Result ^= GetTypeHash(Union.GetSubtype<F>()); break;
-		default: ThrowError(TEXT("Unrecognized TUnion subtype")); break;
+		case 0: Result ^= GetTypeHash(Union.GetSubtype<TypeA>()); break;
+		case 1: Result ^= GetTypeHash(Union.GetSubtype<TypeB>()); break;
+		case 2: Result ^= GetTypeHash(Union.GetSubtype<TypeC>()); break;
+		case 3: Result ^= GetTypeHash(Union.GetSubtype<TypeD>()); break;
+		case 4: Result ^= GetTypeHash(Union.GetSubtype<TypeE>()); break;
+		case 5: Result ^= GetTypeHash(Union.GetSubtype<TypeF>()); break;
+		default: FatalErrorUndefinedSubtype(); break;
 		};
 
 		return Result;
@@ -237,13 +237,13 @@ public:
 		{
 			switch(CurrentSubtypeIndex)
 			{
-			case 0: return GetSubtype<A>() == Other.GetSubtype<A>(); break;
-			case 1: return GetSubtype<B>() == Other.GetSubtype<B>(); break;
-			case 2: return GetSubtype<C>() == Other.GetSubtype<C>(); break;
-			case 3: return GetSubtype<D>() == Other.GetSubtype<D>(); break;
-			case 4: return GetSubtype<E>() == Other.GetSubtype<E>(); break;
-			case 5: return GetSubtype<F>() == Other.GetSubtype<F>(); break;
-			default: ThrowError(TEXT("Unrecognized TUnion subtype")); break;
+			case 0: return GetSubtype<TypeA>() == Other.GetSubtype<TypeA>(); break;
+			case 1: return GetSubtype<TypeB>() == Other.GetSubtype<TypeB>(); break;
+			case 2: return GetSubtype<TypeC>() == Other.GetSubtype<TypeC>(); break;
+			case 3: return GetSubtype<TypeD>() == Other.GetSubtype<TypeD>(); break;
+			case 4: return GetSubtype<TypeE>() == Other.GetSubtype<TypeE>(); break;
+			case 5: return GetSubtype<TypeF>() == Other.GetSubtype<TypeF>(); break;
+			default: FatalErrorUndefinedSubtype(); break;
 			};
 		}
 		else
@@ -262,13 +262,13 @@ public:
 
 			switch(Union.CurrentSubtypeIndex)
 			{
-			case 0: Ar << Union.InitSubtype<A>(); break;
-			case 1: Ar << Union.InitSubtype<B>(); break;
-			case 2: Ar << Union.InitSubtype<C>(); break;
-			case 3: Ar << Union.InitSubtype<D>(); break;
-			case 4: Ar << Union.InitSubtype<E>(); break;
-			case 5: Ar << Union.InitSubtype<F>(); break;
-			default: UE_LOG(LogUnion, Fatal,TEXT("Unrecognized TUnion subtype")); break;
+			case 0: Ar << Union.InitSubtype<TypeA>(); break;
+			case 1: Ar << Union.InitSubtype<TypeB>(); break;
+			case 2: Ar << Union.InitSubtype<TypeC>(); break;
+			case 3: Ar << Union.InitSubtype<TypeD>(); break;
+			case 4: Ar << Union.InitSubtype<TypeE>(); break;
+			case 5: Ar << Union.InitSubtype<TypeF>(); break;
+			default: FatalErrorUndefinedSubtype(); break;
 			};
 		}
 		else
@@ -277,13 +277,13 @@ public:
 
 			switch(Union.CurrentSubtypeIndex)
 			{
-			case 0: Ar << Union.GetSubtype<A>(); break;
-			case 1: Ar << Union.GetSubtype<B>(); break;
-			case 2: Ar << Union.GetSubtype<C>(); break;
-			case 3: Ar << Union.GetSubtype<D>(); break;
-			case 4: Ar << Union.GetSubtype<E>(); break;
-			case 5: Ar << Union.GetSubtype<F>(); break;
-			default: UE_LOG(LogUnion, Fatal,TEXT("Unrecognized TUnion subtype")); break;
+			case 0: Ar << Union.GetSubtype<TypeA>(); break;
+			case 1: Ar << Union.GetSubtype<TypeB>(); break;
+			case 2: Ar << Union.GetSubtype<TypeC>(); break;
+			case 3: Ar << Union.GetSubtype<TypeD>(); break;
+			case 4: Ar << Union.GetSubtype<TypeE>(); break;
+			case 5: Ar << Union.GetSubtype<TypeF>(); break;
+			default: FatalErrorUndefinedSubtype(); break;
 			};
 		}
 		return Ar;
@@ -294,12 +294,12 @@ private:
 	/** The potential values for each subtype of the union. */
 	union
 	{
-		TTypeCompatibleBytes<A> A;
-		TTypeCompatibleBytes<B> B;
-		TTypeCompatibleBytes<C> C;
-		TTypeCompatibleBytes<D> D;
-		TTypeCompatibleBytes<E> E;
-		TTypeCompatibleBytes<F> F;
+		TTypeCompatibleBytes<TypeA> A;
+		TTypeCompatibleBytes<TypeB> B;
+		TTypeCompatibleBytes<TypeC> C;
+		TTypeCompatibleBytes<TypeD> D;
+		TTypeCompatibleBytes<TypeE> E;
+		TTypeCompatibleBytes<TypeF> F;
 	} Values;
 
 	/** The index of the subtype that the union's current value is of. */
@@ -321,32 +321,32 @@ private:
 		PointerType& OutValuePointer
 		)
 	{
-		if(TAreTypesEqual<A,Subtype>::Value)
+		if(TAreTypesEqual<TypeA,Subtype>::Value)
 		{
 			OutIndex = 0;
 			OutValuePointer = (PointerType)&Union.Values.A;
 		}
-		else if(TAreTypesEqual<B,Subtype>::Value)
+		else if(TAreTypesEqual<TypeB,Subtype>::Value)
 		{
 			OutIndex = 1;
 			OutValuePointer = (PointerType)&Union.Values.B;
 		}
-		else if(TAreTypesEqual<C,Subtype>::Value)
+		else if(TAreTypesEqual<TypeC,Subtype>::Value)
 		{
 			OutIndex = 2;
 			OutValuePointer = (PointerType)&Union.Values.C;
 		}
-		else if(TAreTypesEqual<D,Subtype>::Value)
+		else if(TAreTypesEqual<TypeD,Subtype>::Value)
 		{
 			OutIndex = 3;
 			OutValuePointer = (PointerType)&Union.Values.D;
 		}
-		else if(TAreTypesEqual<E,Subtype>::Value)
+		else if(TAreTypesEqual<TypeE,Subtype>::Value)
 		{
 			OutIndex = 4;
 			OutValuePointer = (PointerType)&Union.Values.E;
 		}
-		else if(TAreTypesEqual<F,Subtype>::Value)
+		else if(TAreTypesEqual<TypeF,Subtype>::Value)
 		{
 			OutIndex = 5;
 			OutValuePointer = (PointerType)&Union.Values.F;
@@ -354,17 +354,22 @@ private:
 		else
 		{
 			checkAtCompileTime(
-				TAreTypesEqual<TEMPLATE_PARAMETERS2(A,Subtype)>::Value ||
-				TAreTypesEqual<TEMPLATE_PARAMETERS2(B,Subtype)>::Value ||
-				TAreTypesEqual<TEMPLATE_PARAMETERS2(C,Subtype)>::Value ||
-				TAreTypesEqual<TEMPLATE_PARAMETERS2(D,Subtype)>::Value ||
-				TAreTypesEqual<TEMPLATE_PARAMETERS2(E,Subtype)>::Value ||
-				TAreTypesEqual<TEMPLATE_PARAMETERS2(F,Subtype)>::Value,
+				TAreTypesEqual<TEMPLATE_PARAMETERS2(TypeA, Subtype)>::Value ||
+				TAreTypesEqual<TEMPLATE_PARAMETERS2(TypeB, Subtype)>::Value ||
+				TAreTypesEqual<TEMPLATE_PARAMETERS2(TypeC, Subtype)>::Value ||
+				TAreTypesEqual<TEMPLATE_PARAMETERS2(TypeD, Subtype)>::Value ||
+				TAreTypesEqual<TEMPLATE_PARAMETERS2(TypeE, Subtype)>::Value ||
+				TAreTypesEqual<TEMPLATE_PARAMETERS2(TypeF, Subtype)>::Value,
 				TypeIsNotSubtypeOfUnion
 				);
 			OutIndex = -1;
 			OutValuePointer = NULL;
 		}
 	}
+
+	static void FatalErrorUndefinedSubtype()
+	{
+		UE_LOG(LogUnion, Fatal, TEXT("Unrecognized TUnion subtype"));
+	}
 };
 
