Core: Fix crash when using EvalSecondDerivative near the last point on the curve

diff --git a/Engine/Source/Runtime/Core/Public/Math/InterpCurve.h b/Engine/Source/Runtime/Core/Public/Math/InterpCurve.h
--- a/Engine/Source/Runtime/Core/Public/Math/InterpCurve.h
+++ b/Engine/Source/Runtime/Core/Public/Math/InterpCurve.h
@@ -329,7 +329,7 @@ T FInterpCurve<T>::EvalSecondDerivative( const float InVal, const T& Default ) c
 	}
 
 	// If on or beyond the last point, return 0
-	if (Index >= NumPoints)
+	if (Index >= NumPoints - 1)
 	{
 		return T();
 	}
