Only hide the cursor if the UE4 application is in the foreground, otherwise the cursor state will become very confused when you tab into the UE4 application.

diff --git a/Engine/Source/Runtime/Core/Private/Mac/MacCursor.cpp b/Engine/Source/Runtime/Core/Private/Mac/MacCursor.cpp
--- a/Engine/Source/Runtime/Core/Private/Mac/MacCursor.cpp
+++ b/Engine/Source/Runtime/Core/Private/Mac/MacCursor.cpp
@@ -251,20 +251,23 @@ void FMacCursor::UpdateVisibility()
 	// @TODO: Remove usage of deprecated CGCursorIsVisible function
 #pragma clang diagnostic push
 #pragma clang diagnostic ignored "-Wdeprecated-declarations"
-	if (CurrentCursor && bIsVisible)
+	if([NSApp isActive])
 	{
-		// Enable the cursor.
-		if (!CGCursorIsVisible())
+		if (CurrentCursor && bIsVisible)
 		{
-			CGDisplayShowCursor(kCGDirectMainDisplay);
+			// Enable the cursor.
+			if (!CGCursorIsVisible())
+			{
+				CGDisplayShowCursor(kCGDirectMainDisplay);
+			}
 		}
-	}
-	else
-	{
-		// Disable the cursor.
-		if (CGCursorIsVisible())
+		else
 		{
-			CGDisplayHideCursor(kCGDirectMainDisplay);
+			// Disable the cursor.
+			if (CGCursorIsVisible())
+			{
+				CGDisplayHideCursor(kCGDirectMainDisplay);
+			}
 		}
 	}
 #pragma clang diagnostic pop
