FString::Printf documentation.

diff --git a/Engine/Source/Runtime/Core/Public/Containers/UnrealString.h b/Engine/Source/Runtime/Core/Public/Containers/UnrealString.h
--- a/Engine/Source/Runtime/Core/Public/Containers/UnrealString.h
+++ b/Engine/Source/Runtime/Core/Public/Containers/UnrealString.h
@@ -1089,7 +1089,14 @@ public:
 	/** @return true if the string only contains numeric characters */
 	bool IsNumeric() const;
 	
-	//@todo document
+	/**
+	 * Constructs FString object similarly to how classic sprintf works.
+	 *
+	 * @param Format	Format string that specifies how FString should be built optionally using additional args. Refer to standard printf format.
+	 * @param ...		Depending on format function may require additional arguments to build output object.
+	 *
+	 * @returns FString object that was constructed using format and additional parameters.
+	 */
 	VARARG_DECL( static FString, static FString, return, Printf, VARARG_NONE, const TCHAR*, VARARG_NONE, VARARG_NONE );
 
 	// @return string with Ch character
@@ -1731,4 +1738,3 @@ public:
 
 #endif
 };
-
