Improved logging of FText automation tests to better assess why they fail.

diff --git a/Engine/Source/Runtime/Core/Private/Tests/Internationalization/TextTest.cpp b/Engine/Source/Runtime/Core/Private/Tests/Internationalization/TextTest.cpp
--- a/Engine/Source/Runtime/Core/Private/Tests/Internationalization/TextTest.cpp
+++ b/Engine/Source/Runtime/Core/Private/Tests/Internationalization/TextTest.cpp
@@ -461,6 +461,7 @@ bool FTextTest::RunTest (const FString& Parameters)
 				if(LEETTranslatedString != TEXT("‡««L0r3m» \"L0r3m 1p$um\" «Ip$um»» | ««L0r3m» \"L0r3m 1p$um\" «Ip$um»» | «5.5421» | «5010.89221» | «1944 Month7 14 22:05:06» | «92%» | «¤ 100.25»‡"))
 				{
 					AddError( TEXT("FormattedTestLayer2 did not rebuild to correctly in LEET!") );
+					AddError( TEXT("Formatted Output=") + LEETTranslatedString );
 				}
 			}
 
@@ -474,21 +475,25 @@ bool FTextTest::RunTest (const FString& Parameters)
 				if(AsNumberTest1.CompareTo(FText::AsNumber(5.5421)) != 0)
 				{
 					AddError( TEXT("AsNumberTest1 did not rebuild correctly in French-Canadian") );
+					AddError( TEXT("Number Output=") + AsNumberTest1.ToString() );
 				}
 
 				if(AsPercentTest1.CompareTo(FText::AsPercent(0.925)) != 0)
 				{
 					AddError( TEXT("AsPercentTest1 did not rebuild correctly in French-Canadian") );
+					AddError( TEXT("Percent Output=") + AsPercentTest1.ToString() );
 				}
 
 				if(AsCurrencyTest1.CompareTo(FText::AsCurrency(100.25)) != 0)
 				{
 					AddError( TEXT("AsCurrencyTest1 did not rebuild correctly in French-Canadian") );
+					AddError( TEXT("Currency Output=") + AsCurrencyTest1.ToString() );
 				}
 
 				if(AsDateTimeTest1.CompareTo(FText::AsDateTime(DateTimeInfo)) != 0)
 				{
 					AddError( TEXT("AsDateTimeTest1 did not rebuild correctly in French-Canadian") );
+					AddError( TEXT("DateTime Output=") + AsDateTimeTest1.ToString() );
 				}
 
 				{
@@ -531,6 +536,8 @@ bool FTextTest::RunTest (const FString& Parameters)
 			if(FormattedEnglishTextHistoryAsLeet.CompareTo(FormattedFrenchCanadianTextHistoryAsLeet) != 0)
 			{
 				AddError( TEXT("Serialization of text histories from source English and source French-Canadian to LEET did not produce the same results!") );
+				AddError( TEXT("English Output=") + FormattedEnglishTextHistoryAsLeet.ToString() );
+				AddError( TEXT("French-Canadian Output=") + FormattedFrenchCanadianTextHistoryAsLeet.ToString() );
 			}
 
 			{
@@ -545,6 +552,7 @@ bool FTextTest::RunTest (const FString& Parameters)
 				if(InvariantFText.ToString() != InvariantString)
 				{
 					AddError( TEXT("Invariant FText did not match the original FString after serialization!") );
+					AddError( TEXT("Invariant Output=") + InvariantFText.ToString() );
 				}
 			}
 		}
