Stats - Added options to access the stat description as an ansi and a wide string, this is a constant time access, but requires some memory to allocate the data

diff --git a/Engine/Source/Runtime/Core/Private/Stats/Stats2.cpp b/Engine/Source/Runtime/Core/Private/Stats/Stats2.cpp
--- a/Engine/Source/Runtime/Core/Private/Stats/Stats2.cpp
+++ b/Engine/Source/Runtime/Core/Private/Stats/Stats2.cpp
@@ -80,10 +80,11 @@ DEFINE_STAT(STAT_FrameTime);
 DEFINE_STAT(STAT_FPS);
 DEFINE_STAT(STAT_DrawStats);
 
-DECLARE_DWORD_COUNTER_STAT(TEXT("Frame Packets Received"),STAT_StatFramePacketsRecv,STATGROUP_StatSystem);
+DECLARE_DWORD_COUNTER_STAT( TEXT( "Frame Packets Received" ), STAT_StatFramePacketsRecv, STATGROUP_StatSystem );
+DECLARE_MEMORY_STAT( TEXT( "Stats descriptions (ansi+wide)" ), STAT_StatDescMemory, STATGROUP_StatSystem );
 
-FName TStatId::TStatId_NAME_None;
 
+FName TStatId::TStatId_NAME_None;
 
 void FStartupMessages::AddThreadMetadata( const FName InThreadName, uint32 InThreadID )
 {
@@ -144,20 +145,37 @@ class FStatGroupEnableManager : public IStatGroupEnableManager
 		}
 	};
 
-	enum 
+	enum
 	{
-		NumPerBlock = 16384,
+		/**
+		 *	Increment between two successive long names. 
+		 *	@see TStatId
+		 */
+		FNAME_INCREMENT = 3,
+
+		/** Number of stats pointer allocated per block. */
+		NUM_PER_BLOCK = 16384 * FNAME_INCREMENT,
 	};
+
+
+	TMap<FName, FGroupEnable> HighPerformanceEnable;
+
+	/** Used to synchronize the access to the high performance stats groups. */
 	FCriticalSection SynchronizationObject;
+	
+	/** Pointer to the long name in the names block. */
 	FName* PendingFNames;
+
+	/** Pending count of the name in the names block. */
 	int32 PendingCount;
-	TMap<FName, FGroupEnable> HighPerformanceEnable;
+
+	/** Holds the amount of memory allocated for the stats descriptions. */
+	FThreadSafeCounter MemoryCounter;
 
 	// these control what happens to groups that haven't been registered yet
+	TMap<FName, bool> EnableForNewGroup;
 	bool EnableForNewGroups;
 	bool UseEnableForNewGroups;
-	TMap<FName, bool> EnableForNewGroup;
-
 
 	void EnableStat(FName StatName, FName *DisablePtr)
 	{
@@ -192,7 +210,14 @@ public:
 		check(IsInGameThread());
 	}
 
-	virtual void SetHighPerformanceEnableForGroup(FName Group, bool Enable) OVERRIDE
+	virtual void UpdateMemoryUsage() override
+	{
+		// Update the stats descriptions memory usage.
+		const int32 MemoryUsage = MemoryCounter.GetValue();
+		SET_MEMORY_STAT( STAT_StatDescMemory, MemoryUsage );
+	}
+
+	virtual void SetHighPerformanceEnableForGroup(FName Group, bool Enable) override
 	{
 		FScopeLock ScopeLock(&SynchronizationObject);
 		FThreadStats::MasterDisableChangeTagLockAdd();
@@ -218,7 +243,7 @@ public:
 		FThreadStats::MasterDisableChangeTagLockSubtract();
 	}
 
-	virtual void SetHighPerformanceEnableForAllGroups(bool Enable) OVERRIDE
+	virtual void SetHighPerformanceEnableForAllGroups(bool Enable) override
 	{
 		FScopeLock ScopeLock(&SynchronizationObject);
 		FThreadStats::MasterDisableChangeTagLockAdd();
@@ -242,7 +267,7 @@ public:
 		}
 		FThreadStats::MasterDisableChangeTagLockSubtract();
 	}
-	virtual void ResetHighPerformanceEnableForAllGroups() OVERRIDE
+	virtual void ResetHighPerformanceEnableForAllGroups() override
 	{
 		FScopeLock ScopeLock(&SynchronizationObject);
 		FThreadStats::MasterDisableChangeTagLockAdd();
@@ -267,7 +292,7 @@ public:
 		FThreadStats::MasterDisableChangeTagLockSubtract();
 	}
 
-	virtual TStatId GetHighPerformanceEnableForStat(FName StatShortName, const char* InGroup, const char* InCategory, bool bDefaultEnable, bool bCanBeDisabled, EStatDataType::Type InStatType, TCHAR const* InDescription, bool bCycleStat, FPlatformMemory::EMemoryCounterRegion MemoryRegion = FPlatformMemory::MCR_Invalid) OVERRIDE
+	virtual TStatId GetHighPerformanceEnableForStat(FName StatShortName, const char* InGroup, const char* InCategory, bool bDefaultEnable, bool bCanBeDisabled, EStatDataType::Type InStatType, TCHAR const* InDescription, bool bCycleStat, FPlatformMemory::EMemoryCounterRegion MemoryRegion = FPlatformMemory::MCR_Invalid) override
 	{
 		FScopeLock ScopeLock(&SynchronizationObject);
 
@@ -315,11 +340,32 @@ public:
 		}
 		if (PendingCount < 1)
 		{
-			PendingFNames = new FName[NumPerBlock];
-			PendingCount = NumPerBlock;
+			PendingFNames = new FName[NUM_PER_BLOCK];
+			PendingCount = NUM_PER_BLOCK;
 		}
-		PendingCount--;
-		FName* Result = PendingFNames++;
+		PendingCount -= FNAME_INCREMENT;
+		FName* Result = &PendingFNames[TStatId::INDEX_FNAME];
+
+		// Get the wide stat description.
+		const int32 StatDescLen = FCString::Strlen( InDescription );
+		const WIDECHAR* TempStatDescWide = StringCast<WIDECHAR>( InDescription ).Get();
+		// We are leaking this.
+		WIDECHAR* StatDescWide = new WIDECHAR[StatDescLen + 1];
+		TCString<WIDECHAR>::Strcpy( StatDescWide, StatDescLen, TempStatDescWide );
+
+		// Get the ansi stat description.
+		const ANSICHAR* TempStatDescAnsi = StringCast<ANSICHAR>( InDescription ).Get();
+		// We are leaking this.
+		ANSICHAR* StatDescAnsi = new ANSICHAR[StatDescLen + 1];
+		TCString<ANSICHAR>::Strcpy( StatDescAnsi, StatDescLen, TempStatDescAnsi );
+
+		FMemory::Memcpy( (WIDECHAR*)&PendingFNames[TStatId::INDEX_WIDE_STRING], &StatDescWide, sizeof( WIDECHAR* ) );
+		FMemory::Memcpy( (ANSICHAR*)&PendingFNames[TStatId::INDEX_ANSI_STRING], &StatDescAnsi, sizeof( ANSICHAR* ) );
+
+		MemoryCounter.Add( (StatDescLen + 1)*(sizeof( ANSICHAR ) + sizeof( WIDECHAR )) );
+		
+		PendingFNames += FNAME_INCREMENT;
+
 		if (Found->CurrentEnable)
 		{
 			EnableStat(Stat, Result);
@@ -573,9 +619,9 @@ public:
 	}
 
 	/**
-	 * Returns a pointer to the single threaded interface when mulithreading is disabled.
+	 * Returns a pointer to the single threaded interface when multithreading is disabled.
 	 */
-	virtual FSingleThreadRunnable* GetSingleThreadInterface() OVERRIDE
+	virtual FSingleThreadRunnable* GetSingleThreadInterface() override
 	{
 		return this;
 	}
@@ -589,7 +635,7 @@ public:
 		return 0;
 	}
 
-	virtual void Tick() OVERRIDE
+	virtual void Tick() override
 	{
 		static double LastTime = -1.0;
 		if (bReadyToProcess && FPlatformTime::Seconds() - LastTime > .005) // we won't process more than every 5ms
@@ -597,6 +643,8 @@ public:
 			// Update the seconds per cycle.
 			SET_FLOAT_STAT( STAT_SecondsPerCycle, FPlatformTime::GetSecondsPerCycle() ); 
 
+			IStatGroupEnableManager::Get().UpdateMemoryUsage();
+
 			SCOPE_CYCLE_COUNTER(STAT_StatsNewTick);
 			bReadyToProcess = false;
 			FStatPacketArray NowData; 
diff --git a/Engine/Source/Runtime/Core/Public/Stats/Stats.h b/Engine/Source/Runtime/Core/Public/Stats/Stats.h
--- a/Engine/Source/Runtime/Core/Public/Stats/Stats.h
+++ b/Engine/Source/Runtime/Core/Public/Stats/Stats.h
@@ -50,12 +50,9 @@ public:
 	 * Pushes the specified stat onto the hierarchy for this thread. Starts
 	 * the timing of the cycles used
 	 */
-	FORCEINLINE_STATS FScopeCycleCounter(TStatId StatId)
+	FORCEINLINE_STATS FScopeCycleCounter( TStatId StatId )
 	{
-		if (FThreadStats::IsCollectingData(StatId))
-		{
-			Start(*StatId);
-		}
+		Start( StatId );
 	}
 
 	/**
diff --git a/Engine/Source/Runtime/Core/Public/Stats/Stats2.h b/Engine/Source/Runtime/Core/Public/Stats/Stats2.h
--- a/Engine/Source/Runtime/Core/Public/Stats/Stats2.h
+++ b/Engine/Source/Runtime/Core/Public/Stats/Stats2.h
@@ -42,6 +42,21 @@ CORE_API bool DirectStatsCommand(const TCHAR* Cmd, bool bBlockForCompletion = fa
 
 struct TStatId
 {
+	enum
+	{
+		/**
+		 *	Index of the long name.
+		 *	@see FStatGroupEnableManager
+		 */
+		INDEX_FNAME = 0,
+
+		/** Index of the stat desc as an ansi string. */
+		INDEX_ANSI_STRING = 1,
+
+		/** Index of the stat desc as a wide string. */
+		INDEX_WIDE_STRING = 2,
+	};
+
 	FORCEINLINE TStatId()
 		: StatIdPtr(&TStatId_NAME_None)
 	{
@@ -64,10 +79,41 @@ struct TStatId
 	}
 	FORCEINLINE FName const* GetRawPointer() const
 	{
-		return StatIdPtr;
+		return &StatIdPtr[INDEX_FNAME];
+	}
+
+	/**
+	 * @return a stat description as an ansi string.
+	 * StatIdPtr must point to a valid FName pointer.
+	 * @see FStatGroupEnableManager::GetHighPerformanceEnableForStat
+	 */
+	FORCEINLINE const ANSICHAR* GetStatDescriptionANSI() const
+	{
+		return (ANSICHAR*)*(uint64*)(&StatIdPtr[INDEX_ANSI_STRING]);
+	}
+
+	/**
+	 * @return a stat description as a wide string.
+	 * StatIdPtr must point to a valid FName pointer.
+	 * @see FStatGroupEnableManager::GetHighPerformanceEnableForStat
+	 */
+	FORCEINLINE const WIDECHAR* GetStatDescriptionWIDE() const
+	{
+		return (WIDECHAR*)*(uint64*)(&StatIdPtr[INDEX_WIDE_STRING]);
 	}
 private:
+	/** NAME_None. */
 	CORE_API static FName TStatId_NAME_None;
+
+	/**
+	 *	Holds a pointer to the stat long name if enabled, or to the NAME_None if disabled.
+	 *	@see FStatGroupEnableManager::EnableStat
+	 *	@see FStatGroupEnableManager::DisableStat
+	 *	
+	 *	Next pointer points to the ansi string with a stat description
+	 *	Next pointer points to the wide string with a stat description
+	 *	@see FStatGroupEnableManager::GetHighPerformanceEnableForStat 
+	 */
 	FName const* StatIdPtr;
 };
 
@@ -989,13 +1035,7 @@ public:
 		if (InStatOperation == EStatOperation::CycleScopeStart)
 		{
 			ThreadStats->ScopeCount++;
-			new (ThreadStats->Packet.StatMessages) FStatMessage(InStatName, InStatOperation);
-
-			// Emit named event for active cycle stat.
-			if( GCycleStatsShouldEmitNamedEvents > 0 )
-			{
-				FPlatformMisc::BeginNamedEvent(FColor(0), InStatName.GetPlainANSIString());
-			}
+			new (ThreadStats->Packet.StatMessages) FStatMessage(InStatName, InStatOperation);	
 		}
 		else if (InStatOperation == EStatOperation::CycleScopeEnd)
 		{
@@ -1009,12 +1049,6 @@ public:
 				}
 			}
 			// else we dumped this frame without closing scope, so we just drop the closes on the floor
-
-			// End named event for active cycle stat.
-			if( GCycleStatsShouldEmitNamedEvents > 0 )
-			{
-				FPlatformMisc::EndNamedEvent();
-			}
 		}
 	}
 
@@ -1143,10 +1177,23 @@ public:
 	 * Pushes the specified stat onto the hierarchy for this thread. Starts
 	 * the timing of the cycles used
 	 */
-	FORCEINLINE_STATS void Start(FName InStatId)
+	FORCEINLINE_STATS void Start( TStatId InStatId )
 	{
-		StatId = InStatId;
-		FThreadStats::AddMessage(InStatId, EStatOperation::CycleScopeStart);
+		if( FThreadStats::IsCollectingData( InStatId ) )
+		{
+			StatId = *InStatId;
+			FThreadStats::AddMessage( *InStatId, EStatOperation::CycleScopeStart );
+
+			// Emit named event for active cycle stat.
+			if( GCycleStatsShouldEmitNamedEvents > 0 )
+			{
+#if	PLATFORM_USES_ANSI_STRING_FOR_EXTERNAL_PROFILING
+				FPlatformMisc::BeginNamedEvent( FColor( 0 ), InStatId.GetStatDescriptionANSI() );
+#else
+				FPlatformMisc::BeginNamedEvent( FColor( 0 ), InStatId.GetStatDescriptionWIDE() );
+#endif // PLATFORM_USES_ANSI_STRING_FOR_EXTERNAL_PROFILING
+			}
+		}
 	}
 
 	/**
@@ -1154,9 +1201,15 @@ public:
 	 */
 	FORCEINLINE_STATS void Stop()
 	{
-		if (!StatId.IsNone())
+		if( !StatId.IsNone() )
 		{
 			FThreadStats::AddMessage(StatId, EStatOperation::CycleScopeEnd);
+
+			// End named event for active cycle stat.
+			if( GCycleStatsShouldEmitNamedEvents > 0 )
+			{
+				FPlatformMisc::EndNamedEvent();
+			}
 		}
 	}
 };
@@ -1219,7 +1272,7 @@ public:
 	virtual void SetHighPerformanceEnableForAllGroups(bool Enable)=0;
 
 	/**
-	 * Resets all stats to their default collection state, which was set when they were looked up intially
+	 * Resets all stats to their default collection state, which was set when they were looked up initially
 	 */
 	virtual void ResetHighPerformanceEnableForAllGroups()=0;
 
@@ -1228,6 +1281,9 @@ public:
 	 * @param Cmd, Command to run
 	 */
 	virtual void StatGroupEnableManagerCommand(FString const& Cmd)=0;
+
+	/** Updates memory usage. */
+	virtual void UpdateMemoryUsage() = 0;
 };
 
 
diff --git a/Engine/Source/Runtime/CoreUObject/Public/UObject/UObjectBaseUtility.h b/Engine/Source/Runtime/CoreUObject/Public/UObject/UObjectBaseUtility.h
--- a/Engine/Source/Runtime/CoreUObject/Public/UObject/UObjectBaseUtility.h
+++ b/Engine/Source/Runtime/CoreUObject/Public/UObject/UObjectBaseUtility.h
@@ -457,7 +457,7 @@ struct FScopeCycleCounterUObject : public FCycleCounter
 			TStatId StatId = Object->GetStatID();
 			if (FThreadStats::IsCollectingData(StatId))
 			{
-				Start(*StatId);
+				Start(StatId);
 			}
 		}
 	}
@@ -471,7 +471,7 @@ struct FScopeCycleCounterUObject : public FCycleCounter
 			TStatId StatId = Object->GetStatID();
 			if (!StatId->IsNone())
 			{
-				Start(*StatId);
+				Start(StatId);
 			}
 		}
 	}
