Stats - Fixed a memory leak in the UFE, added better implementation for the stats advance frame

diff --git a/Engine/Source/Programs/UnrealFrontend/Private/UnrealFrontendMain.cpp b/Engine/Source/Programs/UnrealFrontend/Private/UnrealFrontendMain.cpp
--- a/Engine/Source/Programs/UnrealFrontend/Private/UnrealFrontendMain.cpp
+++ b/Engine/Source/Programs/UnrealFrontend/Private/UnrealFrontendMain.cpp
@@ -77,7 +77,6 @@ bool RunDeployCommand()
 
 	double DeltaTime = 0.0;
 	double LastTime = FPlatformTime::Seconds();
-	static int32 MasterDisableChangeTagStartFrame = -1;
 
 	// We track the message sent time because we have to keep updating the loop until the message is *actually sent*. (ie all packets queued, sent, buffer flushed, etc.)
 	double MessageSentTime = 0.0;
@@ -222,6 +221,7 @@ void WriteString(FArchive* File, const ANSICHAR* Format, ...)
 
 void RunStatsConvertCommand()
 {
+#if	STATS
 	// get the target file
 	FString TargetFile;
 	FParse::Value(FCommandLine::Get(), TEXT("-INFILE="), TargetFile);
@@ -337,9 +337,9 @@ void RunStatsConvertCommand()
 			break;
 		}
 	}
+#endif // STATS
 }
 
-
 void RunUI()
 {
 	FString UnrealFrontendLayoutIni = FPaths::GetPath(GEngineIni) + "/Layout.ini";
@@ -409,7 +409,6 @@ void RunUI()
 	double DeltaTime = 0.0;
 	double LastTime = FPlatformTime::Seconds();
 	const float IdealFrameTime = 1.0f / IDEAL_FRAMERATE;
-	static int32 MasterDisableChangeTagStartFrame = -1;
 
 	while (!GIsRequestingExit)
 	{
@@ -430,11 +429,7 @@ void RunUI()
 		DeltaTime =  CurrentTime - LastTime;
 		LastTime = CurrentTime;
 
-#if STATS
-		FThreadStats::ExplicitFlush();
-		FThreadStats::WaitForStats();
-		MasterDisableChangeTagStartFrame = FThreadStats::MasterDisableChangeTag();
-#endif
+		FStats::AdvanceFrame( false );
 
 		GLog->FlushThreadedLogs();
 	}
diff --git a/Engine/Source/Runtime/Core/Private/Stats/Stats2.cpp b/Engine/Source/Runtime/Core/Private/Stats/Stats2.cpp
--- a/Engine/Source/Runtime/Core/Private/Stats/Stats2.cpp
+++ b/Engine/Source/Runtime/Core/Private/Stats/Stats2.cpp
@@ -25,6 +25,55 @@ static struct FForceInitAtBootFStats2 : public TForceInitAtBoot<FStats2Globals>
 
 DECLARE_FLOAT_COUNTER_STAT( TEXT( "Seconds Per Cycle" ), STAT_SecondsPerCycle, STATGROUP_Engine );
 
+/*-----------------------------------------------------------------------------
+	FStats2
+-----------------------------------------------------------------------------*/
+
+void FStats::AdvanceFrame( bool bDiscardCallstack, const FOnAdvanceRenderingThreadStats& AdvanceRenderingThreadStatsDelegate /*= FOnAdvanceRenderingThreadStats()*/ )
+{
+#if STATS
+	check( IsInGameThread() );
+	static int32 MasterDisableChangeTagStartFrame = -1;
+	static int64 StatsFrame = 1;
+	StatsFrame++;
+
+	int64 Frame = StatsFrame;
+	if( bDiscardCallstack )
+	{
+		FThreadStats::FrameDataIsIncomplete(); // we won't collect call stack stats this frame
+	}
+	if( MasterDisableChangeTagStartFrame == -1 )
+	{
+		MasterDisableChangeTagStartFrame = FThreadStats::MasterDisableChangeTag();
+	}
+	if( !FThreadStats::IsCollectingData() || MasterDisableChangeTagStartFrame != FThreadStats::MasterDisableChangeTag() )
+	{
+		Frame = -StatsFrame; // mark this as a bad frame
+	}
+	static FStatNameAndInfo Adv( NAME_AdvanceFrame, "", "", TEXT( "" ), EStatDataType::ST_int64, true, false );
+	FThreadStats::AddMessage( Adv.GetEncodedName(), EStatOperation::AdvanceFrameEventGameThread, Frame ); // we need to flush here if we aren't collecting stats to make sure the meta data is up to date
+	if( FPlatformProperties::IsServerOnly() )
+	{
+		FThreadStats::AddMessage( Adv.GetEncodedName(), EStatOperation::AdvanceFrameEventRenderThread, Frame ); // we need to flush here if we aren't collecting stats to make sure the meta data is up to date
+	}
+
+	if( AdvanceRenderingThreadStatsDelegate.IsBound() )
+	{
+		AdvanceRenderingThreadStatsDelegate.Execute( bDiscardCallstack, StatsFrame, MasterDisableChangeTagStartFrame );
+	}
+	else
+	{
+		// There is no rendering thread, so this message is sufficient to make stats happy and don't leak memory.
+		FThreadStats::AddMessage( Adv.GetEncodedName(), EStatOperation::AdvanceFrameEventRenderThread, Frame );
+	}
+
+	FThreadStats::ExplicitFlush( bDiscardCallstack );
+	FThreadStats::WaitForStats();
+	MasterDisableChangeTagStartFrame = FThreadStats::MasterDisableChangeTag();
+#endif
+}
+
+
 /* Todo
 
 FStatsThreadState::FStatsThreadState(FString const& Filename) - needs to be more careful about files that are "cut off". Should look for < 16 bytes left OR next 8 bytes are zero. Bad files will end with zeros and it can't be a valid record.
@@ -80,7 +129,7 @@ DEFINE_STAT(STAT_FrameTime);
 DEFINE_STAT(STAT_FPS);
 DEFINE_STAT(STAT_DrawStats);
 
-DECLARE_DWORD_COUNTER_STAT( TEXT( "Frame Packets Received" ), STAT_StatFramePacketsRecv, STATGROUP_StatSystem );
+DECLARE_DWORD_COUNTER_STAT(TEXT("Frame Packets Received"),STAT_StatFramePacketsRecv,STATGROUP_StatSystem);
 DECLARE_MEMORY_STAT( TEXT( "Stats descriptions (ansi+wide)" ), STAT_StatDescMemory, STATGROUP_StatSystem );
 
 
@@ -145,7 +194,7 @@ class FStatGroupEnableManager : public IStatGroupEnableManager
 		}
 	};
 
-	enum
+	enum 
 	{
 		/**
 		 *	Increment between two successive long names. 
@@ -937,5 +986,4 @@ void FThreadStats::WaitForStats()
 	}
 }
 
-#endif
-
+#endif
\ No newline at end of file
diff --git a/Engine/Source/Runtime/Core/Public/Stats/Stats2.h b/Engine/Source/Runtime/Core/Public/Stats/Stats2.h
--- a/Engine/Source/Runtime/Core/Public/Stats/Stats2.h
+++ b/Engine/Source/Runtime/Core/Public/Stats/Stats2.h
@@ -38,6 +38,16 @@ public:
 // Pass a console command directly to the stats system, return true if it is known command, false means it might be a stats command
 CORE_API bool DirectStatsCommand(const TCHAR* Cmd, bool bBlockForCompletion = false, FOutputDevice* Ar = nullptr);
 
+/** Helper struct that contains method available even when the stats are disabled. */
+struct CORE_API FStats
+{
+	/** Delegate to fire every time we need to advance the stats for the rendering thread. */
+	DECLARE_DELEGATE_ThreeParams( FOnAdvanceRenderingThreadStats, bool /*bDiscardCallstack*/, int64 /*StatsFrame*/, int32 /*MasterDisableChangeTagStartFrame*/ );
+
+	/** Advances stats for the current frame. */
+	static void AdvanceFrame( bool bDiscardCallstack, const FOnAdvanceRenderingThreadStats& AdvanceRenderingThreadStatsDelegate = FOnAdvanceRenderingThreadStats() );
+};
+
 #if STATS
 
 struct TStatId
diff --git a/Engine/Source/Runtime/Launch/Private/LaunchEngineLoop.cpp b/Engine/Source/Runtime/Launch/Private/LaunchEngineLoop.cpp
--- a/Engine/Source/Runtime/Launch/Private/LaunchEngineLoop.cpp
+++ b/Engine/Source/Runtime/Launch/Private/LaunchEngineLoop.cpp
@@ -1813,59 +1813,12 @@ void FEngineLoop::InitTime()
 	MaxFrameCounter = FMath::TruncToInt(MaxTickTime / FApp::GetFixedDeltaTime());
 }
 
-void FlushStatsFrame(bool bDiscardCallstack)
-{
-	check(IsInGameThread()); 
-	static int32 MasterDisableChangeTagStartFrame = -1;
-	static int64 StatsFrame = 1;
-	StatsFrame++;
-#if STATS
-	int64 Frame = StatsFrame;
-	if (bDiscardCallstack)
-	{
-		FThreadStats::FrameDataIsIncomplete(); // we won't collect call stack stats this frame
-	}
-	if (MasterDisableChangeTagStartFrame == -1)
-	{
-		MasterDisableChangeTagStartFrame = FThreadStats::MasterDisableChangeTag();
-	}
-	if (!FThreadStats::IsCollectingData() ||  MasterDisableChangeTagStartFrame != FThreadStats::MasterDisableChangeTag())
-	{
-		Frame = -StatsFrame; // mark this as a bad frame
-	}
-	static FStatNameAndInfo Adv(NAME_AdvanceFrame, "", "", TEXT(""), EStatDataType::ST_int64, true, false);
-	FThreadStats::AddMessage(Adv.GetEncodedName(), EStatOperation::AdvanceFrameEventGameThread, Frame); // we need to flush here if we aren't collecting stats to make sure the meta data is up to date
-	if (FPlatformProperties::IsServerOnly())
-	{
-		FThreadStats::AddMessage(Adv.GetEncodedName(), EStatOperation::AdvanceFrameEventRenderThread, Frame); // we need to flush here if we aren't collecting stats to make sure the meta data is up to date
-	}
-#endif
-	ENQUEUE_UNIQUE_RENDER_COMMAND_TWOPARAMETER( 
-		RenderingThreadTickCommand, 
-		int64, SentStatsFrame, StatsFrame,
-		int32, SentMasterDisableChangeTagStartFrame, MasterDisableChangeTagStartFrame,
-	{ 
-		RenderingThreadTick(SentStatsFrame, SentMasterDisableChangeTagStartFrame); 
-	} 
-	);
-	if (bDiscardCallstack)
-	{
-		// we need to flush the rendering thread here, otherwise it can get behind and then the stats will get behind.
-		FlushRenderingCommands();
-	}
 
-#if STATS
-	FThreadStats::ExplicitFlush(bDiscardCallstack);
-	FThreadStats::WaitForStats();
-	MasterDisableChangeTagStartFrame = FThreadStats::MasterDisableChangeTag();
-#endif
-
-}
 
 //called via FCoreDelegates::StarvedGameLoop
 void GameLoopIsStarved()
 {
-	FlushStatsFrame(true);
+	FStats::AdvanceFrame( true, FStats::FOnAdvanceRenderingThreadStats::CreateStatic( &AdvanceRenderingThreadStatsGT ) );
 }
 
 
@@ -2076,7 +2029,7 @@ void FEngineLoop::Tick()
 		GMalloc->UpdateStats();
 	} 
 
-	FlushStatsFrame(false);
+	FStats::AdvanceFrame( false, FStats::FOnAdvanceRenderingThreadStats::CreateStatic( &AdvanceRenderingThreadStatsGT ) );
 
 	{ 
 		SCOPE_CYCLE_COUNTER( STAT_FrameTime );
diff --git a/Engine/Source/Runtime/RenderCore/Private/RenderingThread.cpp b/Engine/Source/Runtime/RenderCore/Private/RenderingThread.cpp
--- a/Engine/Source/Runtime/RenderCore/Private/RenderingThread.cpp
+++ b/Engine/Source/Runtime/RenderCore/Private/RenderingThread.cpp
@@ -247,7 +247,7 @@ void RenderingThreadMain( FEvent* TaskGraphBoundSyncEvent )
 /**
  * Advances stats for the rendering thread.
  */
-void RenderingThreadTick(int64 StatsFrame, int32 MasterDisableChangeTagStartFrame)
+static void AdvanceRenderingThreadStats(int64 StatsFrame, int32 MasterDisableChangeTagStartFrame)
 {
 #if STATS
 	int64 Frame = StatsFrame;
@@ -264,6 +264,27 @@ void RenderingThreadTick(int64 StatsFrame, int32 MasterDisableChangeTagStartFram
 #endif
 }
 
+/**
+ * Advances stats for the rendering thread. Called from the game thread.
+ */
+void AdvanceRenderingThreadStatsGT( bool bDiscardCallstack, int64 StatsFrame, int32 MasterDisableChangeTagStartFrame )
+{
+	ENQUEUE_UNIQUE_RENDER_COMMAND_TWOPARAMETER
+	(
+		RenderingThreadTickCommand,
+		int64, SentStatsFrame, StatsFrame,
+		int32, SentMasterDisableChangeTagStartFrame, MasterDisableChangeTagStartFrame,
+		{
+			AdvanceRenderingThreadStats( SentStatsFrame, SentMasterDisableChangeTagStartFrame );
+		}
+	);
+	if( bDiscardCallstack )
+	{
+		// we need to flush the rendering thread here, otherwise it can get behind and then the stats will get behind.
+		FlushRenderingCommands();
+	}
+}
+
 /** The rendering thread runnable object. */
 class FRenderingThread : public FRunnable
 {
diff --git a/Engine/Source/Runtime/RenderCore/Public/RenderingThread.h b/Engine/Source/Runtime/RenderCore/Public/RenderingThread.h
--- a/Engine/Source/Runtime/RenderCore/Public/RenderingThread.h
+++ b/Engine/Source/Runtime/RenderCore/Public/RenderingThread.h
@@ -63,8 +63,10 @@ extern RENDERCORE_API void CheckRenderingThreadHealth();
 /** Checks if the rendering thread is healthy and running, without crashing */
 extern RENDERCORE_API bool IsRenderingThreadHealthy();
 
-/** Advances stats for the rendering thread. */
-extern RENDERCORE_API void RenderingThreadTick(int64 StatsFrame, int32 MasterDisableChangeTagStartFrame);
+/**
+ * Advances stats for the rendering thread. Called from the game thread.
+ */
+extern RENDERCORE_API void AdvanceRenderingThreadStatsGT( bool bDiscardCallstack, int64 StatsFrame, int32 MasterDisableChangeTagStartFrame );
 
 /**
  * Adds a task that must be completed either before the next scene draw or a flush rendering commands
