Stats - Added a missing stat description to the task graph stats

diff --git a/Engine/Source/Runtime/Core/Public/Async/TaskGraphInterfaces.h b/Engine/Source/Runtime/Core/Public/Async/TaskGraphInterfaces.h
--- a/Engine/Source/Runtime/Core/Public/Async/TaskGraphInterfaces.h
+++ b/Engine/Source/Runtime/Core/Public/Async/TaskGraphInterfaces.h
@@ -817,7 +817,8 @@ public:
 	FCustomNameGraphTaskBase(const TCHAR* TaskClassName, const TCHAR* TaskName)
 	{
 #if STATS
-		const FName StatName = FName(*(FString(TaskClassName) + TEXT(".") + FString(TaskName)));
+		const FString LongName = FString(TaskClassName) + TEXT(".") + FString(TaskName);
+		const FName StatName = FName( *LongName );
 
 		FStartupMessages::Get().AddMetadata(StatName, TEXT(""),
 			STAT_GROUP_TO_FStatGroup(STATGROUP_TaskGraphTasks)::GetGroupName(),
@@ -829,7 +830,7 @@ public:
 			STAT_GROUP_TO_FStatGroup(STATGROUP_TaskGraphTasks)::GetGroupName(),
 			STAT_GROUP_TO_FStatGroup(STATGROUP_TaskGraphTasks)::GetGroupCategory(),
 			STAT_GROUP_TO_FStatGroup(STATGROUP_TaskGraphTasks)::DefaultEnable,
-			true, EStatDataType::ST_int64, TEXT(""), true);
+			true, EStatDataType::ST_int64, *LongName, true);
 #endif
 	}
 
