Stats - Compilation fix

diff --git a/Engine/Source/Runtime/Core/Private/Async/TaskGraph.cpp b/Engine/Source/Runtime/Core/Private/Async/TaskGraph.cpp
--- a/Engine/Source/Runtime/Core/Private/Async/TaskGraph.cpp
+++ b/Engine/Source/Runtime/Core/Private/Async/TaskGraph.cpp
@@ -309,7 +309,7 @@ public:
 				if (!bTasksOpen && FThreadStats::IsCollectingData(StatName))
 				{
 					bTasksOpen = true;
-					ProcessingTasks.Start(*StatName);
+					ProcessingTasks.Start(StatName);
 				}
 #endif
 				Task->Execute(NewTasks, ENamedThreads::Type(ThreadId | (QueueIndex << ENamedThreads::QueueIndexShift)));
