Fixed issues with a LEET string in file possibly being corrupted by P4.

diff --git a/Engine/Source/Runtime/Core/Private/Tests/Internationalization/TextTest.cpp b/Engine/Source/Runtime/Core/Private/Tests/Internationalization/TextTest.cpp
--- a/Engine/Source/Runtime/Core/Private/Tests/Internationalization/TextTest.cpp
+++ b/Engine/Source/Runtime/Core/Private/Tests/Internationalization/TextTest.cpp
@@ -458,7 +458,7 @@ bool FTextTest::RunTest (const FString& Parameters)
 				FString LEETTranslatedString = FormattedTestLayer2.ToString();
 
 				// Convert the baked string into an FText, which will be leetified, then compare it to the rebuilt FText
-				if(LEETTranslatedString != TEXT("‡««L0r3m» \"L0r3m 1p$um\" «Ip$um»» | ««L0r3m» \"L0r3m 1p$um\" «Ip$um»» | «5.5421» | «5010.89221» | «1944 Month7 14 22:05:06» | «92%» | «¤ 100.25»‡"))
+				if(LEETTranslatedString != FString(TEXT("\x2021") TEXT("\xAB") TEXT("\xAB") TEXT("L0r3m" TEXT("\xBB") TEXT(" \"L0r3m 1p$um\" ") TEXT("\xAB") TEXT("Ip$um") TEXT("\xBB") TEXT("\xBB") TEXT(" | ") TEXT("\xAB") TEXT("\xAB") TEXT("L0r3m") TEXT("\xBB") TEXT(" \"L0r3m 1p$um\" ") TEXT("\xAB") TEXT("Ip$um") TEXT("\xBB") TEXT("\xBB") TEXT(" | ") TEXT("\xAB") TEXT("5.5421") TEXT("\xBB") TEXT(" | ") TEXT("\xAB") TEXT("5010.89221") TEXT("\xBB") TEXT(" | ") TEXT("\xAB") TEXT("1944 Month7 14 22:05:06") TEXT("\xBB") TEXT(" | ") TEXT("\xAB") TEXT("92%") TEXT("\xBB") TEXT(" | ") TEXT("\xAB") TEXT("\xA4") TEXT("\xA0") TEXT("100.25") TEXT("\xBB") TEXT("\x2021"))))
 				{
 					AddError( TEXT("FormattedTestLayer2 did not rebuild to correctly in LEET!") );
 					AddError( TEXT("Formatted Output=") + LEETTranslatedString );
