UE4: Made TestMap folders configurable. These folders are ignored in automation tests.

diff --git a/Engine/Config/BaseEngine.ini b/Engine/Config/BaseEngine.ini
--- a/Engine/Config/BaseEngine.ini
+++ b/Engine/Config/BaseEngine.ini
@@ -1963,6 +1963,7 @@ TestAllBlueprints=false
 
 [/Script/Engine.AutomationTestSettings]
 +EditorTestModules=StaticMeshEditor
++TestLevelFolders=TestMaps
 
 [AutomationTesting.StaticMeshEditorTest]
 +EditorViewButtons=Wireframe
diff --git a/Engine/Source/Runtime/Core/Private/Misc/AutomationTest.cpp b/Engine/Source/Runtime/Core/Private/Misc/AutomationTest.cpp
--- a/Engine/Source/Runtime/Core/Private/Misc/AutomationTest.cpp
+++ b/Engine/Source/Runtime/Core/Private/Misc/AutomationTest.cpp
@@ -417,10 +417,21 @@ void FAutomationTestFramework::GetValidTestNames( TArray<FAutomationTestInfo>& T
 
 bool FAutomationTestFramework::ShouldTestContent(const FString& Path) const
 {
-	if (Path.Contains(TEXT("/TestMaps/")))
+	static TArray<FString> TestLevelFolders;
+	if ( TestLevelFolders.Num() == 0 )
 	{
-		return false;
+		GConfig->GetArray( TEXT("/Script/Engine.AutomationTestSettings"), TEXT("TestLevelFolders"), TestLevelFolders, GEngineIni);
 	}
+
+	for ( const FString& Folder : TestLevelFolders )
+	{
+		const FString PatternToCheck = FString::Printf(TEXT("/%s/"), *Folder);
+		if ( Path.Contains(*PatternToCheck) )
+		{
+			return false;
+		}
+	}
+
 	FString DevelopersPath = FPaths::GameDevelopersDir().LeftChop(1);
 	return bDeveloperDirectoryIncluded || bVisualCommandletFilterOn || !Path.StartsWith(DevelopersPath);
 }
diff --git a/Engine/Source/Runtime/Engine/Classes/Tests/AutomationTestSettings.h b/Engine/Source/Runtime/Engine/Classes/Tests/AutomationTestSettings.h
--- a/Engine/Source/Runtime/Engine/Classes/Tests/AutomationTestSettings.h
+++ b/Engine/Source/Runtime/Engine/Classes/Tests/AutomationTestSettings.h
@@ -95,6 +95,12 @@ public:
 	UPROPERTY(EditAnywhere, config, Category=Automation)
 	TArray<FString> EditorTestModules;
 
+	/**
+	 * Folders containing levels to exclude from automated tests
+	 */
+	UPROPERTY(EditAnywhere, config, Category=Automation)
+	TArray<FString> TestLevelFolders;
+
 	/**
 	 * Asset to test for open in automation process
 	 */
