Android stack walking displaying function signatures. addr2line can be used offline with the address displayed to get the exact source file and line number.

diff --git a/Engine/Source/Programs/UnrealBuildTool/Android/UEBuildAndroid.cs b/Engine/Source/Programs/UnrealBuildTool/Android/UEBuildAndroid.cs
--- a/Engine/Source/Programs/UnrealBuildTool/Android/UEBuildAndroid.cs
+++ b/Engine/Source/Programs/UnrealBuildTool/Android/UEBuildAndroid.cs
@@ -268,6 +268,9 @@ namespace UnrealBuildTool
 			// link with Android libraries.
 			InBuildTarget.GlobalLinkEnvironment.Config.LibraryPaths.Add("$(NDKROOT)/sources/cxx-stl/stlport/libs/armeabi-v7a");
 
+            // Add path to statically compiled version of cxa_demangle
+            InBuildTarget.GlobalLinkEnvironment.Config.LibraryPaths.Add(UEBuildConfiguration.UEThirdPartyDirectory + "Android/cxa_demangle/armeabi-v7a");
+
 			InBuildTarget.GlobalLinkEnvironment.Config.AdditionalLibraries.Add("stdc++");
 			InBuildTarget.GlobalLinkEnvironment.Config.AdditionalLibraries.Add("gcc");
 			InBuildTarget.GlobalLinkEnvironment.Config.AdditionalLibraries.Add("z");
@@ -279,6 +282,7 @@ namespace UnrealBuildTool
 			InBuildTarget.GlobalLinkEnvironment.Config.AdditionalLibraries.Add("EGL");
 			InBuildTarget.GlobalLinkEnvironment.Config.AdditionalLibraries.Add("OpenSLES");
 			InBuildTarget.GlobalLinkEnvironment.Config.AdditionalLibraries.Add("android");
+			InBuildTarget.GlobalLinkEnvironment.Config.AdditionalLibraries.Add("cxa_demangle");
 
 			UEBuildConfiguration.bCompileSimplygon = false;
 			BuildConfiguration.bDeployAfterCompile = true;
diff --git a/Engine/Source/Runtime/Core/Private/Android/AndroidPlatformStackWalk.cpp b/Engine/Source/Runtime/Core/Private/Android/AndroidPlatformStackWalk.cpp
new file mode 100644
--- /dev/null
+++ b/Engine/Source/Runtime/Core/Private/Android/AndroidPlatformStackWalk.cpp
@@ -0,0 +1,151 @@
+// Copyright 1998-2014 Epic Games, Inc. All Rights Reserved.
+
+/*=============================================================================
+	AndroidPlatformStackWalk.cpp: Android implementations of stack walk functions
+=============================================================================*/
+
+#include "CorePrivate.h"
+#include <unwind.h>
+#include <dlfcn.h>
+#include <cxxabi.h>
+#include <stdio.h>
+
+bool FAndroidPlatformStackWalk::ProgramCounterToHumanReadableString(int32 CurrentCallDepth, uint64 ProgramCounter, ANSICHAR* HumanReadableString, SIZE_T HumanReadableStringSize, EVerbosityFlags VerbosityFlags, FGenericCrashContext* Context)
+{
+	Dl_info DylibInfo;
+	int32 Result = dladdr((const void*)ProgramCounter, &DylibInfo);
+	if (Result == 0)
+	{
+		return false;
+	}
+
+	FProgramCounterSymbolInfo SymbolInfo;
+	ProgramCounterToSymbolInfo(ProgramCounter, SymbolInfo);
+
+	// Write out function name.
+	FCStringAnsi::Strcat(HumanReadableString, HumanReadableStringSize, SymbolInfo.FunctionName);
+
+	if (VerbosityFlags & VF_DISPLAY_FILENAME)
+	{
+		ANSICHAR FileNameLine[MAX_SPRINTF];
+
+		if (SymbolInfo.LineNumber == 0)
+		{
+			// No line number. Print out the module offset address instead.
+			FCStringAnsi::Sprintf(FileNameLine, "Address = 0x%-8x ", SymbolInfo.OffsetInModule);
+		}
+		else
+		{
+			// try to add source file and line number, too
+			FCStringAnsi::Sprintf(FileNameLine, "[%s, line %d] ", SymbolInfo.Filename, SymbolInfo.LineNumber);
+		}
+
+		FCStringAnsi::Strcat(HumanReadableString, HumanReadableStringSize, FileNameLine);
+	}
+
+	if (VerbosityFlags & VF_DISPLAY_MODULE)
+	{
+		ANSICHAR ModuleName[MAX_SPRINTF];
+		// Write out Module information if there is sufficient space.
+		FCStringAnsi::Sprintf(ModuleName, "[in %s]", SymbolInfo.ModuleName);
+		FCStringAnsi::Strcat(HumanReadableString, HumanReadableStringSize, ModuleName);
+	}
+
+	// For the crash reporting code this needs a Windows line ending, the caller is responsible for the '\n'
+	FCStringAnsi::Strcat(HumanReadableString, HumanReadableStringSize, "\r");
+
+	return true;
+}
+
+void FAndroidPlatformStackWalk::ProgramCounterToSymbolInfo(uint64 ProgramCounter, FProgramCounterSymbolInfo&  SymbolInfo)
+{
+	Dl_info DylibInfo;
+	int32 Result = dladdr((const void*)ProgramCounter, &DylibInfo);
+	if (Result == 0)
+	{
+		return;
+	}
+
+	int32 Status = 0;
+	ANSICHAR* DemangledName = NULL;
+
+	// Increased the size of the demangle destination to reduce the chances that abi::__cxa_demangle will allocate
+	// this causes the app to hang as malloc isn't signal handler safe. Ideally we wouldn't call this function in a handler.
+	size_t DemangledNameLen = 65536;
+	ANSICHAR DemangledNameBuffer[65536] = { 0 };
+	DemangledName = abi::__cxa_demangle(DylibInfo.dli_sname, DemangledNameBuffer, &DemangledNameLen, &Status);
+
+	if (DemangledName)
+	{
+		// C++ function
+		FCStringAnsi::Sprintf(SymbolInfo.FunctionName, "%s ", DemangledName);
+	}
+	else if (DylibInfo.dli_sname)
+	{
+		// C function
+		FCStringAnsi::Sprintf(SymbolInfo.FunctionName, "%s() ", DylibInfo.dli_sname);
+	}
+	else
+	{
+		// Unknown!
+		FCStringAnsi::Sprintf(SymbolInfo.FunctionName, "[Unknown]() ");
+	}
+
+	// No line number available.
+	// TODO open libUE4.so from the apk and get the DWARF-2 data.
+	FCStringAnsi::Strcat(SymbolInfo.Filename, "Unknown");
+	SymbolInfo.LineNumber = 0;
+	
+	// Offset of the symbol in the module, eg offset into libUE4.so needed for offline addr2line use.
+	SymbolInfo.OffsetInModule = ProgramCounter - (uint64)DylibInfo.dli_fbase;
+
+	// Write out Module information.
+	ANSICHAR* DylibPath = (ANSICHAR*)DylibInfo.dli_fname;
+	ANSICHAR* DylibName = FCStringAnsi::Strrchr(DylibPath, '/');
+	if (DylibName)
+	{
+		DylibName += 1;
+	}
+	else
+	{
+		DylibName = DylibPath;
+	}
+	FCStringAnsi::Strcpy(SymbolInfo.ModuleName, DylibName);
+}
+
+namespace AndroidStackWalkHelpers
+{
+	uint64* BackTrace;
+	uint32 MaxDepth;
+
+	static _Unwind_Reason_Code BacktraceCallback(struct _Unwind_Context* Context, void* InDepthPtr)
+	{
+		uint32* DepthPtr = (uint32*)InDepthPtr;
+
+		if (*DepthPtr < MaxDepth)
+		{
+			BackTrace[*DepthPtr] = (uint64)_Unwind_GetIP(Context);
+		}
+
+		(*DepthPtr)++;
+		return (_Unwind_Reason_Code)0;
+	}
+}
+
+void FAndroidPlatformStackWalk::CaptureStackBackTrace(uint64* BackTrace, uint32 MaxDepth, void* Context)
+{
+	// Make sure we have place to store the information
+	if (BackTrace == NULL || MaxDepth == 0)
+	{
+		return;
+	}
+
+	// zero results
+	FPlatformMemory::Memzero(BackTrace, MaxDepth*sizeof(uint64));
+
+	// backtrace
+	AndroidStackWalkHelpers::BackTrace = BackTrace;
+	AndroidStackWalkHelpers::MaxDepth = MaxDepth;
+	uint32 Depth = 0;
+	_Unwind_Backtrace(AndroidStackWalkHelpers::BacktraceCallback, &Depth);
+}
diff --git a/Engine/Source/Runtime/Core/Public/Android/AndroidIncludes.h b/Engine/Source/Runtime/Core/Public/Android/AndroidIncludes.h
--- a/Engine/Source/Runtime/Core/Public/Android/AndroidIncludes.h
+++ b/Engine/Source/Runtime/Core/Public/Android/AndroidIncludes.h
@@ -50,7 +50,7 @@ struct RECT
 #include "Android/AndroidMemory.h"
 #include "Android/AndroidString.h"
 #include "Android/AndroidMisc.h"
-#include "Android/AndroidStackWalk.h"
+#include "Android/AndroidPlatformStackWalk.h"
 #include "Android/AndroidMath.h"
 #include "Android/AndroidTime.h"
 #include "Android/AndroidProcess.h"
diff --git a/Engine/Source/Runtime/Core/Public/Android/AndroidPlatformStackWalk.h b/Engine/Source/Runtime/Core/Public/Android/AndroidPlatformStackWalk.h
new file mode 100644
--- /dev/null
+++ b/Engine/Source/Runtime/Core/Public/Android/AndroidPlatformStackWalk.h
@@ -0,0 +1,45 @@
+// Copyright 1998-2014 Epic Games, Inc. All Rights Reserved.
+
+
+/*=============================================================================================
+	AndroidPlatformStackWalk.h: Android platform stack walk functions
+==============================================================================================*/
+
+#pragma once
+
+#undef PLATFORM_SUPPORTS_STACK_SYMBOLS
+#define PLATFORM_SUPPORTS_STACK_SYMBOLS 1
+
+/**
+* Symbol information associated with a program counter.
+*/
+struct FProgramCounterSymbolInfo
+{
+	/** Module name.					*/
+	ANSICHAR	ModuleName[1024];
+	/** Function name.					*/
+	ANSICHAR	FunctionName[1024];
+	/** Filename.						*/
+	ANSICHAR	Filename[1024];
+	/** Line number in file.			*/
+	int32		LineNumber;
+	/** ProgramCounter offset into module.	*/
+	uint64		OffsetInModule;
+	/** Symbol displacement of address.	*/
+	int32		SymbolDisplacement;
+};
+
+
+/**
+* Android platform stack walking
+*/
+struct CORE_API FAndroidPlatformStackWalk : public FGenericPlatformStackWalk
+{
+	typedef FGenericPlatformStackWalk Parent;
+
+	static bool ProgramCounterToHumanReadableString(int32 CurrentCallDepth, uint64 ProgramCounter, ANSICHAR* HumanReadableString, SIZE_T HumanReadableStringSize, EVerbosityFlags VerbosityFlags = VF_DISPLAY_ALL, FGenericCrashContext* Context = NULL);
+	static void ProgramCounterToSymbolInfo(uint64 ProgramCounter, FProgramCounterSymbolInfo&  OutProgramCounterSymbolInfo);
+	static void CaptureStackBackTrace(uint64* BackTrace, uint32 MaxDepth, void* Context = NULL);
+};
+
+typedef FAndroidPlatformStackWalk FPlatformStackWalk;
diff --git a/Engine/Source/Runtime/Core/Public/Android/AndroidStackWalk.h b/Engine/Source/Runtime/Core/Public/Android/AndroidStackWalk.h
deleted file mode 100644
