Docs: Removed file comments and added missing code documentation

diff --git a/Engine/Source/Runtime/Core/Public/Containers/LockFreeListImpl.h b/Engine/Source/Runtime/Core/Public/Containers/LockFreeListImpl.h
--- a/Engine/Source/Runtime/Core/Public/Containers/LockFreeListImpl.h
+++ b/Engine/Source/Runtime/Core/Public/Containers/LockFreeListImpl.h
@@ -1,24 +1,22 @@
 // Copyright 1998-2013 Epic Games, Inc. All Rights Reserved.
 
-/*=============================================================================
-	LockFreeListImpl.h: Implements various lock free linked list containers.
-=============================================================================*/
-
 #pragma once
 
+
 /** 
  *	Wrapper around FLockFreeVoidPointerListBase
  *	Simple lock free list (stack) that does not permit closure.
  *	T is the type of the pointer the list will contain.
-**/
+ */
 template <class T>
 class TLockFreePointerList : private FLockFreeVoidPointerListGeneric 
 {
 public:
 	/**	
-	 *	Push an item onto the head of the list
-	 *	@param NewItem, the new item to push on the list, cannot be NULL
-	**/
+	 *	Push an item onto the head of the list.
+	 *
+	 *	@param NewItem, the new item to push on the list, cannot be NULL.
+	 */
 	void Push(T *NewItem)
 	{
 		FLockFreeListPerfCounter LockFreePerfCounter;
@@ -26,9 +24,9 @@ public:
 	}
 
 	/**	
-	 *	Pop an item from the list or return NULL if the list is empty
-	 *	@return The popped item, if any
-	**/
+	 *	Pop an item from the list or return NULL if the list is empty.
+	 *	@return The popped item, if any.
+	 */
 	T* Pop()
 	{
 		FLockFreeListPerfCounter LockFreePerfCounter;
@@ -36,9 +34,10 @@ public:
 	}
 
 	/**	
-	 *	Pop all items from the list 
+	 *	Pop all items from the list.
+	 *
 	 *	@param Output The array to hold the returned items. Must be empty.
-	**/
+	 */
 	void PopAll(TArray<T *>& Output)
 	{
 		FLockFreeListPerfCounter LockFreePerfCounter;
@@ -46,11 +45,12 @@ public:
 	}	
 
 	/**	
-	 *	Check if the list is empty
+	 *	Check if the list is empty.
+	 *
 	 *	@return true if the list is empty.
 	 *	CAUTION: This methods safety depends on external assumptions. For example, if another thread could add to the list at any time, the return value is no better than a best guess.
 	 *	As typically used, the list is not being access concurrently when this is called.
-	**/
+	 */
 	bool IsEmpty() const  
 	{
 		return FLockFreeVoidPointerListGeneric::IsEmpty();
@@ -82,9 +82,10 @@ class TClosableLockFreePointerList : private FLockFreeVoidPointerListGeneric
 public:
 	/**	
 	 *	Push an item onto the head of the list, unless the list is closed
+	 *
 	 *	@param NewItem, the new item to push on the list, cannot be NULL
 	 *	@return true if the item was pushed on the list, false if the list was closed.
-	**/
+	 */
 	bool PushIfNotClosed(T *NewItem)
 	{
 		FLockFreeListPerfCounter LockFreePerfCounter;
@@ -93,8 +94,9 @@ public:
 
 	/**	
 	 *	Pop all items from the list and atomically close it.
+	 *
 	 *	@param Output The array to hold the returned items. Must be empty.
-	**/
+	 */
 	void PopAllAndClose(TArray<T *>& Output)
 	{
 		FLockFreeListPerfCounter LockFreePerfCounter;
@@ -103,8 +105,9 @@ public:
 
 	/**	
 	 *	Check if the list is closed
+	 *
 	 *	@return true if the list is closed.
-	**/
+	 */
 	bool IsClosed() const
 	{
 		return FLockFreeVoidPointerListGeneric::IsClosed();
@@ -134,10 +137,11 @@ class TReopenableLockFreePointerList : private FLockFreeVoidPointerListGeneric
 {
 public:
 	/**	
-	 *	Push an item onto the head of the list, opening it first if necessary
-	 *	@param NewItem, the new item to push on the list, cannot be NULL
-	 *	@return true if the list needed to be opened first, false if the list was not closed before our push
-	**/
+	 *	Push an item onto the head of the list, opening it first if necessary.
+	 *
+	 *	@param NewItem, the new item to push on the list, cannot be NULL.
+	 *	@return true if the list needed to be opened first, false if the list was not closed before our push.
+	 */
 	bool ReopenIfClosedAndPush(T *NewItem)
 	{
 		FLockFreeListPerfCounter LockFreePerfCounter;
@@ -146,9 +150,9 @@ public:
 	}
 
 	/**	
-	 *	Pop an item from the list or return NULL if the list is empty or closed
+	 *	Pop an item from the list or return NULL if the list is empty or closed.
 	 *	@return The new item, if any
-	**/
+	 */
 	T* PopIfNotClosed()
 	{
 		FLockFreeListPerfCounter LockFreePerfCounter;
@@ -156,19 +160,21 @@ public:
 	}
 
 	/**	
-	 *	Pop all items from the list 
+	 *	Pop all items from the list.
+	 *
 	 *	@param Output The array to hold the returned items. Must be empty.
 	 *	CAUTION: This method should not be used unless the list is known to not be closed.
-	**/
+	 */
 	void PopAll(TArray<T *>& Output)
 	{
 		FLockFreeListPerfCounter LockFreePerfCounter;
 		FLockFreeVoidPointerListGeneric::PopAll<TArray<T *>, T * >(Output);
 	}	
 	/**	
-	 *	Close the list if it is empty
-	 *	@return true if this call actively closed the list
-	**/
+	 *	Close the list if it is empty.
+	 *
+	 *	@return true if this call actively closed the list.
+	 */
 	bool CloseIfEmpty()
 	{
 		FLockFreeListPerfCounter LockFreePerfCounter;
diff --git a/Engine/Source/Runtime/Core/Public/Containers/LockFreeVoidPointerListBase.h b/Engine/Source/Runtime/Core/Public/Containers/LockFreeVoidPointerListBase.h
--- a/Engine/Source/Runtime/Core/Public/Containers/LockFreeVoidPointerListBase.h
+++ b/Engine/Source/Runtime/Core/Public/Containers/LockFreeVoidPointerListBase.h
@@ -1,13 +1,10 @@
 // Copyright 1998-2013 Epic Games, Inc. All Rights Reserved.
 
-/*=============================================================================
-	LockFreeVoidPointerListBase.h: A lock free single linked list
-=============================================================================*/
-
 #pragma once
 
 #define CHECK_NON_CONCURRENT_ASSUMPTIONS (0)
 
+
 /** 
  * Base class for a lock free list of pointers 
  * This class has several different features. Not all methods are supported in all situations.
diff --git a/Engine/Source/Runtime/Core/Public/Containers/LockFreeVoidPointerListBase128.h b/Engine/Source/Runtime/Core/Public/Containers/LockFreeVoidPointerListBase128.h
--- a/Engine/Source/Runtime/Core/Public/Containers/LockFreeVoidPointerListBase128.h
+++ b/Engine/Source/Runtime/Core/Public/Containers/LockFreeVoidPointerListBase128.h
@@ -1,12 +1,8 @@
 // Copyright 1998-2013 Epic Games, Inc. All Rights Reserved.
 
-/*=============================================================================
-	LockFreeVoidPointerListBase128.h: A lock free single linked list
-	Version using 128-bit CAS operation
-=============================================================================*/
-
 #pragma once
 
+
 /** Defines initial values used by the lock free linked list. */
 struct FLockFreeListConstants
 {
diff --git a/Engine/Source/Runtime/Core/Public/Containers/Map.h b/Engine/Source/Runtime/Core/Public/Containers/Map.h
--- a/Engine/Source/Runtime/Core/Public/Containers/Map.h
+++ b/Engine/Source/Runtime/Core/Public/Containers/Map.h
@@ -1,9 +1,5 @@
 // Copyright 1998-2014 Epic Games, Inc. All Rights Reserved.
 
-/*=============================================================================
-	Map.h: Dynamic map definitions.
-=============================================================================*/
-
 #pragma once
 
 #include "Core.h"
@@ -12,6 +8,7 @@
 
 #define ExchangeB(A,B) {bool T=A; A=B; B=T;}
 
+
 /** An initializer type for pairs that's passed to the pair set when adding a new pair. */
 template <typename KeyInitType, typename ValueInitType>
 class TPairInitializer
diff --git a/Engine/Source/Runtime/Core/Public/Containers/MapBuilder.h b/Engine/Source/Runtime/Core/Public/Containers/MapBuilder.h
--- a/Engine/Source/Runtime/Core/Public/Containers/MapBuilder.h
+++ b/Engine/Source/Runtime/Core/Public/Containers/MapBuilder.h
@@ -1,18 +1,14 @@
 // Copyright 1998-2014 Epic Games, Inc. All Rights Reserved.
 
-/*=============================================================================
-	MapBuilder.h: Declares the TMapBuilder template.
-=============================================================================*/
-
 #pragma once
 
 
 /**
  * Template for fluent map builders.
  *
- * @param KeyType - The type of keys stored in the map.
- * @param ValueType - The type of values stored in the map.
- * @param SetAllocator - The allocator to use for key-value pairs.
+ * @param KeyType The type of keys stored in the map.
+ * @param ValueType The type of values stored in the map.
+ * @param SetAllocator The allocator to use for key-value pairs.
  */
 template<typename KeyType, typename ValueType, typename SetAllocator = FDefaultSetAllocator> class TMapBuilder
 {
@@ -26,21 +22,19 @@ public:
 	/**
 	 * Creates and initializes a new map builder from another map.
 	 *
-	 * @param InMap - The map to copy.
+	 * @param InMap The map to copy.
 	 */
 	template<typename OtherAllocator> TMapBuilder( const TMap<KeyType, ValueType, OtherAllocator>& InMap )
 		: Map(InMap)
 	{ }
 
-
 public:
 
 	/**
 	 * Adds a key-value pair to the map.
 	 *
-	 * @param InKey - The key of the pair to add.
-	 * @param InValue - The value of the pair to add.
-	 *
+	 * @param InKey The key of the pair to add.
+	 * @param InValue The value of the pair to add.
 	 * @return This instance (for method chaining).
 	 */
 	TMapBuilder& Add( KeyType InKey, ValueType InValue )
@@ -53,8 +47,7 @@ public:
 	/**
 	 * Appends another map.
 	 *
-	 * @param OtherMap - The map to append.
-	 *
+	 * @param OtherMap The map to append.
 	 * @return This instance (for method chaining).
 	 */
 	TMapBuilder& Append( const TMap<KeyType, ValueType, SetAllocator>& OtherMap )
@@ -64,7 +57,6 @@ public:
 		return *this;
 	}
 
-
 public:
 
 	/**
@@ -87,7 +79,6 @@ public:
 		return Build();
 	}
 
-
 private:
 
 	// Holds the map being built.
diff --git a/Engine/Source/Runtime/Core/Public/Containers/Queue.h b/Engine/Source/Runtime/Core/Public/Containers/Queue.h
--- a/Engine/Source/Runtime/Core/Public/Containers/Queue.h
+++ b/Engine/Source/Runtime/Core/Public/Containers/Queue.h
@@ -1,9 +1,5 @@
 // Copyright 1998-2014 Epic Games, Inc. All Rights Reserved.
 
-/*=============================================================================
-	Queue.h: Declares the TQueue class template.
-=============================================================================*/
-
 #pragma once
 
 
@@ -34,10 +30,9 @@ namespace EQueueMode
  * writing it in a way that does not depend on possible instruction reordering on the CPU.
  * The Enqueue() method uses an atomic compare-and-swap in multiple-producers scenarios.
  *
- * @param ItemType - The type of items stored in the queue.
- * @param Mode - The queue mode (single-producer, single-consumer by default).
- *
- * @todo gmp: Implement node pooling.
+ * @param ItemType The type of items stored in the queue.
+ * @param Mode The queue mode (single-producer, single-consumer by default).
+  * @todo gmp: Implement node pooling.
  */
 template<typename ItemType, EQueueMode::Type Mode = EQueueMode::Spsc>
 class TQueue
diff --git a/Engine/Source/Runtime/Core/Public/Containers/ResourceArray.h b/Engine/Source/Runtime/Core/Public/Containers/ResourceArray.h
--- a/Engine/Source/Runtime/Core/Public/Containers/ResourceArray.h
+++ b/Engine/Source/Runtime/Core/Public/Containers/ResourceArray.h
@@ -1,11 +1,8 @@
 // Copyright 1998-2014 Epic Games, Inc. All Rights Reserved.
 
-/*=============================================================================
-	ResourceArray.h: Resource array definitions and platform includes.
-=============================================================================*/
-
 #pragma once
 
+
 /**
  * An element type independent interface to the resource array.
  */
@@ -14,56 +11,58 @@ class FResourceArrayInterface
 public:
 
 	/**
-	* @return A pointer to the resource data.
-	*/
+	 * @return A pointer to the resource data.
+	 */
 	virtual const void* GetResourceData() const = 0;
 
 	/**
-	* @return size of resource data allocation
-	*/
+	 * @return size of resource data allocation
+	 */
 	virtual uint32 GetResourceDataSize() const = 0;
 
 	/**
-	* Called on non-UMA systems after the RHI has copied the resource data, and no longer needs the CPU's copy.
-	*/
+	 * Called on non-UMA systems after the RHI has copied the resource data, and no longer needs the CPU's copy.
+	 */
 	virtual void Discard() = 0;
 
 	/**
-	* @return true if the resource array is static and shouldn't be modified
-	*/
+	 * @return true if the resource array is static and shouldn't be modified
+	 */
 	virtual bool IsStatic() const = 0;
 
 	/**
-	* @return true if the resource keeps a copy of its resource data after the RHI resource has been created
-	*/
+	 * @return true if the resource keeps a copy of its resource data after the RHI resource has been created
+	 */
 	virtual bool GetAllowCPUAccess() const = 0;
 
 	/** 
-	* Sets whether the resource array will be accessed by CPU. 
-	*/
+	 * Sets whether the resource array will be accessed by CPU. 
+	 */
 	virtual void SetAllowCPUAccess( bool bInNeedsCPUAccess ) = 0;
 };
 
+
 /**
-* allows for direct GPU mem allocation for bulk resource types
-*/
+ * allows for direct GPU mem allocation for bulk resource types
+ */
 class FResourceBulkDataInterface
 {
 public:
 	/** 
-	* @return ptr to the resource memory which has been preallocated
-	*/
+	 * @return ptr to the resource memory which has been preallocated
+	 */
 	virtual const void* GetResourceBulkData() const = 0;
 	/** 
-	* @return size of resource memory
-	*/
+	 * @return size of resource memory
+	 */
 	virtual uint32 GetResourceBulkDataSize() const = 0;
 	/**
-	* Free memory after it has been used to initialize RHI resource 
-	*/
+	 * Free memory after it has been used to initialize RHI resource 
+	 */
 	virtual void Discard() = 0;
 };
 
+
 /**
 * allows for direct GPU mem allocation for texture resource
 */
@@ -71,25 +70,25 @@ class FTexture2DResourceMem : public FResourceBulkDataInterface
 {
 public:
 	/**
-	* @param MipIdx index for mip to retrieve
-	* @return ptr to the offset in bulk memory for the given mip
-	*/
+	 * @param MipIdx index for mip to retrieve
+	 * @return ptr to the offset in bulk memory for the given mip
+	 */
 	virtual void* GetMipData(int32 MipIdx) = 0;
 	/**
-	* @return total number of mips stored in this resource
-	*/
+	 * @return total number of mips stored in this resource
+	 */
 	virtual int32	GetNumMips() = 0;
 	/** 
-	* @return width of texture stored in this resource
-	*/
+	 * @return width of texture stored in this resource
+	 */
 	virtual int32 GetSizeX() = 0;
 	/** 
-	* @return height of texture stored in this resource
-	*/
+	 * @return height of texture stored in this resource
+	 */
 	virtual int32 GetSizeY() = 0;
 	/**
 	 * @return Whether the resource memory is properly allocated or not.
-	 **/
+	 */
 	virtual bool IsValid() = 0;
 
 	/**
@@ -110,6 +109,6 @@ public:
 	virtual ~FTexture2DResourceMem() {}
 };
 
+
 // Default to the dynamically bound RHI.
 #include "DynamicRHIResourceArray.h"
-
diff --git a/Engine/Source/Runtime/Core/Public/Containers/Set.h b/Engine/Source/Runtime/Core/Public/Containers/Set.h
--- a/Engine/Source/Runtime/Core/Public/Containers/Set.h
+++ b/Engine/Source/Runtime/Core/Public/Containers/Set.h
@@ -1,13 +1,10 @@
 // Copyright 1998-2014 Epic Games, Inc. All Rights Reserved.
 
-/*=============================================================================
-	Set.h: Set definitions.
-=============================================================================*/
-
 #pragma once
 
 #include "SparseArray.h"
 
+
 /**
  * The base KeyFuncs type with some useful definitions for all KeyFuncs; meant to be derived from instead of used directly.
  * bInAllowDuplicateKeys=true is slightly faster because it allows the TSet to skip validating that
diff --git a/Engine/Source/Runtime/Core/Public/Containers/SparseArray.h b/Engine/Source/Runtime/Core/Public/Containers/SparseArray.h
--- a/Engine/Source/Runtime/Core/Public/Containers/SparseArray.h
+++ b/Engine/Source/Runtime/Core/Public/Containers/SparseArray.h
@@ -1,18 +1,16 @@
 // Copyright 1998-2014 Epic Games, Inc. All Rights Reserved.
 
-/*=============================================================================
-	SparseArray.h: Sparse array definition.
-=============================================================================*/
-
 #pragma once
 
 #include "Array.h"
 #include "BitArray.h"
 
+
 // Forward declarations.
 template<typename ElementType,typename Allocator = FDefaultSparseArrayAllocator >
 class TSparseArray;
 
+
 /** The result of a sparse array allocation. */
 struct FSparseArrayAllocationInfo
 {
diff --git a/Engine/Source/Runtime/Core/Public/Containers/StackTracker.h b/Engine/Source/Runtime/Core/Public/Containers/StackTracker.h
--- a/Engine/Source/Runtime/Core/Public/Containers/StackTracker.h
+++ b/Engine/Source/Runtime/Core/Public/Containers/StackTracker.h
@@ -1,9 +1,5 @@
 // Copyright 1998-2014 Epic Games, Inc. All Rights Reserved.
 
-/*=============================================================================
-StackTracker.h: Stack Tracking within Unreal Engine.
-=============================================================================*/
-
 #pragma once
 
 #include "UMemoryDefines.h"
@@ -33,7 +29,6 @@ struct FSlackTrackData
 };
 
 
-
 /**
  * Stack tracker. Used to identify callstacks at any point in the codebase.
  */
@@ -66,7 +61,7 @@ public:
 		,	StopFrameCounter(0)
 		,   UpdateFn(InUpdateFn)
 		,   ReportFn(InReportFn)
-	{}
+	{ }
 
 	/**
 	 * Captures the current stack and updates stack tracking information.
@@ -87,6 +82,7 @@ public:
 	CORE_API void ToggleTracking();
 
 private:
+
 	/** Array of unique callstacks. */
 	TArray<FCallStack> CallStacks;
 	/** Mapping from callstack CRC to index in callstack array. */
diff --git a/Engine/Source/Runtime/Core/Public/Containers/StaticArray.h b/Engine/Source/Runtime/Core/Public/Containers/StaticArray.h
--- a/Engine/Source/Runtime/Core/Public/Containers/StaticArray.h
+++ b/Engine/Source/Runtime/Core/Public/Containers/StaticArray.h
@@ -1,13 +1,10 @@
 // Copyright 1998-2014 Epic Games, Inc. All Rights Reserved.
 
-/*=============================================================================
-	StaticArray.h: Static array definitions.
-=============================================================================*/
-
 #pragma once
 
 #include "Containers/ContainerAllocationPolicies.h"
 
+
 /** An array with a static number of elements. */
 template<typename TElement,uint32 NumElements,uint32 Alignment = ALIGNOF(TElement)>
 class TStaticArray
diff --git a/Engine/Source/Runtime/Core/Public/Containers/Ticker.h b/Engine/Source/Runtime/Core/Public/Containers/Ticker.h
--- a/Engine/Source/Runtime/Core/Public/Containers/Ticker.h
+++ b/Engine/Source/Runtime/Core/Public/Containers/Ticker.h
@@ -1,24 +1,20 @@
 // Copyright 1998-2014 Epic Games, Inc. All Rights Reserved.
 
-/*=============================================================================
-	Ticker.h: Efficient scheduled delegate manager	
-=============================================================================*/
-
 #pragma once
 
 #include "Delegate.h"
 
 
 /**
-  * Ticker delegates return true to automatically reschedule at the same delay or false for one-shot.
-  * You will not get more than one fire per "frame", which is just a FTicker::Tick call.
-  * Argument is DeltaTime
-**/
+ * Ticker delegates return true to automatically reschedule at the same delay or false for one-shot.
+ * You will not get more than one fire per "frame", which is just a FTicker::Tick call.
+ * Argument is DeltaTime
+ */
 DECLARE_DELEGATE_RetVal_OneParam(bool, FTickerDelegate, float);
 
 /**
-  * Ticker class. Fires delegates after a delay.
-**/
+ * Ticker class. Fires delegates after a delay.
+ */
 class FTicker
 {
 public:
@@ -32,10 +28,10 @@ public:
 	}
 
 	/**
-	  * Add a new ticker with a given delay / interval
-	  * @param InDelegate	delegate to fire after the delay
-	  * @param InDelay		Delay until next fire; 0 means "next frame"
-	**/
+	 * Add a new ticker with a given delay / interval
+	 * @param InDelegate	delegate to fire after the delay
+	 * @param InDelay		Delay until next fire; 0 means "next frame"
+	 */
 	void AddTicker(const FTickerDelegate& InDelegate, float InDelay = 0.0f)
 	{
 		FElement Ticker(CurrentTime + InDelay, InDelay, InDelegate);
@@ -59,9 +55,9 @@ public:
 	}
 
 	/**
-	  * Fire all tickers who have passed their delay and reschedule the ones that return true
-	  * @param DeltaTime	time that has passed since the last tick call
-	**/
+	 * Fire all tickers who have passed their delay and reschedule the ones that return true
+	 * @param DeltaTime	time that has passed since the last tick call
+	 */
 	void Tick(float DeltaTime)
 	{
 		if (!PriorityQueue.Num())
@@ -90,8 +86,8 @@ public:
 private:
 
 	/**
-	  * Element of the priority queue
-	**/
+	 * Element of the priority queue
+	 */
 	struct FElement
 	{
 		/** Time that this delegate must not fire before **/
@@ -165,8 +161,8 @@ public:
 	 */
 	virtual bool Tick(float DeltaTime) = 0;
 
-private:
+private:	
+
 	/** Delegate for callbacks to Tick */
 	FTickerDelegate TickDelegate;
 };
-
diff --git a/Engine/Source/Runtime/Core/Public/Delegates/DelegateBase.h b/Engine/Source/Runtime/Core/Public/Delegates/DelegateBase.h
--- a/Engine/Source/Runtime/Core/Public/Delegates/DelegateBase.h
+++ b/Engine/Source/Runtime/Core/Public/Delegates/DelegateBase.h
@@ -1,9 +1,5 @@
 // Copyright 1998-2014 Epic Games, Inc. All Rights Reserved.
 
-/*=============================================================================
-	DelegateBase.h: Declares the FDelegateBase class.
-=============================================================================*/
-
 #pragma once
 
 
@@ -27,9 +23,9 @@ public:
 public:
 
 	/**
-	 * If this is a UFunction or UObject delegate, return the UObject
+	 * If this is a UFunction or UObject delegate, return the UObject.
 	 *
-	 * @return  the object associated with this delegate if there is one
+	 * @return The object associated with this delegate if there is one.
 	 */
 	inline class UObject* GetUObject( ) const
 	{
@@ -46,9 +42,9 @@ public:
 	}
 
 	/**
-	 * Checks to see if the user object bound to this delegate is still valid
+	 * Checks to see if the user object bound to this delegate is still valid.
 	 *
-	 * @return  True if the user object is still valid and it's safe to execute the function call
+	 * @return True if the user object is still valid and it's safe to execute the function call.
 	 */
 	inline bool IsBound( ) const
 	{
@@ -58,7 +54,7 @@ public:
 	/** 
 	 * Checks to see if this delegate is bound to the given user object.
 	 *
-	 * @return	True if this delegate is bound to InUserObject, false otherwise.
+	 * @return True if this delegate is bound to InUserObject, false otherwise.
 	 */
 	inline bool IsBoundToObject( void const* InUserObject ) const
 	{
diff --git a/Engine/Source/Runtime/Core/Public/Delegates/DelegateInstancesImpl.inl b/Engine/Source/Runtime/Core/Public/Delegates/DelegateInstancesImpl.inl
--- a/Engine/Source/Runtime/Core/Public/Delegates/DelegateInstancesImpl.inl
+++ b/Engine/Source/Runtime/Core/Public/Delegates/DelegateInstancesImpl.inl
@@ -101,7 +101,7 @@ public:
 
 public:
 
-	// Begin IDelegateInstance interface
+	// IDelegateInstance interface
 
 	virtual FName GetFunctionName( ) const override
 	{
@@ -133,11 +133,9 @@ public:
 		return UserObject.IsValid();
 	}
 
-	// End IDelegateInstance interface
-
 public:
 
-	// Begin DELEGATE_INSTANCE_INTERFACE_CLASS interface
+	// DELEGATE_INSTANCE_INTERFACE_CLASS interface
 
 	virtual DELEGATE_INSTANCE_INTERFACE_CLASS<FUNC_TEMPLATE_ARGS>* CreateCopy() override
 	{
@@ -194,8 +192,6 @@ public:
 		return false;
 	}
 
-	// End DELEGATE_INSTANCE_INTERFACE_CLASS interface
-
 public:
 
 	/**
@@ -295,7 +291,7 @@ public:
 
 public:
 
-	// Begin IDelegateInstance interface
+	// IDelegateInstance interface
 
 	virtual FName GetFunctionName( ) const override
 	{
@@ -329,11 +325,9 @@ public:
 		return true;
 	}
 
-	// End IDelegateInstance interface
-
 public:
 
-	// Begin DELEGATE_INSTANCE_INTERFACE_CLASS interface
+	// DELEGATE_INSTANCE_INTERFACE_CLASS interface
 
 	virtual DELEGATE_INSTANCE_INTERFACE_CLASS<FUNC_TEMPLATE_ARGS>* CreateCopy( ) override
 	{
@@ -462,7 +456,7 @@ public:
 
 public:
 
-	// Begin IDelegateInstance interface
+	// IDelegateInstance interface
 
 	virtual FName GetFunctionName( ) const override
 	{
@@ -499,11 +493,9 @@ public:
 		return !!UserObject.Get();
 	}
 
-	// End IDelegateInstance interface
-
 public:
 
-	// Begin DELEGATE_INSTANCE_INTERFACE_CLASS interface
+	// DELEGATE_INSTANCE_INTERFACE_CLASS interface
 
 	virtual DELEGATE_INSTANCE_INTERFACE_CLASS<FUNC_TEMPLATE_ARGS>* CreateCopy( ) override
 	{
@@ -638,7 +630,7 @@ public:
 
 public:
 
-	// Begin IDelegateInstance interface
+	// IDelegateInstance interface
 
 	virtual FName GetFunctionName( ) const override
 	{
@@ -672,11 +664,9 @@ public:
 		return true;
 	}
 
-	// End IDelegateInstance interface
-
 public:
 
-	// Begin DELEGATE_INSTANCE_INTERFACE_CLASS interface
+	// DELEGATE_INSTANCE_INTERFACE_CLASS interface
 
 	virtual DELEGATE_INSTANCE_INTERFACE_CLASS<FUNC_TEMPLATE_ARGS>* CreateCopy( ) override
 	{
@@ -792,7 +782,7 @@ public:
 	
 public:
 
-	// Begin IDelegateInstance interface
+	// IDelegateInstance interface
 
 	virtual FName GetFunctionName( ) const override
 	{
@@ -829,11 +819,9 @@ public:
 		return UserObjectPtr.IsValid();
 	}
 
-	// End IDelegateInstance interface
-
 public:
 
-	// Begin DELEGATE_INSTANCE_INTERFACE_CLASS interface
+	// DELEGATE_INSTANCE_INTERFACE_CLASS interface
 
 	virtual DELEGATE_INSTANCE_INTERFACE_CLASS<FUNC_TEMPLATE_ARGS>* CreateCopy( ) override
 	{
diff --git a/Engine/Source/Runtime/Core/Public/Delegates/IDelegateInstance.h b/Engine/Source/Runtime/Core/Public/Delegates/IDelegateInstance.h
--- a/Engine/Source/Runtime/Core/Public/Delegates/IDelegateInstance.h
+++ b/Engine/Source/Runtime/Core/Public/Delegates/IDelegateInstance.h
@@ -1,9 +1,5 @@
 // Copyright 1998-2014 Epic Games, Inc. All Rights Reserved.
 
-/*=============================================================================
-	IDelegateInstance.cpp: Declares the IDelegateInstance interface.
-=============================================================================*/
-
 #pragma once
 
 
@@ -52,30 +48,30 @@ public:
 	/**
 	 * Returns raw pointer to the delegate method.
 	 *
-	 * @return  Raw pointer to the delegate method.
+	 * @return Raw pointer to the delegate method.
 	 */
 	virtual const void* GetRawMethodPtr( ) const = 0;
 
 	/**
 	 * Returns raw pointer to UserObject,
 	 *
-	 * @return  Raw pointer to UserObject.
+	 * @return Raw pointer to UserObject.
 	 */
 	virtual const void* GetRawUserObject( ) const = 0;
 
 	/**
 	 * Returns the type of delegate instance
 	 *
-	 * @return  Delegate instance type
+	 * @return Delegate instance type
 	 */
 	virtual EDelegateInstanceType::Type GetType( ) const = 0;
 
 	/**
 	 * Returns true if this delegate is bound to the specified UserObject,
 	 *
-	 * @param  InUserObject
+	 * @param InUserObject
 	 *
-	 * @return  True if delegate is bound to the specified UserObject
+	 * @return True if delegate is bound to the specified UserObject
 	 */
 	virtual bool HasSameObject( const void* InUserObject ) const = 0;
 
@@ -83,7 +79,7 @@ public:
 	 * Checks to see if the user object bound to this delegate can ever be valid again.
 	 * used to compact multicast delegate arrays so they don't expand without limit.
 	 *
-	 * @return  True if the user object can never be used again
+	 * @return True if the user object can never be used again
 	 */
 	virtual bool IsCompactable( ) const
 	{
@@ -93,7 +89,7 @@ public:
 	/**
 	 * Checks to see if the user object bound to this delegate is still valid
 	 *
-	 * @return  True if the user object is still valid and it's safe to execute the function call
+	 * @return True if the user object is still valid and it's safe to execute the function call
 	 */
 	virtual bool IsSafeToExecute( ) const = 0;
 
diff --git a/Engine/Source/Runtime/Core/Public/Delegates/MulticastDelegateBase.h b/Engine/Source/Runtime/Core/Public/Delegates/MulticastDelegateBase.h
--- a/Engine/Source/Runtime/Core/Public/Delegates/MulticastDelegateBase.h
+++ b/Engine/Source/Runtime/Core/Public/Delegates/MulticastDelegateBase.h
@@ -1,9 +1,5 @@
 // Copyright 1998-2014 Epic Games, Inc. All Rights Reserved.
 
-/*=============================================================================
-	MulticastDelegateBase.h: Declares the FMulticastDelegateBase class.
-=============================================================================*/
-
 #pragma once
 
 
diff --git a/Engine/Source/Runtime/Core/Public/Misc/MonitoredProcess.h b/Engine/Source/Runtime/Core/Public/Misc/MonitoredProcess.h
--- a/Engine/Source/Runtime/Core/Public/Misc/MonitoredProcess.h
+++ b/Engine/Source/Runtime/Core/Public/Misc/MonitoredProcess.h
@@ -1,9 +1,5 @@
 // Copyright 1998-2014 Epic Games, Inc. All Rights Reserved.
 
-/*=============================================================================================
-	MonitoredProcess.h: Declares the FMonitoredProcess class
-==============================================================================================*/
-
 #pragma once
 
 /**
@@ -32,9 +28,9 @@ public:
 	/**
 	 * Creates a new monitored process.
 	 *
-	 * @param InURL - The URL of the executable to launch.
-	 * @param InParams - The command line parameters.
-	 * @param InHidden - Whether the window of the process should be hidden.
+	 * @param InURL The URL of the executable to launch.
+	 * @param InParams The command line parameters.
+	 * @param InHidden Whether the window of the process should be hidden.
 	 */
 	FMonitoredProcess( const FString& InURL, const FString& InParams, bool InHidden );
 
@@ -43,13 +39,12 @@ public:
 	 */
 	~FMonitoredProcess( );
 
-
 public:
 
 	/**
 	 * Cancels the process.
 	 *
-	 * @param InKillTree - Whether to kill the entire process tree when canceling this process.
+	 * @param InKillTree Whether to kill the entire process tree when canceling this process.
 	 */
 	void Cancel( bool InKillTree = false )
 	{
@@ -112,38 +107,33 @@ public:
 		return OutputDelegate;
 	}
 
-
 public:
 
-	// Begin FRunnable interface
+	// FRunnable interface
 
-	virtual bool Init( ) OVERRIDE
+	virtual bool Init( ) override
 	{
 		return true;
 	}
 
-	virtual uint32 Run( ) OVERRIDE;
+	virtual uint32 Run( ) override;
 
-	virtual void Stop( ) OVERRIDE
+	virtual void Stop( ) override
 	{
 		Cancel();
 	}
 
-	virtual void Exit( ) OVERRIDE { }
-
-	// End FRunnable interface
-
+	virtual void Exit( ) override { }
 
 protected:
 
 	/**
 	 * Processes the given output string.
 	 *
-	 * @param Output - The output string to process.
+	 * @param Output The output string to process.
 	 */
 	void ProcessOutput( const FString& Output );
 
-
 private:
 
 	// Whether the process is being canceled.
@@ -179,7 +169,6 @@ private:
 	// Holds the write pipe.
 	void* WritePipe;
 
-
 private:
 
 	// Holds a delegate that is executed when the process has been canceled.
diff --git a/Engine/Source/Runtime/Core/Public/UObject/NameTypes.h b/Engine/Source/Runtime/Core/Public/UObject/NameTypes.h
--- a/Engine/Source/Runtime/Core/Public/UObject/NameTypes.h
+++ b/Engine/Source/Runtime/Core/Public/UObject/NameTypes.h
@@ -1,9 +1,5 @@
 // Copyright 1998-2014 Epic Games, Inc. All Rights Reserved.
 
-/*=============================================================================
-	NameTypes.h: Unreal global name types.
-=============================================================================*/
-
 #pragma once
 
 /*----------------------------------------------------------------------------
diff --git a/Engine/Source/Runtime/EngineMessages/Classes/EngineServiceMessages.h b/Engine/Source/Runtime/EngineMessages/Classes/EngineServiceMessages.h
--- a/Engine/Source/Runtime/EngineMessages/Classes/EngineServiceMessages.h
+++ b/Engine/Source/Runtime/EngineMessages/Classes/EngineServiceMessages.h
@@ -1,9 +1,5 @@
 // Copyright 1998-2014 Epic Games, Inc. All Rights Reserved.
 
-/*=============================================================================
-	EngineSerivceMessages.h: Declares message types sent and consumed by FEngineService.
-=============================================================================*/
-
 #pragma once
 
 #include "EngineServiceMessages.generated.h"
@@ -39,45 +35,31 @@ struct FEngineServicePong
 {
 	GENERATED_USTRUCT_BODY()
 
-	/**
-	 * Holds the name of the currently loaded level, if any.
-	 */
+	/** Holds the name of the currently loaded level, if any. */
 	UPROPERTY()
 	FString CurrentLevel;
 
-	/**
-	 * Holds the engine version.
-	 */
+	/** Holds the engine version. */
 	UPROPERTY()
 	int32 EngineVersion;
 
-	/**
-	 * Holds a flag indicating whether game play has begun.
-	 */
+	/** Holds a flag indicating whether game play has begun. */
 	UPROPERTY()
 	bool HasBegunPlay;
 
-	/**
-	 * Holds the instance identifier.
-	 */
+	/** Holds the instance identifier. */
 	UPROPERTY()
 	FGuid InstanceId;
 
-	/**
-	 * Holds the type of the engine instance.
-	 */
+	/** Holds the type of the engine instance. */
 	UPROPERTY()
 	FString InstanceType;
 
-	/**
-	 * Holds the identifier of the session that the application belongs to.
-	 */
+	/** Holds the identifier of the session that the application belongs to. */
 	UPROPERTY()
 	FGuid SessionId;
 
-	/**
-	 * Holds the time in seconds since the world was loaded.
-	 */
+	/** Holds the time in seconds since the world was loaded. */
 	UPROPERTY()
 	float WorldTimeSeconds;
 };
@@ -103,15 +85,11 @@ struct FEngineServiceAuthDeny
 {
 	GENERATED_USTRUCT_BODY()
 
-	/**
-	 * Holds the name of the user that denied access.
-	 */
+	/** Holds the name of the user that denied access. */
 	UPROPERTY()
 	FString UserName;
 
-	/**
-	 * Holds the name of the user that access is denied to.
-	 */
+	/** Holds the name of the user that access is denied to. */
 	UPROPERTY()
 	FString UserToDeny;
 };
@@ -134,15 +112,11 @@ struct FEngineServiceAuthGrant
 {
 	GENERATED_USTRUCT_BODY()
 
-	/**
-	 * Holds the name of the user that granted access.
-	 */
+	/** Holds the name of the user that granted access. */
 	UPROPERTY()
 	FString UserName;
 
-	/**
-	 * Holds the name of the user that access is granted to.
-	 */
+	/** Holds the name of the user that access is granted to. */
 	UPROPERTY()
 	FString UserToGrant;
 
@@ -169,19 +143,14 @@ struct FEngineServiceExecuteCommand
 {
 	GENERATED_USTRUCT_BODY()
 
-	/**
-	 * Holds the command to execute.
-	 */
+	/** Holds the command to execute. */
 	UPROPERTY()
 	FString Command;
 
-	/**
-	 * Holds the name of the user that wants to execute the command.
-	 */
+	/** Holds the name of the user that wants to execute the command. */
 	UPROPERTY()
 	FString UserName;
 
-
 	/**
 	 * Default constructor.
 	 */
@@ -214,13 +183,10 @@ struct FEngineServiceTerminate
 {
 	GENERATED_USTRUCT_BODY()
 
-	/**
-	 * Holds the name of the user that wants to terminate the engine.
-	 */
+	/** Holds the name of the user that wants to terminate the engine. */
 	UPROPERTY()
 	FString UserName;
 
-
 	/**
 	 * Default constructor.
 	 */
@@ -255,19 +221,14 @@ struct FEngineServiceNotification
 {
 	GENERATED_USTRUCT_BODY()
 
-	/**
-	 * Holds the notification text.
-	 */
+	/** Holds the notification text. */
 	UPROPERTY()
 	FString Text;
 
-	/**
-	 * Holds the time in seconds since the engine started.
-	 */
+	/** Holds the time in seconds since the engine started. */
 	UPROPERTY()
 	double TimeSeconds;
 
-
 	/**
 	 * Default constructor.
 	 */
diff --git a/Engine/Source/Runtime/EngineMessages/Private/EngineMessagesModule.cpp b/Engine/Source/Runtime/EngineMessages/Private/EngineMessagesModule.cpp
--- a/Engine/Source/Runtime/EngineMessages/Private/EngineMessagesModule.cpp
+++ b/Engine/Source/Runtime/EngineMessages/Private/EngineMessagesModule.cpp
@@ -1,9 +1,5 @@
 // Copyright 1998-2014 Epic Games, Inc. All Rights Reserved.
 
-/*=============================================================================
-	EngineMessagesModule.cpp: Implements the FEngineMessagesModule class.
-=============================================================================*/
-
 #include "EngineMessagesPrivatePCH.h"
 
 
@@ -15,8 +11,9 @@ class FEngineMessagesModule
 {
 public:
 
-	virtual void StartupModule( ) OVERRIDE { }
+	// IModuleInterface interface
 
+	virtual void StartupModule( ) OVERRIDE { }
 	virtual void ShutdownModule( ) OVERRIDE { }
 
 	virtual bool SupportsDynamicReloading( ) OVERRIDE
diff --git a/Engine/Source/Runtime/EngineMessages/Private/EngineMessagesPrivatePCH.h b/Engine/Source/Runtime/EngineMessages/Private/EngineMessagesPrivatePCH.h
--- a/Engine/Source/Runtime/EngineMessages/Private/EngineMessagesPrivatePCH.h
+++ b/Engine/Source/Runtime/EngineMessages/Private/EngineMessagesPrivatePCH.h
@@ -1,24 +1,11 @@
 // Copyright 1998-2014 Epic Games, Inc. All Rights Reserved.
 
-/*=============================================================================
-	EngineMessagesPrivatePCH.h: Pre-compiled header file for the EngineMessages module.
-=============================================================================*/
+#pragma once
 
-#ifndef ENGINEMESSAGES_PRIVATEPCH_H
-#define ENGINEMESSAGES_PRIVATEPCH_H
+#include "EngineMessages.h"
 
 
-#include "../Public/EngineMessages.h"
-
-
-/* Dependencies
+/* Private dependencies
  *****************************************************************************/
 
 #include "ModuleManager.h"
-
-
-/* Private includes
- *****************************************************************************/
-
-
-#endif
diff --git a/Engine/Source/Runtime/EngineMessages/Public/EngineMessages.h b/Engine/Source/Runtime/EngineMessages/Public/EngineMessages.h
--- a/Engine/Source/Runtime/EngineMessages/Public/EngineMessages.h
+++ b/Engine/Source/Runtime/EngineMessages/Public/EngineMessages.h
@@ -1,19 +1,15 @@
 // Copyright 1998-2014 Epic Games, Inc. All Rights Reserved.
 
-/*=============================================================================
-	EngineMessages.h: EngineMessages module public header file.
-=============================================================================*/
-
 #pragma once
 
 
-/* Dependencies
+/* Public Dependencies
  *****************************************************************************/
 
 #include "CoreUObject.h"
 
 
-/* Messages
+/* Public Includes
  *****************************************************************************/
 
 #include "EngineServiceMessages.h"
diff --git a/Engine/Source/Runtime/EngineSettings/Classes/ConsoleSettings.h b/Engine/Source/Runtime/EngineSettings/Classes/ConsoleSettings.h
--- a/Engine/Source/Runtime/EngineSettings/Classes/ConsoleSettings.h
+++ b/Engine/Source/Runtime/EngineSettings/Classes/ConsoleSettings.h
@@ -1,9 +1,5 @@
 // Copyright 1998-2014 Epic Games, Inc. All Rights Reserved.
 
-/*=============================================================================
-	ConsoleSettings.h: Declares the UConsoleSettings class.
-=============================================================================*/
-
 #pragma once
 
 #include "ConsoleSettings.generated.h"
diff --git a/Engine/Source/Runtime/EngineSettings/Classes/GameMapsSettings.h b/Engine/Source/Runtime/EngineSettings/Classes/GameMapsSettings.h
--- a/Engine/Source/Runtime/EngineSettings/Classes/GameMapsSettings.h
+++ b/Engine/Source/Runtime/EngineSettings/Classes/GameMapsSettings.h
@@ -1,9 +1,5 @@
 // Copyright 1998-2014 Epic Games, Inc. All Rights Reserved.
 
-/*=============================================================================
-	GameMapsSettings.h: Declares the UGameMapsSettings class.
-=============================================================================*/
-
 #pragma once
 
 #include "GameMapsSettings.generated.h"
@@ -62,21 +58,15 @@ class ENGINESETTINGS_API UGameMapsSettings
 
 public:
 
-	/**
-	 * If set, this map will be loaded when the Editor starts up.
-	 */
+	/** If set, this map will be loaded when the Editor starts up. */
 	UPROPERTY(config, EditAnywhere, Category=DefaultMaps)
 	FString EditorStartupMap;
 
-	/**
-	 * The default options that will be appended to a map being loaded.
-	 */
+	/** The default options that will be appended to a map being loaded. */
 	UPROPERTY(config, EditAnywhere, Category=DefaultMaps, AdvancedDisplay)
 	FString LocalMapOptions;
 
-	/**
-	 * The map loaded when transition from one map to another.
-	 */
+	/** The map loaded when transition from one map to another. */
 	UPROPERTY(config, EditAnywhere, Category=DefaultMaps, AdvancedDisplay)
 	FString TransitionMap;
 
@@ -91,21 +81,15 @@ public:
 
 private:
 
-	/**
-	 * The map that will be loaded by default when no other map is loaded.
-	 */
+	/** The map that will be loaded by default when no other map is loaded. */
 	UPROPERTY(config, EditAnywhere, Category=DefaultMaps)
 	FString GameDefaultMap;
 
-	/**
-	 * The map that will be loaded by default when no other map is loaded (DEDICATED SERVER).
-	 */
+	/** The map that will be loaded by default when no other map is loaded (DEDICATED SERVER). */
 	UPROPERTY(config, EditAnywhere, Category=DefaultMaps, AdvancedDisplay)
 	FString ServerDefaultMap;
 
-	/**
-	 * GameMode to use if not specified in any other way. (e.g. per-map DefaultGameMode or on the URL).
-	 */
+	/** GameMode to use if not specified in any other way. (e.g. per-map DefaultGameMode or on the URL). */
 	UPROPERTY(config, noclear, EditAnywhere, Category=DefaultModes, meta=(MetaClass="GameMode", DisplayName="Default GameMode"))
 	FStringClassReference GlobalDefaultGameMode;
 
diff --git a/Engine/Source/Runtime/EngineSettings/Classes/GameNetworkManagerSettings.h b/Engine/Source/Runtime/EngineSettings/Classes/GameNetworkManagerSettings.h
--- a/Engine/Source/Runtime/EngineSettings/Classes/GameNetworkManagerSettings.h
+++ b/Engine/Source/Runtime/EngineSettings/Classes/GameNetworkManagerSettings.h
@@ -1,9 +1,5 @@
 // Copyright 1998-2014 Epic Games, Inc. All Rights Reserved.
 
-/*=============================================================================
-	GameNetworkManagerSettings.h: Declares the UGameNetworkManagerSettings class.
-=============================================================================*/
-
 #pragma once
 
 #include "GameNetworkManagerSettings.generated.h"
@@ -18,69 +14,47 @@ class ENGINESETTINGS_API UGameNetworkManagerSettings
 {
 	GENERATED_UCLASS_BODY()
 
-	/**
-	 * Minimum bandwidth dynamically set per connection.
-	 */
+	/** Minimum bandwidth dynamically set per connection. */
 	UPROPERTY(globalconfig, EditAnywhere, Category=Bandwidth)
 	int32 MinDynamicBandwidth;
 
-	/**
-	 * Maximum bandwidth dynamically set per connection.
-	 */
+	/** Maximum bandwidth dynamically set per connection. */
 	UPROPERTY(globalconfig, EditAnywhere, Category=Bandwidth)
 	int32 MaxDynamicBandwidth;
 
-	/**
-	 * Total available bandwidth for listen server, split dynamically across net connections.
-	 */
+	/** Total available bandwidth for listen server, split dynamically across net connections. */
 	UPROPERTY(globalconfig, EditAnywhere, Category=Bandwidth)
 	int32 TotalNetBandwidth;
 
-	/**
-	 * The point we determine the server is either delaying packets or has bad upstream.
-	 */
+	/** The point we determine the server is either delaying packets or has bad upstream. */
 	UPROPERTY(config, EditAnywhere, Category=StandbyCheats)
 	int32 BadPingThreshold;
 
-	/**
-	 * Used to determine if checking for standby cheats should occur.
-	 */
+	/** Used to determine if checking for standby cheats should occur. */
 	UPROPERTY(config, EditAnywhere, Category=StandbyCheats)
 	uint32 bIsStandbyCheckingEnabled:1;
 
-	/**
-	 * The amount of time without packets before triggering the cheat code.
-	 */
+	/**  The amount of time without packets before triggering the cheat code. */
 	UPROPERTY(config, EditAnywhere, Category=StandbyCheats)
 	float StandbyRxCheatTime;
 
-	/**
-	 * The amount of time without packets before triggering the cheat code.
-	 */
+	/** The amount of time without packets before triggering the cheat code. */
 	UPROPERTY(config, EditAnywhere, Category=StandbyCheats)
 	float StandbyTxCheatTime;
 
-	/**
-	 * The percentage of clients missing RX data before triggering the standby code.
-	 */
+	/** The percentage of clients missing RX data before triggering the standby code. */
 	UPROPERTY(config, EditAnywhere, Category=StandbyCheats)
 	float PercentMissingForRxStandby;
 
-	/**
-	 * The percentage of clients missing TX data before triggering the standby code.
-	 */
+	/** The percentage of clients missing TX data before triggering the standby code. */
 	UPROPERTY(config, EditAnywhere, Category=StandbyCheats)
 	float PercentMissingForTxStandby;
 
-	/**
-	 * The percentage of clients with bad ping before triggering the standby code.
-	 */
+	/** The percentage of clients with bad ping before triggering the standby code. */
 	UPROPERTY(config, EditAnywhere, Category=StandbyCheats)
 	float PercentForBadPing;
 
-	/**
-	 * The amount of time to wait before checking a connection for standby issues.
-	 */
+	/** The amount of time to wait before checking a connection for standby issues. */
 	UPROPERTY(config, EditAnywhere, Category=StandbyCheats)
 	float JoinInProgressStandbyWaitTime;
 };
diff --git a/Engine/Source/Runtime/EngineSettings/Classes/GameSessionSettings.h b/Engine/Source/Runtime/EngineSettings/Classes/GameSessionSettings.h
--- a/Engine/Source/Runtime/EngineSettings/Classes/GameSessionSettings.h
+++ b/Engine/Source/Runtime/EngineSettings/Classes/GameSessionSettings.h
@@ -1,9 +1,5 @@
 // Copyright 1998-2014 Epic Games, Inc. All Rights Reserved.
 
-/*=============================================================================
-	GameSessionSettings.h: Declares the UGameSessionSettings class.
-=============================================================================*/
-
 #pragma once
 
 #include "GameSessionSettings.generated.h"
@@ -15,21 +11,15 @@ class ENGINESETTINGS_API UGameSessionSettings
 {
 	GENERATED_UCLASS_BODY()
 
-	/**
-	 * Maximum number of spectators allowed by this server.
-	 */
+	/** Maximum number of spectators allowed by this server. */
 	UPROPERTY(globalconfig, EditAnywhere, Category=GameSessionSettings)
 	int32 MaxSpectators;
 
-	/**
-	 * Maximum number of players allowed by this server.
-	 */
+	/** Maximum number of players allowed by this server. */
 	UPROPERTY(globalconfig, EditAnywhere, Category=GameSessionSettings)
 	int32 MaxPlayers;
 
-    /**
-	 * Is voice enabled always or via a push to talk key binding.
-	 */
+    /** Is voice enabled always or via a push to talk key binding. */
 	UPROPERTY(globalconfig, EditAnywhere, Category=GameSessionSettings)
 	uint32 bRequiresPushToTalk:1;
 };
diff --git a/Engine/Source/Runtime/EngineSettings/Classes/GeneralEngineSettings.h b/Engine/Source/Runtime/EngineSettings/Classes/GeneralEngineSettings.h
--- a/Engine/Source/Runtime/EngineSettings/Classes/GeneralEngineSettings.h
+++ b/Engine/Source/Runtime/EngineSettings/Classes/GeneralEngineSettings.h
@@ -1,9 +1,5 @@
 // Copyright 1998-2014 Epic Games, Inc. All Rights Reserved.
 
-/*=============================================================================
-	GeneralEngineSettings.h: Declares the UGeneralEngineSettings class.
-=============================================================================*/
-
 #pragma once
 
 #include "GeneralEngineSettings.generated.h"
diff --git a/Engine/Source/Runtime/EngineSettings/Classes/GeneralProjectSettings.h b/Engine/Source/Runtime/EngineSettings/Classes/GeneralProjectSettings.h
--- a/Engine/Source/Runtime/EngineSettings/Classes/GeneralProjectSettings.h
+++ b/Engine/Source/Runtime/EngineSettings/Classes/GeneralProjectSettings.h
@@ -1,9 +1,5 @@
 // Copyright 1998-2014 Epic Games, Inc. All Rights Reserved.
 
-/*=============================================================================
-	GeneralProjectSettings.h: Declares the UGeneralProjectSettings class.
-=============================================================================*/
-
 #pragma once
 
 #include "GeneralProjectSettings.generated.h"
@@ -15,63 +11,43 @@ class ENGINESETTINGS_API UGeneralProjectSettings
 {
 	GENERATED_UCLASS_BODY()
 
-	/**
-	 * The name of the company (author, provider) that created the project.
-	 */
+	/** The name of the company (author, provider) that created the project. */
 	UPROPERTY(config, EditAnywhere, Category=Publisher)
 	FString CompanyName;
 
-	/**
-	 * The project's copyright and/or trademark notices.
-	 */
+	/** The project's copyright and/or trademark notices. */
 	UPROPERTY(config, EditAnywhere, Category=Legal)
 	FString CopyrightNotice;
 
-	/**
-	 * The project's description text.
-	 */
+	/** The project's description text. */
 	UPROPERTY(config, EditAnywhere, Category=About)
 	FString Description;
 
-	/**
-	 * The project's homepage URL.
-	 */
+	/** The project's homepage URL. */
 	UPROPERTY(config, EditAnywhere, Category=Publisher)
 	FString Homepage;
 
-	/**
-	 * The project's licensing terms.
-	 */
+	/** The project's licensing terms. */
 	UPROPERTY(config, EditAnywhere, Category=Legal)
 	FString LicensingTerms;
 
-	/**
-	 * The project's privacy policy.
-	 */
+	/** The project's privacy policy. */
 	UPROPERTY(config, EditAnywhere, Category=Legal)
 	FString PrivacyPolicy;
 
-	/**
-	 * The project's unique identifier.
-	 */
+	/** The project's unique identifier. */
 	UPROPERTY(config, EditAnywhere, Category=About)
 	FGuid ProjectID;
 
-	/**
-	 * The project's name.
-	 */
+	/** The project's name. */
 	UPROPERTY(config, EditAnywhere, Category=About)
 	FString ProjectName;
 
-	/**
-	 * The project's version number.
-	 */
+	/** The project's version number. */
 	UPROPERTY(config, EditAnywhere, Category=About)
 	FString ProjectVersion;
 
-	/**
-	 * The project's support contact information.
-	 */
+	/** The project's support contact information. */
 	UPROPERTY(config, EditAnywhere, Category=Publisher)
 	FString SupportContact;
 };
diff --git a/Engine/Source/Runtime/EngineSettings/Classes/HudSettings.h b/Engine/Source/Runtime/EngineSettings/Classes/HudSettings.h
--- a/Engine/Source/Runtime/EngineSettings/Classes/HudSettings.h
+++ b/Engine/Source/Runtime/EngineSettings/Classes/HudSettings.h
@@ -1,13 +1,10 @@
 // Copyright 1998-2014 Epic Games, Inc. All Rights Reserved.
 
-/*=============================================================================
-	HudSettings.h: Declares the UHudSettings class.
-=============================================================================*/
-
 #pragma once
 
 #include "HudSettings.generated.h"
 
+
 UCLASS(config=Game)
 class ENGINESETTINGS_API UHudSettings
 	: public UObject
diff --git a/Engine/Source/Runtime/EngineSettings/Private/EngineSettingsModule.cpp b/Engine/Source/Runtime/EngineSettings/Private/EngineSettingsModule.cpp
--- a/Engine/Source/Runtime/EngineSettings/Private/EngineSettingsModule.cpp
+++ b/Engine/Source/Runtime/EngineSettings/Private/EngineSettingsModule.cpp
@@ -1,9 +1,5 @@
 // Copyright 1998-2014 Epic Games, Inc. All Rights Reserved.
 
-/*=============================================================================
-	EngineSettingsModule.cpp: Implements the FEngineSettingsModule class.
-=============================================================================*/
-
 #include "EngineSettingsPrivatePCH.h"
 
 
@@ -15,18 +11,15 @@ class FEngineSettingsModule
 {
 public:
 
-	// Begin IModuleInterface interface
+	// IModuleInterface interface
 
 	virtual void StartupModule( ) OVERRIDE { }
-
 	virtual void ShutdownModule( ) OVERRIDE { }
 
 	virtual bool SupportsDynamicReloading( ) OVERRIDE
 	{
 		return true;
 	}
-
-	// End IModuleInterface interface
 };
 
 
diff --git a/Engine/Source/Runtime/EngineSettings/Private/EngineSettingsPrivatePCH.h b/Engine/Source/Runtime/EngineSettings/Private/EngineSettingsPrivatePCH.h
--- a/Engine/Source/Runtime/EngineSettings/Private/EngineSettingsPrivatePCH.h
+++ b/Engine/Source/Runtime/EngineSettings/Private/EngineSettingsPrivatePCH.h
@@ -1,24 +1,11 @@
 // Copyright 1998-2014 Epic Games, Inc. All Rights Reserved.
 
-/*=============================================================================
-	EngineSettingsPrivatePCH.h: Pre-compiled header file for the EngineSettings module.
-=============================================================================*/
+#pragma once
 
-#ifndef ENGINESETTINGS_PRIVATEPCH_H
-#define ENGINESETTINGS_PRIVATEPCH_H
+#include "EngineSettings.h"
 
 
-#include "../Public/EngineSettings.h"
-
-
-/* Dependencies
+/* Private dependencies
  *****************************************************************************/
 
 #include "ModuleManager.h"
-
-
-/* Private includes
- *****************************************************************************/
-
-
-#endif
diff --git a/Engine/Source/Runtime/EngineSettings/Public/EngineSettings.h b/Engine/Source/Runtime/EngineSettings/Public/EngineSettings.h
--- a/Engine/Source/Runtime/EngineSettings/Public/EngineSettings.h
+++ b/Engine/Source/Runtime/EngineSettings/Public/EngineSettings.h
@@ -1,19 +1,15 @@
 // Copyright 1998-2014 Epic Games, Inc. All Rights Reserved.
 
-/*=============================================================================
-	EngineSettings.h: EngineSettings module public header file.
-=============================================================================*/
-
 #pragma once
 
 
-/* Dependencies
+/* Public Dependencies
  *****************************************************************************/
 
 #include "CoreUObject.h"
 
 
-/* Messages
+/* Public Includes
  *****************************************************************************/
 
 #include "ConsoleSettings.h"
diff --git a/Engine/Source/Runtime/LauncherAutomatedService/Private/AutomatedTestManager.cpp b/Engine/Source/Runtime/LauncherAutomatedService/Private/AutomatedTestManager.cpp
--- a/Engine/Source/Runtime/LauncherAutomatedService/Private/AutomatedTestManager.cpp
+++ b/Engine/Source/Runtime/LauncherAutomatedService/Private/AutomatedTestManager.cpp
@@ -1,9 +1,5 @@
 // Copyright 1998-2014 Epic Games, Inc. All Rights Reserved.
 
-/*=============================================================================
-	AutomatedTestManager.cpp: Implements the FAutomatedTestManager class.
-=============================================================================*/
-
 #include "LauncherAutomatedServicePrivatePCH.h"
 
 
diff --git a/Engine/Source/Runtime/LauncherAutomatedService/Private/AutomatedTestManager.h b/Engine/Source/Runtime/LauncherAutomatedService/Private/AutomatedTestManager.h
--- a/Engine/Source/Runtime/LauncherAutomatedService/Private/AutomatedTestManager.h
+++ b/Engine/Source/Runtime/LauncherAutomatedService/Private/AutomatedTestManager.h
@@ -1,11 +1,8 @@
 // Copyright 1998-2014 Epic Games, Inc. All Rights Reserved.
 
-/*=============================================================================
-	FAutomatedTestManager: Declares the FAutomatedTestManager class.
-=============================================================================*/
-
 #pragma once
 
+
 /** States for running the automation process */
 namespace EAutomationTestState
 {
@@ -20,6 +17,7 @@ namespace EAutomationTestState
 	};
 };
 
+
 /**
  * Implements a helper class that runs the automation tests
  */
@@ -27,10 +25,11 @@ class FAutomatedTestManager
 	: public IAutomatedTestManager
 {
 public:
+
 	/**
 	 * Constructor.
 	 *
-	 *@Param SessionID - the session ID used to select the session in the automation controller
+	 * @param SessionID The session ID used to select the session in the automation controller.
 	 */
 	FAutomatedTestManager( const FGuid SessionID );
 
@@ -39,40 +38,37 @@ public:
 	 */
 	~FAutomatedTestManager();
 
-	// Begin IAutomatedTestManager interface
+public:
 
-	virtual bool IsTestingComplete() OVERRIDE;
+	// IAutomatedTestManager interface
 
+	virtual bool IsTestingComplete() OVERRIDE;
 	virtual void RunTests() OVERRIDE;
-
 	virtual void Tick( float DeltaTime ) OVERRIDE;
 
-
-	// End IAutomatedTestManager interface
-
-
 protected:
 
 	/*
-	 * Find some automation workers on the discovered game instances
+	 * Finds some automation workers on the discovered game instances.
 	 *
-	 * @param DeltaTime - Time since last update. Used to add a delay before we send the request
+	 * @param DeltaTime Time since last update. Used to add a delay before we send the request.
 	 */
 	void FindWorkers( float DeltaTime );
 
 	/*
-	 * Generate the automation report, and finish the testing process
+	 * Generates the automation report, and finish the testing process.
 	 */
 	void GenerateReport();
 
 	/*
-	 * Monitor the testing process to see if we have finished
+	 * Monitors the testing process to see if we have finished.
 	 */
 	void MonitorTests();
 
 private:
+
 	/*
-	 * Callback function from the automation controller to let us know there are tests available
+	 * Callback function from the automation controller to let us know there are tests available.
 	 */
 	void HandleRefreshTestCallback();
 
@@ -92,4 +88,4 @@ private:
 
 	/** Session ID to run the tests on */
 	FGuid SessionID;
-};
\ No newline at end of file
+};
diff --git a/Engine/Source/Runtime/LauncherAutomatedService/Private/LauncherAutomatedServiceModule.cpp b/Engine/Source/Runtime/LauncherAutomatedService/Private/LauncherAutomatedServiceModule.cpp
--- a/Engine/Source/Runtime/LauncherAutomatedService/Private/LauncherAutomatedServiceModule.cpp
+++ b/Engine/Source/Runtime/LauncherAutomatedService/Private/LauncherAutomatedServiceModule.cpp
@@ -1,16 +1,13 @@
 // Copyright 1998-2014 Epic Games, Inc. All Rights Reserved.
 
-/*=============================================================================
-	ILauncherServicesModule.h: Declares the ILauncherServicesModule interface.
-=============================================================================*/
-
 #include "LauncherAutomatedServicePrivatePCH.h"
-
 #include "ModuleManager.h"
 
+
 IMPLEMENT_MODULE(FLauncherAutomatedServiceModule, LauncherAutomatedService);
 
+
 ILauncherAutomatedServiceProviderPtr FLauncherAutomatedServiceModule::CreateAutomatedServiceProvider()
 {
 	return MakeShareable(new FLauncherAutomatedServiceProvider());
-}
\ No newline at end of file
+}
diff --git a/Engine/Source/Runtime/LauncherAutomatedService/Private/LauncherAutomatedServiceModule.h b/Engine/Source/Runtime/LauncherAutomatedService/Private/LauncherAutomatedServiceModule.h
--- a/Engine/Source/Runtime/LauncherAutomatedService/Private/LauncherAutomatedServiceModule.h
+++ b/Engine/Source/Runtime/LauncherAutomatedService/Private/LauncherAutomatedServiceModule.h
@@ -1,13 +1,8 @@
 // Copyright 1998-2014 Epic Games, Inc. All Rights Reserved.
 
-/*=============================================================================
-	LauncherAutomatedServiceModule.h: Declares the FLauncherAutomatedServiceModule class.
-=============================================================================*/
-
 #pragma once
 
 
-
 /**
  * Implements the launchers automated service.
  */
@@ -15,11 +10,10 @@ class FLauncherAutomatedServiceModule
 	: public ILauncherAutomatedServiceModule
 {
 public:
-	// Begin ILauncherAutomatedServiceModule interface
 
-	virtual ILauncherAutomatedServiceProviderPtr CreateAutomatedServiceProvider() OVERRIDE;
+	// ILauncherAutomatedServiceModule interface
 
-	// End ILauncherAutomatedServiceModule interface
+	virtual ILauncherAutomatedServiceProviderPtr CreateAutomatedServiceProvider() OVERRIDE;
 
 public:
 
@@ -27,4 +21,4 @@ public:
 	 * Virtual destructor.
 	 */
 	virtual ~FLauncherAutomatedServiceModule( ) { }
-};
\ No newline at end of file
+};
diff --git a/Engine/Source/Runtime/LauncherAutomatedService/Private/LauncherAutomatedServicePrivatePCH.h b/Engine/Source/Runtime/LauncherAutomatedService/Private/LauncherAutomatedServicePrivatePCH.h
--- a/Engine/Source/Runtime/LauncherAutomatedService/Private/LauncherAutomatedServicePrivatePCH.h
+++ b/Engine/Source/Runtime/LauncherAutomatedService/Private/LauncherAutomatedServicePrivatePCH.h
@@ -1,10 +1,6 @@
 // Copyright 1998-2014 Epic Games, Inc. All Rights Reserved.
 
-/*=============================================================================
-	LauncherAutomatedServicePrivatePCH.h: Pre-compiled header file for the LauncherServices module.
-=============================================================================*/
-
-#include "../Public/LauncherAutomatedService.h"
+#include "LauncherAutomatedService.h"
 
 #include "LauncherServices.h"
 #include "TargetDeviceServices.h"
@@ -15,4 +11,4 @@
 
 #include "AutomatedTestManager.h"
 #include "LauncherAutomatedServiceProvider.h"
-#include "LauncherAutomatedServiceModule.h"
\ No newline at end of file
+#include "LauncherAutomatedServiceModule.h"
diff --git a/Engine/Source/Runtime/LauncherAutomatedService/Private/LauncherAutomatedServiceProvider.cpp b/Engine/Source/Runtime/LauncherAutomatedService/Private/LauncherAutomatedServiceProvider.cpp
--- a/Engine/Source/Runtime/LauncherAutomatedService/Private/LauncherAutomatedServiceProvider.cpp
+++ b/Engine/Source/Runtime/LauncherAutomatedService/Private/LauncherAutomatedServiceProvider.cpp
@@ -1,9 +1,5 @@
 // Copyright 1998-2014 Epic Games, Inc. All Rights Reserved.
 
-/*=============================================================================
-	LauncherAutomatedServiceProvider.cpp: Implements the FLauncherAutomatedServiceProvider class.
-=============================================================================*/
-
 #include "LauncherAutomatedServicePrivatePCH.h"
 
 DEFINE_LOG_CATEGORY_STATIC(LauncherAutomatedService, Log, All);
diff --git a/Engine/Source/Runtime/LauncherAutomatedService/Private/LauncherAutomatedServiceProvider.h b/Engine/Source/Runtime/LauncherAutomatedService/Private/LauncherAutomatedServiceProvider.h
--- a/Engine/Source/Runtime/LauncherAutomatedService/Private/LauncherAutomatedServiceProvider.h
+++ b/Engine/Source/Runtime/LauncherAutomatedService/Private/LauncherAutomatedServiceProvider.h
@@ -1,9 +1,5 @@
 // Copyright 1998-2014 Epic Games, Inc. All Rights Reserved.
 
-/*=============================================================================
-	LauncherAutomatedServiceProvider.h: Declares the FLauncherAutomatedServiceProvider class.
-=============================================================================*/
-
 #pragma once
 
 
@@ -20,10 +16,9 @@ public:
 	 */
 	FLauncherAutomatedServiceProvider( );
 
-
 public:
 
-	// Begin ILauncherAutomatedServiceProvider interface
+	// ILauncherAutomatedServiceProvider interface
 
 	virtual int32 GetExitCode( ) OVERRIDE
 	{
@@ -36,20 +31,15 @@ public:
 	}
 
 	virtual void Setup( const TCHAR* Params ) OVERRIDE;
-
 	virtual void Shutdown( ) OVERRIDE;
-
 	virtual void Tick( float DeltaTime ) OVERRIDE;
 
-	// End ILauncherAutomatedServiceProvider interface
-
-
 protected:
 
 	/*
-	 * Sets up the profile manager responsible for deploying the session
+	 * Sets up the profile manager responsible for deploying the session.
 	 *
-	 * @param Params - The settings required to determine the profile/device group to use.
+	 * @param Params The settings required to determine the profile/device group to use.
 	 */
 	void SetupProfileAndGroupSettings( const TCHAR* Params );
 
@@ -58,17 +48,15 @@ protected:
 	 */
 	void StartAutomationTests( );
 
-
 private:
 
 	/*
-	 * Handles devices being added to the proxy manager
+	 * Callback for devices being added to the proxy manager.
 	 *
-	 * @param AddedProxy - The newly added device proxy to the device proxy manager
+	 * @param AddedProxy The newly added device proxy to the device proxy manager.
 	 */
 	void HandleDeviceProxyManagerProxyAdded( const ITargetDeviceProxyRef& AddedProxy );
 
-
 private:
 
 	/** The device group which holds device ids for the session we are about to launch */
diff --git a/Engine/Source/Runtime/LauncherAutomatedService/Public/Interfaces/IAutomatedTestManager.h b/Engine/Source/Runtime/LauncherAutomatedService/Public/Interfaces/IAutomatedTestManager.h
--- a/Engine/Source/Runtime/LauncherAutomatedService/Public/Interfaces/IAutomatedTestManager.h
+++ b/Engine/Source/Runtime/LauncherAutomatedService/Public/Interfaces/IAutomatedTestManager.h
@@ -1,20 +1,12 @@
 // Copyright 1998-2014 Epic Games, Inc. All Rights Reserved.
 
-/*=============================================================================
-	IAutomatedTestManager.h: Declares the IAutomatedTestManager interface.
-=============================================================================*/
-
 #pragma once
 
 
-/**
- * Type definition for shared pointers to instances of IAutomatedTestManager.
- */
+/** Type definition for shared pointers to instances of IAutomatedTestManager. */
 typedef TSharedPtr<class IAutomatedTestManager> IAutomatedTestManagerPtr;
 
-/**
- * Type definition for shared references to instances of IAutomatedTestManager.
- */
+/** Type definition for shared references to instances of IAutomatedTestManager. */
 typedef TSharedRef<class IAutomatedTestManager> IAutomatedTestManagerRef;
 
 
@@ -27,21 +19,21 @@ class IAutomatedTestManager
 public:
 
 	/*
-	 * Check to see if we have finished running the automation tests
+	 * Checks whether we have finished running the automation tests.
 	 *
 	 * @return True if we have finished
 	 */
 	virtual bool IsTestingComplete() = 0;
 
 	/*
-	 * Start the testing process. Called after we have discovered some devices
+	 * Starts the testing process. Called after we have discovered some devices.
 	 */
 	virtual void RunTests() = 0;
 
 	/*
-	 * Tick function to update automation tests
+	 * Ticks the automation tests.
 	 *
-	 * @param DeltaTime - Time since last update
+	 * @param DeltaTime Time since last update.
 	 */
 	virtual void Tick( float DeltaTime ) = 0;
 
@@ -51,4 +43,4 @@ public:
 	 * Virtual destructor.
 	 */
 	virtual ~IAutomatedTestManager( ) { }
-};
\ No newline at end of file
+};
diff --git a/Engine/Source/Runtime/LauncherAutomatedService/Public/Interfaces/ILauncherAutomatedServiceModule.h b/Engine/Source/Runtime/LauncherAutomatedService/Public/Interfaces/ILauncherAutomatedServiceModule.h
--- a/Engine/Source/Runtime/LauncherAutomatedService/Public/Interfaces/ILauncherAutomatedServiceModule.h
+++ b/Engine/Source/Runtime/LauncherAutomatedService/Public/Interfaces/ILauncherAutomatedServiceModule.h
@@ -1,9 +1,5 @@
 // Copyright 1998-2014 Epic Games, Inc. All Rights Reserved.
 
-/*=============================================================================
-	ILauncherAutomatedServiceModule.h: Declares the ILauncherAutomatedServiceModule interface.
-=============================================================================*/
-
 #pragma once
 
 
@@ -28,4 +24,4 @@ public:
 	 * Virtual destructor.
 	 */
 	virtual ~ILauncherAutomatedServiceModule( ) { }
-};
\ No newline at end of file
+};
diff --git a/Engine/Source/Runtime/LauncherAutomatedService/Public/Interfaces/ILauncherAutomatedServiceProvider.h b/Engine/Source/Runtime/LauncherAutomatedService/Public/Interfaces/ILauncherAutomatedServiceProvider.h
--- a/Engine/Source/Runtime/LauncherAutomatedService/Public/Interfaces/ILauncherAutomatedServiceProvider.h
+++ b/Engine/Source/Runtime/LauncherAutomatedService/Public/Interfaces/ILauncherAutomatedServiceProvider.h
@@ -1,20 +1,12 @@
 // Copyright 1998-2014 Epic Games, Inc. All Rights Reserved.
 
-/*=============================================================================
-	ILauncherAutomatedServiceProvider.h: Declares the ILauncherAutomatedServiceProvider interface.
-=============================================================================*/
-
 #pragma once
 
 
-/**
- * Type definition for shared pointers to instances of ILauncherAutomatedServiceProvider.
- */
+/** Type definition for shared pointers to instances of ILauncherAutomatedServiceProvider. */
 typedef TSharedPtr<class ILauncherAutomatedServiceProvider> ILauncherAutomatedServiceProviderPtr;
 
-/**
- * Type definition for shared references to instances of ILauncherAutomatedServiceProvider.
- */
+/** Type definition for shared references to instances of ILauncherAutomatedServiceProvider. */
 typedef TSharedRef<class ILauncherAutomatedServiceProvider> ILauncherAutomatedServiceProviderRef;
 
 
@@ -42,7 +34,7 @@ public:
 	/**
 	 * Setup the service pre-requisites
 	 *
-	 * @param Params - The settings to use for the service, I.e. Game/Config to deploy etc
+	 * @param Params The settings to use for the service, I.e. Game/Config to deploy etc
 	 */
 	virtual void Setup( const TCHAR* Params ) = 0;
 	
@@ -54,13 +46,14 @@ public:
 	/**
 	 * Tick the service so it can continue with its work.
 	 *
-	 * @param DeltaTime - Time since last tick
+	 * @param DeltaTime Time since last tick
 	 */
 	virtual void Tick( float DeltaTime ) = 0;
+
 public:
 
 	/**
 	 * Virtual destructor.
 	 */
 	virtual ~ILauncherAutomatedServiceProvider( ) { }
-};
\ No newline at end of file
+};
diff --git a/Engine/Source/Runtime/LauncherAutomatedService/Public/LauncherAutomatedService.h b/Engine/Source/Runtime/LauncherAutomatedService/Public/LauncherAutomatedService.h
--- a/Engine/Source/Runtime/LauncherAutomatedService/Public/LauncherAutomatedService.h
+++ b/Engine/Source/Runtime/LauncherAutomatedService/Public/LauncherAutomatedService.h
@@ -1,24 +1,18 @@
 // Copyright 1998-2014 Epic Games, Inc. All Rights Reserved.
 
-/*=============================================================================
-	LauncherAutomatedService.h: LauncherAutomatedService module public header file.
-=============================================================================*/
-
 #pragma once
 
 
-/* Dependencies
+/* Public Dependencies
  *****************************************************************************/
 
 #include "Core.h"
 #include "ModuleInterface.h"
 
 
-/* Interfaces
+/* Public Includes
  *****************************************************************************/
 
 #include "IAutomatedTestManager.h"
-
 #include "ILauncherAutomatedServiceProvider.h"
-
-#include "ILauncherAutomatedServiceModule.h"
\ No newline at end of file
+#include "ILauncherAutomatedServiceModule.h"
diff --git a/Engine/Source/Runtime/Messaging/Private/Bridge/MessageAddressBook.h b/Engine/Source/Runtime/Messaging/Private/Bridge/MessageAddressBook.h
--- a/Engine/Source/Runtime/Messaging/Private/Bridge/MessageAddressBook.h
+++ b/Engine/Source/Runtime/Messaging/Private/Bridge/MessageAddressBook.h
@@ -1,9 +1,5 @@
 // Copyright 1998-2014 Epic Games, Inc. All Rights Reserved.
 
-/*=============================================================================
-	MessageAddressBook.h: Declares the FMessageAddressBook class.
-=============================================================================*/
-
 #pragma once
 
 
diff --git a/Engine/Source/Runtime/Messaging/Private/Bridge/MessageBridge.cpp b/Engine/Source/Runtime/Messaging/Private/Bridge/MessageBridge.cpp
--- a/Engine/Source/Runtime/Messaging/Private/Bridge/MessageBridge.cpp
+++ b/Engine/Source/Runtime/Messaging/Private/Bridge/MessageBridge.cpp
@@ -1,9 +1,5 @@
 // Copyright 1998-2014 Epic Games, Inc. All Rights Reserved.
 
-/*=============================================================================
-	MessageBridge.cpp: Implements the FMessageBridge class.
-=============================================================================*/
-
 #include "MessagingPrivatePCH.h"
 
 
diff --git a/Engine/Source/Runtime/Messaging/Private/Bridge/MessageBridge.h b/Engine/Source/Runtime/Messaging/Private/Bridge/MessageBridge.h
--- a/Engine/Source/Runtime/Messaging/Private/Bridge/MessageBridge.h
+++ b/Engine/Source/Runtime/Messaging/Private/Bridge/MessageBridge.h
@@ -1,9 +1,5 @@
 // Copyright 1998-2014 Epic Games, Inc. All Rights Reserved.
 
-/*=============================================================================
-	MessageBridge.h: Declares the FMessageBridge class.
-=============================================================================*/
-
 #pragma once
 
 
@@ -35,10 +31,9 @@ public:
 
 public:
 
-	// Begin IMessageBridge interface
+	// IMessageBridge interface
 
 	virtual void Disable( ) override;
-
 	virtual void Enable( ) override;
 
 	virtual bool IsEnabled( ) const override
@@ -46,11 +41,9 @@ public:
 		return Enabled;
 	}
 
-	// End IMessageBridge interface
-
 public:
 
-	// Begin IReceiveMessages interface
+	// IReceiveMessages interface
 
 	virtual FName GetDebugName( ) const override
 	{
@@ -74,11 +67,9 @@ public:
 
 	virtual void ReceiveMessage( const IMessageContextRef& Context ) override;
 
-	// End IReceiveMessages interface
-
 public:
 
-	// Begin ISendMessages interface
+	// ISendMessages interface
 
 	virtual FMessageAddress GetSenderAddress( ) override
 	{
@@ -87,8 +78,6 @@ public:
 
 	virtual void NotifyMessageError( const IMessageContextRef& Context, const FString& Error ) override;
 
-	// End ISendMessages interface
-
 protected:
 
 	/**
diff --git a/Engine/Source/Runtime/Messaging/Private/Bridge/MessageData.h b/Engine/Source/Runtime/Messaging/Private/Bridge/MessageData.h
--- a/Engine/Source/Runtime/Messaging/Private/Bridge/MessageData.h
+++ b/Engine/Source/Runtime/Messaging/Private/Bridge/MessageData.h
@@ -1,9 +1,5 @@
 // Copyright 1998-2014 Epic Games, Inc. All Rights Reserved.
 
-/*=============================================================================
-	MessageData.h: Declares the FMessageData class.
-=============================================================================*/
-
 #pragma once
 
 
@@ -46,7 +42,7 @@ public:
 
 public:
 
-	// Begin IMessageAttachment interface
+	// IMessageAttachment interface
 
 	virtual FArchive* CreateReader( ) override
 	{
@@ -63,8 +59,6 @@ public:
 		return StateChangedDelegate;
 	}
 
-	// End IMessageAttachment interface
-
 private:
 
 	// Holds the data.
diff --git a/Engine/Source/Runtime/Messaging/Private/Bridge/MessageDeserializeTask.h b/Engine/Source/Runtime/Messaging/Private/Bridge/MessageDeserializeTask.h
--- a/Engine/Source/Runtime/Messaging/Private/Bridge/MessageDeserializeTask.h
+++ b/Engine/Source/Runtime/Messaging/Private/Bridge/MessageDeserializeTask.h
@@ -1,9 +1,5 @@
 // Copyright 1998-2014 Epic Games, Inc. All Rights Reserved.
 
-/*=============================================================================
-	MessageDeserializeTask.h: Declares the FMessageDeserializeTask class.
-=============================================================================*/
-
 #pragma once
 
 
diff --git a/Engine/Source/Runtime/Messaging/Private/Bridge/MessageForwardTask.h b/Engine/Source/Runtime/Messaging/Private/Bridge/MessageForwardTask.h
--- a/Engine/Source/Runtime/Messaging/Private/Bridge/MessageForwardTask.h
+++ b/Engine/Source/Runtime/Messaging/Private/Bridge/MessageForwardTask.h
@@ -1,9 +1,5 @@
 // Copyright 1998-2014 Epic Games, Inc. All Rights Reserved.
 
-/*=============================================================================
-	MessageForwardTask.h: Declares the FMessageForwardTask class.
-=============================================================================*/
-
 #pragma once
 
 
diff --git a/Engine/Source/Runtime/Messaging/Private/Bridge/MessageSerializeTask.h b/Engine/Source/Runtime/Messaging/Private/Bridge/MessageSerializeTask.h
--- a/Engine/Source/Runtime/Messaging/Private/Bridge/MessageSerializeTask.h
+++ b/Engine/Source/Runtime/Messaging/Private/Bridge/MessageSerializeTask.h
@@ -1,9 +1,5 @@
 // Copyright 1998-2014 Epic Games, Inc. All Rights Reserved.
 
-/*=============================================================================
-	MessageSerializeTask.h: Declares the FMessageSerializeTask class.
-=============================================================================*/
-
 #pragma once
 
 
diff --git a/Engine/Source/Runtime/Messaging/Private/Bus/MessageBus.cpp b/Engine/Source/Runtime/Messaging/Private/Bus/MessageBus.cpp
--- a/Engine/Source/Runtime/Messaging/Private/Bus/MessageBus.cpp
+++ b/Engine/Source/Runtime/Messaging/Private/Bus/MessageBus.cpp
@@ -1,9 +1,5 @@
 // Copyright 1998-2014 Epic Games, Inc. All Rights Reserved.
 
-/*=============================================================================
-	MessageBus.cpp: Implements the FMessageBus class.
-=============================================================================*/
-
 #include "MessagingPrivatePCH.h"
 
 
diff --git a/Engine/Source/Runtime/Messaging/Private/Bus/MessageBus.h b/Engine/Source/Runtime/Messaging/Private/Bus/MessageBus.h
--- a/Engine/Source/Runtime/Messaging/Private/Bus/MessageBus.h
+++ b/Engine/Source/Runtime/Messaging/Private/Bus/MessageBus.h
@@ -1,9 +1,5 @@
 // Copyright 1998-2014 Epic Games, Inc. All Rights Reserved.
 
-/*=============================================================================
-	MessageBus.h: Declares the FMessageBus class.
-=============================================================================*/
-
 #pragma once
 
 
@@ -30,7 +26,7 @@ public:
 
 public:
 
-	// Begin IMessageBus interface
+	// IMessageBus interface
 
 	virtual void Forward( const IMessageContextRef& Context, const TArray<FMessageAddress>& Recipients, EMessageScope::Type ForwardingScope, const FTimespan& Delay, const ISendMessagesRef& Forwarder ) override
 	{
@@ -118,8 +114,6 @@ public:
 		}
 	}
 
-	// End IMessageBus interface
-
 private:
 
 	// Holds the message router.
diff --git a/Engine/Source/Runtime/Messaging/Private/Bus/MessageContext.h b/Engine/Source/Runtime/Messaging/Private/Bus/MessageContext.h
--- a/Engine/Source/Runtime/Messaging/Private/Bus/MessageContext.h
+++ b/Engine/Source/Runtime/Messaging/Private/Bus/MessageContext.h
@@ -1,9 +1,5 @@
 // Copyright 1998-2014 Epic Games, Inc. All Rights Reserved.
 
-/*=============================================================================
-	MessageContext.h: Declares the FMessageContext class.
-=============================================================================*/
-
 #pragma once
 
 
@@ -87,7 +83,7 @@ public:
 
 public:
 
-	// Begin IMessageContext interface
+	// IMessageContext interface
 
 	virtual IMessageAttachmentPtr GetAttachment( ) const override
 	{
@@ -214,11 +210,9 @@ public:
 		return ((Message != nullptr) && TypeInfo.IsValid(false, true));
 	}
 
-	// End IMessageContext interface
-
 public:
 
-	// Begin IMutableMessageContext interface
+	// IMutableMessageContext interface
 
 	virtual void AddRecipient( const FMessageAddress& Recipient ) override
 	{
@@ -261,8 +255,6 @@ public:
 		TimeSent = InTimeSent;
 	}
 
-	// Begin IMutableMessageContext interface
-
 private:
 
 	// Holds a pointer to attached binary data.
diff --git a/Engine/Source/Runtime/Messaging/Private/Bus/MessageDispatchTask.h b/Engine/Source/Runtime/Messaging/Private/Bus/MessageDispatchTask.h
--- a/Engine/Source/Runtime/Messaging/Private/Bus/MessageDispatchTask.h
+++ b/Engine/Source/Runtime/Messaging/Private/Bus/MessageDispatchTask.h
@@ -1,9 +1,5 @@
 // Copyright 1998-2014 Epic Games, Inc. All Rights Reserved.
 
-/*=============================================================================
-	DispatchMessageTask.h: Declares the FDispatchMessageTask class.
-=============================================================================*/
-
 #pragma once
 
 
diff --git a/Engine/Source/Runtime/Messaging/Private/Bus/MessageRouter.cpp b/Engine/Source/Runtime/Messaging/Private/Bus/MessageRouter.cpp
--- a/Engine/Source/Runtime/Messaging/Private/Bus/MessageRouter.cpp
+++ b/Engine/Source/Runtime/Messaging/Private/Bus/MessageRouter.cpp
@@ -1,9 +1,5 @@
 // Copyright 1998-2014 Epic Games, Inc. All Rights Reserved.
 
-/*=============================================================================
-	MessageRouter.cpp: Implements the FMessageRouter class.
-=============================================================================*/
-
 #include "MessagingPrivatePCH.h"
 
 
diff --git a/Engine/Source/Runtime/Messaging/Private/Bus/MessageRouter.h b/Engine/Source/Runtime/Messaging/Private/Bus/MessageRouter.h
--- a/Engine/Source/Runtime/Messaging/Private/Bus/MessageRouter.h
+++ b/Engine/Source/Runtime/Messaging/Private/Bus/MessageRouter.h
@@ -1,9 +1,5 @@
 // Copyright 1998-2014 Epic Games, Inc. All Rights Reserved.
 
-/*=============================================================================
-	MessageRouter.h: Declares the FMessageRouter class.
-=============================================================================*/
-
 #pragma once
 
 
@@ -114,18 +110,13 @@ public:
 
 public:
 
-	// Begin FRunnable interface
+	// FRunnable interface
 
 	virtual bool Init( ) override;
-
 	virtual uint32 Run( ) override;
-
 	virtual void Stop( ) override;
-
 	virtual void Exit( ) override;
 
-	// End FRunnable interface
-
 protected:
 
 	/**
diff --git a/Engine/Source/Runtime/Messaging/Private/Bus/MessageSubscription.h b/Engine/Source/Runtime/Messaging/Private/Bus/MessageSubscription.h
--- a/Engine/Source/Runtime/Messaging/Private/Bus/MessageSubscription.h
+++ b/Engine/Source/Runtime/Messaging/Private/Bus/MessageSubscription.h
@@ -1,9 +1,5 @@
 // Copyright 1998-2014 Epic Games, Inc. All Rights Reserved.
 
-/*=============================================================================
-	MessageSubscription.h: Declares the FMessageSubscription structure.
-=============================================================================*/
-
 #pragma once
 
 
@@ -36,7 +32,7 @@ public:
 
 public:
 
-	// Begin IMessageSubscription interface
+	// IMessageSubscription interface
 
 	virtual void Disable( ) override
 	{
@@ -68,8 +64,6 @@ public:
 		return Enabled;
 	}
 
-	// End IMessageSubscription interface
-
 private:
 
 	// Holds a flag indicating whether this subscription is enabled.
diff --git a/Engine/Source/Runtime/Messaging/Private/Bus/MessageTracer.cpp b/Engine/Source/Runtime/Messaging/Private/Bus/MessageTracer.cpp
--- a/Engine/Source/Runtime/Messaging/Private/Bus/MessageTracer.cpp
+++ b/Engine/Source/Runtime/Messaging/Private/Bus/MessageTracer.cpp
@@ -1,9 +1,5 @@
 // Copyright 1998-2014 Epic Games, Inc. All Rights Reserved.
 
-/*=============================================================================
-	MessageTracer.cpp: Implements the FMessageTracer class.
-=============================================================================*/
-
 #include "MessagingPrivatePCH.h"
 
 
diff --git a/Engine/Source/Runtime/Messaging/Private/Bus/MessageTracer.h b/Engine/Source/Runtime/Messaging/Private/Bus/MessageTracer.h
--- a/Engine/Source/Runtime/Messaging/Private/Bus/MessageTracer.h
+++ b/Engine/Source/Runtime/Messaging/Private/Bus/MessageTracer.h
@@ -1,9 +1,5 @@
 // Copyright 1998-2014 Epic Games, Inc. All Rights Reserved.
 
-/*=============================================================================
-	MessageTracer.h: Declares the FMessageTracer class.
-=============================================================================*/
-
 #pragma once
 
 
@@ -223,7 +219,7 @@ public:
 
 public:
 
-	// Begin IMessageTracer interface
+	// IMessageTracer interface
 
 	virtual void Break( ) override
 	{
@@ -242,9 +238,7 @@ public:
 	}
 
 	virtual int32 GetEndpoints( TArray<FMessageTracerEndpointInfoPtr>& OutEndpoints ) const override;
-
 	virtual int32 GetMessages( TArray<FMessageTracerMessageInfoPtr>& OutMessages ) const override;
-
 	virtual int32 GetMessageTypes( TArray<FMessageTracerTypeInfoPtr>& OutTypes ) const override;
 
 	virtual bool HasMessages( ) const override
@@ -320,8 +314,6 @@ public:
 
 	virtual bool Tick( float DeltaTime ) override;
 
-	// End IMessageTracer interface
-
 protected:
 
 	/**
diff --git a/Engine/Source/Runtime/Messaging/Private/MessagingModule.cpp b/Engine/Source/Runtime/Messaging/Private/MessagingModule.cpp
--- a/Engine/Source/Runtime/Messaging/Private/MessagingModule.cpp
+++ b/Engine/Source/Runtime/Messaging/Private/MessagingModule.cpp
@@ -1,9 +1,5 @@
 // Copyright 1998-2014 Epic Games, Inc. All Rights Reserved.
 
-/*=============================================================================
-	MessagingModule.cpp: Implements the FMessagingModule class.
-=============================================================================*/
-
 #include "MessagingPrivatePCH.h"
 #include "ModuleManager.h"
 
@@ -22,9 +18,9 @@ class FMessagingModule
 {
 public:
 
-	// Begin FSelfRegisteringExec interface
+	// FSelfRegisteringExec interface
 
-	virtual bool Exec( UWorld* InWorld, const TCHAR* Cmd, FOutputDevice& Ar ) OVERRIDE
+	virtual bool Exec( UWorld* InWorld, const TCHAR* Cmd, FOutputDevice& Ar ) override
 	{
 		if (FParse::Command(&Cmd, TEXT("MESSAGING")))
 		{
@@ -54,39 +50,35 @@ public:
 		return false;
 	}
 
-	// End FSelfRegisteringExec interface
-
 public:
 
-	// Begin IMessagingModule interface
+	// IMessagingModule interface
 
-	virtual IMessageBridgePtr CreateBridge( const FMessageAddress& Address, const IMessageBusRef& Bus, const ISerializeMessagesRef& Serializer, const ITransportMessagesRef& Transport ) OVERRIDE
+	virtual IMessageBridgePtr CreateBridge( const FMessageAddress& Address, const IMessageBusRef& Bus, const ISerializeMessagesRef& Serializer, const ITransportMessagesRef& Transport ) override
 	{
 		return MakeShareable(new FMessageBridge(Address, Bus, Serializer, Transport));
 	}
 
-	virtual IMessageBusPtr CreateBus( const IAuthorizeMessageRecipientsPtr& RecipientAuthorizer ) OVERRIDE
+	virtual IMessageBusPtr CreateBus( const IAuthorizeMessageRecipientsPtr& RecipientAuthorizer ) override
 	{
 		return MakeShareable(new FMessageBus(RecipientAuthorizer));
 	}
 
-	virtual ISerializeMessagesPtr CreateJsonMessageSerializer( ) OVERRIDE
+	virtual ISerializeMessagesPtr CreateJsonMessageSerializer( ) override
 	{
 		return MakeShareable(new FJsonMessageSerializer());
 	}
 
-	virtual IMessageBusPtr GetDefaultBus( ) const OVERRIDE
+	virtual IMessageBusPtr GetDefaultBus( ) const override
 	{
 		return DefaultBus;
 	}
 
-	// End IMessagingModule interface
-
 public:
 
-	// Begin IModuleInterface interface
+	// IModuleInterface interface
 
-	virtual void StartupModule( ) OVERRIDE
+	virtual void StartupModule( ) override
 	{
 #if PLATFORM_SUPPORTS_MESSAGEBUS
 		FCoreDelegates::OnPreExit.AddRaw(this, &FMessagingModule::HandleCorePreExit);
@@ -94,18 +86,16 @@ public:
 #endif	//PLATFORM_SUPPORTS_MESSAGEBUS
 	}
 
-	virtual void ShutdownModule( ) OVERRIDE
+	virtual void ShutdownModule( ) override
 	{
 		ShutdownDefaultBus();
 	}
 
-	virtual bool SupportsDynamicReloading( ) OVERRIDE
+	virtual bool SupportsDynamicReloading( ) override
 	{
 		return false;
 	}
 
-	// End IModuleInterface interface
-
 protected:
 
 	void ShutdownDefaultBus( )
diff --git a/Engine/Source/Runtime/Messaging/Private/MessagingPrivatePCH.h b/Engine/Source/Runtime/Messaging/Private/MessagingPrivatePCH.h
--- a/Engine/Source/Runtime/Messaging/Private/MessagingPrivatePCH.h
+++ b/Engine/Source/Runtime/Messaging/Private/MessagingPrivatePCH.h
@@ -1,16 +1,11 @@
 // Copyright 1998-2014 Epic Games, Inc. All Rights Reserved.
 
-/*=============================================================================
-	MessagingPrivatePCH.h: Pre-compiled header file for the Messaging module.
-=============================================================================*/
-
 #pragma once
 
-
-#include "../Public/Messaging.h"
+#include "Messaging.h"
 
 
-/* Private Dependencies
+/* Private dependencies
  *****************************************************************************/
 
 #include "Json.h"
diff --git a/Engine/Source/Runtime/Messaging/Private/Serialization/BsonMessageSerializer.h b/Engine/Source/Runtime/Messaging/Private/Serialization/BsonMessageSerializer.h
--- a/Engine/Source/Runtime/Messaging/Private/Serialization/BsonMessageSerializer.h
+++ b/Engine/Source/Runtime/Messaging/Private/Serialization/BsonMessageSerializer.h
@@ -1,9 +1,5 @@
 // Copyright 1998-2014 Epic Games, Inc. All Rights Reserved.
 
-/*=============================================================================
-	BsonMessageSerializer.h: Declares the FBsonMessageSerializer structure.
-=============================================================================*/
-
 #pragma once
 
 
@@ -13,7 +9,9 @@
 class FBsonMessageSerializer
 	: public ISerializeMessages
 {
-	// Begin ISerializeMessages interface
+public:
+
+	// ISerializeMessages interface
 
 	virtual bool DeserializeMessage( FArchive& Archive, IMutableMessageContextRef& OutContext ) override
 	{
@@ -26,6 +24,4 @@ class FBsonMessageSerializer
 		// not implemented yet
 		return false;
 	}
-
-	// End ISerializeMessages interface
 };
diff --git a/Engine/Source/Runtime/Messaging/Private/Serialization/JsonMessageSerializer.cpp b/Engine/Source/Runtime/Messaging/Private/Serialization/JsonMessageSerializer.cpp
--- a/Engine/Source/Runtime/Messaging/Private/Serialization/JsonMessageSerializer.cpp
+++ b/Engine/Source/Runtime/Messaging/Private/Serialization/JsonMessageSerializer.cpp
@@ -1,9 +1,5 @@
 // Copyright 1998-2014 Epic Games, Inc. All Rights Reserved.
 
-/*=============================================================================
-	JsonMessageSerializer.cpp: Implements the FJsonMessageSerializer structure.
-=============================================================================*/
-
 #include "MessagingPrivatePCH.h"
 
 
diff --git a/Engine/Source/Runtime/Messaging/Private/Serialization/JsonMessageSerializer.h b/Engine/Source/Runtime/Messaging/Private/Serialization/JsonMessageSerializer.h
--- a/Engine/Source/Runtime/Messaging/Private/Serialization/JsonMessageSerializer.h
+++ b/Engine/Source/Runtime/Messaging/Private/Serialization/JsonMessageSerializer.h
@@ -1,9 +1,5 @@
 // Copyright 1998-2014 Epic Games, Inc. All Rights Reserved.
 
-/*=============================================================================
-	JsonMessageSerializer.h: Declares the FJsonMessageSerializer structure.
-=============================================================================*/
-
 #pragma once
 
 
@@ -15,14 +11,11 @@ class FJsonMessageSerializer
 {
 public:
 
-	// Begin ISerializeMessages interface
+	// ISerializeMessages interface
 
 	virtual bool DeserializeMessage( FArchive& Archive, IMutableMessageContextRef& OutContext ) override;
-
 	virtual bool SerializeMessage( const IMessageContextRef& Context, FArchive& Archive ) override;
 
-	// End ISerializeMessages interface
-
 protected:
 
 	/**
diff --git a/Engine/Source/Runtime/Messaging/Private/Serialization/XmlMessageSerializer.h b/Engine/Source/Runtime/Messaging/Private/Serialization/XmlMessageSerializer.h
--- a/Engine/Source/Runtime/Messaging/Private/Serialization/XmlMessageSerializer.h
+++ b/Engine/Source/Runtime/Messaging/Private/Serialization/XmlMessageSerializer.h
@@ -1,9 +1,5 @@
 // Copyright 1998-2014 Epic Games, Inc. All Rights Reserved.
 
-/*=============================================================================
-	XmlMessageSerializer.h: Declares the FXmlMessageSerializer structure.
-=============================================================================*/
-
 #pragma once
 
 
@@ -13,7 +9,9 @@
 class FXmlMessageSerializer
 	: public ISerializeMessages
 {
-	// Begin ISerializeMessages interface
+public:
+
+	// ISerializeMessages interface
 
 	virtual bool DeserializeMessage( FArchive& Archive, IMutableMessageContextRef& OutContext ) override
 	{
@@ -26,6 +24,4 @@ class FXmlMessageSerializer
 		// not implemented yet
 		return false;
 	}
-
-	// End ISerializeMessages interface
 };
diff --git a/Engine/Source/Runtime/Messaging/Public/Common/FileMessageAttachment.h b/Engine/Source/Runtime/Messaging/Public/Common/FileMessageAttachment.h
--- a/Engine/Source/Runtime/Messaging/Public/Common/FileMessageAttachment.h
+++ b/Engine/Source/Runtime/Messaging/Public/Common/FileMessageAttachment.h
@@ -1,9 +1,5 @@
 // Copyright 1998-2014 Epic Games, Inc. All Rights Reserved.
 
-/*=============================================================================
-	FileMessageAttachment.h: Declares the FFileMessageAttachment class.
-=============================================================================*/
-
 #pragma once
 
 
@@ -49,15 +45,13 @@ public:
 
 public:
 
-	// Begin IMessageAttachment interface
+	// IMessageAttachment interface
 
 	virtual FArchive* CreateReader( ) override
 	{
 		return IFileManager::Get().CreateFileReader(*Filename);
 	}
 
-	// End IMessageAttachment interface
-
 private:
 
 	// Holds a flag indicating whether the file should be deleted.
diff --git a/Engine/Source/Runtime/Messaging/Public/Common/MessageBridgeBuilder.h b/Engine/Source/Runtime/Messaging/Public/Common/MessageBridgeBuilder.h
--- a/Engine/Source/Runtime/Messaging/Public/Common/MessageBridgeBuilder.h
+++ b/Engine/Source/Runtime/Messaging/Public/Common/MessageBridgeBuilder.h
@@ -1,9 +1,5 @@
 // Copyright 1998-2014 Epic Games, Inc. All Rights Reserved.
 
-/*=============================================================================
-	MessageBridgeBuilder.h: Declares the FMessageBridgeBuilder class.
-=============================================================================*/
-
 #pragma once
 
 
diff --git a/Engine/Source/Runtime/Messaging/Public/Common/MessageEndpoint.h b/Engine/Source/Runtime/Messaging/Public/Common/MessageEndpoint.h
--- a/Engine/Source/Runtime/Messaging/Public/Common/MessageEndpoint.h
+++ b/Engine/Source/Runtime/Messaging/Public/Common/MessageEndpoint.h
@@ -1,9 +1,5 @@
 // Copyright 1998-2014 Epic Games, Inc. All Rights Reserved.
 
-/*=============================================================================
-	MessageEndpoint.h: Declares the FMessageEndpoint class.
-=============================================================================*/
-
 #pragma once
 
 
@@ -290,11 +286,7 @@ public:
 	 *
 	 * The inbox is disabled by default.
 	 *
-	 * @see EnableInbox
-	 * @see IsInboxEmpty
-	 * @see IsInboxEnabled
-	 * @see ProcessInbox
-	 * @see ReceiveFromInbox
+	 * @see EnableInbox, IsInboxEmpty, IsInboxEnabled, ProcessInbox, ReceiveFromInbox
 	 */
 	void DisableInbox( )
 	{
@@ -308,11 +300,7 @@ public:
 	 * invoke the registered message handlers for all queued up messages, or ReceiveFromInbox() to
 	 * manually receive one message from the inbox at a time. The inbox is disabled by default.
 	 *
-	 * @see DisableInbox
-	 * @see IsInboxEmpty
-	 * @see IsInboxEnabled
-	 * @see ProcessInbox
-	 * @see ReceiveFromInbox
+	 * @see DisableInbox, IsInboxEmpty, IsInboxEnabled, ProcessInbox, ReceiveFromInbox
 	 */
 	void EnableInbox( )
 	{
@@ -324,11 +312,7 @@ public:
 	 *
 	 * @return true if the inbox is empty, false otherwise.
 	 *
-	 * @see DisableInbox
-	 * @see EnableInbox
-	 * @see IsInboxEnabled
-	 * @see ProcessInbox
-	 * @see ReceiveFromInbox
+	 * @see DisableInbox, EnableInbox, IsInboxEnabled, ProcessInbox, ReceiveFromInbox
 	 */
 	bool IsInboxEmpty( ) const
 	{
@@ -338,11 +322,7 @@ public:
 	/**
 	 * Checks whether the inbox is enabled.
 	 *
-	 * @see DisableInbox
-	 * @see EnableInbox
-	 * @see IsInboxEmpty
-	 * @see ProcessInbox
-	 * @see ReceiveFromInbox
+	 * @see DisableInbox, EnableInbox, IsInboxEmpty, ProcessInbox, ReceiveFromInbox
 	 */
 	bool IsInboxEnabled( ) const
 	{
@@ -376,9 +356,7 @@ public:
 	 * must be enabled using the EnableInbox() method.
 	 *
 	 * @param OutContext Will hold the context of the received message.
-	 *
 	 * @return true if a message was received, false if the inbox was empty.
-	 *
 	 * @see DisableInbox, EnableInbox, IsInboxEnabled, ProcessInbox
 	 */
 	bool ReceiveFromInbox( IMessageContextPtr& OutContext )
@@ -403,7 +381,7 @@ public:
 
 public:
 
-	// Begin IReceiveMessages interface
+	// IReceiveMessages interface
 
 	virtual FName GetDebugName( ) const override
 	{
@@ -447,11 +425,9 @@ public:
 		}
 	}
 
-	// End IReceiveMessages interface
-
 public:
 
-	// Begin ISendMessages interface
+	// ISendMessages interface
 
 	virtual FMessageAddress GetSenderAddress( ) override
 	{
@@ -463,8 +439,6 @@ public:
 		ErrorDelegate.ExecuteIfBound(Context, Error);
 	}
 
-	// End ISendMessages interface
-
 public:
 
 	/**
diff --git a/Engine/Source/Runtime/Messaging/Public/Common/MessageEndpointBuilder.h b/Engine/Source/Runtime/Messaging/Public/Common/MessageEndpointBuilder.h
--- a/Engine/Source/Runtime/Messaging/Public/Common/MessageEndpointBuilder.h
+++ b/Engine/Source/Runtime/Messaging/Public/Common/MessageEndpointBuilder.h
@@ -1,9 +1,5 @@
 // Copyright 1998-2014 Epic Games, Inc. All Rights Reserved.
 
-/*=============================================================================
-	MessageEndpointConfigurator.h: Declares the FMessageEndpointConfigurator class.
-=============================================================================*/
-
 #pragma once
 
 
diff --git a/Engine/Source/Runtime/Messaging/Public/Interfaces/IAuthorizeMessageRecipients.h b/Engine/Source/Runtime/Messaging/Public/Interfaces/IAuthorizeMessageRecipients.h
--- a/Engine/Source/Runtime/Messaging/Public/Interfaces/IAuthorizeMessageRecipients.h
+++ b/Engine/Source/Runtime/Messaging/Public/Interfaces/IAuthorizeMessageRecipients.h
@@ -1,9 +1,5 @@
 // Copyright 1998-2014 Epic Games, Inc. All Rights Reserved.
 
-/*=============================================================================
-	IAuthorizeMessageRecipients.h: Declares the IAuthorizeMessageRecipients interface.
-=============================================================================*/
-
 #pragma once
 
 
diff --git a/Engine/Source/Runtime/Messaging/Public/Interfaces/IInterceptMessages.h b/Engine/Source/Runtime/Messaging/Public/Interfaces/IInterceptMessages.h
--- a/Engine/Source/Runtime/Messaging/Public/Interfaces/IInterceptMessages.h
+++ b/Engine/Source/Runtime/Messaging/Public/Interfaces/IInterceptMessages.h
@@ -1,9 +1,5 @@
 // Copyright 1998-2014 Epic Games, Inc. All Rights Reserved.
 
-/*=============================================================================
-	IInterceptMessages.h: Declares the IInterceptMessages interface.
-=============================================================================*/
-
 #pragma once
 
 
diff --git a/Engine/Source/Runtime/Messaging/Public/Interfaces/IMessageAttachment.h b/Engine/Source/Runtime/Messaging/Public/Interfaces/IMessageAttachment.h
--- a/Engine/Source/Runtime/Messaging/Public/Interfaces/IMessageAttachment.h
+++ b/Engine/Source/Runtime/Messaging/Public/Interfaces/IMessageAttachment.h
@@ -1,9 +1,5 @@
 // Copyright 1998-2014 Epic Games, Inc. All Rights Reserved.
 
-/*=============================================================================
-	IMessageAttachment.h: Declares the IMessageAttachment interface.
-=============================================================================*/
-
 #pragma once
 
 
diff --git a/Engine/Source/Runtime/Messaging/Public/Interfaces/IMessageBridge.h b/Engine/Source/Runtime/Messaging/Public/Interfaces/IMessageBridge.h
--- a/Engine/Source/Runtime/Messaging/Public/Interfaces/IMessageBridge.h
+++ b/Engine/Source/Runtime/Messaging/Public/Interfaces/IMessageBridge.h
@@ -1,9 +1,5 @@
 // Copyright 1998-2014 Epic Games, Inc. All Rights Reserved.
 
-/*=============================================================================
-	IMessageBridge.h: Declares the IMessageBridge interface.
-=============================================================================*/
-
 #pragma once
 
 
diff --git a/Engine/Source/Runtime/Messaging/Public/Interfaces/IMessageBus.h b/Engine/Source/Runtime/Messaging/Public/Interfaces/IMessageBus.h
--- a/Engine/Source/Runtime/Messaging/Public/Interfaces/IMessageBus.h
+++ b/Engine/Source/Runtime/Messaging/Public/Interfaces/IMessageBus.h
@@ -1,9 +1,5 @@
 // Copyright 1998-2014 Epic Games, Inc. All Rights Reserved.
 
-/*=============================================================================
-	IMessageBus.h: Declares the IMessageBus interface.
-=============================================================================*/
-
 #pragma once
 
 
