UE4 Refactoring. Changed OVERRIDE and FINAL macros to keywords override and final.

diff --git a/Engine/Source/Runtime/Core/Private/Windows/WindowsEvent.h b/Engine/Source/Runtime/Core/Private/Windows/WindowsEvent.h
--- a/Engine/Source/Runtime/Core/Private/Windows/WindowsEvent.h
+++ b/Engine/Source/Runtime/Core/Private/Windows/WindowsEvent.h
@@ -38,7 +38,7 @@ public:
 
 public:
 
-	virtual bool Create (bool bIsManualReset = false) OVERRIDE
+	virtual bool Create (bool bIsManualReset = false) override
 	{
 		// Create the event and default it to non-signaled
 		Event = CreateEvent(NULL, bIsManualReset, 0, NULL);
@@ -46,21 +46,21 @@ public:
 		return Event != NULL;
 	}
 
-	virtual void Trigger () OVERRIDE
+	virtual void Trigger () override
 	{
 		check(Event);
 
 		SetEvent(Event);
 	}
 
-	virtual void Reset () OVERRIDE
+	virtual void Reset () override
 	{
 		check(Event);
 
 		ResetEvent(Event);
 	}
 
-	virtual bool Wait (uint32 WaitTime) OVERRIDE;
+	virtual bool Wait (uint32 WaitTime) override;
 
 private:
 
diff --git a/Engine/Source/Runtime/Core/Private/Windows/WindowsPlatformFile.cpp b/Engine/Source/Runtime/Core/Private/Windows/WindowsPlatformFile.cpp
--- a/Engine/Source/Runtime/Core/Private/Windows/WindowsPlatformFile.cpp
+++ b/Engine/Source/Runtime/Core/Private/Windows/WindowsPlatformFile.cpp
@@ -51,24 +51,24 @@ public:
 		CloseHandle(FileHandle);
 		FileHandle = NULL;
 	}
-	virtual int64 Tell() OVERRIDE
+	virtual int64 Tell() override
 	{
 		check(IsValid());
 		return FileSeek(0, FILE_CURRENT);
 	}
-	virtual bool Seek(int64 NewPosition) OVERRIDE
+	virtual bool Seek(int64 NewPosition) override
 	{
 		check(IsValid());
 		check(NewPosition >= 0);
 		return FileSeek(NewPosition, FILE_BEGIN) != -1;
 	}
-	virtual bool SeekFromEnd(int64 NewPositionRelativeToEnd = 0) OVERRIDE
+	virtual bool SeekFromEnd(int64 NewPositionRelativeToEnd = 0) override
 	{
 		check(IsValid());
 		check(NewPositionRelativeToEnd <= 0);
 		return FileSeek(NewPositionRelativeToEnd, FILE_END) != -1;
 	}
-	virtual bool Read(uint8* Destination, int64 BytesToRead) OVERRIDE
+	virtual bool Read(uint8* Destination, int64 BytesToRead) override
 	{
 		check(IsValid());
 		while (BytesToRead)
@@ -86,7 +86,7 @@ public:
 		}
 		return true;
 	}
-	virtual bool Write(const uint8* Source, int64 BytesToWrite) OVERRIDE
+	virtual bool Write(const uint8* Source, int64 BytesToWrite) override
 	{
 		check(IsValid());
 		while (BytesToWrite)
@@ -133,7 +133,7 @@ protected:
 		return FPaths::ConvertRelativePathToFull(Result);
 	}
 public:
-	virtual bool FileExists(const TCHAR* Filename) OVERRIDE
+	virtual bool FileExists(const TCHAR* Filename) override
 	{
 		uint32 Result = GetFileAttributesW(*NormalizeFilename(Filename));
 		if (Result != 0xFFFFFFFF && !(Result & FILE_ATTRIBUTE_DIRECTORY))
@@ -142,7 +142,7 @@ public:
 		}
 		return false;
 	}
-	virtual int64 FileSize(const TCHAR* Filename) OVERRIDE
+	virtual int64 FileSize(const TCHAR* Filename) override
 	{
 		WIN32_FILE_ATTRIBUTE_DATA Info;
 		if (!!GetFileAttributesExW(*NormalizeFilename(Filename), GetFileExInfoStandard, &Info))
@@ -157,12 +157,12 @@ public:
 		}
 		return -1;
 	}
-	virtual bool DeleteFile(const TCHAR* Filename) OVERRIDE
+	virtual bool DeleteFile(const TCHAR* Filename) override
 	{
 		const FString NormalizedFilename = NormalizeFilename(Filename);
 		return !!DeleteFileW(*NormalizedFilename);
 	}
-	virtual bool IsReadOnly(const TCHAR* Filename) OVERRIDE
+	virtual bool IsReadOnly(const TCHAR* Filename) override
 	{
 		uint32 Result = GetFileAttributes(*NormalizeFilename(Filename));
 		if (Result != 0xFFFFFFFF)
@@ -171,16 +171,16 @@ public:
 		}
 		return false;
 	}
-	virtual bool MoveFile(const TCHAR* To, const TCHAR* From) OVERRIDE
+	virtual bool MoveFile(const TCHAR* To, const TCHAR* From) override
 	{
 		return !!MoveFileW(*NormalizeFilename(From), *NormalizeFilename(To));
 	}
-	virtual bool SetReadOnly(const TCHAR* Filename, bool bNewReadOnlyValue) OVERRIDE
+	virtual bool SetReadOnly(const TCHAR* Filename, bool bNewReadOnlyValue) override
 	{
 		return !!SetFileAttributesW(*NormalizeFilename(Filename), bNewReadOnlyValue ? FILE_ATTRIBUTE_READONLY : FILE_ATTRIBUTE_NORMAL);
 	}
 
-	virtual FDateTime GetTimeStamp(const TCHAR* Filename) OVERRIDE
+	virtual FDateTime GetTimeStamp(const TCHAR* Filename) override
 	{
 		// get file times
 		struct _stati64 FileInfo;
@@ -194,7 +194,7 @@ public:
 		return WindowsEpoch + TimeSinceEpoch;
 	}
 
-	virtual void SetTimeStamp(const TCHAR* Filename, FDateTime DateTime) OVERRIDE
+	virtual void SetTimeStamp(const TCHAR* Filename, FDateTime DateTime) override
 	{
 		// get file times
 		struct _stati64 FileInfo;
@@ -210,7 +210,7 @@ public:
 		_wutime(Filename, &Times);
 	}
 
-	virtual FDateTime GetAccessTimeStamp(const TCHAR* Filename) OVERRIDE
+	virtual FDateTime GetAccessTimeStamp(const TCHAR* Filename) override
 	{
 		// get file times
 		struct _stati64 FileInfo;
@@ -224,7 +224,7 @@ public:
 		return WindowsEpoch + TimeSinceEpoch;
 	}
 
-	virtual IFileHandle* OpenRead(const TCHAR* Filename) OVERRIDE
+	virtual IFileHandle* OpenRead(const TCHAR* Filename) override
 	{
 		uint32  Access    = GENERIC_READ;
 		uint32  WinFlags  = FILE_SHARE_READ;
@@ -236,7 +236,7 @@ public:
 		}
 		return NULL;
 	}
-	virtual IFileHandle* OpenWrite(const TCHAR* Filename, bool bAppend = false, bool bAllowRead = false) OVERRIDE
+	virtual IFileHandle* OpenWrite(const TCHAR* Filename, bool bAppend = false, bool bAllowRead = false) override
 	{
 		uint32  Access    = GENERIC_WRITE;
 		uint32  WinFlags  = bAllowRead ? FILE_SHARE_READ : 0;
@@ -249,7 +249,7 @@ public:
 		return NULL;
 	}
 
-	virtual bool DirectoryExists(const TCHAR* Directory) OVERRIDE
+	virtual bool DirectoryExists(const TCHAR* Directory) override
 	{
 		// Empty Directory is the current directory so assume it always exists.
 		bool bExists = !FCString::Strlen(Directory);
@@ -260,11 +260,11 @@ public:
 		}
 		return bExists;
 	}
-	virtual bool CreateDirectory(const TCHAR* Directory) OVERRIDE
+	virtual bool CreateDirectory(const TCHAR* Directory) override
 	{
 		return CreateDirectoryW(*NormalizeDirectory(Directory), NULL) || GetLastError() == ERROR_ALREADY_EXISTS;
 	}
-	virtual bool DeleteDirectory(const TCHAR* Directory) OVERRIDE
+	virtual bool DeleteDirectory(const TCHAR* Directory) override
 	{
 		RemoveDirectoryW(*NormalizeDirectory(Directory));
 		return !DirectoryExists(Directory);
diff --git a/Engine/Source/Runtime/Core/Private/Windows/WindowsPlatformOutputDevicesPrivate.h b/Engine/Source/Runtime/Core/Private/Windows/WindowsPlatformOutputDevicesPrivate.h
--- a/Engine/Source/Runtime/Core/Private/Windows/WindowsPlatformOutputDevicesPrivate.h
+++ b/Engine/Source/Runtime/Core/Private/Windows/WindowsPlatformOutputDevicesPrivate.h
@@ -61,7 +61,7 @@ public:
 	}
 
 
-	virtual void Serialize(const TCHAR* Buffer, ELogVerbosity::Type Verbosity, const class FName& Category) OVERRIDE
+	virtual void Serialize(const TCHAR* Buffer, ELogVerbosity::Type Verbosity, const class FName& Category) override
 	{
 		if (EventLog != NULL)
 		{
@@ -129,7 +129,7 @@ public:
 	 * @param	Data	Text to log
 	 * @param	Event	Event name used for suppression purposes
 	 */
-	virtual void Serialize( const TCHAR* Msg, ELogVerbosity::Type Verbosity, const class FName& Category ) OVERRIDE;
+	virtual void Serialize( const TCHAR* Msg, ELogVerbosity::Type Verbosity, const class FName& Category ) override;
 
 	/**
 	 * Error handling function that is being called from within the system wide global
diff --git a/Engine/Source/Runtime/Core/Private/Windows/WindowsTextInputMethodSystem.cpp b/Engine/Source/Runtime/Core/Private/Windows/WindowsTextInputMethodSystem.cpp
--- a/Engine/Source/Runtime/Core/Private/Windows/WindowsTextInputMethodSystem.cpp
+++ b/Engine/Source/Runtime/Core/Private/Windows/WindowsTextInputMethodSystem.cpp
@@ -1,4 +1,4 @@
-﻿// Copyright 1998-2014 Epic Games, Inc. All Rights Reserved.
+// Copyright 1998-2014 Epic Games, Inc. All Rights Reserved.
 
 #include "CorePrivate.h"
 #include "WindowsTextInputMethodSystem.h"
@@ -29,10 +29,10 @@ namespace
 
 		virtual ~FTextInputMethodChangeNotifier() {}
 
-		virtual void NotifyLayoutChanged(const ELayoutChangeType ChangeType) OVERRIDE;
-		virtual void NotifySelectionChanged() OVERRIDE;
-		virtual void NotifyTextChanged(const uint32 BeginIndex, const uint32 OldLength, const uint32 NewLength) OVERRIDE;
-		virtual void CancelComposition() OVERRIDE;
+		virtual void NotifyLayoutChanged(const ELayoutChangeType ChangeType) override;
+		virtual void NotifySelectionChanged() override;
+		virtual void NotifyTextChanged(const uint32 BeginIndex, const uint32 OldLength, const uint32 NewLength) override;
+		virtual void CancelComposition() override;
 
 	private:
 		const FCOMPtr<FTextStoreACP> TextStoreACP;
diff --git a/Engine/Source/Runtime/Core/Public/Android/AndroidApplication.h b/Engine/Source/Runtime/Core/Public/Android/AndroidApplication.h
--- a/Engine/Source/Runtime/Core/Public/Android/AndroidApplication.h
+++ b/Engine/Source/Runtime/Core/Public/Android/AndroidApplication.h
@@ -26,13 +26,13 @@ public:
 
 	void SetMessageHandler( const TSharedRef< FGenericApplicationMessageHandler >& InMessageHandler );
 
-	virtual void PollGameDeviceState( const float TimeDelta ) OVERRIDE;
+	virtual void PollGameDeviceState( const float TimeDelta ) override;
 
-	virtual FPlatformRect GetWorkArea( const FPlatformRect& CurrentWindow ) const OVERRIDE;
+	virtual FPlatformRect GetWorkArea( const FPlatformRect& CurrentWindow ) const override;
 
-	virtual void GetDisplayMetrics( FDisplayMetrics& OutDesktopMetrics ) const OVERRIDE;
+	virtual void GetDisplayMetrics( FDisplayMetrics& OutDesktopMetrics ) const override;
 
-	virtual TSharedRef< FGenericWindow > MakeWindow() OVERRIDE;
+	virtual TSharedRef< FGenericWindow > MakeWindow() override;
 	
 	void InitializeWindow( const TSharedRef< FGenericWindow >& InWindow, const TSharedRef< FGenericWindowDefinition >& InDefinition, const TSharedPtr< FGenericWindow >& InParent, const bool bShowImmediately );
 
diff --git a/Engine/Source/Runtime/Core/Public/Android/AndroidWindow.h b/Engine/Source/Runtime/Core/Public/Android/AndroidWindow.h
--- a/Engine/Source/Runtime/Core/Public/Android/AndroidWindow.h
+++ b/Engine/Source/Runtime/Core/Public/Android/AndroidWindow.h
@@ -25,12 +25,12 @@ public:
 	static TSharedRef<FAndroidWindow> Make();
 
 	
-	virtual void* GetOSWindowHandle() const OVERRIDE { return Window; } //can be null.
+	virtual void* GetOSWindowHandle() const override { return Window; } //can be null.
 
 	void Initialize( class FAndroidApplication* const Application, const TSharedRef< FGenericWindowDefinition >& InDefinition, const TSharedPtr< FAndroidWindow >& InParent, const bool bShowImmediately );
 
 	/** Returns the rectangle of the screen the window is associated with */
-	virtual bool GetFullScreenInfo( int32& X, int32& Y, int32& Width, int32& Height ) const OVERRIDE;
+	virtual bool GetFullScreenInfo( int32& X, int32& Y, int32& Width, int32& Height ) const override;
 
 	
 	virtual void SetOSWindowHandle(void*);
@@ -42,7 +42,7 @@ public:
 
 protected:
 	/** @return true if the native window is currently in fullscreen mode, false otherwise */
-	virtual EWindowMode::Type GetWindowMode() const OVERRIDE { return EWindowMode::Fullscreen; }
+	virtual EWindowMode::Type GetWindowMode() const override { return EWindowMode::Fullscreen; }
 
 private:
 	/**
diff --git a/Engine/Source/Runtime/Core/Public/Apple/ApplePlatformFile.h b/Engine/Source/Runtime/Core/Public/Apple/ApplePlatformFile.h
--- a/Engine/Source/Runtime/Core/Public/Apple/ApplePlatformFile.h
+++ b/Engine/Source/Runtime/Core/Public/Apple/ApplePlatformFile.h
@@ -15,25 +15,25 @@ protected:
 	virtual FString NormalizeFilename(const TCHAR* Filename);
 	virtual FString NormalizeDirectory(const TCHAR* Directory);
 public:
-	virtual bool FileExists(const TCHAR* Filename) OVERRIDE;
-	virtual int64 FileSize(const TCHAR* Filename) OVERRIDE;
-	virtual bool DeleteFile(const TCHAR* Filename) OVERRIDE;
-	virtual bool IsReadOnly(const TCHAR* Filename) OVERRIDE;
-	virtual bool MoveFile(const TCHAR* To, const TCHAR* From) OVERRIDE;
-	virtual bool SetReadOnly(const TCHAR* Filename, bool bNewReadOnlyValue) OVERRIDE;
+	virtual bool FileExists(const TCHAR* Filename) override;
+	virtual int64 FileSize(const TCHAR* Filename) override;
+	virtual bool DeleteFile(const TCHAR* Filename) override;
+	virtual bool IsReadOnly(const TCHAR* Filename) override;
+	virtual bool MoveFile(const TCHAR* To, const TCHAR* From) override;
+	virtual bool SetReadOnly(const TCHAR* Filename, bool bNewReadOnlyValue) override;
 
 
-	virtual FDateTime GetTimeStamp(const TCHAR* Filename) OVERRIDE;
+	virtual FDateTime GetTimeStamp(const TCHAR* Filename) override;
 
-	virtual void SetTimeStamp(const TCHAR* Filename, const FDateTime DateTime) OVERRIDE;
+	virtual void SetTimeStamp(const TCHAR* Filename, const FDateTime DateTime) override;
 
-	virtual FDateTime GetAccessTimeStamp(const TCHAR* Filename) OVERRIDE;
+	virtual FDateTime GetAccessTimeStamp(const TCHAR* Filename) override;
 
-	virtual IFileHandle* OpenRead(const TCHAR* Filename) OVERRIDE;
-	virtual IFileHandle* OpenWrite(const TCHAR* Filename, bool bAppend = false, bool bAllowRead = false) OVERRIDE;
-	virtual bool DirectoryExists(const TCHAR* Directory) OVERRIDE;
-	virtual bool CreateDirectory(const TCHAR* Directory) OVERRIDE;
-	virtual bool DeleteDirectory(const TCHAR* Directory) OVERRIDE;
+	virtual IFileHandle* OpenRead(const TCHAR* Filename) override;
+	virtual IFileHandle* OpenWrite(const TCHAR* Filename, bool bAppend = false, bool bAllowRead = false) override;
+	virtual bool DirectoryExists(const TCHAR* Directory) override;
+	virtual bool CreateDirectory(const TCHAR* Directory) override;
+	virtual bool DeleteDirectory(const TCHAR* Directory) override;
 	bool IterateDirectory(const TCHAR* Directory, FDirectoryVisitor& Visitor);
 private:
 	int32 Stat(const TCHAR* Filename, struct stat* OutFileInfo, ANSICHAR* OutFixedPath = NULL, SIZE_T FixedPathSize = 0);
diff --git a/Engine/Source/Runtime/Core/Public/Containers/UnrealString.h b/Engine/Source/Runtime/Core/Public/Containers/UnrealString.h
--- a/Engine/Source/Runtime/Core/Public/Containers/UnrealString.h
+++ b/Engine/Source/Runtime/Core/Public/Containers/UnrealString.h
@@ -1597,7 +1597,7 @@ public:
 	{
 		bAutoEmitLineTerminator = false;
 	}
-	virtual void Serialize( const TCHAR* InData, ELogVerbosity::Type Verbosity, const class FName& Category ) OVERRIDE
+	virtual void Serialize( const TCHAR* InData, ELogVerbosity::Type Verbosity, const class FName& Category ) override
 	{
 		*this += (TCHAR*)InData;
 		if(bAutoEmitLineTerminator)
@@ -1667,7 +1667,7 @@ public:
 	,	LineCount(0)
 	{}
 
-	virtual void Serialize( const TCHAR* InData, ELogVerbosity::Type Verbosity, const class FName& Category ) OVERRIDE
+	virtual void Serialize( const TCHAR* InData, ELogVerbosity::Type Verbosity, const class FName& Category ) override
 	{
 		Super::Serialize(InData, Verbosity, Category);
 		int32 TermLength = FCString::Strlen(LINE_TERMINATOR);
diff --git a/Engine/Source/Runtime/Core/Public/GenericPlatform/GenericPlatformChunkInstall.h b/Engine/Source/Runtime/Core/Public/GenericPlatform/GenericPlatformChunkInstall.h
--- a/Engine/Source/Runtime/Core/Public/GenericPlatform/GenericPlatformChunkInstall.h
+++ b/Engine/Source/Runtime/Core/Public/GenericPlatform/GenericPlatformChunkInstall.h
@@ -116,7 +116,7 @@ public:
 	 * @param ChunkID		The id of the chunk to check.
 	 * @return				Enum specifying whether the chunk is available to use, waiting to install, or does not exist.
 	 **/
-	virtual EChunkLocation::Type GetChunkLocation( uint32 ChunkID ) OVERRIDE
+	virtual EChunkLocation::Type GetChunkLocation( uint32 ChunkID ) override
 	{
 		return EChunkLocation::LocalFast;
 	}
@@ -126,7 +126,7 @@ public:
 	 * @param ReportType	Enum specifying how progress is reported.
 	 * @return				true if reporting type is supported on the current platform.
 	 **/
-	virtual bool GetProgressReportingTypeSupported(EChunkProgressReportingType::Type ReportType) OVERRIDE
+	virtual bool GetProgressReportingTypeSupported(EChunkProgressReportingType::Type ReportType) override
 	{
 		if (ReportType == EChunkProgressReportingType::PercentageComplete)
 		{
@@ -142,7 +142,7 @@ public:
 	 * @param ReportType	The type of progress report you want.
 	 * @return				A value whose meaning is dependent on the ReportType param.
 	 **/
-	virtual float GetChunkProgress( uint32 ChunkID, EChunkProgressReportingType::Type ReportType ) OVERRIDE
+	virtual float GetChunkProgress( uint32 ChunkID, EChunkProgressReportingType::Type ReportType ) override
 	{
 		if (ReportType == EChunkProgressReportingType::PercentageComplete)
 		{
@@ -155,7 +155,7 @@ public:
 	 * Inquire about the priority of chunk installation vs. game IO.
 	 * @return				Paused, low or high priority.
 	 **/
-	virtual EChunkInstallSpeed::Type GetInstallSpeed() OVERRIDE
+	virtual EChunkInstallSpeed::Type GetInstallSpeed() override
 	{
 		return EChunkInstallSpeed::Paused;
 	}
@@ -165,7 +165,7 @@ public:
 	 * @param InstallSpeed	Pause, low or high priority.
 	 * @return				false if the operation is not allowed, otherwise true.
 	 **/
-	virtual bool SetInstallSpeed( EChunkInstallSpeed::Type InstallSpeed ) OVERRIDE
+	virtual bool SetInstallSpeed( EChunkInstallSpeed::Type InstallSpeed ) override
 	{
 		return false;
 	}
@@ -175,7 +175,7 @@ public:
 	 * @param ChunkID		The id of the chunk to prioritize.
 	 * @return				false if the operation is not allowed or the chunk doesn't exist, otherwise true.
 	 **/
-	virtual bool PrioritizeChunk( uint32 ChunkID ) OVERRIDE
+	virtual bool PrioritizeChunk( uint32 ChunkID ) override
 	{
 		return false;
 	}
diff --git a/Engine/Source/Runtime/Core/Public/GenericPlatform/GenericPlatformFile.h b/Engine/Source/Runtime/Core/Public/GenericPlatform/GenericPlatformFile.h
--- a/Engine/Source/Runtime/Core/Public/GenericPlatform/GenericPlatformFile.h
+++ b/Engine/Source/Runtime/Core/Public/GenericPlatform/GenericPlatformFile.h
@@ -261,16 +261,16 @@ class CORE_API IPhysicalPlatformFile : public IPlatformFile
 {
 public:
 	// BEGIN IPlatformFile Interface
-	virtual bool ShouldBeUsed(IPlatformFile* Inner, const TCHAR* CmdLine) const OVERRIDE
+	virtual bool ShouldBeUsed(IPlatformFile* Inner, const TCHAR* CmdLine) const override
 	{
 		return true;
 	}
-	virtual bool Initialize(IPlatformFile* Inner, const TCHAR* CmdLine) OVERRIDE;
-	virtual IPlatformFile* GetLowerLevel() OVERRIDE
+	virtual bool Initialize(IPlatformFile* Inner, const TCHAR* CmdLine) override;
+	virtual IPlatformFile* GetLowerLevel() override
 	{
 		return NULL;
 	}
-	virtual const TCHAR* GetName() const OVERRIDE
+	virtual const TCHAR* GetName() const override
 	{
 		return IPlatformFile::GetPhysicalTypeName();
 	}
diff --git a/Engine/Source/Runtime/Core/Public/HAL/ConsoleManager.h b/Engine/Source/Runtime/Core/Public/HAL/ConsoleManager.h
--- a/Engine/Source/Runtime/Core/Public/HAL/ConsoleManager.h
+++ b/Engine/Source/Runtime/Core/Public/HAL/ConsoleManager.h
@@ -47,31 +47,31 @@ public:
 
 	// interface IConsoleManager -----------------------------------
 
-	virtual IConsoleVariable* RegisterConsoleVariable(const TCHAR* Name, int32 DefaultValue, const TCHAR* Help, uint32 Flags) OVERRIDE;
-	virtual IConsoleVariable* RegisterConsoleVariable(const TCHAR* Name, float DefaultValue, const TCHAR* Help, uint32 Flags) OVERRIDE;
-	virtual IConsoleVariable* RegisterConsoleVariable(const TCHAR* Name, const TCHAR *DefaultValue, const TCHAR* Help, uint32 Flags) OVERRIDE;
-	virtual IConsoleVariable* RegisterConsoleVariableRef(const TCHAR* Name, int32& RefValue, const TCHAR* Help, uint32 Flags) OVERRIDE;
-	virtual IConsoleVariable* RegisterConsoleVariableRef(const TCHAR* Name, float& RefValue, const TCHAR* Help, uint32 Flags) OVERRIDE;
-	virtual IConsoleVariable* RegisterConsoleVariableBitRef(const TCHAR* CVarName, const TCHAR* FlagName, uint32 BitNumber, uint8* Force0MaskPtr, uint8* Force1MaskPtr, const TCHAR* Help, uint32 Flags) OVERRIDE;
+	virtual IConsoleVariable* RegisterConsoleVariable(const TCHAR* Name, int32 DefaultValue, const TCHAR* Help, uint32 Flags) override;
+	virtual IConsoleVariable* RegisterConsoleVariable(const TCHAR* Name, float DefaultValue, const TCHAR* Help, uint32 Flags) override;
+	virtual IConsoleVariable* RegisterConsoleVariable(const TCHAR* Name, const TCHAR *DefaultValue, const TCHAR* Help, uint32 Flags) override;
+	virtual IConsoleVariable* RegisterConsoleVariableRef(const TCHAR* Name, int32& RefValue, const TCHAR* Help, uint32 Flags) override;
+	virtual IConsoleVariable* RegisterConsoleVariableRef(const TCHAR* Name, float& RefValue, const TCHAR* Help, uint32 Flags) override;
+	virtual IConsoleVariable* RegisterConsoleVariableBitRef(const TCHAR* CVarName, const TCHAR* FlagName, uint32 BitNumber, uint8* Force0MaskPtr, uint8* Force1MaskPtr, const TCHAR* Help, uint32 Flags) override;
 	
 	virtual void CallAllConsoleVariableSinks();
 	virtual void RegisterConsoleVariableSink(const FConsoleCommandDelegate& Command);
 	virtual void UnregisterConsoleVariableSink(const FConsoleCommandDelegate& Command);
 
-	virtual IConsoleCommand* RegisterConsoleCommand(const TCHAR* Name, const TCHAR* Help, const FConsoleCommandDelegate& Command, uint32 Flags) OVERRIDE;
-	virtual IConsoleCommand* RegisterConsoleCommand(const TCHAR* Name, const TCHAR* Help, const FConsoleCommandWithArgsDelegate& Command, uint32 Flags) OVERRIDE;
-	virtual IConsoleCommand* RegisterConsoleCommand(const TCHAR* Name, const TCHAR* Help, const FConsoleCommandWithWorldDelegate& Command, uint32 Flags) OVERRIDE;
-	virtual IConsoleCommand* RegisterConsoleCommand(const TCHAR* Name, const TCHAR* Help, const FConsoleCommandWithWorldAndArgsDelegate& Command, uint32 Flags) OVERRIDE;
-	virtual IConsoleCommand* RegisterConsoleCommand(const TCHAR* Name, const TCHAR* Help, const FConsoleCommandWithOutputDeviceDelegate& Command, uint32 Flags) OVERRIDE;
-	virtual IConsoleCommand* RegisterConsoleCommand(const TCHAR* Name, const TCHAR* Help, uint32 Flags) OVERRIDE;
-	virtual IConsoleVariable* FindConsoleVariable(const TCHAR* Name) const OVERRIDE;
-	virtual void ForEachConsoleObject(const FConsoleObjectVisitor& Visitor, const TCHAR* ThatStartsWith) const OVERRIDE;
-	virtual bool ProcessUserConsoleInput(const TCHAR* InInput, FOutputDevice& Ar, UWorld* InWorld) OVERRIDE;
+	virtual IConsoleCommand* RegisterConsoleCommand(const TCHAR* Name, const TCHAR* Help, const FConsoleCommandDelegate& Command, uint32 Flags) override;
+	virtual IConsoleCommand* RegisterConsoleCommand(const TCHAR* Name, const TCHAR* Help, const FConsoleCommandWithArgsDelegate& Command, uint32 Flags) override;
+	virtual IConsoleCommand* RegisterConsoleCommand(const TCHAR* Name, const TCHAR* Help, const FConsoleCommandWithWorldDelegate& Command, uint32 Flags) override;
+	virtual IConsoleCommand* RegisterConsoleCommand(const TCHAR* Name, const TCHAR* Help, const FConsoleCommandWithWorldAndArgsDelegate& Command, uint32 Flags) override;
+	virtual IConsoleCommand* RegisterConsoleCommand(const TCHAR* Name, const TCHAR* Help, const FConsoleCommandWithOutputDeviceDelegate& Command, uint32 Flags) override;
+	virtual IConsoleCommand* RegisterConsoleCommand(const TCHAR* Name, const TCHAR* Help, uint32 Flags) override;
+	virtual IConsoleVariable* FindConsoleVariable(const TCHAR* Name) const override;
+	virtual void ForEachConsoleObject(const FConsoleObjectVisitor& Visitor, const TCHAR* ThatStartsWith) const override;
+	virtual bool ProcessUserConsoleInput(const TCHAR* InInput, FOutputDevice& Ar, UWorld* InWorld) override;
 	virtual void AddConsoleHistoryEntry(const TCHAR* Input);
 	virtual void GetConsoleHistory(TArray<FString>& Out) const; 
-	virtual bool IsNameRegistered(const TCHAR* Name) const OVERRIDE;	
-	virtual void RegisterThreadPropagation(uint32 ThreadId, IConsoleThreadPropagation* InCallback) OVERRIDE;
-	virtual void UnregisterConsoleObject( IConsoleObject* Object, bool bKeepState) OVERRIDE;
+	virtual bool IsNameRegistered(const TCHAR* Name) const override;	
+	virtual void RegisterThreadPropagation(uint32 ThreadId, IConsoleThreadPropagation* InCallback) override;
+	virtual void UnregisterConsoleObject( IConsoleObject* Object, bool bKeepState) override;
 
 private: // ----------------------------------------------------
 
diff --git a/Engine/Source/Runtime/Core/Public/HAL/FeedbackContextAnsi.h b/Engine/Source/Runtime/Core/Public/HAL/FeedbackContextAnsi.h
--- a/Engine/Source/Runtime/Core/Public/HAL/FeedbackContextAnsi.h
+++ b/Engine/Source/Runtime/Core/Public/HAL/FeedbackContextAnsi.h
@@ -81,7 +81,7 @@ public:
 		fflush( stdout );
 	}
 
-	virtual bool YesNof( const FText& Question ) OVERRIDE
+	virtual bool YesNof( const FText& Question ) override
 	{
 		if( (GIsClient || GIsEditor) )
 		{
diff --git a/Engine/Source/Runtime/Core/Public/HAL/FileManagerGeneric.h b/Engine/Source/Runtime/Core/Public/HAL/FileManagerGeneric.h
--- a/Engine/Source/Runtime/Core/Public/HAL/FileManagerGeneric.h
+++ b/Engine/Source/Runtime/Core/Public/HAL/FileManagerGeneric.h
@@ -32,38 +32,38 @@ public:
 	{
 	}
 
-	virtual void ProcessCommandLineOptions() OVERRIDE;
+	virtual void ProcessCommandLineOptions() override;
 
-	virtual void SetSandboxEnabled(bool bInEnabled) OVERRIDE
+	virtual void SetSandboxEnabled(bool bInEnabled) override
 	{
 		GetLowLevel().SetSandboxEnabled(bInEnabled);
 	}
 
-	virtual bool IsSandboxEnabled() const OVERRIDE
+	virtual bool IsSandboxEnabled() const override
 	{
 		return GetLowLevel().IsSandboxEnabled();
 	}
 
-	FArchive* CreateFileReader( const TCHAR* Filename, uint32 ReadFlags=0 ) OVERRIDE;	
-	FArchive* CreateFileWriter( const TCHAR* Filename, uint32 WriteFlags=0 ) OVERRIDE;
+	FArchive* CreateFileReader( const TCHAR* Filename, uint32 ReadFlags=0 ) override;	
+	FArchive* CreateFileWriter( const TCHAR* Filename, uint32 WriteFlags=0 ) override;
 
 #if ALLOW_DEBUG_FILES
-	FArchive* CreateDebugFileWriter( const TCHAR* Filename, uint32 WriteFlags=0 ) OVERRIDE
+	FArchive* CreateDebugFileWriter( const TCHAR* Filename, uint32 WriteFlags=0 ) override
 	{
 		return CreateFileWriter( Filename, WriteFlags );
 	}
 #endif
 
-	bool Delete( const TCHAR* Filename, bool RequireExists=0, bool EvenReadOnly=0, bool Quiet=0 ) OVERRIDE;
-	bool IsReadOnly( const TCHAR* Filename ) OVERRIDE;
-	bool Move( const TCHAR* Dest, const TCHAR* Src, bool Replace=1, bool EvenIfReadOnly=0, bool Attributes=0, bool bDoNotRetryOrError=0 ) OVERRIDE;
-	bool DirectoryExists(const TCHAR* InDirectory) OVERRIDE;
-	void FindFiles( TArray<FString>& Result, const TCHAR* Filename, bool Files, bool Directories ) OVERRIDE;
-	void FindFilesRecursive( TArray<FString>& FileNames, const TCHAR* StartDirectory, const TCHAR* Filename, bool Files, bool Directories, bool bClearFileNames=true) OVERRIDE;
-	double GetFileAgeSeconds( const TCHAR* Filename ) OVERRIDE;
-	FDateTime GetTimeStamp( const TCHAR* Filename ) OVERRIDE;
-	FDateTime GetAccessTimeStamp( const TCHAR* Filename ) OVERRIDE;
-	bool SetTimeStamp( const TCHAR* Filename, FDateTime Timestamp ) OVERRIDE;
+	bool Delete( const TCHAR* Filename, bool RequireExists=0, bool EvenReadOnly=0, bool Quiet=0 ) override;
+	bool IsReadOnly( const TCHAR* Filename ) override;
+	bool Move( const TCHAR* Dest, const TCHAR* Src, bool Replace=1, bool EvenIfReadOnly=0, bool Attributes=0, bool bDoNotRetryOrError=0 ) override;
+	bool DirectoryExists(const TCHAR* InDirectory) override;
+	void FindFiles( TArray<FString>& Result, const TCHAR* Filename, bool Files, bool Directories ) override;
+	void FindFilesRecursive( TArray<FString>& FileNames, const TCHAR* StartDirectory, const TCHAR* Filename, bool Files, bool Directories, bool bClearFileNames=true) override;
+	double GetFileAgeSeconds( const TCHAR* Filename ) override;
+	FDateTime GetTimeStamp( const TCHAR* Filename ) override;
+	FDateTime GetAccessTimeStamp( const TCHAR* Filename ) override;
+	bool SetTimeStamp( const TCHAR* Filename, FDateTime Timestamp ) override;
 
 	virtual uint32	Copy( const TCHAR* InDestFile, const TCHAR* InSrcFile, bool ReplaceExisting, bool EvenIfReadOnly, bool Attributes, FCopyProgress* Progress );
 	virtual bool	MakeDirectory( const TCHAR* Path, bool Tree=0 );
@@ -85,7 +85,7 @@ public:
 	 * 
 	 * @return	filename using relative path
 	 */
-	FString ConvertToRelativePath( const TCHAR* Filename ) OVERRIDE;
+	FString ConvertToRelativePath( const TCHAR* Filename ) override;
 
 	/**
 	 * Converts passed in filename to use an absolute path (for reading)
@@ -94,7 +94,7 @@ public:
 	 * 
 	 * @return	filename using absolute path
 	 */
-	FString ConvertToAbsolutePathForExternalAppForRead( const TCHAR* Filename ) OVERRIDE;
+	FString ConvertToAbsolutePathForExternalAppForRead( const TCHAR* Filename ) override;
 
 	/**
 	 * Converts passed in filename to use an absolute path (for writing)
@@ -103,7 +103,7 @@ public:
 	 * 
 	 * @return	filename using absolute path
 	 */
-	FString ConvertToAbsolutePathForExternalAppForWrite( const TCHAR* Filename ) OVERRIDE;
+	FString ConvertToAbsolutePathForExternalAppForWrite( const TCHAR* Filename ) override;
 
 	/**
 	 *	Returns the size of a file. (Thread-safe)
@@ -111,7 +111,7 @@ public:
 	 *	@param Filename		Platform-independent Unreal filename.
 	 *	@return				File size in bytes or INDEX_NONE if the file didn't exist.
 	 **/
-	int64 FileSize( const TCHAR* Filename ) OVERRIDE;
+	int64 FileSize( const TCHAR* Filename ) override;
 
 	/**
 	 * Sends a message to the file server, and will block until it's complete. Will return 
@@ -121,7 +121,7 @@ public:
 	 *
 	 * @return			true if the message was sent to server and it returned success, or false if there is no server, or the command failed
 	 */
-	virtual bool SendMessageToServer(const TCHAR* Message, IPlatformFile::IFileServerMessageHandler* Handler) OVERRIDE
+	virtual bool SendMessageToServer(const TCHAR* Message, IPlatformFile::IFileServerMessageHandler* Handler) override
 	{
 		return GetLowLevel().SendMessageToServer(Message, Handler);
 	}
@@ -146,17 +146,17 @@ public:
 	FArchiveFileReaderGeneric( IFileHandle* InHandle, const TCHAR* InFilename, int64 InSize );
 	~FArchiveFileReaderGeneric();
 
-	virtual void Seek( int64 InPos ) FINAL;
-	virtual int64 Tell() FINAL
+	virtual void Seek( int64 InPos ) final;
+	virtual int64 Tell() final
 	{
 		return Pos;
 	}
-	virtual int64 TotalSize() FINAL
+	virtual int64 TotalSize() final
 	{
 		return Size;
 	}
-	virtual bool Close() FINAL;
-	virtual void Serialize( void* V, int64 Length ) FINAL;
+	virtual bool Close() final;
+	virtual void Serialize( void* V, int64 Length ) final;
 
 protected:
 	bool InternalPrecache( int64 PrecacheOffset, int64 PrecacheSize );
@@ -197,15 +197,15 @@ public:
 	FArchiveFileWriterGeneric( IFileHandle* InHandle, const TCHAR* InFilename, int64 InPos );
 	~FArchiveFileWriterGeneric();
 
-	virtual void Seek( int64 InPos ) FINAL;
-	virtual int64 Tell() FINAL
+	virtual void Seek( int64 InPos ) final;
+	virtual int64 Tell() final
 	{
 		return Pos;
 	}
 	virtual int64 TotalSize();
-	virtual bool Close() FINAL;
-	virtual void Serialize( void* V, int64 Length ) FINAL;
-	virtual void Flush() FINAL;
+	virtual bool Close() final;
+	virtual void Serialize( void* V, int64 Length ) final;
+	virtual void Flush() final;
 
 protected:
 	/** 
diff --git a/Engine/Source/Runtime/Core/Public/HAL/IPlatformFileCachedWrapper.h b/Engine/Source/Runtime/Core/Public/HAL/IPlatformFileCachedWrapper.h
--- a/Engine/Source/Runtime/Core/Public/HAL/IPlatformFileCachedWrapper.h
+++ b/Engine/Source/Runtime/Core/Public/HAL/IPlatformFileCachedWrapper.h
@@ -22,12 +22,12 @@ public:
 	}
 
 	
-	virtual int64		Tell() OVERRIDE
+	virtual int64		Tell() override
 	{
 		return FilePos;
 	}
 
-	virtual bool		Seek(int64 NewPosition) OVERRIDE
+	virtual bool		Seek(int64 NewPosition) override
 	{
 		if (NewPosition < 0 || NewPosition > FileSize)
 		{
@@ -37,12 +37,12 @@ public:
 		return true;
 	}
 
-	virtual bool		SeekFromEnd(int64 NewPositionRelativeToEnd = 0) OVERRIDE
+	virtual bool		SeekFromEnd(int64 NewPositionRelativeToEnd = 0) override
 	{
 		return Seek(FileSize - NewPositionRelativeToEnd);
 	}
 
-	virtual bool		Read(uint8* Destination, int64 BytesToRead) OVERRIDE
+	virtual bool		Read(uint8* Destination, int64 BytesToRead) override
 	{
 		if (!bReadable || BytesToRead < 0 || (BytesToRead + FilePos > FileSize))
 		{
@@ -117,7 +117,7 @@ public:
 		return Result;
 	}
 
-	virtual bool		Write(const uint8* Source, int64 BytesToWrite) OVERRIDE
+	virtual bool		Write(const uint8* Source, int64 BytesToWrite) override
 	{
 		if (!bWritable || BytesToWrite < 0)
 		{
@@ -141,7 +141,7 @@ public:
 		return Result;
 	}
 
-	virtual int64		Size() OVERRIDE
+	virtual int64		Size() override
 	{
 		return FileSize;
 	}
@@ -218,7 +218,7 @@ public:
 		: LowerLevel(NULL)
 	{
 	}
-	virtual bool Initialize(IPlatformFile* Inner, const TCHAR* CommandLineParam) OVERRIDE
+	virtual bool Initialize(IPlatformFile* Inner, const TCHAR* CommandLineParam) override
 	{
 		// Inner is required.
 		check(Inner != NULL);
@@ -229,52 +229,52 @@ public:
 	{
 		return FPlatformProperties::RequiresCookedData();
 	}
-	IPlatformFile* GetLowerLevel() OVERRIDE
+	IPlatformFile* GetLowerLevel() override
 	{
 		return LowerLevel;
 	}
 
-	virtual const TCHAR* GetName() const OVERRIDE
+	virtual const TCHAR* GetName() const override
 	{
 		return FCachedReadPlatformFile::GetTypeName();
 	}
-	virtual bool		FileExists(const TCHAR* Filename) OVERRIDE
+	virtual bool		FileExists(const TCHAR* Filename) override
 	{
 		return LowerLevel->FileExists(Filename);
 	}
-	virtual int64		FileSize(const TCHAR* Filename) OVERRIDE
+	virtual int64		FileSize(const TCHAR* Filename) override
 	{
 		return LowerLevel->FileSize(Filename);
 	}
-	virtual bool		DeleteFile(const TCHAR* Filename) OVERRIDE
+	virtual bool		DeleteFile(const TCHAR* Filename) override
 	{
 		return LowerLevel->DeleteFile(Filename);
 	}
-	virtual bool		IsReadOnly(const TCHAR* Filename) OVERRIDE
+	virtual bool		IsReadOnly(const TCHAR* Filename) override
 	{
 		return LowerLevel->IsReadOnly(Filename);
 	}
-	virtual bool		MoveFile(const TCHAR* To, const TCHAR* From) OVERRIDE
+	virtual bool		MoveFile(const TCHAR* To, const TCHAR* From) override
 	{
 		return LowerLevel->MoveFile(To, From);
 	}
-	virtual bool		SetReadOnly(const TCHAR* Filename, bool bNewReadOnlyValue) OVERRIDE
+	virtual bool		SetReadOnly(const TCHAR* Filename, bool bNewReadOnlyValue) override
 	{
 		return LowerLevel->SetReadOnly(Filename, bNewReadOnlyValue);
 	}
-	virtual FDateTime	GetTimeStamp(const TCHAR* Filename) OVERRIDE
+	virtual FDateTime	GetTimeStamp(const TCHAR* Filename) override
 	{
 		return LowerLevel->GetTimeStamp(Filename);
 	}
-	virtual void		SetTimeStamp(const TCHAR* Filename, FDateTime DateTime) OVERRIDE
+	virtual void		SetTimeStamp(const TCHAR* Filename, FDateTime DateTime) override
 	{
 		LowerLevel->SetTimeStamp(Filename, DateTime);
 	}
-	virtual FDateTime	GetAccessTimeStamp(const TCHAR* Filename) OVERRIDE
+	virtual FDateTime	GetAccessTimeStamp(const TCHAR* Filename) override
 	{
 		return LowerLevel->GetAccessTimeStamp(Filename);
 	}
-	virtual IFileHandle*	OpenRead(const TCHAR* Filename) OVERRIDE
+	virtual IFileHandle*	OpenRead(const TCHAR* Filename) override
 	{
 		IFileHandle* InnerHandle=LowerLevel->OpenRead(Filename);
 		if (!InnerHandle)
@@ -283,7 +283,7 @@ public:
 		}
 		return new FCachedFileHandle(InnerHandle, true, false);
 	}
-	virtual IFileHandle*	OpenWrite(const TCHAR* Filename, bool bAppend = false, bool bAllowRead = false) OVERRIDE
+	virtual IFileHandle*	OpenWrite(const TCHAR* Filename, bool bAppend = false, bool bAllowRead = false) override
 	{
 		IFileHandle* InnerHandle=LowerLevel->OpenWrite(Filename, bAppend, bAllowRead);
 		if (!InnerHandle)
@@ -292,51 +292,51 @@ public:
 		}
 		return new FCachedFileHandle(InnerHandle, bAllowRead, true);
 	}
-	virtual bool		DirectoryExists(const TCHAR* Directory) OVERRIDE
+	virtual bool		DirectoryExists(const TCHAR* Directory) override
 	{
 		return LowerLevel->DirectoryExists(Directory);
 	}
-	virtual bool		CreateDirectory(const TCHAR* Directory) OVERRIDE
+	virtual bool		CreateDirectory(const TCHAR* Directory) override
 	{
 		return LowerLevel->CreateDirectory(Directory);
 	}
-	virtual bool		DeleteDirectory(const TCHAR* Directory) OVERRIDE
+	virtual bool		DeleteDirectory(const TCHAR* Directory) override
 	{
 		return LowerLevel->DeleteDirectory(Directory);
 	}
-	virtual bool		IterateDirectory(const TCHAR* Directory, IPlatformFile::FDirectoryVisitor& Visitor) OVERRIDE
+	virtual bool		IterateDirectory(const TCHAR* Directory, IPlatformFile::FDirectoryVisitor& Visitor) override
 	{
 		return LowerLevel->IterateDirectory(Directory, Visitor);
 	}
-	virtual bool		IterateDirectoryRecursively(const TCHAR* Directory, IPlatformFile::FDirectoryVisitor& Visitor) OVERRIDE
+	virtual bool		IterateDirectoryRecursively(const TCHAR* Directory, IPlatformFile::FDirectoryVisitor& Visitor) override
 	{
 		return LowerLevel->IterateDirectoryRecursively(Directory, Visitor);
 	}
-	virtual bool		DeleteDirectoryRecursively(const TCHAR* Directory) OVERRIDE
+	virtual bool		DeleteDirectoryRecursively(const TCHAR* Directory) override
 	{
 		return LowerLevel->DeleteDirectoryRecursively(Directory);
 	}
-	virtual bool		CopyFile(const TCHAR* To, const TCHAR* From) OVERRIDE
+	virtual bool		CopyFile(const TCHAR* To, const TCHAR* From) override
 	{
 		return LowerLevel->CopyFile(To, From);
 	}
-	virtual bool		CreateDirectoryTree(const TCHAR* Directory) OVERRIDE
+	virtual bool		CreateDirectoryTree(const TCHAR* Directory) override
 	{
 		return LowerLevel->CreateDirectoryTree(Directory);
 	}
-	virtual bool		CopyDirectoryTree(const TCHAR* DestinationDirectory, const TCHAR* Source, bool bOverwriteAllExisting) OVERRIDE
+	virtual bool		CopyDirectoryTree(const TCHAR* DestinationDirectory, const TCHAR* Source, bool bOverwriteAllExisting) override
 	{
 		return LowerLevel->CopyDirectoryTree(DestinationDirectory, Source, bOverwriteAllExisting);
 	}
-	virtual FString		ConvertToAbsolutePathForExternalAppForRead( const TCHAR* Filename ) OVERRIDE
+	virtual FString		ConvertToAbsolutePathForExternalAppForRead( const TCHAR* Filename ) override
 	{
 		return LowerLevel->ConvertToAbsolutePathForExternalAppForRead(Filename);
 	}
-	virtual FString		ConvertToAbsolutePathForExternalAppForWrite( const TCHAR* Filename ) OVERRIDE
+	virtual FString		ConvertToAbsolutePathForExternalAppForWrite( const TCHAR* Filename ) override
 	{
 		return LowerLevel->ConvertToAbsolutePathForExternalAppForWrite(Filename);
 	}
-	virtual bool		SendMessageToServer(const TCHAR* Message, IFileServerMessageHandler* Handler) OVERRIDE
+	virtual bool		SendMessageToServer(const TCHAR* Message, IFileServerMessageHandler* Handler) override
 	{
 		return LowerLevel->SendMessageToServer(Message, Handler);
 	}
diff --git a/Engine/Source/Runtime/Core/Public/HAL/IPlatformFileLogWrapper.h b/Engine/Source/Runtime/Core/Public/HAL/IPlatformFileLogWrapper.h
--- a/Engine/Source/Runtime/Core/Public/HAL/IPlatformFileLogWrapper.h
+++ b/Engine/Source/Runtime/Core/Public/HAL/IPlatformFileLogWrapper.h
@@ -31,7 +31,7 @@ public:
 	{
 	}
 
-	virtual int64		Tell() OVERRIDE
+	virtual int64		Tell() override
 	{
 		FILE_LOG(LogPlatformFile, Log, TEXT("Tell %s"), *Filename);
 		double StartTime = FPlatformTime::Seconds();
@@ -40,7 +40,7 @@ public:
 		FILE_LOG(LogPlatformFile, Log, TEXT("Tell return %lld [%fms]"), Result, ThisTime);
 		return Result;
 	}
-	virtual bool		Seek(int64 NewPosition) OVERRIDE
+	virtual bool		Seek(int64 NewPosition) override
 	{
 		FILE_LOG(LogPlatformFile, Log, TEXT("Seek %s %lld"), *Filename, NewPosition);
 		double StartTime = FPlatformTime::Seconds();
@@ -49,7 +49,7 @@ public:
 		FILE_LOG(LogPlatformFile, Log, TEXT("Seek return %d [%fms]"), int32(Result), ThisTime);
 		return Result;
 	}
-	virtual bool		SeekFromEnd(int64 NewPositionRelativeToEnd) OVERRIDE
+	virtual bool		SeekFromEnd(int64 NewPositionRelativeToEnd) override
 	{
 		FILE_LOG(LogPlatformFile, Log, TEXT("SeekFromEnd %s %lld"), *Filename, NewPositionRelativeToEnd);
 		double StartTime = FPlatformTime::Seconds();
@@ -58,7 +58,7 @@ public:
 		FILE_LOG(LogPlatformFile, Log, TEXT("SeekFromEnd return %d [%fms]"), int32(Result), ThisTime);
 		return Result;
 	}
-	virtual bool		Read(uint8* Destination, int64 BytesToRead) OVERRIDE
+	virtual bool		Read(uint8* Destination, int64 BytesToRead) override
 	{
 		FILE_LOG(LogPlatformFile, Log, TEXT("Read %s %lld"), *Filename, BytesToRead);
 		double StartTime = FPlatformTime::Seconds();
@@ -67,7 +67,7 @@ public:
 		FILE_LOG(LogPlatformFile, Log, TEXT("Read return %d [%fms]"), int32(Result), ThisTime);
 		return Result;
 	}
-	virtual bool		Write(const uint8* Source, int64 BytesToWrite) OVERRIDE
+	virtual bool		Write(const uint8* Source, int64 BytesToWrite) override
 	{
 		FILE_LOG(LogPlatformFile, Log, TEXT("Write %s %lld"), *Filename, BytesToWrite);
 		double StartTime = FPlatformTime::Seconds();
@@ -76,7 +76,7 @@ public:
 		FILE_LOG(LogPlatformFile, Log, TEXT("Write return %d [%fms]"), int32(Result), ThisTime);
 		return Result;
 	}
-	virtual int64		Size() OVERRIDE
+	virtual int64		Size() override
 	{
 		FILE_LOG(LogPlatformFile, Log, TEXT("Size %s"), *Filename);
 		double StartTime = FPlatformTime::Seconds();
@@ -101,11 +101,11 @@ public:
 	{
 	}
 
-	virtual bool ShouldBeUsed(IPlatformFile* Inner, const TCHAR* CmdLine) const OVERRIDE
+	virtual bool ShouldBeUsed(IPlatformFile* Inner, const TCHAR* CmdLine) const override
 	{
 		return FParse::Param(CmdLine, TEXT("FileLog"));
 	}
-	virtual bool Initialize(IPlatformFile* Inner, const TCHAR* CommandLineParam) OVERRIDE
+	virtual bool Initialize(IPlatformFile* Inner, const TCHAR* CommandLineParam) override
 	{
 		// Inner is required.
 		check(Inner != NULL);
@@ -113,17 +113,17 @@ public:
 		return !!LowerLevel;
 	}
 
-	IPlatformFile* GetLowerLevel() OVERRIDE
+	IPlatformFile* GetLowerLevel() override
 	{
 		return LowerLevel;
 	}
 
-	virtual const TCHAR* GetName() const OVERRIDE
+	virtual const TCHAR* GetName() const override
 	{
 		return FLoggedPlatformFile::GetTypeName();
 	}
 
-	virtual bool		FileExists(const TCHAR* Filename) OVERRIDE
+	virtual bool		FileExists(const TCHAR* Filename) override
 	{
 		FILE_LOG(LogPlatformFile, Log, TEXT("FileExists %s"), Filename);
 		double StartTime = FPlatformTime::Seconds();
@@ -132,7 +132,7 @@ public:
 		FILE_LOG(LogPlatformFile, Log, TEXT("FileExists return %d [%fms]"), int32(Result), ThisTime);
 		return Result;
 	}
-	virtual int64		FileSize(const TCHAR* Filename) OVERRIDE
+	virtual int64		FileSize(const TCHAR* Filename) override
 	{
 		FILE_LOG(LogPlatformFile, Log, TEXT("FileSize %s"), Filename);
 		double StartTime = FPlatformTime::Seconds();
@@ -141,7 +141,7 @@ public:
 		FILE_LOG(LogPlatformFile, Log, TEXT("FileSize return %lld [%fms]"), Result, ThisTime);
 		return Result;
 	}
-	virtual bool		DeleteFile(const TCHAR* Filename) OVERRIDE
+	virtual bool		DeleteFile(const TCHAR* Filename) override
 	{
 		FILE_LOG(LogPlatformFile, Log, TEXT("DeleteFile %s"), Filename);
 		double StartTime = FPlatformTime::Seconds();
@@ -150,7 +150,7 @@ public:
 		FILE_LOG(LogPlatformFile, Log, TEXT("DeleteFile return %d [%fms]"), int32(Result), ThisTime);
 		return Result;
 	}
-	virtual bool		IsReadOnly(const TCHAR* Filename) OVERRIDE
+	virtual bool		IsReadOnly(const TCHAR* Filename) override
 	{
 		FILE_LOG(LogPlatformFile, Log, TEXT("IsReadOnly %s"), Filename);
 		double StartTime = FPlatformTime::Seconds();
@@ -159,7 +159,7 @@ public:
 		FILE_LOG(LogPlatformFile, Log, TEXT("IsReadOnly return %d [%fms]"), int32(Result), ThisTime);
 		return Result;
 	}
-	virtual bool		MoveFile(const TCHAR* To, const TCHAR* From) OVERRIDE
+	virtual bool		MoveFile(const TCHAR* To, const TCHAR* From) override
 	{
 		FILE_LOG(LogPlatformFile, Log, TEXT("MoveFile %s %s"), To, From);
 		double StartTime = FPlatformTime::Seconds();
@@ -168,7 +168,7 @@ public:
 		FILE_LOG(LogPlatformFile, Log, TEXT("MoveFile return %d [%fms]"), int32(Result), ThisTime);
 		return Result;
 	}
-	virtual bool		SetReadOnly(const TCHAR* Filename, bool bNewReadOnlyValue) OVERRIDE
+	virtual bool		SetReadOnly(const TCHAR* Filename, bool bNewReadOnlyValue) override
 	{
 		FILE_LOG(LogPlatformFile, Log, TEXT("SetReadOnly %s %d"), Filename, int32(bNewReadOnlyValue));
 		double StartTime = FPlatformTime::Seconds();
@@ -177,7 +177,7 @@ public:
 		FILE_LOG(LogPlatformFile, Log, TEXT("SetReadOnly return %d [%fms]"), int32(Result), ThisTime);
 		return Result;
 	}
-	virtual FDateTime	GetTimeStamp(const TCHAR* Filename) OVERRIDE
+	virtual FDateTime	GetTimeStamp(const TCHAR* Filename) override
 	{
 		FILE_LOG(LogPlatformFile, Log, TEXT("GetTimeStamp %s"), Filename);
 		double StartTime = FPlatformTime::Seconds();
@@ -186,7 +186,7 @@ public:
 		FILE_LOG(LogPlatformFile, Log, TEXT("GetTimeStamp return %llx [%fms]"), Result.GetTicks(), ThisTime);
 		return Result;
 	}
-	virtual void		SetTimeStamp(const TCHAR* Filename, FDateTime DateTime) OVERRIDE
+	virtual void		SetTimeStamp(const TCHAR* Filename, FDateTime DateTime) override
 	{
 		FILE_LOG(LogPlatformFile, Log, TEXT("SetTimeStamp %s"), Filename);
 		double StartTime = FPlatformTime::Seconds();
@@ -194,7 +194,7 @@ public:
 		float ThisTime = 1000.0f * float(FPlatformTime::Seconds() - StartTime);
 		FILE_LOG(LogPlatformFile, Log, TEXT("SetTimeStamp [%fms]"), ThisTime);
 	}
-	virtual FDateTime	GetAccessTimeStamp(const TCHAR* Filename) OVERRIDE
+	virtual FDateTime	GetAccessTimeStamp(const TCHAR* Filename) override
 	{
 		FILE_LOG(LogPlatformFile, Log, TEXT("GetAccessTimeStamp %s"), Filename);
 		double StartTime = FPlatformTime::Seconds();
@@ -203,7 +203,7 @@ public:
 		FILE_LOG(LogPlatformFile, Log, TEXT("GetAccessTimeStamp return %llx [%fms]"), Result.GetTicks(), ThisTime);
 		return Result;
 	}
-	virtual IFileHandle*	OpenRead(const TCHAR* Filename) OVERRIDE
+	virtual IFileHandle*	OpenRead(const TCHAR* Filename) override
 	{
 		FILE_LOG(LogPlatformFile, Log, TEXT("OpenRead %s"), Filename);
 		double StartTime = FPlatformTime::Seconds();
@@ -212,7 +212,7 @@ public:
 		FILE_LOG(LogPlatformFile, Log, TEXT("OpenRead return %llx [%fms]"), uint64(Result), ThisTime);
 		return Result ? (new FLoggedFileHandle(Result, Filename)) : Result;
 	}
-	virtual IFileHandle*	OpenWrite(const TCHAR* Filename, bool bAppend = false, bool bAllowRead = false) OVERRIDE
+	virtual IFileHandle*	OpenWrite(const TCHAR* Filename, bool bAppend = false, bool bAllowRead = false) override
 	{
 		FILE_LOG(LogPlatformFile, Log, TEXT("OpenWrite %s %d %d"), Filename, int32(bAppend), int32(bAllowRead));
 		double StartTime = FPlatformTime::Seconds();
@@ -222,7 +222,7 @@ public:
 		return Result ? (new FLoggedFileHandle(Result, Filename)) : Result;
 	}
 
-	virtual bool		DirectoryExists(const TCHAR* Directory) OVERRIDE
+	virtual bool		DirectoryExists(const TCHAR* Directory) override
 	{
 		FILE_LOG(LogPlatformFile, Log, TEXT("DirectoryExists %s"), Directory);
 		double StartTime = FPlatformTime::Seconds();
@@ -231,7 +231,7 @@ public:
 		FILE_LOG(LogPlatformFile, Log, TEXT("DirectoryExists return %d [%fms]"), int32(Result), ThisTime);
 		return Result;
 	}
-	virtual bool		CreateDirectory(const TCHAR* Directory) OVERRIDE
+	virtual bool		CreateDirectory(const TCHAR* Directory) override
 	{
 		FILE_LOG(LogPlatformFile, Log, TEXT("CreateDirectory %s"), Directory);
 		double StartTime = FPlatformTime::Seconds();
@@ -240,7 +240,7 @@ public:
 		FILE_LOG(LogPlatformFile, Log, TEXT("CreateDirectory return %d [%fms]"), int32(Result), ThisTime);
 		return Result;
 	}
-	virtual bool		DeleteDirectory(const TCHAR* Directory) OVERRIDE
+	virtual bool		DeleteDirectory(const TCHAR* Directory) override
 	{
 		FILE_LOG(LogPlatformFile, Log, TEXT("DeleteDirectory %s"), Directory);
 		double StartTime = FPlatformTime::Seconds();
@@ -269,7 +269,7 @@ public:
 		}
 	};
 
-	virtual bool		IterateDirectory(const TCHAR* Directory, IPlatformFile::FDirectoryVisitor& Visitor) OVERRIDE
+	virtual bool		IterateDirectory(const TCHAR* Directory, IPlatformFile::FDirectoryVisitor& Visitor) override
 	{
 		FILE_LOG(LogPlatformFile, Log, TEXT("IterateDirectory %s"), Directory);
 		double StartTime = FPlatformTime::Seconds();
@@ -279,7 +279,7 @@ public:
 		FILE_LOG(LogPlatformFile, Log, TEXT("IterateDirectory return %d [%fms]"), int32(Result), ThisTime);
 		return Result;
 	}
-	virtual bool		IterateDirectoryRecursively(const TCHAR* Directory, IPlatformFile::FDirectoryVisitor& Visitor) OVERRIDE
+	virtual bool		IterateDirectoryRecursively(const TCHAR* Directory, IPlatformFile::FDirectoryVisitor& Visitor) override
 	{
 		FILE_LOG(LogPlatformFile, Log, TEXT("IterateDirectoryRecursively %s"), Directory);
 		double StartTime = FPlatformTime::Seconds();
@@ -289,7 +289,7 @@ public:
 		FILE_LOG(LogPlatformFile, Log, TEXT("IterateDirectoryRecursively return %d [%fms]"), int32(Result), ThisTime);
 		return Result;
 	}
-	virtual bool		DeleteDirectoryRecursively(const TCHAR* Directory) OVERRIDE
+	virtual bool		DeleteDirectoryRecursively(const TCHAR* Directory) override
 	{
 		FILE_LOG(LogPlatformFile, Log, TEXT("DeleteDirectoryRecursively %s"), Directory);
 		double StartTime = FPlatformTime::Seconds();
@@ -298,7 +298,7 @@ public:
 		FILE_LOG(LogPlatformFile, Log, TEXT("DeleteDirectoryRecursively return %d [%fms]"), int32(Result), ThisTime);
 		return Result;
 	}
-	virtual bool		CopyFile(const TCHAR* To, const TCHAR* From) OVERRIDE
+	virtual bool		CopyFile(const TCHAR* To, const TCHAR* From) override
 	{
 		FILE_LOG(LogPlatformFile, Log, TEXT("CopyFile %s %s"), To, From);
 		double StartTime = FPlatformTime::Seconds();
diff --git a/Engine/Source/Runtime/Core/Public/HAL/IPlatformFileOpenLogWrapper.h b/Engine/Source/Runtime/Core/Public/HAL/IPlatformFileOpenLogWrapper.h
--- a/Engine/Source/Runtime/Core/Public/HAL/IPlatformFileOpenLogWrapper.h
+++ b/Engine/Source/Runtime/Core/Public/HAL/IPlatformFileOpenLogWrapper.h
@@ -26,13 +26,13 @@ public:
 	{
 	}
 
-	virtual bool ShouldBeUsed(IPlatformFile* Inner, const TCHAR* CmdLine) const OVERRIDE
+	virtual bool ShouldBeUsed(IPlatformFile* Inner, const TCHAR* CmdLine) const override
 	{
 		bool bResult = FParse::Param(CmdLine, TEXT("FileOpenLog"));
 		return bResult;
 	}
 
-	virtual bool Initialize(IPlatformFile* Inner, const TCHAR* CommandLineParam) OVERRIDE
+	virtual bool Initialize(IPlatformFile* Inner, const TCHAR* CommandLineParam) override
 	{
 		LowerLevel = Inner;
 		FString LogFileDirectory;
@@ -72,7 +72,7 @@ public:
 		}
 		return true;
 	}
-	virtual IPlatformFile* GetLowerLevel() OVERRIDE
+	virtual IPlatformFile* GetLowerLevel() override
 	{
 		return LowerLevel;
 	}
@@ -80,47 +80,47 @@ public:
 	{
 		return TEXT("FileOpenLog");
 	}
-	virtual const TCHAR* GetName() const OVERRIDE
+	virtual const TCHAR* GetName() const override
 	{
 		return GetTypeName();
 	}
-	virtual bool		FileExists(const TCHAR* Filename) OVERRIDE
+	virtual bool		FileExists(const TCHAR* Filename) override
 	{
 		return LowerLevel->FileExists(Filename);
 	}
-	virtual int64		FileSize(const TCHAR* Filename) OVERRIDE
+	virtual int64		FileSize(const TCHAR* Filename) override
 	{
 		return LowerLevel->FileSize(Filename);
 	}
-	virtual bool		DeleteFile(const TCHAR* Filename) OVERRIDE
+	virtual bool		DeleteFile(const TCHAR* Filename) override
 	{
 		return LowerLevel->DeleteFile(Filename);
 	}
-	virtual bool		IsReadOnly(const TCHAR* Filename) OVERRIDE
+	virtual bool		IsReadOnly(const TCHAR* Filename) override
 	{
 		return LowerLevel->IsReadOnly(Filename);
 	}
-	virtual bool		MoveFile(const TCHAR* To, const TCHAR* From) OVERRIDE
+	virtual bool		MoveFile(const TCHAR* To, const TCHAR* From) override
 	{
 		return LowerLevel->MoveFile(To, From);
 	}
-	virtual bool		SetReadOnly(const TCHAR* Filename, bool bNewReadOnlyValue) OVERRIDE
+	virtual bool		SetReadOnly(const TCHAR* Filename, bool bNewReadOnlyValue) override
 	{
 		return LowerLevel->SetReadOnly(Filename, bNewReadOnlyValue);
 	}
-	virtual FDateTime	GetTimeStamp(const TCHAR* Filename) OVERRIDE
+	virtual FDateTime	GetTimeStamp(const TCHAR* Filename) override
 	{
 		return LowerLevel->GetTimeStamp(Filename);
 	}
-	virtual void		SetTimeStamp(const TCHAR* Filename, FDateTime DateTime) OVERRIDE
+	virtual void		SetTimeStamp(const TCHAR* Filename, FDateTime DateTime) override
 	{
 		LowerLevel->SetTimeStamp(Filename, DateTime);
 	}
-	virtual FDateTime	GetAccessTimeStamp(const TCHAR* Filename) OVERRIDE
+	virtual FDateTime	GetAccessTimeStamp(const TCHAR* Filename) override
 	{
 		return LowerLevel->GetAccessTimeStamp(Filename);
 	}
-	virtual IFileHandle*	OpenRead(const TCHAR* Filename) OVERRIDE
+	virtual IFileHandle*	OpenRead(const TCHAR* Filename) override
 	{
 		IFileHandle* Result = LowerLevel->OpenRead(Filename);
 		if (Result)
@@ -139,55 +139,55 @@ public:
 		}
 		return Result;
 	}
-	virtual IFileHandle*	OpenWrite(const TCHAR* Filename, bool bAppend = false, bool bAllowRead = false) OVERRIDE
+	virtual IFileHandle*	OpenWrite(const TCHAR* Filename, bool bAppend = false, bool bAllowRead = false) override
 	{
 		return LowerLevel->OpenWrite(Filename, bAppend, bAllowRead);
 	}
-	virtual bool		DirectoryExists(const TCHAR* Directory) OVERRIDE
+	virtual bool		DirectoryExists(const TCHAR* Directory) override
 	{
 		return LowerLevel->DirectoryExists(Directory);
 	}
-	virtual bool		CreateDirectory(const TCHAR* Directory) OVERRIDE
+	virtual bool		CreateDirectory(const TCHAR* Directory) override
 	{
 		return LowerLevel->CreateDirectory(Directory);
 	}
-	virtual bool		DeleteDirectory(const TCHAR* Directory) OVERRIDE
+	virtual bool		DeleteDirectory(const TCHAR* Directory) override
 	{
 		return LowerLevel->DeleteDirectory(Directory);
 	}
-	virtual bool		IterateDirectory(const TCHAR* Directory, IPlatformFile::FDirectoryVisitor& Visitor) OVERRIDE
+	virtual bool		IterateDirectory(const TCHAR* Directory, IPlatformFile::FDirectoryVisitor& Visitor) override
 	{
 		return LowerLevel->IterateDirectory( Directory, Visitor );
 	}
-	virtual bool		IterateDirectoryRecursively(const TCHAR* Directory, IPlatformFile::FDirectoryVisitor& Visitor) OVERRIDE
+	virtual bool		IterateDirectoryRecursively(const TCHAR* Directory, IPlatformFile::FDirectoryVisitor& Visitor) override
 	{
 		return LowerLevel->IterateDirectoryRecursively( Directory, Visitor );
 	}
-	virtual bool		DeleteDirectoryRecursively(const TCHAR* Directory) OVERRIDE
+	virtual bool		DeleteDirectoryRecursively(const TCHAR* Directory) override
 	{
 		return LowerLevel->DeleteDirectoryRecursively( Directory );
 	}
-	virtual bool		CopyFile(const TCHAR* To, const TCHAR* From) OVERRIDE
+	virtual bool		CopyFile(const TCHAR* To, const TCHAR* From) override
 	{
 		return LowerLevel->CopyFile( To, From );
 	}
-	virtual bool		CreateDirectoryTree(const TCHAR* Directory) OVERRIDE
+	virtual bool		CreateDirectoryTree(const TCHAR* Directory) override
 	{
 		return LowerLevel->CreateDirectoryTree(Directory);
 	}
-	virtual bool		CopyDirectoryTree(const TCHAR* DestinationDirectory, const TCHAR* Source, bool bOverwriteAllExisting) OVERRIDE
+	virtual bool		CopyDirectoryTree(const TCHAR* DestinationDirectory, const TCHAR* Source, bool bOverwriteAllExisting) override
 	{
 		return LowerLevel->CopyDirectoryTree(DestinationDirectory, Source, bOverwriteAllExisting);
 	}
-	virtual FString		ConvertToAbsolutePathForExternalAppForRead( const TCHAR* Filename ) OVERRIDE
+	virtual FString		ConvertToAbsolutePathForExternalAppForRead( const TCHAR* Filename ) override
 	{
 		return LowerLevel->ConvertToAbsolutePathForExternalAppForRead(Filename);
 	}
-	virtual FString		ConvertToAbsolutePathForExternalAppForWrite( const TCHAR* Filename ) OVERRIDE
+	virtual FString		ConvertToAbsolutePathForExternalAppForWrite( const TCHAR* Filename ) override
 	{
 		return LowerLevel->ConvertToAbsolutePathForExternalAppForWrite(Filename);
 	}
-	virtual bool		SendMessageToServer(const TCHAR* Message, IFileServerMessageHandler* Handler) OVERRIDE
+	virtual bool		SendMessageToServer(const TCHAR* Message, IFileServerMessageHandler* Handler) override
 	{
 		return LowerLevel->SendMessageToServer(Message, Handler);
 	}
diff --git a/Engine/Source/Runtime/Core/Public/HAL/IPlatformFileProfilerWrapper.h b/Engine/Source/Runtime/Core/Public/HAL/IPlatformFileProfilerWrapper.h
--- a/Engine/Source/Runtime/Core/Public/HAL/IPlatformFileProfilerWrapper.h
+++ b/Engine/Source/Runtime/Core/Public/HAL/IPlatformFileProfilerWrapper.h
@@ -146,28 +146,28 @@ public:
 	{
 	}
 
-	virtual int64		Tell() OVERRIDE
+	virtual int64		Tell() override
 	{
 		FProfiledFileStatsOp* Stat( FileStats->CreateOpStat( FProfiledFileStatsOp::Tell ) );
 		int64 Result = FileHandle->Tell();
 		Stat->Duration += FPlatformTime::Seconds() * 1000.0 - Stat->LastOpTime;
 		return Result;
 	}
-	virtual bool		Seek(int64 NewPosition) OVERRIDE
+	virtual bool		Seek(int64 NewPosition) override
 	{
 		FProfiledFileStatsOp* Stat( FileStats->CreateOpStat( FProfiledFileStatsOp::Seek ) );
 		bool Result = FileHandle->Seek(NewPosition);
 		Stat->Duration += FPlatformTime::Seconds() * 1000.0 - Stat->LastOpTime;
 		return Result;
 	}
-	virtual bool		SeekFromEnd(int64 NewPositionRelativeToEnd) OVERRIDE
+	virtual bool		SeekFromEnd(int64 NewPositionRelativeToEnd) override
 	{
 		FProfiledFileStatsOp* Stat( FileStats->CreateOpStat( FProfiledFileStatsOp::Seek ) );
 		bool Result = FileHandle->SeekFromEnd(NewPositionRelativeToEnd);
 		Stat->Duration += FPlatformTime::Seconds() * 1000.0 - Stat->LastOpTime;
 		return Result;
 	}
-	virtual bool		Read(uint8* Destination, int64 BytesToRead) OVERRIDE
+	virtual bool		Read(uint8* Destination, int64 BytesToRead) override
 	{
 		FProfiledFileStatsOp* Stat( FileStats->CreateOpStat( FProfiledFileStatsOp::Read ) );
 		bool Result = FileHandle->Read(Destination, BytesToRead);
@@ -175,7 +175,7 @@ public:
 		Stat->Bytes += BytesToRead;
 		return Result;
 	}
-	virtual bool		Write(const uint8* Source, int64 BytesToWrite) OVERRIDE
+	virtual bool		Write(const uint8* Source, int64 BytesToWrite) override
 	{
 		FProfiledFileStatsOp* Stat( FileStats->CreateOpStat( FProfiledFileStatsOp::Write ) );
 		bool Result = FileHandle->Write(Source, BytesToWrite);
@@ -183,7 +183,7 @@ public:
 		Stat->Bytes += BytesToWrite;
 		return Result;
 	}
-	virtual int64		Size() OVERRIDE
+	virtual int64		Size() override
 	{
 		FProfiledFileStatsOp* Stat( FileStats->CreateOpStat( FProfiledFileStatsOp::Size ) );
 		int64 Result = FileHandle->Size();
@@ -213,12 +213,12 @@ public:
 	{
 	}
 
-	virtual bool ShouldBeUsed(IPlatformFile* Inner, const TCHAR* CmdLine) const OVERRIDE
+	virtual bool ShouldBeUsed(IPlatformFile* Inner, const TCHAR* CmdLine) const override
 	{
 		return false;
 	}
 
-	virtual bool Initialize(IPlatformFile* Inner, const TCHAR* CommandLineParam) OVERRIDE
+	virtual bool Initialize(IPlatformFile* Inner, const TCHAR* CommandLineParam) override
 	{
 		// Inner is required.
 		check(Inner != NULL);
@@ -227,7 +227,7 @@ public:
 		return !!LowerLevel;
 	}
 
-	virtual IPlatformFile* GetLowerLevel() OVERRIDE
+	virtual IPlatformFile* GetLowerLevel() override
 	{
 		return LowerLevel;
 	}
@@ -275,11 +275,11 @@ public:
 	{
 		return NULL;
 	}
-	virtual const TCHAR* GetName() const OVERRIDE
+	virtual const TCHAR* GetName() const override
 	{
 		return TProfiledPlatformFile::GetTypeName();
 	}
-	virtual bool		FileExists(const TCHAR* Filename) OVERRIDE
+	virtual bool		FileExists(const TCHAR* Filename) override
 	{
 		StatsType* FileStat = CreateStat( Filename );
 		FProfiledFileStatsOp* OpStat = FileStat->CreateOpStat( FProfiledFileStatsOp::Exists );
@@ -287,7 +287,7 @@ public:
 		OpStat->Duration += FPlatformTime::Seconds() * 1000.0 - OpStat->LastOpTime;
 		return Result;
 	}
-	virtual int64		FileSize(const TCHAR* Filename) OVERRIDE
+	virtual int64		FileSize(const TCHAR* Filename) override
 	{
 		StatsType* FileStat = CreateStat( Filename );
 		FProfiledFileStatsOp* OpStat = FileStat->CreateOpStat( FProfiledFileStatsOp::Size );
@@ -295,7 +295,7 @@ public:
 		OpStat->Duration += FPlatformTime::Seconds() * 1000.0 - OpStat->LastOpTime;
 		return Result;
 	}
-	virtual bool		DeleteFile(const TCHAR* Filename) OVERRIDE
+	virtual bool		DeleteFile(const TCHAR* Filename) override
 	{
 		StatsType* FileStat = CreateStat( Filename );
 		FProfiledFileStatsOp* OpStat = FileStat->CreateOpStat( FProfiledFileStatsOp::Delete );
@@ -303,7 +303,7 @@ public:
 		OpStat->Duration += FPlatformTime::Seconds() * 1000.0 - OpStat->LastOpTime;
 		return Result;
 	}
-	virtual bool		IsReadOnly(const TCHAR* Filename) OVERRIDE
+	virtual bool		IsReadOnly(const TCHAR* Filename) override
 	{
 		StatsType* FileStat = CreateStat( Filename );
 		FProfiledFileStatsOp* OpStat = FileStat->CreateOpStat( FProfiledFileStatsOp::IsReadOnly );
@@ -311,7 +311,7 @@ public:
 		OpStat->Duration += FPlatformTime::Seconds() * 1000.0 - OpStat->LastOpTime;
 		return Result;
 	}
-	virtual bool		MoveFile(const TCHAR* To, const TCHAR* From) OVERRIDE
+	virtual bool		MoveFile(const TCHAR* To, const TCHAR* From) override
 	{
 		StatsType* FileStat = CreateStat( From );
 		FProfiledFileStatsOp* OpStat = FileStat->CreateOpStat( FProfiledFileStatsOp::Move );
@@ -319,7 +319,7 @@ public:
 		OpStat->Duration += FPlatformTime::Seconds() * 1000.0 - OpStat->LastOpTime;
 		return Result;
 	}
-	virtual bool		SetReadOnly(const TCHAR* Filename, bool bNewReadOnlyValue) OVERRIDE
+	virtual bool		SetReadOnly(const TCHAR* Filename, bool bNewReadOnlyValue) override
 	{
 		StatsType* FileStat = CreateStat( Filename );
 		FProfiledFileStatsOp* OpStat = FileStat->CreateOpStat( FProfiledFileStatsOp::SetReadOnly );
@@ -327,7 +327,7 @@ public:
 		OpStat->Duration += FPlatformTime::Seconds() * 1000.0 - OpStat->LastOpTime;
 		return Result;
 	}
-	virtual FDateTime	GetTimeStamp(const TCHAR* Filename) OVERRIDE
+	virtual FDateTime	GetTimeStamp(const TCHAR* Filename) override
 	{
 		StatsType* FileStat = CreateStat( Filename );
 		FProfiledFileStatsOp* OpStat = FileStat->CreateOpStat( FProfiledFileStatsOp::GetTimeStamp );
@@ -336,7 +336,7 @@ public:
 		OpStat->Duration += FPlatformTime::Seconds() * 1000 - OpStat->LastOpTime;
 		return Result;
 	}
-	virtual void		SetTimeStamp(const TCHAR* Filename, FDateTime DateTime) OVERRIDE
+	virtual void		SetTimeStamp(const TCHAR* Filename, FDateTime DateTime) override
 	{
 		StatsType* FileStat = CreateStat( Filename );
 		FProfiledFileStatsOp* OpStat = FileStat->CreateOpStat( FProfiledFileStatsOp::SetTimeStamp );
@@ -344,7 +344,7 @@ public:
 		LowerLevel->SetTimeStamp(Filename, DateTime);
 		OpStat->Duration += FPlatformTime::Seconds() * 1000 - OpStat->LastOpTime;
 	}
-	virtual FDateTime	GetAccessTimeStamp(const TCHAR* Filename) OVERRIDE
+	virtual FDateTime	GetAccessTimeStamp(const TCHAR* Filename) override
 	{
 		StatsType* FileStat = CreateStat( Filename );
 		FProfiledFileStatsOp* OpStat = FileStat->CreateOpStat( FProfiledFileStatsOp::GetTimeStamp );
@@ -353,7 +353,7 @@ public:
 		OpStat->Duration += FPlatformTime::Seconds() * 1000 - OpStat->LastOpTime;
 		return Result;
 	}
-	virtual IFileHandle*	OpenRead(const TCHAR* Filename) OVERRIDE
+	virtual IFileHandle*	OpenRead(const TCHAR* Filename) override
 	{
 		StatsType* FileStat = CreateStat( Filename );
 		FProfiledFileStatsOp* OpStat = FileStat->CreateOpStat( FProfiledFileStatsOp::OpenRead );
@@ -361,7 +361,7 @@ public:
 		OpStat->Duration += FPlatformTime::Seconds() * 1000.0 - OpStat->LastOpTime;
 		return Result ? (new TProfiledFileHandle< StatsType >( Result, Filename, FileStat )) : Result;
 	}
-	virtual IFileHandle*	OpenWrite(const TCHAR* Filename, bool bAppend = false, bool bAllowRead = false) OVERRIDE
+	virtual IFileHandle*	OpenWrite(const TCHAR* Filename, bool bAppend = false, bool bAllowRead = false) override
 	{
 		StatsType* FileStat = CreateStat( Filename );
 		FProfiledFileStatsOp* OpStat = FileStat->CreateOpStat( FProfiledFileStatsOp::OpenWrite );
@@ -370,7 +370,7 @@ public:
 		return Result ? (new TProfiledFileHandle< StatsType >( Result, Filename, FileStat )) : Result;
 	}
 
-	virtual bool		DirectoryExists(const TCHAR* Directory) OVERRIDE
+	virtual bool		DirectoryExists(const TCHAR* Directory) override
 	{
 		StatsType* FileStat = CreateStat( Directory );
 		FProfiledFileStatsOp* OpStat = FileStat->CreateOpStat( FProfiledFileStatsOp::Exists );
@@ -378,7 +378,7 @@ public:
 		OpStat->Duration += FPlatformTime::Seconds() * 1000.0 - OpStat->LastOpTime;
 		return Result;
 	}
-	virtual bool		CreateDirectory(const TCHAR* Directory) OVERRIDE
+	virtual bool		CreateDirectory(const TCHAR* Directory) override
 	{
 		StatsType* FileStat = CreateStat( Directory );
 		FProfiledFileStatsOp* OpStat = FileStat->CreateOpStat( FProfiledFileStatsOp::Create );
@@ -386,7 +386,7 @@ public:
 		OpStat->Duration += FPlatformTime::Seconds() * 1000.0 - OpStat->LastOpTime;
 		return Result;
 	}
-	virtual bool		DeleteDirectory(const TCHAR* Directory) OVERRIDE
+	virtual bool		DeleteDirectory(const TCHAR* Directory) override
 	{
 		StatsType* FileStat = CreateStat( Directory );
 		FProfiledFileStatsOp* OpStat = FileStat->CreateOpStat( FProfiledFileStatsOp::Delete );
@@ -395,7 +395,7 @@ public:
 		return Result;
 	}
 
-	virtual bool		IterateDirectory(const TCHAR* Directory, IPlatformFile::FDirectoryVisitor& Visitor) OVERRIDE
+	virtual bool		IterateDirectory(const TCHAR* Directory, IPlatformFile::FDirectoryVisitor& Visitor) override
 	{
 		StatsType* FileStat = CreateStat( Directory );
 		FProfiledFileStatsOp* OpStat = FileStat->CreateOpStat( FProfiledFileStatsOp::Iterate );
@@ -403,7 +403,7 @@ public:
 		OpStat->Duration += FPlatformTime::Seconds() * 1000.0 - OpStat->LastOpTime;
 		return Result;
 	}
-	virtual bool		IterateDirectoryRecursively(const TCHAR* Directory, IPlatformFile::FDirectoryVisitor& Visitor) OVERRIDE
+	virtual bool		IterateDirectoryRecursively(const TCHAR* Directory, IPlatformFile::FDirectoryVisitor& Visitor) override
 	{
 		StatsType* FileStat = CreateStat( Directory );
 		FProfiledFileStatsOp* OpStat = FileStat->CreateOpStat( FProfiledFileStatsOp::Iterate );
@@ -411,7 +411,7 @@ public:
 		OpStat->Duration += FPlatformTime::Seconds() * 1000.0 - OpStat->LastOpTime;
 		return Result;
 	}
-	virtual bool		DeleteDirectoryRecursively(const TCHAR* Directory) OVERRIDE
+	virtual bool		DeleteDirectoryRecursively(const TCHAR* Directory) override
 	{
 		StatsType* FileStat = CreateStat( Directory );
 		FProfiledFileStatsOp* OpStat = FileStat->CreateOpStat( FProfiledFileStatsOp::Delete );
@@ -419,7 +419,7 @@ public:
 		OpStat->Duration += FPlatformTime::Seconds() * 1000.0 - OpStat->LastOpTime;
 		return Result;
 	}
-	virtual bool		CopyFile(const TCHAR* To, const TCHAR* From) OVERRIDE
+	virtual bool		CopyFile(const TCHAR* To, const TCHAR* From) override
 	{
 		StatsType* FileStat = CreateStat( From );
 		FProfiledFileStatsOp* OpStat = FileStat->CreateOpStat( FProfiledFileStatsOp::Copy );
@@ -462,24 +462,24 @@ public:
 	{
 	}
 
-	virtual int64		Tell() OVERRIDE
+	virtual int64		Tell() override
 	{
 		return FileHandle->Tell();
 	}
-	virtual bool		Seek(int64 NewPosition) OVERRIDE
+	virtual bool		Seek(int64 NewPosition) override
 	{
 		return FileHandle->Seek(NewPosition);
 	}
-	virtual bool		SeekFromEnd(int64 NewPositionRelativeToEnd) OVERRIDE
+	virtual bool		SeekFromEnd(int64 NewPositionRelativeToEnd) override
 	{
 		return FileHandle->SeekFromEnd(NewPositionRelativeToEnd);
 	}
-	virtual bool		Read(uint8* Destination, int64 BytesToRead) OVERRIDE;
-	virtual bool		Write(const uint8* Source, int64 BytesToWrite) OVERRIDE
+	virtual bool		Read(uint8* Destination, int64 BytesToRead) override;
+	virtual bool		Write(const uint8* Source, int64 BytesToWrite) override
 	{
 		return FileHandle->Write(Source, BytesToWrite);
 	}
-	virtual int64		Size() OVERRIDE
+	virtual int64		Size() override
 	{
 		return FileHandle->Size();
 	}
@@ -510,7 +510,7 @@ public:
 	{
 	}
 
-	virtual bool ShouldBeUsed(IPlatformFile* Inner, const TCHAR* CmdLine) const OVERRIDE
+	virtual bool ShouldBeUsed(IPlatformFile* Inner, const TCHAR* CmdLine) const override
 	{
 #if STATS
 		bool bResult = FParse::Param(CmdLine, TEXT("FileReadStats"));
@@ -520,9 +520,9 @@ public:
 #endif
 	}
 
-	virtual bool Initialize(IPlatformFile* Inner, const TCHAR* CommandLineParam) OVERRIDE;
+	virtual bool Initialize(IPlatformFile* Inner, const TCHAR* CommandLineParam) override;
 	bool Tick(float Delta);
-	virtual IPlatformFile* GetLowerLevel() OVERRIDE
+	virtual IPlatformFile* GetLowerLevel() override
 	{
 		return LowerLevel;
 	}
@@ -530,83 +530,83 @@ public:
 	{
 		return TEXT("FileReadStats");
 	}
-	virtual const TCHAR* GetName() const OVERRIDE
+	virtual const TCHAR* GetName() const override
 	{
 		return GetTypeName();
 	}
-	virtual bool		FileExists(const TCHAR* Filename) OVERRIDE
+	virtual bool		FileExists(const TCHAR* Filename) override
 	{
 		return LowerLevel->FileExists(Filename);
 	}
-	virtual int64		FileSize(const TCHAR* Filename) OVERRIDE
+	virtual int64		FileSize(const TCHAR* Filename) override
 	{
 		return LowerLevel->FileSize(Filename);
 	}
-	virtual bool		DeleteFile(const TCHAR* Filename) OVERRIDE
+	virtual bool		DeleteFile(const TCHAR* Filename) override
 	{
 		return LowerLevel->DeleteFile(Filename);
 	}
-	virtual bool		IsReadOnly(const TCHAR* Filename) OVERRIDE
+	virtual bool		IsReadOnly(const TCHAR* Filename) override
 	{
 		return LowerLevel->IsReadOnly(Filename);
 	}
-	virtual bool		MoveFile(const TCHAR* To, const TCHAR* From) OVERRIDE
+	virtual bool		MoveFile(const TCHAR* To, const TCHAR* From) override
 	{
 		return LowerLevel->MoveFile(To, From);
 	}
-	virtual bool		SetReadOnly(const TCHAR* Filename, bool bNewReadOnlyValue) OVERRIDE
+	virtual bool		SetReadOnly(const TCHAR* Filename, bool bNewReadOnlyValue) override
 	{
 		return LowerLevel->SetReadOnly(Filename, bNewReadOnlyValue);
 	}
-	virtual FDateTime	GetTimeStamp(const TCHAR* Filename) OVERRIDE
+	virtual FDateTime	GetTimeStamp(const TCHAR* Filename) override
 	{
 		return LowerLevel->GetTimeStamp(Filename);
 	}
-	virtual void		SetTimeStamp(const TCHAR* Filename, FDateTime DateTime) OVERRIDE
+	virtual void		SetTimeStamp(const TCHAR* Filename, FDateTime DateTime) override
 	{
 		LowerLevel->SetTimeStamp(Filename, DateTime);
 	}
-	virtual FDateTime	GetAccessTimeStamp(const TCHAR* Filename) OVERRIDE
+	virtual FDateTime	GetAccessTimeStamp(const TCHAR* Filename) override
 	{
 		return LowerLevel->GetAccessTimeStamp(Filename);
 	}
-	virtual IFileHandle*	OpenRead(const TCHAR* Filename) OVERRIDE
+	virtual IFileHandle*	OpenRead(const TCHAR* Filename) override
 	{
 		IFileHandle* Result = LowerLevel->OpenRead(Filename);
 		return Result ? (new FPlatformFileReadStatsHandle(Result, Filename, &BytePerSecThisTick, &BytesReadThisTick, &ReadsThisTick)) : Result;
 	}
-	virtual IFileHandle*	OpenWrite(const TCHAR* Filename, bool bAppend = false, bool bAllowRead = false) OVERRIDE
+	virtual IFileHandle*	OpenWrite(const TCHAR* Filename, bool bAppend = false, bool bAllowRead = false) override
 	{
 		IFileHandle* Result = LowerLevel->OpenWrite(Filename, bAppend, bAllowRead);
 		return Result ? (new FPlatformFileReadStatsHandle(Result, Filename, &BytePerSecThisTick, &BytesReadThisTick, &ReadsThisTick)) : Result;
 	}
 
-	virtual bool		DirectoryExists(const TCHAR* Directory) OVERRIDE
+	virtual bool		DirectoryExists(const TCHAR* Directory) override
 	{
 		return LowerLevel->DirectoryExists(Directory);
 	}
-	virtual bool		CreateDirectory(const TCHAR* Directory) OVERRIDE
+	virtual bool		CreateDirectory(const TCHAR* Directory) override
 	{
 		return LowerLevel->CreateDirectory(Directory);
 	}
-	virtual bool		DeleteDirectory(const TCHAR* Directory) OVERRIDE
+	virtual bool		DeleteDirectory(const TCHAR* Directory) override
 	{
 		return LowerLevel->DeleteDirectory(Directory);
 	}
 
-	virtual bool		IterateDirectory(const TCHAR* Directory, IPlatformFile::FDirectoryVisitor& Visitor) OVERRIDE
+	virtual bool		IterateDirectory(const TCHAR* Directory, IPlatformFile::FDirectoryVisitor& Visitor) override
 	{
 		return LowerLevel->IterateDirectory( Directory, Visitor );
 	}
-	virtual bool		IterateDirectoryRecursively(const TCHAR* Directory, IPlatformFile::FDirectoryVisitor& Visitor) OVERRIDE
+	virtual bool		IterateDirectoryRecursively(const TCHAR* Directory, IPlatformFile::FDirectoryVisitor& Visitor) override
 	{
 		return LowerLevel->IterateDirectoryRecursively( Directory, Visitor );
 	}
-	virtual bool		DeleteDirectoryRecursively(const TCHAR* Directory) OVERRIDE
+	virtual bool		DeleteDirectoryRecursively(const TCHAR* Directory) override
 	{
 		return LowerLevel->DeleteDirectoryRecursively( Directory );
 	}
-	virtual bool		CopyFile(const TCHAR* To, const TCHAR* From) OVERRIDE
+	virtual bool		CopyFile(const TCHAR* To, const TCHAR* From) override
 	{
 		return LowerLevel->CopyFile( To, From );
 	}
diff --git a/Engine/Source/Runtime/Core/Public/HAL/MallocAnsi.h b/Engine/Source/Runtime/Core/Public/HAL/MallocAnsi.h
--- a/Engine/Source/Runtime/Core/Public/HAL/MallocAnsi.h
+++ b/Engine/Source/Runtime/Core/Public/HAL/MallocAnsi.h
@@ -38,7 +38,7 @@ public:
 	}
 
 	// FMalloc interface.
-	virtual void* Malloc( SIZE_T Size, uint32 Alignment ) OVERRIDE
+	virtual void* Malloc( SIZE_T Size, uint32 Alignment ) override
 	{
 		Alignment = FMath::Max(Size >= 16 ? (uint32)16 : (uint32)8, Alignment);
 
@@ -59,7 +59,7 @@ public:
 		return Result;
 	}
 
-	virtual void* Realloc( void* Ptr, SIZE_T NewSize, uint32 Alignment ) OVERRIDE
+	virtual void* Realloc( void* Ptr, SIZE_T NewSize, uint32 Alignment ) override
 	{
 		void* Result;
 		Alignment = FMath::Max(NewSize >= 16 ? (uint32)16 : (uint32)8, Alignment);
@@ -104,7 +104,7 @@ public:
 		return Result;
 	}
 
-	virtual void Free( void* Ptr ) OVERRIDE
+	virtual void Free( void* Ptr ) override
 	{
 #if USE_ALIGNED_MALLOC
 		_aligned_free( Ptr );
@@ -116,7 +116,7 @@ public:
 #endif
 	}
 
-	virtual bool GetAllocationSize(void *Original, SIZE_T &SizeOut) OVERRIDE
+	virtual bool GetAllocationSize(void *Original, SIZE_T &SizeOut) override
 	{ 
 		return false;
 	}
@@ -157,7 +157,7 @@ public:
 	 *
 	 * @return true as we're using system allocator
 	 */
-	virtual bool IsInternallyThreadSafe() const OVERRIDE
+	virtual bool IsInternallyThreadSafe() const override
 	{
 #if PLATFORM_MAC
 			return true;
@@ -171,7 +171,7 @@ public:
 	/**
 	 * Validates the allocator's heap
 	 */
-	virtual bool ValidateHeap() OVERRIDE
+	virtual bool ValidateHeap() override
 	{
 #if PLATFORM_WINDOWS
 		int32 Result = _heapchk();
@@ -186,6 +186,6 @@ public:
 		return true;
 	}
 
-	virtual const TCHAR * GetDescriptiveName() OVERRIDE { return TEXT("ANSI"); }
+	virtual const TCHAR * GetDescriptiveName() override { return TEXT("ANSI"); }
 };
 
diff --git a/Engine/Source/Runtime/Core/Public/HAL/MallocBinned.h b/Engine/Source/Runtime/Core/Public/HAL/MallocBinned.h
--- a/Engine/Source/Runtime/Core/Public/HAL/MallocBinned.h
+++ b/Engine/Source/Runtime/Core/Public/HAL/MallocBinned.h
@@ -878,7 +878,7 @@ public:
 	 * Returns if the allocator is guaranteed to be thread-safe and therefore
 	 * doesn't need a unnecessary thread-safety wrapper around it.
 	 */
-	virtual bool IsInternallyThreadSafe() const OVERRIDE
+	virtual bool IsInternallyThreadSafe() const override
 	{ 
 #ifdef USE_INTERNAL_LOCKS
 		return true;
@@ -890,7 +890,7 @@ public:
 	/** 
 	 * Malloc
 	 */
-	virtual void* Malloc( SIZE_T Size, uint32 Alignment ) OVERRIDE
+	virtual void* Malloc( SIZE_T Size, uint32 Alignment ) override
 	{
 #ifdef USE_COARSE_GRAIN_LOCKS
 		FScopeLock ScopedLock(&AccessGuard);
@@ -986,7 +986,7 @@ public:
 	/** 
 	 * Realloc
 	 */
-	virtual void* Realloc( void* Ptr, SIZE_T NewSize, uint32 Alignment ) OVERRIDE
+	virtual void* Realloc( void* Ptr, SIZE_T NewSize, uint32 Alignment ) override
 	{
 		// Handle DEFAULT_ALIGNMENT for binned allocator.
 		if (Alignment == DEFAULT_ALIGNMENT)
@@ -1055,7 +1055,7 @@ public:
 	/** 
 	 * Free
 	 */
-	virtual void Free( void* Ptr ) OVERRIDE
+	virtual void Free( void* Ptr ) override
 	{
 		if( !Ptr )
 		{
@@ -1072,7 +1072,7 @@ public:
 	 * @param SizeOut - If possible, this value is set to the size of the passed in pointer
 	 * @return true if succeeded
 	 */
-	virtual bool GetAllocationSize(void *Original, SIZE_T &SizeOut) OVERRIDE
+	virtual bool GetAllocationSize(void *Original, SIZE_T &SizeOut) override
 	{
 		if (!Original)
 		{
@@ -1087,7 +1087,7 @@ public:
 	/**
 	 * Validates the allocator's heap
 	 */
-	virtual bool ValidateHeap() OVERRIDE
+	virtual bool ValidateHeap() override
 	{
 #ifdef USE_INTERNAL_LOCKS
 		FScopeLock ScopedLock(&AccessGuard);
@@ -1118,7 +1118,7 @@ public:
 
 
 	/** Called once per frame, gathers and sets all memory allocator statistics into the corresponding stats. */
-	virtual void UpdateStats() OVERRIDE
+	virtual void UpdateStats() override
 	{
 #if STATS
 		SIZE_T	LocalOsCurrent = 0;
@@ -1162,14 +1162,14 @@ public:
 	}
 
 	/** Writes allocator stats from the last update into the specified destination. */
-	virtual void GetAllocatorStats( FGenericMemoryStats& out_Stats ) OVERRIDE;
+	virtual void GetAllocatorStats( FGenericMemoryStats& out_Stats ) override;
 
 	/**
 	 * Dumps allocator stats to an output device. Subclasses should override to add additional info
 	 *
 	 * @param Ar	[in] Output device
 	 */
-	virtual void DumpAllocatorStats( class FOutputDevice& Ar ) OVERRIDE
+	virtual void DumpAllocatorStats( class FOutputDevice& Ar ) override
 	{
 		FBufferedOutputDevice BufferedOutput;
 		{
@@ -1258,7 +1258,7 @@ public:
 		BufferedOutput.RedirectTo( Ar );
 	}
 
-	virtual const TCHAR * GetDescriptiveName() OVERRIDE { return TEXT("binned"); }
+	virtual const TCHAR * GetDescriptiveName() override { return TEXT("binned"); }
 
 protected:
 
diff --git a/Engine/Source/Runtime/Core/Public/HAL/MallocDebug.h b/Engine/Source/Runtime/Core/Public/HAL/MallocDebug.h
--- a/Engine/Source/Runtime/Core/Public/HAL/MallocDebug.h
+++ b/Engine/Source/Runtime/Core/Public/HAL/MallocDebug.h
@@ -55,7 +55,7 @@ public:
 	/** 
 	 * Malloc
 	 */
-	virtual void* Malloc( SIZE_T Size, uint32 Alignment ) OVERRIDE
+	virtual void* Malloc( SIZE_T Size, uint32 Alignment ) override
 	{
 		check(Alignment == DEFAULT_ALIGNMENT && "Alignment currently unsupported in this allocator");
 		FMemDebug* Ptr = NULL;
@@ -87,7 +87,7 @@ public:
 	/** 
 	 * Realloc
 	 */
-	virtual void* Realloc( void* InPtr, SIZE_T NewSize, uint32 Alignment ) OVERRIDE
+	virtual void* Realloc( void* InPtr, SIZE_T NewSize, uint32 Alignment ) override
 	{
 		if( InPtr && NewSize )
 		{
@@ -112,7 +112,7 @@ public:
 	/** 
 	 * Free
 	 */
-	virtual void Free( void* InPtr ) OVERRIDE
+	virtual void Free( void* InPtr ) override
 	{
 		if( !InPtr )
 		{
@@ -150,7 +150,7 @@ public:
 	 * @param SizeOut - If possible, this value is set to the size of the passed in pointer
 	 * @return true if succeeded
 	 */
-	virtual bool GetAllocationSize(void *Original, SIZE_T &SizeOut) OVERRIDE
+	virtual bool GetAllocationSize(void *Original, SIZE_T &SizeOut) override
 	{
 		if( !Original )
 		{
@@ -170,7 +170,7 @@ public:
 	 *
 	 * @param Ar	[in] Output device
 	 */
-	virtual void DumpAllocatorStats( FOutputDevice& Ar ) OVERRIDE
+	virtual void DumpAllocatorStats( FOutputDevice& Ar ) override
 	{
 		Ar.Logf( TEXT( "Total Allocation Size: %u" ), TotalAllocationSize );
 		Ar.Logf( TEXT( "Total Waste Size: %u" ), TotalWasteSize );
@@ -193,7 +193,7 @@ public:
 	/**
 	 * Validates the allocator's heap
 	 */
-	virtual bool ValidateHeap() OVERRIDE
+	virtual bool ValidateHeap() override
 	{
 		for( FMemDebug** Link = &GFirstDebug; *Link; Link=&(*Link)->Next )
 		{
@@ -210,12 +210,12 @@ public:
 		return( true );
 	}
 
-	virtual bool Exec( UWorld* InWorld, const TCHAR* Cmd, FOutputDevice& Ar ) OVERRIDE
+	virtual bool Exec( UWorld* InWorld, const TCHAR* Cmd, FOutputDevice& Ar ) override
 	{
 		return false;
 	}
 
-	virtual const TCHAR * GetDescriptiveName() OVERRIDE { return TEXT("debug"); }
+	virtual const TCHAR * GetDescriptiveName() override { return TEXT("debug"); }
 
 };
 
diff --git a/Engine/Source/Runtime/Core/Public/HAL/MallocJemalloc.h b/Engine/Source/Runtime/Core/Public/HAL/MallocJemalloc.h
--- a/Engine/Source/Runtime/Core/Public/HAL/MallocJemalloc.h
+++ b/Engine/Source/Runtime/Core/Public/HAL/MallocJemalloc.h
@@ -34,13 +34,13 @@ public:
 	
 	// Begin FMalloc interface.
 
-	virtual void* Malloc( SIZE_T Size, uint32 Alignment ) OVERRIDE;
-	virtual void* Realloc( void* Ptr, SIZE_T NewSize, uint32 Alignment ) OVERRIDE;
-	virtual void Free( void* Ptr ) OVERRIDE;
-	virtual void DumpAllocatorStats( FOutputDevice& Ar ) OVERRIDE;
-	virtual bool GetAllocationSize(void *Original, SIZE_T &SizeOut) OVERRIDE;
-	virtual bool IsInternallyThreadSafe() const OVERRIDE {  return true; }
-	virtual const TCHAR * GetDescriptiveName() OVERRIDE { return TEXT("jemalloc"); }
+	virtual void* Malloc( SIZE_T Size, uint32 Alignment ) override;
+	virtual void* Realloc( void* Ptr, SIZE_T NewSize, uint32 Alignment ) override;
+	virtual void Free( void* Ptr ) override;
+	virtual void DumpAllocatorStats( FOutputDevice& Ar ) override;
+	virtual bool GetAllocationSize(void *Original, SIZE_T &SizeOut) override;
+	virtual bool IsInternallyThreadSafe() const override {  return true; }
+	virtual const TCHAR * GetDescriptiveName() override { return TEXT("jemalloc"); }
 
 	// End FMalloc interface.
 
diff --git a/Engine/Source/Runtime/Core/Public/HAL/MallocTBB.h b/Engine/Source/Runtime/Core/Public/HAL/MallocTBB.h
--- a/Engine/Source/Runtime/Core/Public/HAL/MallocTBB.h
+++ b/Engine/Source/Runtime/Core/Public/HAL/MallocTBB.h
@@ -31,17 +31,17 @@ public:
 	/** 
 	 * Malloc
 	 */
-	virtual void* Malloc( SIZE_T Size, uint32 Alignment ) OVERRIDE;
+	virtual void* Malloc( SIZE_T Size, uint32 Alignment ) override;
 	
 	/** 
 	 * Realloc
 	 */
-	virtual void* Realloc( void* Ptr, SIZE_T NewSize, uint32 Alignment ) OVERRIDE;
+	virtual void* Realloc( void* Ptr, SIZE_T NewSize, uint32 Alignment ) override;
 	
 	/** 
 	 * Free
 	 */
-	virtual void Free( void* Ptr ) OVERRIDE;
+	virtual void Free( void* Ptr ) override;
 
 	/**
 	 * If possible determine the size of the memory allocated at the given address
@@ -50,17 +50,17 @@ public:
 	 * @param SizeOut - If possible, this value is set to the size of the passed in pointer
 	 * @return true if succeeded
 	 */
-	virtual bool GetAllocationSize(void *Original, SIZE_T &SizeOut) OVERRIDE;
+	virtual bool GetAllocationSize(void *Original, SIZE_T &SizeOut) override;
 
 	/**
 	 * Returns if the allocator is guaranteed to be thread-safe and therefore
 	 * doesn't need a unnecessary thread-safety wrapper around it.
 	 */
-	virtual bool IsInternallyThreadSafe() const OVERRIDE
+	virtual bool IsInternallyThreadSafe() const override
 	{ 
 		return true;
 	}
 
-	virtual const TCHAR * GetDescriptiveName() OVERRIDE { return TEXT("TBB"); }
+	virtual const TCHAR * GetDescriptiveName() override { return TEXT("TBB"); }
 };
 
diff --git a/Engine/Source/Runtime/Core/Public/HAL/MallocThreadSafeProxy.h b/Engine/Source/Runtime/Core/Public/HAL/MallocThreadSafeProxy.h
--- a/Engine/Source/Runtime/Core/Public/HAL/MallocThreadSafeProxy.h
+++ b/Engine/Source/Runtime/Core/Public/HAL/MallocThreadSafeProxy.h
@@ -40,14 +40,14 @@ public:
 	 * @param Alignment		The alignment of a hypothetical allocation request
 	 * @return				Returns the usable size that the allocation request would return. In other words you can ask for this greater amount without using any more actual memory.
 	 */
-	virtual SIZE_T QuantizeSize( SIZE_T Size, uint32 Alignment ) OVERRIDE
+	virtual SIZE_T QuantizeSize( SIZE_T Size, uint32 Alignment ) override
 	{
 		return UsedMalloc->QuantizeSize(Size,Alignment); 
 	}
 	/** 
 	 * Malloc
 	 */
-	void* Malloc( SIZE_T Size, uint32 Alignment ) OVERRIDE
+	void* Malloc( SIZE_T Size, uint32 Alignment ) override
 	{
 		FScopeLock ScopeLock( &SynchronizationObject );
 		STAT(TotalMallocCalls++);
@@ -57,7 +57,7 @@ public:
 	/** 
 	 * Realloc
 	 */
-	void* Realloc( void* Ptr, SIZE_T NewSize, uint32 Alignment ) OVERRIDE
+	void* Realloc( void* Ptr, SIZE_T NewSize, uint32 Alignment ) override
 	{
 		FScopeLock ScopeLock( &SynchronizationObject );
 		STAT(TotalReallocCalls++);
@@ -67,7 +67,7 @@ public:
 	/** 
 	 * Free
 	 */
-	void Free( void* Ptr ) OVERRIDE
+	void Free( void* Ptr ) override
 	{
 		if( Ptr )
 		{
@@ -101,13 +101,13 @@ public:
 	/**
 	 * Validates the allocator's heap
 	 */
-	virtual bool ValidateHeap() OVERRIDE
+	virtual bool ValidateHeap() override
 	{
 		FScopeLock Lock( &SynchronizationObject );
 		return( UsedMalloc->ValidateHeap() );
 	}
 
-	bool Exec( UWorld* InWorld, const TCHAR* Cmd, FOutputDevice& Ar ) OVERRIDE
+	bool Exec( UWorld* InWorld, const TCHAR* Cmd, FOutputDevice& Ar ) override
 	{
 		FScopeLock ScopeLock( &SynchronizationObject );
 		return UsedMalloc->Exec( InWorld, Cmd, Ar);
@@ -120,13 +120,13 @@ public:
 	* @param SizeOut - If possible, this value is set to the size of the passed in pointer
 	* @return true if succeeded
 	*/
-	bool GetAllocationSize(void *Original, SIZE_T &SizeOut) OVERRIDE
+	bool GetAllocationSize(void *Original, SIZE_T &SizeOut) override
 	{
 		FScopeLock ScopeLock( &SynchronizationObject );
 		return UsedMalloc->GetAllocationSize(Original,SizeOut);
 	}
 
-	virtual const TCHAR * GetDescriptiveName() OVERRIDE
+	virtual const TCHAR * GetDescriptiveName() override
 	{ 
 		FScopeLock ScopeLock( &SynchronizationObject );
 		check(UsedMalloc);
diff --git a/Engine/Source/Runtime/Core/Public/HAL/MemoryBase.h b/Engine/Source/Runtime/Core/Public/HAL/MemoryBase.h
--- a/Engine/Source/Runtime/Core/Public/HAL/MemoryBase.h
+++ b/Engine/Source/Runtime/Core/Public/HAL/MemoryBase.h
@@ -111,7 +111,7 @@ public:
 	/** 
 	 * Handles any commands passed in on the command line
 	 */
-	virtual bool Exec( UWorld* InWorld, const TCHAR* Cmd, FOutputDevice& Ar ) OVERRIDE
+	virtual bool Exec( UWorld* InWorld, const TCHAR* Cmd, FOutputDevice& Ar ) override
 	{ 
 		return false; 
 	}
diff --git a/Engine/Source/Runtime/Core/Public/HAL/OutputDevices.h b/Engine/Source/Runtime/Core/Public/HAL/OutputDevices.h
--- a/Engine/Source/Runtime/Core/Public/HAL/OutputDevices.h
+++ b/Engine/Source/Runtime/Core/Public/HAL/OutputDevices.h
@@ -352,7 +352,7 @@ public:
 	 * @param	Data	Text to log
 	 * @param	Event	Event name used for suppression purposes
 	 */
-	virtual void Serialize( const TCHAR* Data, ELogVerbosity::Type Verbosity, const class FName& Category) OVERRIDE;
+	virtual void Serialize( const TCHAR* Data, ELogVerbosity::Type Verbosity, const class FName& Category) override;
 
 	/**
 	 * Passes on the flush request to all current output devices.
@@ -404,9 +404,9 @@ public:
 	 */
 	void Flush();
 
-	virtual void Serialize( const TCHAR* Data, ELogVerbosity::Type Verbosity, const class FName& Category ) OVERRIDE;
+	virtual void Serialize( const TCHAR* Data, ELogVerbosity::Type Verbosity, const class FName& Category ) override;
 
-	virtual bool CanBeUsedOnAnyThread() const OVERRIDE
+	virtual bool CanBeUsedOnAnyThread() const override
 	{
 		return true;
 	}
@@ -433,7 +433,7 @@ public:
 	 * @param	Data	unused
 	 * @param	Event	unused
 	 */
-	virtual void Serialize( const TCHAR* V, ELogVerbosity::Type Verbosity, const class FName& Category ) OVERRIDE
+	virtual void Serialize( const TCHAR* V, ELogVerbosity::Type Verbosity, const class FName& Category ) override
 	{}
 };
 
@@ -446,9 +446,9 @@ public:
 	 * @param	Data	Text to log
 	 * @param	Event	Event name used for suppression purposes
 	 */
-	virtual void Serialize( const TCHAR* Data, ELogVerbosity::Type Verbosity, const class FName& Category ) OVERRIDE;
+	virtual void Serialize( const TCHAR* Data, ELogVerbosity::Type Verbosity, const class FName& Category ) override;
 
-	virtual bool CanBeUsedOnAnyThread() const OVERRIDE
+	virtual bool CanBeUsedOnAnyThread() const override
 	{
 		return true;
 	}
@@ -460,7 +460,7 @@ class FBufferedOutputDevice : public FOutputDevice
 	TArray<FBufferedLine> BufferedLines;
 
 public:
-	virtual void Serialize( const TCHAR* InData, ELogVerbosity::Type Verbosity, const class FName& Category ) OVERRIDE
+	virtual void Serialize( const TCHAR* InData, ELogVerbosity::Type Verbosity, const class FName& Category ) override
 	{
 		new(BufferedLines)FBufferedLine( InData, Verbosity, Category );
 	}
@@ -491,9 +491,9 @@ public:
 	 * @param	Data	Text to log
 	 * @param	Event	Event name used for suppression purposes
 	 */
-	virtual void Serialize( const TCHAR* Msg, ELogVerbosity::Type Verbosity, const class FName& Category ) OVERRIDE;
+	virtual void Serialize( const TCHAR* Msg, ELogVerbosity::Type Verbosity, const class FName& Category ) override;
 
-	virtual bool CanBeUsedOnAnyThread() const OVERRIDE
+	virtual bool CanBeUsedOnAnyThread() const override
 	{
 		return true;
 	}
diff --git a/Engine/Source/Runtime/Core/Public/HAL/PThreadEvent.h b/Engine/Source/Runtime/Core/Public/HAL/PThreadEvent.h
--- a/Engine/Source/Runtime/Core/Public/HAL/PThreadEvent.h
+++ b/Engine/Source/Runtime/Core/Public/HAL/PThreadEvent.h
@@ -91,7 +91,7 @@ public:
 		}
 	}
 
-	virtual bool Create(bool _bIsManualReset = false) OVERRIDE
+	virtual bool Create(bool _bIsManualReset = false) override
 	{
 		check(!bInitialized);
 		bool RetVal = false;
@@ -113,7 +113,7 @@ public:
 		return RetVal;
 	}
 
-	virtual void Trigger() OVERRIDE
+	virtual void Trigger() override
 	{
 		check(bInitialized);
 
@@ -138,7 +138,7 @@ public:
 		UnlockEventMutex();
 	}
 
-	virtual void Reset() OVERRIDE
+	virtual void Reset() override
 	{
 		check(bInitialized);
 		LockEventMutex();
@@ -146,7 +146,7 @@ public:
 		UnlockEventMutex();
 	}
 
-	virtual bool Wait(uint32 WaitTime = (uint32)-1) OVERRIDE;
+	virtual bool Wait(uint32 WaitTime = (uint32)-1) override;
 
 };
 
diff --git a/Engine/Source/Runtime/Core/Public/HTML5/HTML5Application.h b/Engine/Source/Runtime/Core/Public/HTML5/HTML5Application.h
--- a/Engine/Source/Runtime/Core/Public/HTML5/HTML5Application.h
+++ b/Engine/Source/Runtime/Core/Public/HTML5/HTML5Application.h
@@ -21,11 +21,11 @@ public:
 
 	void SetMessageHandler( const TSharedRef< FGenericApplicationMessageHandler >& InMessageHandler );
 
-	virtual void PollGameDeviceState( const float TimeDelta ) OVERRIDE;
+	virtual void PollGameDeviceState( const float TimeDelta ) override;
 
-	virtual FPlatformRect GetWorkArea( const FPlatformRect& CurrentWindow ) const OVERRIDE;
+	virtual FPlatformRect GetWorkArea( const FPlatformRect& CurrentWindow ) const override;
 
-	virtual void GetDisplayMetrics( FDisplayMetrics& OutDesktopMetrics ) const OVERRIDE;
+	virtual void GetDisplayMetrics( FDisplayMetrics& OutDesktopMetrics ) const override;
 
 
 	TSharedRef< FGenericWindow > MakeWindow();
diff --git a/Engine/Source/Runtime/Core/Public/HTML5/HTML5Cursor.h b/Engine/Source/Runtime/Core/Public/HTML5/HTML5Cursor.h
--- a/Engine/Source/Runtime/Core/Public/HTML5/HTML5Cursor.h
+++ b/Engine/Source/Runtime/Core/Public/HTML5/HTML5Cursor.h
@@ -12,17 +12,17 @@ public:
 
 	virtual ~FHTML5Cursor();
 
-	virtual FVector2D GetPosition() const OVERRIDE;
+	virtual FVector2D GetPosition() const override;
 
-	virtual void SetPosition( const int32 X, const int32 Y ) OVERRIDE;
+	virtual void SetPosition( const int32 X, const int32 Y ) override;
 
-	virtual void SetType( const EMouseCursor::Type InNewCursor ) OVERRIDE;
+	virtual void SetType( const EMouseCursor::Type InNewCursor ) override;
 
-	virtual void GetSize( int32& Width, int32& Height ) const OVERRIDE;
+	virtual void GetSize( int32& Width, int32& Height ) const override;
 
-	virtual void Show( bool bShow ) OVERRIDE;
+	virtual void Show( bool bShow ) override;
 
-	virtual void Lock( const RECT* const Bounds ) OVERRIDE;
+	virtual void Lock( const RECT* const Bounds ) override;
 
 
 private:
diff --git a/Engine/Source/Runtime/Core/Public/HTML5/HTML5Window.h b/Engine/Source/Runtime/Core/Public/HTML5/HTML5Window.h
--- a/Engine/Source/Runtime/Core/Public/HTML5/HTML5Window.h
+++ b/Engine/Source/Runtime/Core/Public/HTML5/HTML5Window.h
@@ -22,12 +22,12 @@ public:
 	static TSharedRef<FHTML5Window> Make();
 
 	
-	virtual void* GetOSWindowHandle() const OVERRIDE { return NULL; } //can be null.
+	virtual void* GetOSWindowHandle() const override { return NULL; } //can be null.
 
 	void Initialize( class FHTML5Application* const Application, const TSharedRef< FGenericWindowDefinition >& InDefinition, const TSharedPtr< FHTML5Window >& InParent, const bool bShowImmediately );
 
 	/** Returns the rectangle of the screen the window is associated with */
-	virtual bool GetFullScreenInfo( int32& X, int32& Y, int32& Width, int32& Height ) const OVERRIDE;
+	virtual bool GetFullScreenInfo( int32& X, int32& Y, int32& Width, int32& Height ) const override;
 
 	virtual void ReshapeWindow(int32 X, int32 Y, int32 Width, int32 Height);
 
@@ -40,7 +40,7 @@ public:
 
 protected:
 	/** @return true if the native window is currently in fullscreen mode, false otherwise */
-	virtual EWindowMode::Type GetWindowMode() const OVERRIDE;
+	virtual EWindowMode::Type GetWindowMode() const override;
 
 private:
 	/**
diff --git a/Engine/Source/Runtime/Core/Public/IOS/IOSApplication.h b/Engine/Source/Runtime/Core/Public/IOS/IOSApplication.h
--- a/Engine/Source/Runtime/Core/Public/IOS/IOSApplication.h
+++ b/Engine/Source/Runtime/Core/Public/IOS/IOSApplication.h
@@ -19,16 +19,16 @@ public:
 
 	void SetMessageHandler( const TSharedRef< FGenericApplicationMessageHandler >& InMessageHandler );
 
-	virtual void PollGameDeviceState( const float TimeDelta ) OVERRIDE;
+	virtual void PollGameDeviceState( const float TimeDelta ) override;
 
-	virtual FPlatformRect GetWorkArea( const FPlatformRect& CurrentWindow ) const OVERRIDE;
+	virtual FPlatformRect GetWorkArea( const FPlatformRect& CurrentWindow ) const override;
 
-	virtual void GetDisplayMetrics( FDisplayMetrics& OutDisplayMetrics ) const OVERRIDE;
+	virtual void GetDisplayMetrics( FDisplayMetrics& OutDisplayMetrics ) const override;
 
-	virtual TSharedRef< FGenericWindow > MakeWindow() OVERRIDE;
+	virtual TSharedRef< FGenericWindow > MakeWindow() override;
 
 protected:
-	virtual void InitializeWindow( const TSharedRef< FGenericWindow >& Window, const TSharedRef< FGenericWindowDefinition >& InDefinition, const TSharedPtr< FGenericWindow >& InParent, const bool bShowImmediately ) OVERRIDE;
+	virtual void InitializeWindow( const TSharedRef< FGenericWindow >& Window, const TSharedRef< FGenericWindowDefinition >& InDefinition, const TSharedPtr< FGenericWindow >& InParent, const bool bShowImmediately ) override;
 
 private:
 
diff --git a/Engine/Source/Runtime/Core/Public/IOS/IOSInputInterface.h b/Engine/Source/Runtime/Core/Public/IOS/IOSInputInterface.h
--- a/Engine/Source/Runtime/Core/Public/IOS/IOSInputInterface.h
+++ b/Engine/Source/Runtime/Core/Public/IOS/IOSInputInterface.h
@@ -45,7 +45,7 @@ public:
 	static void QueueTouchInput(TArray<TouchInput> InTouchEvents);
 
 	// Begin Exec Interface
-	virtual bool Exec(UWorld* InWorld, const TCHAR* Cmd, FOutputDevice& Ar) OVERRIDE;
+	virtual bool Exec(UWorld* InWorld, const TCHAR* Cmd, FOutputDevice& Ar) override;
 	// End Exec Interface
 
 private:
diff --git a/Engine/Source/Runtime/Core/Public/Internationalization/InternationalizationArchiveJsonSerializer.h b/Engine/Source/Runtime/Core/Public/Internationalization/InternationalizationArchiveJsonSerializer.h
--- a/Engine/Source/Runtime/Core/Public/Internationalization/InternationalizationArchiveJsonSerializer.h
+++ b/Engine/Source/Runtime/Core/Public/Internationalization/InternationalizationArchiveJsonSerializer.h
@@ -37,7 +37,7 @@ public:
 	 *
 	 * @return true if deserialization was successful, false otherwise.
 	 */
-	virtual bool DeserializeArchive( FArchive& Archive, TSharedRef< FInternationalizationArchive > InternationalizationArchive ) OVERRIDE;
+	virtual bool DeserializeArchive( FArchive& Archive, TSharedRef< FInternationalizationArchive > InternationalizationArchive ) override;
 
 	/**
 	 * Deserializes a Internationalization archive from a string.
@@ -67,7 +67,7 @@ public:
 	 *
 	 * @return true if serialization was successful, false otherwise.
 	 */
-	virtual bool SerializeArchive( TSharedRef< const FInternationalizationArchive > InternationalizationArchive, FArchive& Archive ) OVERRIDE;
+	virtual bool SerializeArchive( TSharedRef< const FInternationalizationArchive > InternationalizationArchive, FArchive& Archive ) override;
 
 	/**
 	 * Serializes a Internationalization archive to a string.
diff --git a/Engine/Source/Runtime/Core/Public/Internationalization/InternationalizationManifestJsonSerializer.h b/Engine/Source/Runtime/Core/Public/Internationalization/InternationalizationManifestJsonSerializer.h
--- a/Engine/Source/Runtime/Core/Public/Internationalization/InternationalizationManifestJsonSerializer.h
+++ b/Engine/Source/Runtime/Core/Public/Internationalization/InternationalizationManifestJsonSerializer.h
@@ -32,7 +32,7 @@ public:
 	 *
 	 * @return true if deserialization was successful, false otherwise.
 	 */
-	virtual bool DeserializeManifest( FArchive& Archive, TSharedRef< FInternationalizationManifest > Manifest ) OVERRIDE;
+	virtual bool DeserializeManifest( FArchive& Archive, TSharedRef< FInternationalizationManifest > Manifest ) override;
 
 	/**
 	 * Deserializes a Internationalization manifest from a string.
@@ -62,7 +62,7 @@ public:
 	 *
 	 * @return true if serialization was successful, false otherwise.
 	 */
-	virtual bool SerializeManifest(TSharedRef< const FInternationalizationManifest > Manifest, FArchive& Archive ) OVERRIDE;
+	virtual bool SerializeManifest(TSharedRef< const FInternationalizationManifest > Manifest, FArchive& Archive ) override;
 
 	/**
 	 * Serializes a Internationalization manifest to a string.
diff --git a/Engine/Source/Runtime/Core/Public/Internationalization/InternationalizationMetadata.h b/Engine/Source/Runtime/Core/Public/Internationalization/InternationalizationMetadata.h
--- a/Engine/Source/Runtime/Core/Public/Internationalization/InternationalizationMetadata.h
+++ b/Engine/Source/Runtime/Core/Public/Internationalization/InternationalizationMetadata.h
@@ -180,10 +180,10 @@ public:
 protected:
 	FString Value;
 
-	virtual bool EqualTo( const FLocMetadataValue& Other ) const OVERRIDE;
-	virtual bool LessThan( const FLocMetadataValue& Other ) const OVERRIDE;
+	virtual bool EqualTo( const FLocMetadataValue& Other ) const override;
+	virtual bool LessThan( const FLocMetadataValue& Other ) const override;
 
-	virtual FString GetType() const OVERRIDE {return TEXT("String");}
+	virtual FString GetType() const override {return TEXT("String");}
 };
 
 
@@ -199,10 +199,10 @@ public:
 protected:
 	bool Value;
 
-	virtual bool EqualTo( const FLocMetadataValue& Other ) const OVERRIDE;
-	virtual bool LessThan( const FLocMetadataValue& Other ) const OVERRIDE;
+	virtual bool EqualTo( const FLocMetadataValue& Other ) const override;
+	virtual bool LessThan( const FLocMetadataValue& Other ) const override;
 
-	virtual FString GetType() const OVERRIDE {return TEXT("Boolean");}
+	virtual FString GetType() const override {return TEXT("Boolean");}
 };
 
 /** A LocMetadata Array Value. */
@@ -217,10 +217,10 @@ public:
 protected:
 	TArray< TSharedPtr<FLocMetadataValue> > Value;
 
-	virtual bool EqualTo( const FLocMetadataValue& Other ) const OVERRIDE;
-	virtual bool LessThan( const FLocMetadataValue& Other ) const OVERRIDE;
+	virtual bool EqualTo( const FLocMetadataValue& Other ) const override;
+	virtual bool LessThan( const FLocMetadataValue& Other ) const override;
 
-	virtual FString GetType() const OVERRIDE {return TEXT("Array");}
+	virtual FString GetType() const override {return TEXT("Array");}
 };
 
 /** A LocMetadata Object Value. */
@@ -235,10 +235,10 @@ public:
 protected:
 	TSharedPtr<FLocMetadataObject> Value;
 
-	virtual bool EqualTo( const FLocMetadataValue& Other ) const OVERRIDE;
-	virtual bool LessThan( const FLocMetadataValue& Other ) const OVERRIDE;
+	virtual bool EqualTo( const FLocMetadataValue& Other ) const override;
+	virtual bool LessThan( const FLocMetadataValue& Other ) const override;
 
-	virtual FString GetType() const OVERRIDE {return TEXT("Object");}
+	virtual FString GetType() const override {return TEXT("Object");}
 
 private:
 	// Hidden default ctor
