Stats - Fixed an issue with string conversion

diff --git a/Engine/Source/Runtime/Core/Private/Stats/Stats2.cpp b/Engine/Source/Runtime/Core/Private/Stats/Stats2.cpp
--- a/Engine/Source/Runtime/Core/Private/Stats/Stats2.cpp
+++ b/Engine/Source/Runtime/Core/Private/Stats/Stats2.cpp
@@ -390,6 +390,7 @@ public:
 		if (PendingCount < 1)
 		{
 			PendingFNames = new FName[NUM_PER_BLOCK];
+			FMemory::Memzero( PendingFNames, NUM_PER_BLOCK *sizeof( FName ) );
 			PendingCount = NUM_PER_BLOCK;
 		}
 		PendingCount -= FNAME_INCREMENT;
@@ -397,19 +398,17 @@ public:
 
 		// Get the wide stat description.
 		const int32 StatDescLen = FCString::Strlen( InDescription );
-		const WIDECHAR* TempStatDescWide = StringCast<WIDECHAR>( InDescription ).Get();
-		// We are leaking this.
+		// We are leaking this. @see STAT_StatDescMemory
 		WIDECHAR* StatDescWide = new WIDECHAR[StatDescLen + 1];
-		TCString<WIDECHAR>::Strcpy( StatDescWide, StatDescLen, TempStatDescWide );
+		TCString<WIDECHAR>::Strcpy( StatDescWide, StatDescLen, StringCast<WIDECHAR>( InDescription ).Get() );
 
 		// Get the ansi stat description.
-		const ANSICHAR* TempStatDescAnsi = StringCast<ANSICHAR>( InDescription ).Get();
-		// We are leaking this.
+		// We are leaking this. @see STAT_StatDescMemory
 		ANSICHAR* StatDescAnsi = new ANSICHAR[StatDescLen + 1];
-		TCString<ANSICHAR>::Strcpy( StatDescAnsi, StatDescLen, TempStatDescAnsi );
+		TCString<ANSICHAR>::Strcpy( StatDescAnsi, StatDescLen, StringCast<ANSICHAR>( InDescription ).Get() );
 
-		FMemory::Memcpy( (WIDECHAR*)&PendingFNames[TStatId::INDEX_WIDE_STRING], &StatDescWide, sizeof( WIDECHAR* ) );
-		FMemory::Memcpy( (ANSICHAR*)&PendingFNames[TStatId::INDEX_ANSI_STRING], &StatDescAnsi, sizeof( ANSICHAR* ) );
+		*(UPTRINT*)&PendingFNames[TStatId::INDEX_WIDE_STRING] = (UPTRINT)(UPTRINT*)StatDescWide;
+		*(UPTRINT*)&PendingFNames[TStatId::INDEX_ANSI_STRING] = (UPTRINT)(UPTRINT*)StatDescAnsi;
 
 		MemoryCounter.Add( (StatDescLen + 1)*(sizeof( ANSICHAR ) + sizeof( WIDECHAR )) );
 		
