Fix AssetTool's diff command generation generating garbage strings and modify Mac's ExecProcess and CreateProc to handle finding applications by name and using application bundles.

diff --git a/Engine/Source/Developer/AssetTools/Private/AssetTools.cpp b/Engine/Source/Developer/AssetTools/Private/AssetTools.cpp
--- a/Engine/Source/Developer/AssetTools/Private/AssetTools.cpp
+++ b/Engine/Source/Developer/AssetTools/Private/AssetTools.cpp
@@ -923,9 +923,9 @@ FString FAssetTools::DumpAssetToTempFile(UObject* Asset) const
 FString WrapArgument(const FString& Argument)
 {
 	// Wrap the passed in argument so it changes from Argument to "Argument"
-	return FString::Printf(TEXT("%s%s%s"),	(Argument.StartsWith("\"")) ? "": "\"",
+	return FString::Printf(TEXT("%s%s%s"),	(Argument.StartsWith("\"")) ? TEXT(""): TEXT("\""),
 											*Argument,
-											(Argument.EndsWith("\"")) ? "": "\"");
+											(Argument.EndsWith("\"")) ? TEXT(""): TEXT("\""));
 }
 
 bool FAssetTools::CreateDiffProcess(const FString& DiffCommand,  const FString& OldTextFilename,  const FString& NewTextFilename, const FString& DiffArgs) const
diff --git a/Engine/Source/Runtime/Core/Private/Mac/MacPlatformProcess.cpp b/Engine/Source/Runtime/Core/Private/Mac/MacPlatformProcess.cpp
--- a/Engine/Source/Runtime/Core/Private/Mac/MacPlatformProcess.cpp
+++ b/Engine/Source/Runtime/Core/Private/Mac/MacPlatformProcess.cpp
@@ -205,6 +205,29 @@ bool FMacPlatformProcess::ExecProcess( const TCHAR* URL, const TCHAR* Params, in
 	NSString* LaunchPath = (NSString*)FPlatformString::TCHARToCFString(*ProcessPath);
 	
 	if (![[NSFileManager defaultManager] fileExistsAtPath: LaunchPath])
+	{
+		NSString* AppName = [[LaunchPath lastPathComponent] stringByDeletingPathExtension];
+		LaunchPath = [[NSWorkspace sharedWorkspace] fullPathForApplication:AppName];
+	}
+	
+	if ([[NSFileManager defaultManager] fileExistsAtPath: LaunchPath])
+	{
+		if([[NSWorkspace sharedWorkspace] isFilePackageAtPath: LaunchPath])
+		{
+			NSBundle* Bundle = [NSBundle bundleWithPath:LaunchPath];
+			LaunchPath = Bundle ? [Bundle executablePath] : NULL;
+		}
+		else
+		{
+			LaunchPath = LaunchPath;
+		}
+	}
+	else
+	{
+		LaunchPath = NULL;
+	}
+	
+	if(LaunchPath == NULL)
 	{
 		*OutReturnCode = ENOENT;
 		*OutStdErr = TEXT("No such executable");
@@ -312,6 +335,29 @@ FProcHandle FMacPlatformProcess::CreateProc( const TCHAR* URL, const TCHAR* Parm
 	NSString* LaunchPath = (NSString*)FPlatformString::TCHARToCFString(*ProcessPath);
 
 	if (![[NSFileManager defaultManager] fileExistsAtPath: LaunchPath])
+	{
+		NSString* AppName = [[LaunchPath lastPathComponent] stringByDeletingPathExtension];
+		LaunchPath = [[NSWorkspace sharedWorkspace] fullPathForApplication:AppName];
+	}
+	
+	if ([[NSFileManager defaultManager] fileExistsAtPath: LaunchPath])
+	{
+		if([[NSWorkspace sharedWorkspace] isFilePackageAtPath: LaunchPath])
+		{
+			NSBundle* Bundle = [NSBundle bundleWithPath:LaunchPath];
+			LaunchPath = Bundle ? [Bundle executablePath] : NULL;
+		}
+		else
+		{
+			LaunchPath = LaunchPath;
+		}
+	}
+	else
+	{
+		LaunchPath = NULL;
+	}
+	
+	if(LaunchPath == NULL)
 	{
 		return FProcHandle(NULL);
 	}
