Cocoa events on Mac are now processed as soon as we get them

diff --git a/Engine/Source/Runtime/Core/Private/Mac/MacApplication.cpp b/Engine/Source/Runtime/Core/Private/Mac/MacApplication.cpp
--- a/Engine/Source/Runtime/Core/Private/Mac/MacApplication.cpp
+++ b/Engine/Source/Runtime/Core/Private/Mac/MacApplication.cpp
@@ -60,7 +60,7 @@ FMacApplication::FMacApplication()
 {
 	CGDisplayRegisterReconfigurationCallback(FMacApplication::OnDisplayReconfiguration, this);
 
-	[NSEvent addGlobalMonitorForEventsMatchingMask:NSMouseMovedMask handler:^(NSEvent* Event){AddPendingEvent(Event);}];
+	[NSEvent addGlobalMonitorForEventsMatchingMask:NSMouseMovedMask handler:^(NSEvent* Event){ ProcessEvent(Event); }];
 
 #if WITH_EDITOR
 	NSMutableArray* MultiTouchDevices = (__bridge NSMutableArray*)MTDeviceCreateList();
@@ -136,12 +136,6 @@ void FMacApplication::PumpMessages( const float TimeDelta )
 	FPlatformMisc::PumpMessages( true );
 }
 
-void FMacApplication::AddPendingEvent( NSEvent* Event )
-{
-	SCOPED_AUTORELEASE_POOL;
-	PendingEvents.Add( [Event retain] );
-}
-
 void FMacApplication::OnWindowDraggingFinished()
 {
 	if( DraggedWindow )
@@ -199,21 +193,6 @@ bool FMacApplication::IsWindowMovable(FSlateCocoaWindow* Win, bool* OutMovableBy
 	return true;
 }
 
-void FMacApplication::ProcessDeferredEvents( const float TimeDelta )
-{
-	SCOPED_AUTORELEASE_POOL;
-
-	// Processing some events may result in adding more events and calling this function recursively. To avoid problems,
-	// we make a local copy of pending events and empty the global one.
-	TArray< NSEvent* > Events( PendingEvents );
-	PendingEvents.Empty();
-	for( int32 Index = 0; Index < Events.Num(); Index++ )
-	{
-		ProcessEvent( Events[Index] );
-		[Events[Index] release];
-	}
-}
-
 void FMacApplication::HandleModifierChange(TSharedPtr< FMacWindow > CurrentEventWindow, NSUInteger NewModifierFlags, NSUInteger FlagsShift, NSUInteger UE4Shift, EMacModifierKeys TranslatedCode)
 {
 	bool CurrentPressed = (CurrentModifierFlags & FlagsShift) != 0;
diff --git a/Engine/Source/Runtime/Core/Private/Mac/MacPlatformMisc.cpp b/Engine/Source/Runtime/Core/Private/Mac/MacPlatformMisc.cpp
--- a/Engine/Source/Runtime/Core/Private/Mac/MacPlatformMisc.cpp
+++ b/Engine/Source/Runtime/Core/Private/Mac/MacPlatformMisc.cpp
@@ -568,7 +568,7 @@ void FMacPlatformMisc::PumpMessages( bool bFromMainLoop )
 			{
 				if( !bIsMouseClickEvent || [Event window] == NULL )
 				{
-					MacApplication->AddPendingEvent( Event );
+					MacApplication->ProcessEvent( Event );
 				}
 
 				if( [Event type] == NSLeftMouseUp )
diff --git a/Engine/Source/Runtime/Core/Private/Mac/MacWindow.cpp b/Engine/Source/Runtime/Core/Private/Mac/MacWindow.cpp
--- a/Engine/Source/Runtime/Core/Private/Mac/MacWindow.cpp
+++ b/Engine/Source/Runtime/Core/Private/Mac/MacWindow.cpp
@@ -548,7 +548,7 @@
 {
 	if(self.bForwardEvents)
 	{
-		MacApplication->AddPendingEvent( Event );
+		MacApplication->ProcessEvent( Event );
 	}
 }
 
@@ -563,7 +563,7 @@
 
 	if(self.bForwardEvents)
 	{
-		MacApplication->AddPendingEvent( Event );
+		MacApplication->ProcessEvent( Event );
 	}
 }
 
@@ -571,7 +571,7 @@
 {
 	if(self.bForwardEvents)
 	{
-		MacApplication->AddPendingEvent( Event );
+		MacApplication->ProcessEvent( Event );
 	}
 }
 
@@ -579,7 +579,7 @@
 {
 	if(self.bForwardEvents)
 	{
-		MacApplication->AddPendingEvent( Event );
+		MacApplication->ProcessEvent( Event );
 	}
 }
 
@@ -587,7 +587,7 @@
 {
 	if(self.bForwardEvents)
 	{
-		MacApplication->AddPendingEvent( Event );
+		MacApplication->ProcessEvent( Event );
 	}
 }
 
@@ -595,7 +595,7 @@
 {
 	if(self.bForwardEvents)
 	{
-		MacApplication->AddPendingEvent( Event );
+		MacApplication->ProcessEvent( Event );
 	}
 }
 
@@ -679,32 +679,32 @@
 
 - (void)mouseDown:(NSEvent*)Event
 {
-	MacApplication->AddPendingEvent( Event );
+	MacApplication->ProcessEvent( Event );
 }
 
 - (void)rightMouseDown:(NSEvent*)Event
 {
-	MacApplication->AddPendingEvent( Event );
+	MacApplication->ProcessEvent( Event );
 }
 
 - (void)otherMouseDown:(NSEvent*)Event
 {
-	MacApplication->AddPendingEvent( Event );
+	MacApplication->ProcessEvent( Event );
 }
 
 - (void)mouseUp:(NSEvent*)Event
 {
-	MacApplication->AddPendingEvent( Event );
+	MacApplication->ProcessEvent( Event );
 }
 
 - (void)rightMouseUp:(NSEvent*)Event
 {
-	MacApplication->AddPendingEvent( Event );
+	MacApplication->ProcessEvent( Event );
 }
 
 - (void)otherMouseUp:(NSEvent*)Event
 {
-	MacApplication->AddPendingEvent( Event );
+	MacApplication->ProcessEvent( Event );
 }
 
 @end
@@ -1047,8 +1047,7 @@ void FMacWindow::Show()
 			bool isMainAndKey = false;
 			do
 			{
-				OwningApplication->PumpMessages(0.f);
-				OwningApplication->ProcessDeferredEvents(0.f);
+				FPlatformMisc::PumpMessages(true);
 				bIsVisible = [WindowHandle isVisible];
 				isMainAndKey = [WindowHandle isKeyWindow];
 			} while(!bIsVisible && isMainAndKey != bMakeMainAndKey);
@@ -1094,8 +1093,7 @@ void FMacWindow::SetWindowMode( EWindowMode::Type NewWindowMode )
 		// and OpenGL contexts due to bad event ordering.
 		do
 		{
-			OwningApplication->PumpMessages(0.f);
-			OwningApplication->ProcessDeferredEvents(0.f);
+			FPlatformMisc::PumpMessages(true);
 			WindowIsFullScreen = [WindowHandle windowMode] != EWindowMode::Windowed;
 		} while(WindowIsFullScreen != bMakeFullscreen);
 		
diff --git a/Engine/Source/Runtime/Core/Public/Mac/MacApplication.h b/Engine/Source/Runtime/Core/Public/Mac/MacApplication.h
--- a/Engine/Source/Runtime/Core/Public/Mac/MacApplication.h
+++ b/Engine/Source/Runtime/Core/Public/Mac/MacApplication.h
@@ -31,8 +31,6 @@ public:
 
 	virtual void PumpMessages( const float TimeDelta ) override;
 
-	virtual void ProcessDeferredEvents( const float TimeDelta ) override;
-
 	virtual TSharedRef< FGenericWindow > MakeWindow() override;
 
 	virtual void InitializeWindow( const TSharedRef< FGenericWindow >& Window, const TSharedRef< FGenericWindowDefinition >& InDefinition, const TSharedPtr< FGenericWindow >& InParent, const bool bShowImmediately ) override;
@@ -63,7 +61,7 @@ public:
 		return TextInputMethodSystem.Get();
 	}
 
-	void AddPendingEvent( NSEvent* Event );
+	void ProcessEvent(NSEvent* Event);
 
 	void OnWindowDraggingFinished();
 
@@ -75,7 +73,7 @@ public:
 
 	uint32 GetModifierKeysFlags() { return ModifierKeysFlags; }
 
-	void UseMouseCaptureWindow( bool bUseMouseCaptureWindow );
+	void UseMouseCaptureWindow(bool bUseMouseCaptureWindow);
 
 #if WITH_EDITOR
     virtual void SendAnalytics(IAnalyticsProvider* Provider) override;
@@ -115,7 +113,6 @@ private:
 	TCHAR ConvertChar( TCHAR Character );
 	TCHAR TranslateCharCode( TCHAR CharCode, uint32 KeyCode );
 
-	void ProcessEvent( NSEvent* Event );
 	FSlateCocoaWindow* FindEventWindow( NSEvent* CocoaEvent );
 
 	NSScreen* FindScreenByPoint( int32 X, int32 Y ) const;
@@ -142,8 +139,6 @@ private:
 
 	TArray< TSharedRef< FMacWindow > > Windows;
 
-	TArray< NSEvent* > PendingEvents;
-
 	TSharedRef< class HIDInputInterface > HIDInput;
 
 	FSlateCocoaWindow* DraggedWindow;
