Added handling of empty paths to FMacPlatformMisc::NormalizePath()

diff --git a/Engine/Source/Runtime/Core/Private/Mac/MacPlatformMisc.cpp b/Engine/Source/Runtime/Core/Private/Mac/MacPlatformMisc.cpp
--- a/Engine/Source/Runtime/Core/Private/Mac/MacPlatformMisc.cpp
+++ b/Engine/Source/Runtime/Core/Private/Mac/MacPlatformMisc.cpp
@@ -997,11 +997,14 @@ int32 FMacPlatformMisc::NumberOfCoresIncludingHyperthreads()
 
 void FMacPlatformMisc::NormalizePath(FString& InPath)
 {
-	const bool bAppendSlash = InPath[InPath.Len() - 1] == '/'; // NSString will remove the trailing slash, if present, so we need to restore it after conversion
-	InPath = [[InPath.GetNSString() stringByStandardizingPath] stringByResolvingSymlinksInPath];
-	if (bAppendSlash)
+	if (InPath.Len() > 1)
 	{
-		InPath += TEXT("/");
+		const bool bAppendSlash = InPath[InPath.Len() - 1] == '/'; // NSString will remove the trailing slash, if present, so we need to restore it after conversion
+		InPath = [[InPath.GetNSString() stringByStandardizingPath] stringByResolvingSymlinksInPath];
+		if (bAppendSlash)
+		{
+			InPath += TEXT("/");
+		}
 	}
 }
 
