Fix or suppress VS2013 static analysis warnings in the Core module.

diff --git a/Engine/Source/Runtime/Core/Private/Containers/StackTracker.cpp b/Engine/Source/Runtime/Core/Private/Containers/StackTracker.cpp
--- a/Engine/Source/Runtime/Core/Private/Containers/StackTracker.cpp
+++ b/Engine/Source/Runtime/Core/Private/Containers/StackTracker.cpp
@@ -25,6 +25,8 @@ void FStackTracker::CaptureStackTrace(int32 EntriesToIgnore /*=2*/, void* UserDa
 		FullBackTrace = static_cast<uint64*>(FMemory_Alloca((MAX_BACKTRACE_DEPTH + EntriesToIgnore) * sizeof(uint64)));
 
 		FPlatformStackWalk::CaptureStackBackTrace( FullBackTrace, MAX_BACKTRACE_DEPTH + EntriesToIgnore );
+		CA_ASSUME(FullBackTrace);
+
 		// Skip first NUM_ENTRIES_TO_SKIP entries as they are inside this code
 		uint64* BackTrace = &FullBackTrace[EntriesToIgnore];
 		uint32 CRC = FCrc::MemCrc_DEPRECATED( BackTrace, MAX_BACKTRACE_DEPTH * sizeof(uint64) );
diff --git a/Engine/Source/Runtime/Core/Private/HAL/ConsoleManager.cpp b/Engine/Source/Runtime/Core/Private/HAL/ConsoleManager.cpp
--- a/Engine/Source/Runtime/Core/Private/HAL/ConsoleManager.cpp
+++ b/Engine/Source/Runtime/Core/Private/HAL/ConsoleManager.cpp
@@ -931,7 +931,7 @@ IConsoleObject* FConsoleManager::AddConsoleObject(const TCHAR* Name, IConsoleObj
 			UE_LOG(LogConsoleManager, Fatal, TEXT( "Console object named '%s' can't be replaced with the new one of different type!"), Name );
 		}
 
-		if( ExistingVar )
+		if( ExistingVar && Var )
 		{
 			if(ExistingVar->TestFlags(ECVF_CreatedFromIni))
 			{
diff --git a/Engine/Source/Runtime/Core/Private/Internationalization/ICUText.cpp b/Engine/Source/Runtime/Core/Private/Internationalization/ICUText.cpp
--- a/Engine/Source/Runtime/Core/Private/Internationalization/ICUText.cpp
+++ b/Engine/Source/Runtime/Core/Private/Internationalization/ICUText.cpp
@@ -263,6 +263,7 @@ FText FText::FormatInternal(const FText& Pattern, const FFormatNamedArguments& A
 	checkf(FInternationalization::Get().IsInitialized() == true, TEXT("FInternationalization is not initialized. An FText formatting method was likely used in static object initialization - this is not supported."));
 	//SCOPE_CYCLE_COUNTER( STAT_TextFormat );
 
+	CA_SUPPRESS(6316)
 	const bool bEnableErrorResults = ENABLE_TEXT_ERROR_CHECKING_RESULTS && bEnableErrorCheckingResults;
 
 	TArray<icu::UnicodeString> ArgumentNames;
@@ -351,6 +352,7 @@ FText FText::FormatInternal(const FText& Pattern, const FFormatOrderedArguments&
 	checkf(FInternationalization::Get().IsInitialized() == true, TEXT("FInternationalization is not initialized. An FText formatting method was likely used in static object initialization - this is not supported."));
 	//SCOPE_CYCLE_COUNTER( STAT_TextFormat );
 
+	CA_SUPPRESS(6316)
 	const bool bEnableErrorResults = ENABLE_TEXT_ERROR_CHECKING_RESULTS && bEnableErrorCheckingResults;
 
 	TArray<icu::Formattable> ArgumentValues;
@@ -436,6 +438,7 @@ FText FText::FormatInternal(const FText& Pattern, const TArray< FFormatArgumentD
 	checkf(FInternationalization::Get().IsInitialized() == true, TEXT("FInternationalization is not initialized. An FText formatting method was likely used in static object initialization - this is not supported."));
 	//SCOPE_CYCLE_COUNTER( STAT_TextFormat );
 
+	CA_SUPPRESS(6316)
 	const bool bEnableErrorResults = ENABLE_TEXT_ERROR_CHECKING_RESULTS && bEnableErrorCheckingResults;
 
 	TArray<icu::UnicodeString> ArgumentNames;
diff --git a/Engine/Source/Runtime/Core/Private/Internationalization/TextHistory.cpp b/Engine/Source/Runtime/Core/Private/Internationalization/TextHistory.cpp
--- a/Engine/Source/Runtime/Core/Private/Internationalization/TextHistory.cpp
+++ b/Engine/Source/Runtime/Core/Private/Internationalization/TextHistory.cpp
@@ -246,6 +246,7 @@ void FTextHistory_FormatNumber::Serialize(FArchive& Ar)
 		{
 			FormatOptions = new FNumberFormattingOptions;
 		}
+		CA_SUPPRESS(6011)
 		Ar << *FormatOptions;
 	}
 
diff --git a/Engine/Source/Runtime/Core/Private/Internationalization/TextLocalizationManager.cpp b/Engine/Source/Runtime/Core/Private/Internationalization/TextLocalizationManager.cpp
--- a/Engine/Source/Runtime/Core/Private/Internationalization/TextLocalizationManager.cpp
+++ b/Engine/Source/Runtime/Core/Private/Internationalization/TextLocalizationManager.cpp
@@ -511,6 +511,8 @@ TSharedRef<FString> FTextLocalizationManager::GetString(const FString& Namespace
 		const TSharedRef<FString> UnlocalizedString = MakeShareable( new FString( SourceString ? **SourceString : TEXT("") ) );
 
 		// If live-culture-swap is enabled or the system is uninitialized - make entries so that they can be updated when system is initialized or a culture swap occurs.
+		CA_SUPPRESS(6236)
+		CA_SUPPRESS(6316)
 		if( !(bIsInitialized) || ENABLE_LOC_TESTING )
 		{
 #if ENABLE_LOC_TESTING
diff --git a/Engine/Source/Runtime/Core/Private/Misc/EngineVersion.cpp b/Engine/Source/Runtime/Core/Private/Misc/EngineVersion.cpp
--- a/Engine/Source/Runtime/Core/Private/Misc/EngineVersion.cpp
+++ b/Engine/Source/Runtime/Core/Private/Misc/EngineVersion.cpp
@@ -141,7 +141,7 @@ bool FEngineVersion::Parse(const FString &Text, FEngineVersion &OutVersion)
 
 EVersionComparison::Type FEngineVersion::GetNewest(const FEngineVersion &First, const FEngineVersion &Second, EVersionComponent::Type *OutComponent)
 {
-	EVersionComponent::Type LocalComponent;
+	EVersionComponent::Type LocalComponent = EVersionComponent::Minor;
 	auto& Component = OutComponent ? *OutComponent : LocalComponent;
 
 	// Compare major versions
diff --git a/Engine/Source/Runtime/Core/Private/Windows/TextStoreACP.cpp b/Engine/Source/Runtime/Core/Private/Windows/TextStoreACP.cpp
--- a/Engine/Source/Runtime/Core/Private/Windows/TextStoreACP.cpp
+++ b/Engine/Source/Runtime/Core/Private/Windows/TextStoreACP.cpp
@@ -70,7 +70,7 @@ STDAPI_(ULONG) FTextStoreACP::Release()
 	return LocalReferenceCount;
 }
 
-STDAPI FTextStoreACP::AdviseSink(REFIID riid, IUnknown *punk, DWORD dwMask)
+STDAPI FTextStoreACP::AdviseSink(__RPC__in REFIID riid, __RPC__in_opt IUnknown *punk, DWORD dwMask)
 {
 	UE_LOG(LogTextStoreACP, Verbose, TEXT("AdviseSink"));
 
@@ -128,7 +128,7 @@ STDAPI FTextStoreACP::AdviseSink(REFIID riid, IUnknown *punk, DWORD dwMask)
 	return S_OK;
 }
 
-STDAPI FTextStoreACP::UnadviseSink(IUnknown *punk)
+STDAPI FTextStoreACP::UnadviseSink(__RPC__in_opt IUnknown *punk)
 {
 	UE_LOG(LogTextStoreACP, Verbose, TEXT("UnadviseSink"));
 
@@ -232,7 +232,7 @@ STDAPI FTextStoreACP::RequestLock(DWORD dwLockFlags, HRESULT *phrSession)
 	return S_OK;
 }
 
-STDAPI FTextStoreACP::GetStatus(TS_STATUS *pdcs)
+STDAPI FTextStoreACP::GetStatus(__RPC__out TS_STATUS *pdcs)
 {
 	UE_LOG(LogTextStoreACP, Verbose, TEXT("GetStatus"));
 
@@ -248,7 +248,7 @@ STDAPI FTextStoreACP::GetStatus(TS_STATUS *pdcs)
 	return S_OK;
 }
 
-STDAPI FTextStoreACP::GetEndACP(LONG *pacp)
+STDAPI FTextStoreACP::GetEndACP(__RPC__out LONG *pacp)
 {
 	UE_LOG(LogTextStoreACP, Verbose, TEXT("GetEndACP"));
 
@@ -262,7 +262,7 @@ STDAPI FTextStoreACP::GetEndACP(LONG *pacp)
 	return S_OK;
 }
 
-STDAPI FTextStoreACP::GetSelection(ULONG ulIndex, ULONG ulCount, TS_SELECTION_ACP *pSelection, ULONG *pcFetched)
+STDAPI FTextStoreACP::GetSelection(ULONG ulIndex, ULONG ulCount, __RPC__out_ecount_part(ulCount, *pcFetched) TS_SELECTION_ACP *pSelection, __RPC__out ULONG *pcFetched)
 {
 	UE_LOG(LogTextStoreACP, Verbose, TEXT("GetSelection"));
 
@@ -303,7 +303,7 @@ STDAPI FTextStoreACP::GetSelection(ULONG ulIndex, ULONG ulCount, TS_SELECTION_AC
 	return S_OK;
 }
 
-STDAPI FTextStoreACP::SetSelection(ULONG ulCount, const TS_SELECTION_ACP *pSelection)
+STDAPI FTextStoreACP::SetSelection(ULONG ulCount, __RPC__in_ecount_full(ulCount) const TS_SELECTION_ACP *pSelection)
 {
 	UE_LOG(LogTextStoreACP, Verbose, TEXT("SetSelection - From %d to %d"), pSelection->acpStart, pSelection->acpEnd);
 
@@ -338,7 +338,7 @@ STDAPI FTextStoreACP::SetSelection(ULONG ulCount, const TS_SELECTION_ACP *pSelec
 	return S_OK;
 }
 
-STDAPI FTextStoreACP::RequestSupportedAttrs(DWORD dwFlags, ULONG cFilterAttrs, const TS_ATTRID *paFilterAttrs)
+STDAPI FTextStoreACP::RequestSupportedAttrs(DWORD dwFlags, ULONG cFilterAttrs, __RPC__in_ecount_full_opt(cFilterAttrs) const TS_ATTRID *paFilterAttrs)
 {
 	UE_LOG(LogTextStoreACP, Verbose, TEXT("RequestSupportedAttrs"));
 
@@ -359,7 +359,7 @@ STDAPI FTextStoreACP::RequestSupportedAttrs(DWORD dwFlags, ULONG cFilterAttrs, c
 	return S_OK;
 }
 
-STDAPI FTextStoreACP::RequestAttrsAtPosition(LONG acpPos, ULONG cFilterAttrs, const TS_ATTRID *paFilterAttrs, DWORD dwFlags)
+STDAPI FTextStoreACP::RequestAttrsAtPosition(LONG acpPos, ULONG cFilterAttrs, __RPC__in_ecount_full_opt(cFilterAttrs) const TS_ATTRID *paFilterAttrs, DWORD dwFlags)
 {
 	UE_LOG(LogTextStoreACP, Verbose, TEXT("RequestAttrsAtPosition"));
 
@@ -380,19 +380,20 @@ STDAPI FTextStoreACP::RequestAttrsAtPosition(LONG acpPos, ULONG cFilterAttrs, co
 	return S_OK;
 }
 
-STDAPI FTextStoreACP::RequestAttrsTransitioningAtPosition(LONG acpPos, ULONG cFilterAttrs, const TS_ATTRID *paFilterAttrs, DWORD dwFlags)
+STDAPI FTextStoreACP::RequestAttrsTransitioningAtPosition(LONG acpPos, ULONG cFilterAttrs, __RPC__in_ecount_full_opt(cFilterAttrs) const TS_ATTRID *paFilterAttrs, DWORD dwFlags)
 {
 	UE_LOG(LogTextStoreACP, Verbose, TEXT("RequestAttrsTransitioningAtPosition"));
 	return E_NOTIMPL; // TODO: No meaningful documentation to this method, not implemented.
 }
 
-STDAPI FTextStoreACP::FindNextAttrTransition(LONG acpStart, LONG acpHalt, ULONG cFilterAttrs, const TS_ATTRID *paFilterAttrs, DWORD dwFlags, LONG *pacpNext, BOOL *pfFound, LONG *plFoundOffset)
+STDAPI FTextStoreACP::FindNextAttrTransition(LONG acpStart, LONG acpHalt, ULONG cFilterAttrs, __RPC__in_ecount_full_opt(cFilterAttrs) const TS_ATTRID *paFilterAttrs, DWORD dwFlags, __RPC__out LONG *pacpNext, __RPC__out BOOL *pfFound, __RPC__out LONG *plFoundOffset)
 {
 	UE_LOG(LogTextStoreACP, Verbose, TEXT("FindNextAttrTransition"));
 	return E_NOTIMPL; // TODO: No meaningful documentation to this method, not implemented.
 }
 
-STDAPI FTextStoreACP::RetrieveRequestedAttrs(ULONG ulCount, TS_ATTRVAL *paAttrVals, ULONG *pcFetched)
+CA_SUPPRESS(6101)
+STDAPI FTextStoreACP::RetrieveRequestedAttrs(ULONG ulCount, __RPC__out_ecount_part(ulCount, *pcFetched) TS_ATTRVAL *paAttrVals, __RPC__out ULONG *pcFetched)
 {
 	UE_LOG(LogTextStoreACP, Verbose, TEXT("RetrieveRequestedAttrs"));
 
@@ -403,6 +404,7 @@ STDAPI FTextStoreACP::RetrieveRequestedAttrs(ULONG ulCount, TS_ATTRVAL *paAttrVa
 		if(SupportedAttributes[i].WantsDefault)
 		{
 			paAttrVals[i].idAttr = *(SupportedAttributes[i].Id);
+			CA_SUPPRESS(6031)
 			VariantCopy(&(paAttrVals[*pcFetched].varValue), &(SupportedAttributes[i].DefaultValue));
 			paAttrVals[i].dwOverlapId = 0;
 		}
@@ -413,7 +415,7 @@ STDAPI FTextStoreACP::RetrieveRequestedAttrs(ULONG ulCount, TS_ATTRVAL *paAttrVa
 	return S_OK;
 }
 
-STDAPI FTextStoreACP::GetActiveView(TsViewCookie *pvcView)
+STDAPI FTextStoreACP::GetActiveView(__RPC__out TsViewCookie *pvcView)
 {
 	UE_LOG(LogTextStoreACP, Verbose, TEXT("GetActiveView"));
 	*pvcView = 0;
@@ -421,7 +423,7 @@ STDAPI FTextStoreACP::GetActiveView(TsViewCookie *pvcView)
 	return S_OK;
 }
 
-STDAPI FTextStoreACP::GetACPFromPoint(TsViewCookie vcView, const POINT *pt, DWORD dwFlags, LONG *pacp)
+STDAPI FTextStoreACP::GetACPFromPoint(TsViewCookie vcView, __RPC__in const POINT *pt, DWORD dwFlags, __RPC__out LONG *pacp)
 {
 	UE_LOG(LogTextStoreACP, Verbose, TEXT("GetACPFromPoint - At (%d, %d)"), pt->x, pt->y);
 
@@ -438,7 +440,7 @@ STDAPI FTextStoreACP::GetACPFromPoint(TsViewCookie vcView, const POINT *pt, DWOR
 	return S_OK;
 }
 
-STDAPI FTextStoreACP::GetTextExt(TsViewCookie vcView, LONG acpStart, LONG acpEnd, RECT *prc, BOOL *pfClipped)
+STDAPI FTextStoreACP::GetTextExt(TsViewCookie vcView, LONG acpStart, LONG acpEnd, __RPC__out RECT *prc, __RPC__out BOOL *pfClipped)
 {
 	UE_LOG(LogTextStoreACP, Verbose, TEXT("GetTextExt - From %d to %d"), acpStart, acpEnd);
 
@@ -472,7 +474,7 @@ STDAPI FTextStoreACP::GetTextExt(TsViewCookie vcView, LONG acpStart, LONG acpEnd
 	return S_OK;
 }
 
-STDAPI FTextStoreACP::GetScreenExt(TsViewCookie vcView, RECT *prc)
+STDAPI FTextStoreACP::GetScreenExt(TsViewCookie vcView, __RPC__out RECT *prc)
 {
 	UE_LOG(LogTextStoreACP, Verbose, TEXT("GetScreenExt"));
 
@@ -493,7 +495,9 @@ STDAPI FTextStoreACP::GetScreenExt(TsViewCookie vcView, RECT *prc)
 	return S_OK;
 }
 
-STDAPI FTextStoreACP::GetWnd(TsViewCookie vcView, HWND *phwnd)
+CA_SUPPRESS(28196)
+CA_SUPPRESS(6387)
+STDAPI FTextStoreACP::GetWnd(TsViewCookie vcView, __RPC__deref_out_opt HWND *phwnd)
 {
 	UE_LOG(LogTextStoreACP, Verbose, TEXT("GetWnd"));
 
@@ -503,7 +507,7 @@ STDAPI FTextStoreACP::GetWnd(TsViewCookie vcView, HWND *phwnd)
 	return S_OK;
 }
 
-STDAPI FTextStoreACP::GetText(LONG acpStart, LONG acpEnd, __out_ecount(cchPlainReq) WCHAR *pchPlain, ULONG cchPlainReq, ULONG *pcchPlainOut, TS_RUNINFO *prgRunInfo, ULONG ulRunInfoReq, ULONG *pulRunInfoOut, LONG *pacpNext)
+STDAPI FTextStoreACP::GetText(LONG acpStart, LONG acpEnd, __RPC__out_ecount_part(cchPlainReq, *pcchPlainOut) WCHAR *pchPlain, ULONG cchPlainReq, __RPC__out ULONG *pcchPlainOut, __RPC__out_ecount_part(ulRunInfoReq, *pulRunInfoOut) TS_RUNINFO *prgRunInfo, ULONG ulRunInfoReq, __RPC__out ULONG *pulRunInfoOut, __RPC__out LONG *pacpNext)
 {
 	UE_LOG(LogTextStoreACP, Verbose, TEXT("GetText - From %d to %d"), acpStart, acpEnd);
 
@@ -552,7 +556,7 @@ STDAPI FTextStoreACP::GetText(LONG acpStart, LONG acpEnd, __out_ecount(cchPlainR
 	return S_OK;
 }
 
-STDAPI FTextStoreACP::QueryInsert(LONG acpInsertStart, LONG acpInsertEnd, ULONG cch, LONG *pacpResultStart, LONG *pacpResultEnd)
+STDAPI FTextStoreACP::QueryInsert(LONG acpInsertStart, LONG acpInsertEnd, ULONG cch, __RPC__out LONG *pacpResultStart, __RPC__out LONG *pacpResultEnd)
 {
 	UE_LOG(LogTextStoreACP, Verbose, TEXT("QueryInsert"));
 
@@ -576,7 +580,8 @@ STDAPI FTextStoreACP::QueryInsert(LONG acpInsertStart, LONG acpInsertEnd, ULONG
 	return S_OK;
 }
 
-STDAPI FTextStoreACP::InsertTextAtSelection(DWORD dwFlags, __in_ecount(cch) const WCHAR *pchText, ULONG cch, LONG *pacpStart, LONG *pacpEnd, TS_TEXTCHANGE *pChange)
+CA_SUPPRESS(6101)
+STDAPI FTextStoreACP::InsertTextAtSelection(DWORD dwFlags, __RPC__in_ecount_full(cch) const WCHAR *pchText, ULONG cch, __RPC__out LONG *pacpStart, __RPC__out LONG *pacpEnd, __RPC__out TS_TEXTCHANGE *pChange)
 {
 	UE_LOG(LogTextStoreACP, Verbose, TEXT("InsertTextAtSelection - %s"), pchText);
 
@@ -651,7 +656,8 @@ STDAPI FTextStoreACP::InsertTextAtSelection(DWORD dwFlags, __in_ecount(cch) cons
 	return S_OK;
 }
 
-STDAPI FTextStoreACP::SetText(DWORD dwFlags, LONG acpStart, LONG acpEnd, __in_ecount(cch) const WCHAR *pchText, ULONG cch, TS_TEXTCHANGE *pChange)
+CA_SUPPRESS(6103)
+STDAPI FTextStoreACP::SetText(DWORD dwFlags, LONG acpStart, LONG acpEnd, __RPC__in_ecount_full(cch) const WCHAR *pchText, ULONG cch, __RPC__out TS_TEXTCHANGE *pChange)
 {
 	UE_LOG(LogTextStoreACP, Verbose, TEXT("SetText - From %d to %d, set %s"), acpStart, acpEnd, pchText);
 
@@ -682,7 +688,7 @@ STDAPI FTextStoreACP::SetText(DWORD dwFlags, LONG acpStart, LONG acpEnd, __in_ec
 	return S_OK;
 }
 
-STDAPI FTextStoreACP::GetEmbedded(LONG acpPos, REFGUID rguidService, REFIID riid, IUnknown **ppunk)
+STDAPI FTextStoreACP::GetEmbedded(LONG acpPos, __RPC__in REFGUID rguidService, __RPC__in REFIID riid, __RPC__deref_out_opt IUnknown **ppunk)
 {
 	UE_LOG(LogTextStoreACP, Verbose, TEXT("GetEmbedded"));
 
@@ -694,7 +700,7 @@ STDAPI FTextStoreACP::GetEmbedded(LONG acpPos, REFGUID rguidService, REFIID riid
 	return E_NOTIMPL;
 }
 
-STDAPI FTextStoreACP::GetFormattedText(LONG acpStart, LONG acpEnd, IDataObject **ppDataObject)
+STDAPI FTextStoreACP::GetFormattedText(LONG acpStart, LONG acpEnd, __RPC__deref_out_opt IDataObject **ppDataObject)
 {
 	UE_LOG(LogTextStoreACP, Verbose, TEXT("GetFormattedText"));
 
@@ -706,14 +712,14 @@ STDAPI FTextStoreACP::GetFormattedText(LONG acpStart, LONG acpEnd, IDataObject *
 	return E_NOTIMPL;
 }
 
-STDAPI FTextStoreACP::QueryInsertEmbedded(const GUID *pguidService, const FORMATETC *pFormatEtc, BOOL *pfInsertable)
+STDAPI FTextStoreACP::QueryInsertEmbedded(__RPC__in const GUID *pguidService, __RPC__in const FORMATETC *pFormatEtc, __RPC__out BOOL *pfInsertable)
 {
 	UE_LOG(LogTextStoreACP, Verbose, TEXT("QueryInsertEmbedded"));
 
 	return E_NOTIMPL;
 }
 
-STDAPI FTextStoreACP::InsertEmbedded(DWORD dwFlags, LONG acpStart, LONG acpEnd, IDataObject *pDataObject, TS_TEXTCHANGE *pChange)
+STDAPI FTextStoreACP::InsertEmbedded(DWORD dwFlags, LONG acpStart, LONG acpEnd, __RPC__in_opt IDataObject *pDataObject, __RPC__out TS_TEXTCHANGE *pChange)
 {
 	UE_LOG(LogTextStoreACP, Verbose, TEXT("InsertEmbedded"));
 
@@ -725,7 +731,7 @@ STDAPI FTextStoreACP::InsertEmbedded(DWORD dwFlags, LONG acpStart, LONG acpEnd,
 	return E_NOTIMPL;
 }
 
-STDAPI FTextStoreACP::InsertEmbeddedAtSelection(DWORD dwFlags, IDataObject *pDataObject, LONG *pacpStart, LONG *pacpEnd, TS_TEXTCHANGE *pChange)
+STDAPI FTextStoreACP::InsertEmbeddedAtSelection(DWORD dwFlags, __RPC__in_opt IDataObject *pDataObject, __RPC__out LONG *pacpStart, __RPC__out LONG *pacpEnd, __RPC__out TS_TEXTCHANGE *pChange)
 {
 	UE_LOG(LogTextStoreACP, Verbose, TEXT("InsertEmbeddedAtSelection"));
 
@@ -737,7 +743,7 @@ STDAPI FTextStoreACP::InsertEmbeddedAtSelection(DWORD dwFlags, IDataObject *pDat
 	return E_NOTIMPL;
 }
 
-STDAPI FTextStoreACP::OnStartComposition(ITfCompositionView *pComposition, BOOL *pfOk)
+STDAPI FTextStoreACP::OnStartComposition(__RPC__in_opt ITfCompositionView *pComposition, __RPC__out BOOL *pfOk)
 {
 	UE_LOG(LogTextStoreACP, Verbose, TEXT("OnStartComposition"));
 
@@ -758,7 +764,7 @@ STDAPI FTextStoreACP::OnStartComposition(ITfCompositionView *pComposition, BOOL
 	return S_OK;
 }
 
-STDAPI FTextStoreACP::OnUpdateComposition(ITfCompositionView *pComposition, ITfRange *pRangeNew)
+STDAPI FTextStoreACP::OnUpdateComposition(__RPC__in_opt ITfCompositionView *pComposition, __RPC__in_opt ITfRange *pRangeNew)
 {
 	UE_LOG(LogTextStoreACP, Verbose, TEXT("OnUpdateComposition"));
 
@@ -781,7 +787,12 @@ STDAPI FTextStoreACP::OnUpdateComposition(ITfCompositionView *pComposition, ITfR
 
 		LONG BeginIndex;
 		LONG Length;
-		RangeACP->GetExtent(&(BeginIndex), &(Length));
+		const auto Result = RangeACP->GetExtent(&(BeginIndex), &(Length));
+
+		if(FAILED(Result))
+		{
+			return E_FAIL;
+		}
 
 		UE_LOG(LogTextStoreACP, Verbose, TEXT("OnUpdateComposition - From %d to %d"), BeginIndex, BeginIndex + Length);
 
@@ -791,7 +802,7 @@ STDAPI FTextStoreACP::OnUpdateComposition(ITfCompositionView *pComposition, ITfR
 	return S_OK;
 }
 
-STDAPI FTextStoreACP::OnEndComposition(ITfCompositionView *pComposition)
+STDAPI FTextStoreACP::OnEndComposition(__RPC__in_opt ITfCompositionView *pComposition)
 {
 	UE_LOG(LogTextStoreACP, Verbose, TEXT("OnEndComposition"));
 
diff --git a/Engine/Source/Runtime/Core/Private/Windows/WindowsApplication.cpp b/Engine/Source/Runtime/Core/Private/Windows/WindowsApplication.cpp
--- a/Engine/Source/Runtime/Core/Private/Windows/WindowsApplication.cpp
+++ b/Engine/Source/Runtime/Core/Private/Windows/WindowsApplication.cpp
@@ -61,6 +61,7 @@ FWindowsApplication::FWindowsApplication( const HINSTANCE HInstance, const HICON
 	const bool bClassRegistered = RegisterClass( InstanceHandle, IconHandle );
 
 	// Initialize OLE for Drag and Drop support.
+	CA_SUPPRESS(6031);
 	OleInitialize( NULL );
 
 	TextInputMethodSystem = MakeShareable( new FWindowsTextInputMethodSystem );
@@ -159,7 +160,7 @@ void FWindowsApplication::SetCapture( const TSharedPtr< FGenericWindow >& InWind
 	}
 	else
 	{
-		::SetCapture( NULL );
+		::ReleaseCapture();
 	}
 }
 
@@ -413,13 +414,16 @@ void FWindowsApplication::GetDisplayMetrics( FDisplayMetrics& OutDisplayMetrics
 
 	// Get the screen rect of the primary monitor, excluding taskbar etc.
 	RECT WorkAreaRect;
-	WorkAreaRect.top = WorkAreaRect.bottom = WorkAreaRect.left = WorkAreaRect.right = 0;
-	SystemParametersInfo( SPI_GETWORKAREA, 0, &WorkAreaRect, 0 );
+	if(!SystemParametersInfo(SPI_GETWORKAREA, 0, &WorkAreaRect, 0))
+	{
+		WorkAreaRect.top = WorkAreaRect.bottom = WorkAreaRect.left = WorkAreaRect.right = 0;
+	}
+
 	OutDisplayMetrics.PrimaryDisplayWorkAreaRect.Left = WorkAreaRect.left;
 	OutDisplayMetrics.PrimaryDisplayWorkAreaRect.Top = WorkAreaRect.top;
 	OutDisplayMetrics.PrimaryDisplayWorkAreaRect.Right = WorkAreaRect.right;
 	OutDisplayMetrics.PrimaryDisplayWorkAreaRect.Bottom = WorkAreaRect.bottom;
-
+	
 	// Virtual desktop area
 	OutDisplayMetrics.VirtualDisplayRect.Left = ::GetSystemMetrics( SM_XVIRTUALSCREEN );
 	OutDisplayMetrics.VirtualDisplayRect.Top = ::GetSystemMetrics( SM_YVIRTUALSCREEN );
diff --git a/Engine/Source/Runtime/Core/Private/Windows/WindowsPlatformExceptionHandling.cpp b/Engine/Source/Runtime/Core/Private/Windows/WindowsPlatformExceptionHandling.cpp
--- a/Engine/Source/Runtime/Core/Private/Windows/WindowsPlatformExceptionHandling.cpp
+++ b/Engine/Source/Runtime/Core/Private/Windows/WindowsPlatformExceptionHandling.cpp
@@ -103,20 +103,20 @@ void GetModuleVersion( TCHAR* ModuleName, TCHAR* StringBuffer, DWORD MaxSize )
 	DWORD InfoSize = GetFileVersionInfoSize( ModuleName, &Handle );
 	if( InfoSize > 0 )
 	{
-		char* VersionInfo = new char[InfoSize];
-		if( GetFileVersionInfo( ModuleName, 0, InfoSize, VersionInfo ) )
+		TArray<char> VersionInfo;
+		VersionInfo.SetNum(InfoSize);
+
+		if( GetFileVersionInfo( ModuleName, 0, InfoSize, VersionInfo.GetData() ) )
 		{
 			VS_FIXEDFILEINFO* FixedFileInfo;
 
 			UINT InfoLength = 0;
-			if( VerQueryValue( VersionInfo, TEXT( "\\" ), ( void** )&FixedFileInfo, &InfoLength ) )
+			if( VerQueryValue( VersionInfo.GetData(), TEXT( "\\" ), ( void** )&FixedFileInfo, &InfoLength ) )
 			{
-				StringCchPrintf( StringBuffer, MaxSize, TEXT( "%d.%d.%d.%d" ), 
+				StringCchPrintf( StringBuffer, MaxSize, TEXT( "%u.%u.%u.%u" ), 
 					HIWORD( FixedFileInfo->dwProductVersionMS ), LOWORD( FixedFileInfo->dwProductVersionMS ), HIWORD( FixedFileInfo->dwProductVersionLS ), LOWORD( FixedFileInfo->dwProductVersionLS ) );
 			}
 		}
-
-		delete VersionInfo;
 	}
 }
 
@@ -166,6 +166,7 @@ void SetReportParameters( HREPORT ReportHandle, EXCEPTION_POINTERS* ExceptionInf
 	Result = WerReportSetParameter( ReportHandle, WER_P6, TEXT( "Exception Code" ), StringBuffer );
 
 	INT_PTR ExceptionOffset = ( char* )( ExceptionInfo->ExceptionRecord->ExceptionAddress ) - ( char* )FaultModuleHandle;
+	CA_SUPPRESS(6066) // The format specifier should probably be something like %tX, but VS 2013 doesn't support 't'.
 	StringCchPrintf( StringBuffer, LocalBufferSize, TEXT( "%p" ), ExceptionOffset );
 	Result = WerReportSetParameter( ReportHandle, WER_P7, TEXT( "Exception Offset" ), StringBuffer );
 
@@ -392,6 +393,7 @@ void NewReportEnsure( const TCHAR* ErrorMessage )
 	}
 #if !PLATFORM_SEH_EXCEPTIONS_DISABLED
 	__except(ReportCrashUsingCrashReportClient(GetExceptionInformation(), ErrorMessage, EErrorReportUI::ReportInUnattendedMode))
+	CA_SUPPRESS(6322)
 	{
 	}
 #endif
diff --git a/Engine/Source/Runtime/Core/Private/Windows/WindowsPlatformMemory.cpp b/Engine/Source/Runtime/Core/Private/Windows/WindowsPlatformMemory.cpp
--- a/Engine/Source/Runtime/Core/Private/Windows/WindowsPlatformMemory.cpp
+++ b/Engine/Source/Runtime/Core/Private/Windows/WindowsPlatformMemory.cpp
@@ -165,6 +165,7 @@ void* FWindowsPlatformMemory::BinnedAllocFromOS( SIZE_T Size )
 
 void FWindowsPlatformMemory::BinnedFreeToOS( void* Ptr )
 {
+	CA_SUPPRESS(6001)
 	verify(VirtualFree( Ptr, 0, MEM_RELEASE ) != 0);
 }
 
diff --git a/Engine/Source/Runtime/Core/Private/Windows/WindowsPlatformMisc.cpp b/Engine/Source/Runtime/Core/Private/Windows/WindowsPlatformMisc.cpp
--- a/Engine/Source/Runtime/Core/Private/Windows/WindowsPlatformMisc.cpp
+++ b/Engine/Source/Runtime/Core/Private/Windows/WindowsPlatformMisc.cpp
@@ -60,6 +60,7 @@ int32 FWindowsOSVersionHelper::GetOSVersions( FString& out_OSVersionLabel, FStri
 	out_OSVersionLabel = TEXT( "Windows (unknown version)" );
 	out_OSSubVersionLabel = FString();
 #pragma warning(disable : 4996) // 'function' was declared deprecated
+	CA_SUPPRESS(28159)
 	if( GetVersionEx( (LPOSVERSIONINFO)&OsVersionInfo ) )
 #pragma warning(default : 4996)
 	{
@@ -908,6 +909,7 @@ uint32 FWindowsPlatformMisc::GetKeyMap( uint16* KeyCodes, FString* KeyNames, uin
 
 void FWindowsPlatformMisc::SetUTF8Output()
 {
+	CA_SUPPRESS(6031)
 	_setmode(_fileno(stdout), _O_U8TEXT);
 }
 
@@ -1349,6 +1351,7 @@ static bool HandleGameExplorerIntegration()
 
 					// convert guid to a string
 					TCHAR GuidDir[MAX_PATH];
+					CA_SUPPRESS(6031)
 					StringFromGUID2(GEGuid, GuidDir, MAX_PATH - 1);
 
 					// make the base path for all tasks
@@ -1389,6 +1392,7 @@ static bool HandleGameExplorerIntegration()
 					Link->Release();
 
 					IUniformResourceLocator* InternetLink;
+					CA_SUPPRESS(6031)
 					CoCreateInstance (CLSID_InternetShortcut, NULL, 
 						CLSCTX_INPROC_SERVER, IID_IUniformResourceLocator, (LPVOID*) &InternetLink);
 
@@ -1769,7 +1773,6 @@ bool FWindowsPlatformMisc::OsExecute(const TCHAR* CommandType, const TCHAR* Comm
 		NULL,
 		SW_SHOWNORMAL);
 	bool bSucceeded = hApp > (HINSTANCE)32;
-	CloseHandle(hApp);
 	return bSucceeded;
 }
 
diff --git a/Engine/Source/Runtime/Core/Private/Windows/WindowsPlatformSplash.cpp b/Engine/Source/Runtime/Core/Private/Windows/WindowsPlatformSplash.cpp
--- a/Engine/Source/Runtime/Core/Private/Windows/WindowsPlatformSplash.cpp
+++ b/Engine/Source/Runtime/Core/Private/Windows/WindowsPlatformSplash.cpp
@@ -326,18 +326,37 @@ HBITMAP LoadSplashBitmap()
 		}
 
 		// Image Extraction failed, clear allocated memory
-		if (FAILED(hr))
+		if (FAILED(hr) && hDIBBitmap)
 		{
 			DeleteObject(hDIBBitmap);
 			hDIBBitmap = NULL;
 		}
 	}
 
-	DeleteObject(OriginalBitmapSource);
-	DeleteObject(ToRenderBitmapSource);
-	DeleteObject(Decoder);
-	DeleteObject(Frame);
-	DeleteObject(Factory);
+	if ( OriginalBitmapSource )
+	{
+		DeleteObject( OriginalBitmapSource );
+	}
+
+	if ( ToRenderBitmapSource )
+	{
+		DeleteObject( ToRenderBitmapSource );
+	}
+
+	if ( Decoder )
+	{
+		DeleteObject( Decoder );
+	}
+
+	if ( Frame )
+	{
+		DeleteObject( Frame );
+	}
+
+	if ( Factory )
+	{
+		DeleteObject( Factory );
+	}
 
 	return hDIBBitmap;
 }
@@ -524,6 +543,7 @@ uint32 WINAPI StartSplashScreenThread( LPVOID unused )
 				}
 
 				// Update window opacity
+				CA_SUPPRESS(6239)
 				if( bAllowFading && CurrentOpacityByte < 255 )
 				{
 					// Set window to fully transparent to start out
diff --git a/Engine/Source/Runtime/Core/Private/Windows/WindowsPlatformStackWalk.cpp b/Engine/Source/Runtime/Core/Private/Windows/WindowsPlatformStackWalk.cpp
--- a/Engine/Source/Runtime/Core/Private/Windows/WindowsPlatformStackWalk.cpp
+++ b/Engine/Source/Runtime/Core/Private/Windows/WindowsPlatformStackWalk.cpp
@@ -491,9 +491,13 @@ static void LoadProcessModules()
 		// Set the search path to find PDBs in the same folder as the DLL.
 		ANSICHAR SearchPath[1024];
 		ANSICHAR* FileName = NULL;
-		GetFullPathNameA( ImageName, 1024, SearchPath, &FileName );
-		*FileName = 0;
-		SymSetSearchPath( GetCurrentProcess(), SearchPath );
+		const auto Result =	GetFullPathNameA( ImageName, ARRAY_COUNT( SearchPath ), SearchPath, &FileName );
+
+		if(Result != 0 && Result < ARRAY_COUNT(SearchPath))
+		{
+			*FileName = 0;
+			SymSetSearchPath(GetCurrentProcess(), SearchPath);
+		}
 
 		// Load module.
 		DWORD64 BaseAddress = SymLoadModule64( ProcessHandle, ModuleHandleArray[ModuleIndex], ImageName, ModuleName, (DWORD64) ModuleInfo.lpBaseOfDll, (uint32) ModuleInfo.SizeOfImage );
diff --git a/Engine/Source/Runtime/Core/Private/Windows/WindowsPlatformSurvey.cpp b/Engine/Source/Runtime/Core/Private/Windows/WindowsPlatformSurvey.cpp
--- a/Engine/Source/Runtime/Core/Private/Windows/WindowsPlatformSurvey.cpp
+++ b/Engine/Source/Runtime/Core/Private/Windows/WindowsPlatformSurvey.cpp
@@ -465,9 +465,9 @@ void FWindowsPlatformSurvey::TickSurveyHardware( FHardwareSurveyResults& OutResu
 	LCID DefaultLocale = GetSystemDefaultLCID();
 	const int32 MaxLocaleStringLength = 9;
 	TCHAR LangBuffer[MaxLocaleStringLength];
-	int LangReturn = GetLocaleInfo(DefaultLocale, LOCALE_SISO639LANGNAME, LangBuffer, MaxLocaleStringLength * sizeof(TCHAR));
+	int LangReturn = GetLocaleInfo(DefaultLocale, LOCALE_SISO639LANGNAME, LangBuffer, ARRAY_COUNT(LangBuffer));
 	TCHAR CountryBuffer[MaxLocaleStringLength];
-	int CountryReturn = GetLocaleInfo(DefaultLocale, LOCALE_SISO3166CTRYNAME, CountryBuffer, MaxLocaleStringLength * sizeof(TCHAR));
+	int CountryReturn = GetLocaleInfo(DefaultLocale, LOCALE_SISO3166CTRYNAME, CountryBuffer, ARRAY_COUNT(CountryBuffer));
 
 	if (LangReturn == 0 || CountryReturn == 0)
 	{
@@ -588,10 +588,7 @@ void FWindowsPlatformSurvey::TickSurveyHardware( FHardwareSurveyResults& OutResu
 						OutResults.GPUPerformanceIndex = 0.5f * (GPU3DScoreIndex + GPUDesktopScoreIndex);
 					}
 				}
-			}
-
-			if (WinSATResults)
-			{
+				
 				WinSATResults->Release();
 			}
 		}
@@ -668,10 +665,7 @@ bool FWindowsPlatformSurvey::GetSubComponentIndex( IProvideWinSATResultsInfo* Wi
 		{
 			bSuccess = true;
 		}
-	}
 
-	if (AssessmentInfo)
-	{
 		AssessmentInfo->Release();
 	}
 #endif
@@ -697,6 +691,7 @@ void FWindowsPlatformSurvey::GetOSVersionLabels(const SYSTEM_INFO& SystemInfo, F
 		OSVERSIONINFOEX OsVersionInfo = {0};
 		OsVersionInfo.dwOSVersionInfoSize = sizeof( OSVERSIONINFOEX );
 #pragma warning(disable : 4996) // 'function' was declared deprecated
+		CA_SUPPRESS(28159)
 		GetVersionEx( (LPOSVERSIONINFO)&OsVersionInfo );
 #pragma warning(default : 4996)
 
diff --git a/Engine/Source/Runtime/Core/Private/Windows/WindowsRunnableThread.h b/Engine/Source/Runtime/Core/Private/Windows/WindowsRunnableThread.h
--- a/Engine/Source/Runtime/Core/Private/Windows/WindowsRunnableThread.h
+++ b/Engine/Source/Runtime/Core/Private/Windows/WindowsRunnableThread.h
@@ -82,7 +82,8 @@ class FRunnableThreadWin : public FRunnableThread
 		}
 #if !PLATFORM_SEH_EXCEPTIONS_DISABLED
 		__except( EXCEPTION_EXECUTE_HANDLER )
-		{
+		CA_SUPPRESS(6322)
+		{	
 		}
 #endif
 	}
diff --git a/Engine/Source/Runtime/Core/Private/Windows/WindowsTextInputMethodSystem.cpp b/Engine/Source/Runtime/Core/Private/Windows/WindowsTextInputMethodSystem.cpp
--- a/Engine/Source/Runtime/Core/Private/Windows/WindowsTextInputMethodSystem.cpp
+++ b/Engine/Source/Runtime/Core/Private/Windows/WindowsTextInputMethodSystem.cpp
@@ -119,7 +119,7 @@ STDMETHODIMP_(ULONG) FTSFActivationProxy::Release()
 	return LocalReferenceCount;
 }
 
-STDAPI FTSFActivationProxy::OnActivated(DWORD dwProfileType, LANGID langid, REFCLSID clsid, REFGUID catid, REFGUID guidProfile, HKL hkl, DWORD dwFlags)
+STDAPI FTSFActivationProxy::OnActivated(DWORD dwProfileType, LANGID langid, __RPC__in REFCLSID clsid, __RPC__in REFGUID catid, __RPC__in REFGUID guidProfile, HKL hkl, DWORD dwFlags)
 {
 	FString APIString;
 	if(::ImmGetIMEFileName(::GetKeyboardLayout(0), nullptr, 0) > 0)
@@ -325,6 +325,7 @@ bool FWindowsTextInputMethodSystem::InitializeTSF()
 		TSFActivationProxy = new FTSFActivationProxy(this);
 
 #pragma warning(disable : 4996) // 'function' was was declared deprecated
+		CA_SUPPRESS(28159)
 		const DWORD WindowsVersion = ::GetVersion();
 #pragma warning(default : 4996)
 
@@ -556,7 +557,7 @@ void FWindowsTextInputMethodSystem::DeactivateContext(const TSharedRef<ITextInpu
 	{
 		UE_LOG(LogWindowsTextInputMethodSystem, Error, TEXT("Deactivating a context failed while getting the currently focused document manager."));
 	}
-	if(FocusedDocumentManger == TextStore->TSFDocumentManager)
+	else if(FocusedDocumentManger == TextStore->TSFDocumentManager)
 	{
 		Result = TSFThreadManager->SetFocus(TSFDisabledDocumentManager);
 		if(FAILED(Result))
diff --git a/Engine/Source/Runtime/Core/Public/GenericPlatform/GenericPlatformMemory.h b/Engine/Source/Runtime/Core/Public/GenericPlatform/GenericPlatformMemory.h
--- a/Engine/Source/Runtime/Core/Public/GenericPlatform/GenericPlatformMemory.h
+++ b/Engine/Source/Runtime/Core/Public/GenericPlatform/GenericPlatformMemory.h
@@ -161,7 +161,7 @@ struct CORE_API FGenericPlatformMemory
 	/** Initializes platform memory specific constants. */
 	static void Init();
 	
-	static void OnOutOfMemory(uint64 Size, uint32 Alignment);
+	static CA_NORETURN void OnOutOfMemory(uint64 Size, uint32 Alignment);
 
 	/** Initializes the memory pools, should be called by the init function. */
 	static void SetupMemoryPools();
diff --git a/Engine/Source/Runtime/Core/Public/HAL/IPlatformFileCachedWrapper.h b/Engine/Source/Runtime/Core/Public/HAL/IPlatformFileCachedWrapper.h
--- a/Engine/Source/Runtime/Core/Public/HAL/IPlatformFileCachedWrapper.h
+++ b/Engine/Source/Runtime/Core/Public/HAL/IPlatformFileCachedWrapper.h
@@ -106,6 +106,8 @@ public:
 				// copy from the cache to the destination
 				if (Result)
 				{
+					// Analyzer doesn't see this - if this code ever changes make sure there are no buffer overruns!
+					CA_ASSUME(CacheIndex < CacheCount);
 					uint64 CorrectedBytesToRead=FMath::Min<uint64>(BytesToRead, CacheEnd[CacheIndex]-FilePos);
 					FMemory::Memcpy(Destination, BufferCache[CacheIndex]+(FilePos-CacheStart[CacheIndex]), CorrectedBytesToRead);
 					FilePos += CorrectedBytesToRead;
diff --git a/Engine/Source/Runtime/Core/Public/HAL/MallocBinned.h b/Engine/Source/Runtime/Core/Public/HAL/MallocBinned.h
--- a/Engine/Source/Runtime/Core/Public/HAL/MallocBinned.h
+++ b/Engine/Source/Runtime/Core/Public/HAL/MallocBinned.h
@@ -317,7 +317,7 @@ private:
 #endif
 
 	// Implementation. 
-	void OutOfMemory(uint64 Size, uint32 Alignment=0)
+	CA_NORETURN void OutOfMemory(uint64 Size, uint32 Alignment=0)
 	{
 		// this is expected not to return
 		FPlatformMemory::OnOutOfMemory(Size, Alignment);
@@ -384,6 +384,7 @@ private:
 				{
 					collision->Key=Key;
 					InitializeHashBucket(collision);
+					CA_ASSUME(collision->FirstPool);
 				}
 				return &collision->FirstPool[PoolIndex];
 			}
diff --git a/Engine/Source/Runtime/Core/Public/Misc/AssertionMacros.h b/Engine/Source/Runtime/Core/Public/Misc/AssertionMacros.h
--- a/Engine/Source/Runtime/Core/Public/Misc/AssertionMacros.h
+++ b/Engine/Source/Runtime/Core/Public/Misc/AssertionMacros.h
@@ -105,7 +105,7 @@
  * ensure() resolves to a regular assertion (crash) in shipping or test builds.
  */
 
-#if DO_CHECK
+#if DO_CHECK && !USING_CODE_ANALYSIS // The Visual Studio 2013 analyzer doesn't understand these complex conditionals
 
 	#define ensure( InExpression ) \
 		FDebug::EnsureNotFalse( ( ( InExpression ) != 0 ), #InExpression, __FILE__, __LINE__ )
diff --git a/Engine/Source/Runtime/Core/Public/Misc/Char.h b/Engine/Source/Runtime/Core/Public/Misc/Char.h
--- a/Engine/Source/Runtime/Core/Public/Misc/Char.h
+++ b/Engine/Source/Runtime/Core/Public/Misc/Char.h
@@ -216,7 +216,7 @@ template <> inline bool TChar<WIDECHAR>::IsPunct( CharType c )				{ return ::isw
 -----------------------------------------------------------------------------*/
 template <> inline TChar<ANSICHAR>::CharType TChar<ANSICHAR>::ToUpper( CharType c )	{ return (ANSICHAR)::toupper(c); }
 template <> inline TChar<ANSICHAR>::CharType TChar<ANSICHAR>::ToLower( CharType c )	{ return (ANSICHAR)::tolower(c); }
-template <> inline bool TChar<ANSICHAR>::IsUpper( CharType c )							{ return ::isupper(c) != 0; }
-template <> inline bool TChar<ANSICHAR>::IsLower( CharType c )							{ return ::islower(c) != 0; }
-template <> inline bool TChar<ANSICHAR>::IsAlpha( CharType c )							{ return ::isalpha(c) != 0; }
-template <> inline bool TChar<ANSICHAR>::IsPunct( CharType c )							{ return ::ispunct( c ) != 0; }
+template <> inline bool TChar<ANSICHAR>::IsUpper( CharType c )							{ return ::isupper(static_cast<unsigned char>(c)) != 0; }
+template <> inline bool TChar<ANSICHAR>::IsLower(CharType c)							{ return ::islower(static_cast<unsigned char>(c)) != 0; }
+template <> inline bool TChar<ANSICHAR>::IsAlpha(CharType c)							{ return ::isalpha(static_cast<unsigned char>(c)) != 0; }
+template <> inline bool TChar<ANSICHAR>::IsPunct(CharType c)							{ return ::ispunct(static_cast<unsigned char>(c)) != 0; }
diff --git a/Engine/Source/Runtime/Core/Public/Misc/CoreMiscDefines.h b/Engine/Source/Runtime/Core/Public/Misc/CoreMiscDefines.h
--- a/Engine/Source/Runtime/Core/Public/Misc/CoreMiscDefines.h
+++ b/Engine/Source/Runtime/Core/Public/Misc/CoreMiscDefines.h
@@ -61,7 +61,7 @@
 #endif
 
 #if USING_CODE_ANALYSIS
-	#if !defined( CA_IN ) || !defined( CA_OUT ) || !defined( CA_READ_ONLY ) || !defined( CA_WRITE_ONLY ) || !defined( CA_VALID_POINTER ) || !defined( CA_CHECK_RETVAL ) || !defined( CA_SUPPRESS ) || !defined( CA_ASSUME )
+	#if !defined( CA_IN ) || !defined( CA_OUT ) || !defined( CA_READ_ONLY ) || !defined( CA_WRITE_ONLY ) || !defined( CA_VALID_POINTER ) || !defined( CA_CHECK_RETVAL ) || !defined( CA_NORETURN) || !defined( CA_SUPPRESS ) || !defined( CA_ASSUME )
 		#error Code analysis macros are not configured correctly for this platform
 	#endif
 #else
@@ -72,6 +72,7 @@
 	#define CA_WRITE_ONLY
 	#define CA_VALID_POINTER
 	#define CA_CHECK_RETVAL
+	#define CA_NORETURN
 	#define CA_SUPPRESS( WarningNumber )
 	#define CA_ASSUME( Expr )
 #endif
diff --git a/Engine/Source/Runtime/Core/Public/Templates/UnrealTemplate.h b/Engine/Source/Runtime/Core/Public/Templates/UnrealTemplate.h
--- a/Engine/Source/Runtime/Core/Public/Templates/UnrealTemplate.h
+++ b/Engine/Source/Runtime/Core/Public/Templates/UnrealTemplate.h
@@ -162,7 +162,7 @@ template <typename T, uint32 N>
 char (&ArrayCountHelper(const T (&)[N]))[N];
 
 // Number of elements in an array.
-#define ARRAY_COUNT( array ) sizeof(ArrayCountHelper(array))
+#define ARRAY_COUNT( array ) (sizeof(ArrayCountHelper(array))+0)
 
 // Offset of a struct member.
 #define STRUCT_OFFSET( struc, member )	offsetof(struc, member)
diff --git a/Engine/Source/Runtime/Core/Public/Windows/TextStoreACP.h b/Engine/Source/Runtime/Core/Public/Windows/TextStoreACP.h
--- a/Engine/Source/Runtime/Core/Public/Windows/TextStoreACP.h
+++ b/Engine/Source/Runtime/Core/Public/Windows/TextStoreACP.h
@@ -22,49 +22,49 @@ public:
 	// IUnknown Interface End
 
 	// ITextStoreACP Interface Begin
-	STDMETHODIMP AdviseSink(REFIID riid, IUnknown *punk, DWORD dwMask);
-	STDMETHODIMP UnadviseSink(IUnknown *punk);
+	STDMETHODIMP AdviseSink(__RPC__in REFIID riid, __RPC__in_opt IUnknown *punk, DWORD dwMask);
+	STDMETHODIMP UnadviseSink(__RPC__in_opt IUnknown *punk);
 	STDMETHODIMP RequestLock(DWORD dwLockFlags, HRESULT *phrSession);
 
-	STDMETHODIMP GetStatus(TS_STATUS *pdcs);
-	STDMETHODIMP GetEndACP(LONG *pacp);
+	STDMETHODIMP GetStatus(__RPC__out TS_STATUS *pdcs);
+	STDMETHODIMP GetEndACP(__RPC__out LONG *pacp);
 
 	// Selection Methods
-	STDMETHODIMP GetSelection(ULONG ulIndex, ULONG ulCount, TS_SELECTION_ACP *pSelection, ULONG *pcFetched);
-	STDMETHODIMP SetSelection(ULONG ulCount, const TS_SELECTION_ACP *pSelection);
+	STDMETHODIMP GetSelection(ULONG ulIndex, ULONG ulCount, __RPC__out_ecount_part(ulCount, *pcFetched) TS_SELECTION_ACP *pSelection, __RPC__out ULONG *pcFetched);
+	STDMETHODIMP SetSelection(ULONG ulCount, __RPC__in_ecount_full(ulCount) const TS_SELECTION_ACP *pSelection);
 
 	// Attributes Methods
-	STDMETHODIMP RequestSupportedAttrs(DWORD dwFlags, ULONG cFilterAttrs, const TS_ATTRID *paFilterAttrs);
-	STDMETHODIMP RequestAttrsAtPosition(LONG acpPos, ULONG cFilterAttrs, const TS_ATTRID *paFilterAttrs, DWORD dwFlags);
-	STDMETHODIMP RequestAttrsTransitioningAtPosition(LONG acpPos, ULONG cFilterAttrs, const TS_ATTRID *paFilterAttrs, DWORD dwFlags);
-	STDMETHODIMP FindNextAttrTransition(LONG acpStart, LONG acpHalt, ULONG cFilterAttrs, const TS_ATTRID *paFilterAttrs, DWORD dwFlags, LONG *pacpNext, BOOL *pfFound, LONG *plFoundOffset);
-	STDMETHODIMP RetrieveRequestedAttrs(ULONG ulCount, TS_ATTRVAL *paAttrVals, ULONG *pcFetched);
+	STDMETHODIMP RequestSupportedAttrs(DWORD dwFlags, ULONG cFilterAttrs, __RPC__in_ecount_full_opt(cFilterAttrs) const TS_ATTRID *paFilterAttrs);
+	STDMETHODIMP RequestAttrsAtPosition(LONG acpPos, ULONG cFilterAttrs, __RPC__in_ecount_full_opt(cFilterAttrs) const TS_ATTRID *paFilterAttrs, DWORD dwFlags);
+	STDMETHODIMP RequestAttrsTransitioningAtPosition(LONG acpPos, ULONG cFilterAttrs, __RPC__in_ecount_full_opt(cFilterAttrs) const TS_ATTRID *paFilterAttrs, DWORD dwFlags);
+	STDMETHODIMP FindNextAttrTransition(LONG acpStart, LONG acpHalt, ULONG cFilterAttrs, __RPC__in_ecount_full_opt(cFilterAttrs) const TS_ATTRID *paFilterAttrs, DWORD dwFlags, __RPC__out LONG *pacpNext, __RPC__out BOOL *pfFound, __RPC__out LONG *plFoundOffset);
+	STDMETHODIMP RetrieveRequestedAttrs(ULONG ulCount, __RPC__out_ecount_part(ulCount, *pcFetched) TS_ATTRVAL *paAttrVals, __RPC__out ULONG *pcFetched);
 
 	// View Methods
-	STDMETHODIMP GetActiveView(TsViewCookie *pvcView);
-	STDMETHODIMP GetACPFromPoint(TsViewCookie vcView, const POINT *pt, DWORD dwFlags, LONG *pacp);
-	STDMETHODIMP GetTextExt(TsViewCookie vcView, LONG acpStart, LONG acpEnd, RECT *prc, BOOL *pfClipped);
-	STDMETHODIMP GetScreenExt(TsViewCookie vcView, RECT *prc);
-	STDMETHODIMP GetWnd(TsViewCookie vcView, HWND *phwnd);
+	STDMETHODIMP GetActiveView(__RPC__out TsViewCookie *pvcView);
+	STDMETHODIMP GetACPFromPoint(TsViewCookie vcView, __RPC__in const POINT *pt, DWORD dwFlags, __RPC__out LONG *pacp);
+	STDMETHODIMP GetTextExt(TsViewCookie vcView, LONG acpStart, LONG acpEnd, __RPC__out RECT *prc, __RPC__out BOOL *pfClipped);
+	STDMETHODIMP GetScreenExt(TsViewCookie vcView, __RPC__out RECT *prc);
+	STDMETHODIMP GetWnd(TsViewCookie vcView, __RPC__deref_out_opt HWND *phwnd);
 
 	// Plain Character Methods
-	STDMETHODIMP GetText(LONG acpStart, LONG acpEnd, __out_ecount(cchPlainReq) WCHAR *pchPlain, ULONG cchPlainReq, ULONG *pcchPlainOut, TS_RUNINFO *prgRunInfo, ULONG ulRunInfoReq, ULONG *pulRunInfoOut, LONG *pacpNext);
-	STDMETHODIMP QueryInsert(LONG acpInsertStart, LONG acpInsertEnd, ULONG cch, LONG *pacpResultStart, LONG *pacpResultEnd);
-	STDMETHODIMP InsertTextAtSelection(DWORD dwFlags, __in_ecount(cch) const WCHAR *pchText, ULONG cch, LONG *pacpStart, LONG *pacpEnd, TS_TEXTCHANGE *pChange);
-	STDMETHODIMP SetText(DWORD dwFlags, LONG acpStart, LONG acpEnd, __in_ecount(cch) const WCHAR *pchText, ULONG cch, TS_TEXTCHANGE *pChange);
+	STDMETHODIMP GetText(LONG acpStart, LONG acpEnd, __RPC__out_ecount_part(cchPlainReq, *pcchPlainOut) WCHAR *pchPlain, ULONG cchPlainReq, __RPC__out ULONG *pcchPlainOut, __RPC__out_ecount_part(ulRunInfoReq, *pulRunInfoOut) TS_RUNINFO *prgRunInfo, ULONG ulRunInfoReq, __RPC__out ULONG *pulRunInfoOut, __RPC__out LONG *pacpNext);
+	STDMETHODIMP QueryInsert(LONG acpInsertStart, LONG acpInsertEnd, ULONG cch, __RPC__out LONG *pacpResultStart, __RPC__out LONG *pacpResultEnd);
+	STDMETHODIMP InsertTextAtSelection(DWORD dwFlags, __RPC__in_ecount_full(cch) const WCHAR *pchText, ULONG cch, __RPC__out LONG *pacpStart, __RPC__out LONG *pacpEnd, __RPC__out TS_TEXTCHANGE *pChange);
+	STDMETHODIMP SetText(DWORD dwFlags, LONG acpStart, LONG acpEnd, __RPC__in_ecount_full(cch) const WCHAR *pchText, ULONG cch, __RPC__out TS_TEXTCHANGE *pChange);
 
 	// Embedded Character Methods
-	STDMETHODIMP GetEmbedded(LONG acpPos, REFGUID rguidService, REFIID riid, IUnknown **ppunk);
-	STDMETHODIMP GetFormattedText(LONG acpStart, LONG acpEnd, IDataObject **ppDataObject);
-	STDMETHODIMP QueryInsertEmbedded(const GUID *pguidService, const FORMATETC *pFormatEtc, BOOL *pfInsertable);
-	STDMETHODIMP InsertEmbedded(DWORD dwFlags, LONG acpStart, LONG acpEnd, IDataObject *pDataObject, TS_TEXTCHANGE *pChange);
-	STDMETHODIMP InsertEmbeddedAtSelection(DWORD dwFlags, IDataObject *pDataObject, LONG *pacpStart, LONG *pacpEnd, TS_TEXTCHANGE *pChange);
+	STDMETHODIMP GetEmbedded(LONG acpPos, __RPC__in REFGUID rguidService, __RPC__in REFIID riid, __RPC__deref_out_opt IUnknown **ppunk);
+	STDMETHODIMP GetFormattedText(LONG acpStart, LONG acpEnd, __RPC__deref_out_opt IDataObject **ppDataObject);
+	STDMETHODIMP QueryInsertEmbedded(__RPC__in const GUID *pguidService, __RPC__in const FORMATETC *pFormatEtc, __RPC__out BOOL *pfInsertable);
+	STDMETHODIMP InsertEmbedded(DWORD dwFlags, LONG acpStart, LONG acpEnd, __RPC__in_opt IDataObject *pDataObject, __RPC__out TS_TEXTCHANGE *pChange);
+	STDMETHODIMP InsertEmbeddedAtSelection(DWORD dwFlags, __RPC__in_opt IDataObject *pDataObject, __RPC__out LONG *pacpStart, __RPC__out LONG *pacpEnd, __RPC__out TS_TEXTCHANGE *pChange);
 	// ITextStoreACP Interface End
 
 	// ITfContextOwnerCompositionSink Interface Begin
-	STDMETHODIMP OnStartComposition(ITfCompositionView *pComposition, BOOL *pfOk);
-	STDMETHODIMP OnUpdateComposition(ITfCompositionView *pComposition, ITfRange *pRangeNew);
-	STDMETHODIMP OnEndComposition(ITfCompositionView *pComposition);
+	STDMETHODIMP OnStartComposition(__RPC__in_opt ITfCompositionView *pComposition, __RPC__out BOOL *pfOk);
+	STDMETHODIMP OnUpdateComposition(__RPC__in_opt ITfCompositionView *pComposition, __RPC__in_opt ITfRange *pRangeNew);
+	STDMETHODIMP OnEndComposition(__RPC__in_opt ITfCompositionView *pComposition);
 	// ITfContextOwnerCompositionSink Interface End
 
 private:
diff --git a/Engine/Source/Runtime/Core/Public/Windows/WindowsPlatformCodeAnalysis.h b/Engine/Source/Runtime/Core/Public/Windows/WindowsPlatformCodeAnalysis.h
--- a/Engine/Source/Runtime/Core/Public/Windows/WindowsPlatformCodeAnalysis.h
+++ b/Engine/Source/Runtime/Core/Public/Windows/WindowsPlatformCodeAnalysis.h
@@ -61,6 +61,9 @@
 	// Example:  CA_CHECK_RETVAL int32 GetNumber();
 	#define CA_CHECK_RETVAL [returnvalue:Post(MustCheck=Yes)]
 
+	// Function is expected to never return
+	#define CA_NORETURN __declspec(noreturn)
+
 	// Suppresses a warning for a single occurrence.  Should be used only for code analysis warnings on Windows platform!
 	#define CA_SUPPRESS( WarningNumber ) __pragma( warning( suppress: WarningNumber ) )
 
@@ -81,4 +84,12 @@
 	// appAlloca macro (not inline methods) and don't want to suppress at all call sites.
 	#pragma warning( disable : 6255 )	// warning C6255: _alloca indicates failure by raising a stack overflow exception.  Consider using _malloca instead.
 
+	// This a very common false positive warning (but some cases are not false positives!). Disabling for now so that we
+	// can more quickly see the benefits of static analysis on new code.
+	#pragma warning(disable : 6102) // warning C6102: Using 'variable' from failed function call at line 'line'.
+
+	// We use this exception handler in Windows code, it may be worth a closer look, but disabling for now
+	// so we can get the benefits of analysis on cross-platform code sooner.
+	#pragma warning(disable : 6320) // warning C6320: Exception-filter expression is the constant EXCEPTION_EXECUTE_HANDLER. This might mask exceptions that were not intended to be handled.
+
 #endif
\ No newline at end of file
diff --git a/Engine/Source/Runtime/Core/Public/Windows/WindowsPlatformMath.h b/Engine/Source/Runtime/Core/Public/Windows/WindowsPlatformMath.h
--- a/Engine/Source/Runtime/Core/Public/Windows/WindowsPlatformMath.h
+++ b/Engine/Source/Runtime/Core/Public/Windows/WindowsPlatformMath.h
@@ -91,19 +91,25 @@ struct FWindowsPlatformMath : public FGenericPlatformMath
 	#pragma intrinsic( _BitScanReverse )
 	static FORCEINLINE uint32 FloorLog2(uint32 Value) 
 	{
-		if (Value == 0) return 0;
 		// Use BSR to return the log2 of the integer
 		DWORD Log2;
-		_BitScanReverse( &Log2, Value );
-		return Log2;
+		if(_BitScanReverse(&Log2, Value) != 0)
+		{
+			return Log2;
+		}
+
+		return 0;
 	}
 	static FORCEINLINE uint32 CountLeadingZeros(uint32 Value)
 	{
-		if (Value == 0) return 32;
 		// Use BSR to return the log2 of the integer
 		DWORD Log2;
-		_BitScanReverse( &Log2, Value );
-		return 31 - Log2;
+		if(_BitScanReverse(&Log2, Value) != 0)
+		{
+			return 31 - Log2;
+		}
+
+		return 32;
 	}
 	static FORCEINLINE uint32 CeilLogTwo( uint32 Arg )
 	{
diff --git a/Engine/Source/Runtime/Core/Public/Windows/WindowsTextInputMethodSystem.h b/Engine/Source/Runtime/Core/Public/Windows/WindowsTextInputMethodSystem.h
--- a/Engine/Source/Runtime/Core/Public/Windows/WindowsTextInputMethodSystem.h
+++ b/Engine/Source/Runtime/Core/Public/Windows/WindowsTextInputMethodSystem.h
@@ -25,7 +25,7 @@ public:
 	// IUnknown Interface End
 
 	// ITfInputProcessorProfileActivationSink Interface Begin
-	STDMETHODIMP OnActivated(DWORD dwProfileType, LANGID langid, REFCLSID clsid, REFGUID catid, REFGUID guidProfile, HKL hkl, DWORD dwFlags) override;
+	STDMETHODIMP OnActivated(DWORD dwProfileType, LANGID langid, __RPC__in REFCLSID clsid, __RPC__in REFGUID catid, __RPC__in REFGUID guidProfile, HKL hkl, DWORD dwFlags) override;
 	// ITfInputProcessorProfileActivationSink Interface End
 
 	// ITfActiveLanguageProfileNotifySink Interface Begin
