Fix for missing EMIT_DEPRECATED_WARNING_MESSAGE macro on platforms other than Windows

diff --git a/Engine/Source/Runtime/Core/Public/HAL/Platform.h b/Engine/Source/Runtime/Core/Public/HAL/Platform.h
--- a/Engine/Source/Runtime/Core/Public/HAL/Platform.h
+++ b/Engine/Source/Runtime/Core/Public/HAL/Platform.h
@@ -292,6 +292,10 @@
 	#define TYPENAME_OUTSIDE_TEMPLATE	typename
 #endif
 
+#ifndef EMIT_DEPRECATED_WARNING_MESSAGE
+	#define EMIT_DEPRECATED_WARNING_MESSAGE(Msg)
+#endif
+
 // Legacy method modifier macros.  You shouldn't use these macros in modern code.  Use the built-in keyword directly.
 #ifndef OVERRIDE
 	#define OVERRIDE \
