Deprecate and remove checkAtCompileTime.
#ttp 337754

diff --git a/Engine/Source/Developer/AudioFormatOgg/Private/AudioFormatOgg.cpp b/Engine/Source/Developer/AudioFormatOgg/Private/AudioFormatOgg.cpp
--- a/Engine/Source/Developer/AudioFormatOgg/Private/AudioFormatOgg.cpp
+++ b/Engine/Source/Developer/AudioFormatOgg/Private/AudioFormatOgg.cpp
@@ -14,7 +14,7 @@
 	#pragma pack(pop)
 #endif
 
-checkAtCompileTime(WITH_OGGVORBIS, no_point_in_compiling_the_ogg_compressor_if_we_dont_have_vorbis);
+static_assert(WITH_OGGVORBIS, "No point in compiling the OGG compressor if we don't have Vorbis.");
 
 // Vorbis encoded sound is about 15% better quality than XMA - adjust the quality setting to get consistent cross platform sound quality
 #define VORBIS_QUALITY_MODIFIER		-15
diff --git a/Engine/Source/Developer/DerivedDataCache/Private/FileSystemDerivedDataBackend.cpp b/Engine/Source/Developer/DerivedDataCache/Private/FileSystemDerivedDataBackend.cpp
--- a/Engine/Source/Developer/DerivedDataCache/Private/FileSystemDerivedDataBackend.cpp
+++ b/Engine/Source/Developer/DerivedDataCache/Private/FileSystemDerivedDataBackend.cpp
@@ -32,7 +32,8 @@ public:
 		, DaysToDeleteUnusedFiles(InDaysToDeleteUnusedFiles)
 	{
 		// If we find a platform that has more stingent limits, this needs to be rethought.
-		checkAtCompileTime(MAX_BACKEND_KEY_LENGTH + MAX_CACHE_DIR_LEN + MAX_BACKEND_NUMBERED_SUBFOLDER_LENGTH + MAX_CACHE_EXTENTION_LEN < PLATFORM_MAX_FILEPATH_LENGTH, not_enough_room_left_for_cache_keys_in_max_path);
+		static_assert(MAX_BACKEND_KEY_LENGTH + MAX_CACHE_DIR_LEN + MAX_BACKEND_NUMBERED_SUBFOLDER_LENGTH + MAX_CACHE_EXTENTION_LEN < PLATFORM_MAX_FILEPATH_LENGTH,
+			"Not enough room left for cache keys in max path.");
 		const double SlowInitDuration = 10.0;
 		double AccessDuration = 0.0;
 
diff --git a/Engine/Source/Developer/MeshSimplifier/Private/Cache.h b/Engine/Source/Developer/MeshSimplifier/Private/Cache.h
--- a/Engine/Source/Developer/MeshSimplifier/Private/Cache.h
+++ b/Engine/Source/Developer/MeshSimplifier/Private/Cache.h
@@ -32,7 +32,7 @@ protected:
 template< typename T, uint32 Size >
 FORCEINLINE TCacheDirect<T, Size>::TCacheDirect()
 {
-	checkAtCompileTime( ( Size & (Size - 1) ) == 0, Size_must_be_power_of_2 );
+	static_assert((Size & (Size - 1)) == 0, "Size must be power of 2.");
 
 	FMemory::MemSet( Keys, 0xff );
 }
@@ -92,7 +92,7 @@ protected:
 template< typename T, uint32 size >
 FORCEINLINE TCache4Way<T, size>::TCache4Way()
 {
-	checkAtCompileTime( ( Size & (Size - 1) ) == 0, Size_must_be_power_of_2 );
+	static_assert((Size & (Size - 1)) == 0, "Size must be power of 2.");
 
 	FMemory::MemSet( Keys, 0xff );
 
diff --git a/Engine/Source/Developer/MeshSimplifier/Private/QuadricMeshReduction.cpp b/Engine/Source/Developer/MeshSimplifier/Private/QuadricMeshReduction.cpp
--- a/Engine/Source/Developer/MeshSimplifier/Private/QuadricMeshReduction.cpp
+++ b/Engine/Source/Developer/MeshSimplifier/Private/QuadricMeshReduction.cpp
@@ -605,7 +605,7 @@ private:
 			8.0f,	// Highest
 		};
 
-		checkAtCompileTime( ARRAY_COUNT( ImportanceTable ) == (EMeshFeatureImportance::Highest+1), ImportanceTableSizeMismatch ); // -1 because of TEMP_BROKEN(?)
+		static_assert(ARRAY_COUNT(ImportanceTable) == (EMeshFeatureImportance::Highest + 1), "Importance table size mismatch."); // -1 because of TEMP_BROKEN(?)
 		check(Settings.SilhouetteImportance < EMeshFeatureImportance::Highest+1); // -1 because of TEMP_BROKEN(?)
 		check(Settings.TextureImportance < EMeshFeatureImportance::Highest+1); // -1 because of TEMP_BROKEN(?)
 		check(Settings.ShadingImportance < EMeshFeatureImportance::Highest+1); // -1 because of TEMP_BROKEN(?)
diff --git a/Engine/Source/Developer/MeshSimplifier/Private/UnrolledLinkList.h b/Engine/Source/Developer/MeshSimplifier/Private/UnrolledLinkList.h
--- a/Engine/Source/Developer/MeshSimplifier/Private/UnrolledLinkList.h
+++ b/Engine/Source/Developer/MeshSimplifier/Private/UnrolledLinkList.h
@@ -92,7 +92,7 @@ FORCEINLINE TUnrolledLinkList<T, blockSize>::TUnrolledLinkList()
 	: freeNode( &head )
 	, num(0)
 {
-	checkAtCompileTime( ( blockSize & (blockSize - 1) ) == 0, Block_size_must_be_power_of_2 );
+	static_assert((blockSize & (blockSize - 1)) == 0, "Block size must be power of 2.");
 }
 
 template< typename T, uint32 blockSize >
diff --git a/Engine/Source/Developer/MeshUtilities/Private/MeshUtilities.cpp b/Engine/Source/Developer/MeshUtilities/Private/MeshUtilities.cpp
--- a/Engine/Source/Developer/MeshUtilities/Private/MeshUtilities.cpp
+++ b/Engine/Source/Developer/MeshUtilities/Private/MeshUtilities.cpp
@@ -595,7 +595,7 @@ namespace NvTriStrip
 	template<typename IndexDataType, typename Allocator>
 	void CacheOptimizeIndexBuffer(TArray<IndexDataType,Allocator>& Indices)
 	{
-		checkAtCompileTime(sizeof(IndexDataType) == 2 || sizeof(IndexDataType) == 4, IndicesMustBeShortOrInt);
+		static_assert(sizeof(IndexDataType) == 2 || sizeof(IndexDataType) == 4, "Indices must be short or int.");
 
 		PrimitiveGroup*	PrimitiveGroups = NULL;
 		uint32			NumPrimitiveGroups = 0;
@@ -668,7 +668,7 @@ namespace Forsyth
 	template<typename IndexDataType, typename Allocator>
 	void CacheOptimizeIndexBuffer(TArray<IndexDataType,Allocator>& Indices)
 	{
-		checkAtCompileTime(sizeof(IndexDataType) == 2 || sizeof(IndexDataType) == 4, IndicesMustBeShortOrInt);
+		static_assert(sizeof(IndexDataType) == 2 || sizeof(IndexDataType) == 4, "Indices must be short or int.");
 		bool Is32Bit = sizeof(IndexDataType) == 4;
 
 		// Count the number of vertices
diff --git a/Engine/Source/Developer/PhysXFormats/Private/PhysXFormats.cpp b/Engine/Source/Developer/PhysXFormats/Private/PhysXFormats.cpp
--- a/Engine/Source/Developer/PhysXFormats/Private/PhysXFormats.cpp
+++ b/Engine/Source/Developer/PhysXFormats/Private/PhysXFormats.cpp
@@ -8,7 +8,7 @@
 #include "TargetPlatform.h"
 #include "PhysXFormats.h"
 
-checkAtCompileTime(WITH_PHYSX, no_point_in_compiling_physx_cooker_if_we_dont_have_physx);
+static_assert(WITH_PHYSX, "No point in compiling PhysX cooker, if we don't have PhysX.");
 
 static FName NAME_PhysXPC(TEXT("PhysXPC"));
 static FName NAME_PhysXXboxOne(TEXT("PhysXXboxOne"));
diff --git a/Engine/Source/Developer/Profiler/Private/ProfilerCommands.cpp b/Engine/Source/Developer/Profiler/Private/ProfilerCommands.cpp
--- a/Engine/Source/Developer/Profiler/Private/ProfilerCommands.cpp
+++ b/Engine/Source/Developer/Profiler/Private/ProfilerCommands.cpp
@@ -58,7 +58,7 @@ void FProfilerCommands::RegisterCommands()
 }
 PRAGMA_ENABLE_OPTIMIZATION
 
-//checkAtCompileTime(sizeof(FProfilerActionManager) == 0, FProfilerActionManager_CannotContainAnyVariablesAtThisMoment);
+//static_assert(sizeof(FProfilerActionManager) == 0, "Cannot contain any variables at this moment.");
 
 /*-----------------------------------------------------------------------------
 	FProfilerMenuBuilder
diff --git a/Engine/Source/Developer/Profiler/Private/ProfilerDataSource.h b/Engine/Source/Developer/Profiler/Private/ProfilerDataSource.h
--- a/Engine/Source/Developer/Profiler/Private/ProfilerDataSource.h
+++ b/Engine/Source/Developer/Profiler/Private/ProfilerDataSource.h
@@ -1709,7 +1709,7 @@ public:
 	template< const EEventPropertyIndex::Type BooleanPropertyIndex >
 	void SetBooleanStateForAllChildren( const bool bValue )
 	{
-		//checkAtCompileTime( false, FEventGraphSample_SetBooleanStateForAllChildren_BooleanPropertyIndex_NotSupported );
+		//static_assert(false, "Method not supported.");
  
 		// Not sure why but this if-statement is not optimized based on BooleanPropertyIndex
 		if( BooleanPropertyIndex == EEventPropertyIndex::bIsCulled )
diff --git a/Engine/Source/Developer/ShaderFormatOpenGL/Private/OpenGLShaderCompiler.cpp b/Engine/Source/Developer/ShaderFormatOpenGL/Private/OpenGLShaderCompiler.cpp
--- a/Engine/Source/Developer/ShaderFormatOpenGL/Private/OpenGLShaderCompiler.cpp
+++ b/Engine/Source/Developer/ShaderFormatOpenGL/Private/OpenGLShaderCompiler.cpp
@@ -339,8 +339,8 @@ const TCHAR* GLFrequencyStringTable[] =
 };
 
 /** Compile time check to verify that the GL mapping tables are up-to-date. */
-checkAtCompileTime(SF_NumFrequencies == ARRAY_COUNT(GLFrequencyTable), SF_NumFrequencesChanged_PleaseUpdateTables);
-checkAtCompileTime(ARRAY_COUNT(GLFrequencyTable) == ARRAY_COUNT(GLFrequencyStringTable), FrequencyTableSizeMismatch);
+static_assert(SF_NumFrequencies == ARRAY_COUNT(GLFrequencyTable), "NumFrequencies changed. Please update tables.");
+static_assert(ARRAY_COUNT(GLFrequencyTable) == ARRAY_COUNT(GLFrequencyStringTable), "Frequency table size mismatch.");
 
 /**
  * Parse a GLSL error.
diff --git a/Engine/Source/Developer/SimplygonMeshReduction/Private/SimplygonMeshReduction.cpp b/Engine/Source/Developer/SimplygonMeshReduction/Private/SimplygonMeshReduction.cpp
--- a/Engine/Source/Developer/SimplygonMeshReduction/Private/SimplygonMeshReduction.cpp
+++ b/Engine/Source/Developer/SimplygonMeshReduction/Private/SimplygonMeshReduction.cpp
@@ -843,7 +843,7 @@ private:
 			8.0f,	// Highest
 		};
 
-		checkAtCompileTime( ARRAY_COUNT( ImportanceTable ) == (EMeshFeatureImportance::Highest+1), ImportanceTableSizeMismatch ); // -1 because of TEMP_BROKEN(?)
+		static_assert(ARRAY_COUNT(ImportanceTable) == (EMeshFeatureImportance::Highest + 1), "Importance table size mismatch."); // -1 because of TEMP_BROKEN(?)
 		check(Settings.SilhouetteImportance < EMeshFeatureImportance::Highest+1); // -1 because of TEMP_BROKEN(?)
 		check(Settings.TextureImportance < EMeshFeatureImportance::Highest+1); // -1 because of TEMP_BROKEN(?)
 		check(Settings.ShadingImportance < EMeshFeatureImportance::Highest+1); // -1 because of TEMP_BROKEN(?)
@@ -1362,7 +1362,7 @@ private:
 			8.0f,	// Highest
 		};
 
-		checkAtCompileTime( ARRAY_COUNT( ImportanceTable ) == SMOI_MAX, BadImportanceTableSize );
+		static_assert(ARRAY_COUNT(ImportanceTable) == SMOI_MAX, "Bad importance table size.");
 		check( Settings.SilhouetteImportance < SMOI_MAX );
 		check( Settings.TextureImportance < SMOI_MAX );
 		check( Settings.ShadingImportance < SMOI_MAX );
diff --git a/Engine/Source/Editor/BlueprintGraph/Private/EdGraphSchema_K2.cpp b/Engine/Source/Editor/BlueprintGraph/Private/EdGraphSchema_K2.cpp
--- a/Engine/Source/Editor/BlueprintGraph/Private/EdGraphSchema_K2.cpp
+++ b/Engine/Source/Editor/BlueprintGraph/Private/EdGraphSchema_K2.cpp
@@ -1704,7 +1704,7 @@ void UEdGraphSchema_K2::ClearPinWatch(UEdGraphPin const* Pin) const
 bool UEdGraphSchema_K2::DefaultValueSimpleValidation(const FEdGraphPinType& PinType, const FString& PinName, const FString& NewDefaultValue, UObject* NewDefaultObject, const FText& InNewDefaultText, FString* OutMsg /*= NULL*/) const
 {
 #ifdef DVSV_RETURN_MSG
-	checkAtCompileTime(false, DVSV_RETURN_MSG_macro_redefined);
+	static_assert(false, "Macro redefinition.");
 #endif
 #define DVSV_RETURN_MSG(Str) if(NULL != OutMsg) { *OutMsg = Str; } return false;
 
diff --git a/Engine/Source/Editor/GameProjectGeneration/Private/GameProjectUtils.cpp b/Engine/Source/Editor/GameProjectGeneration/Private/GameProjectUtils.cpp
--- a/Engine/Source/Editor/GameProjectGeneration/Private/GameProjectUtils.cpp
+++ b/Engine/Source/Editor/GameProjectGeneration/Private/GameProjectUtils.cpp
@@ -21,7 +21,7 @@
 
 #define MAX_PROJECT_PATH_BUFFER_SPACE 130 // Leave a reasonable buffer of additional characters to account for files created in the content directory during or after project generation
 #define MAX_PROJECT_NAME_LENGTH 20 // Enforce a reasonable project name length so the path is not too long for PLATFORM_MAX_FILEPATH_LENGTH
-checkAtCompileTime(PLATFORM_MAX_FILEPATH_LENGTH - MAX_PROJECT_PATH_BUFFER_SPACE > 0, filesystem_path_shorter_than_project_creation_buffer_space);
+static_assert(PLATFORM_MAX_FILEPATH_LENGTH - MAX_PROJECT_PATH_BUFFER_SPACE > 0, "File system path shorter than project creation buffer space.");
 
 #define MAX_CLASS_NAME_LENGTH 32 // Enforce a reasonable class name length so the path is not too long for PLATFORM_MAX_FILEPATH_LENGTH
 
diff --git a/Engine/Source/Editor/KismetCompiler/Private/KismetCompilerVMBackend.cpp b/Engine/Source/Editor/KismetCompiler/Private/KismetCompilerVMBackend.cpp
--- a/Engine/Source/Editor/KismetCompiler/Private/KismetCompilerVMBackend.cpp
+++ b/Engine/Source/Editor/KismetCompiler/Private/KismetCompilerVMBackend.cpp
@@ -108,11 +108,11 @@ public:
 	{
 		//@TODO: Any endian issues?
 #if SCRIPT_LIMIT_BYTECODE_TO_64KB
-		checkAtCompileTime(sizeof(CodeSkipSizeType) == 2, updateThisCodeAsSizeChanged);
+		static_assert(sizeof(CodeSkipSizeType) == 2, "Update this code as size changed.");
 		ScriptBuffer[WriteOffset] = NewValue & 0xFF;
 		ScriptBuffer[WriteOffset+1] = (NewValue >> 8) & 0xFF;
 #else
-		checkAtCompileTime(sizeof(CodeSkipSizeType) == 4, updateThisCodeAsSizeChanged);
+		static_assert(sizeof(CodeSkipSizeType) == 4, "Update this code as size changed.");
 		ScriptBuffer[WriteOffset] = NewValue & 0xFF;
 		ScriptBuffer[WriteOffset+1] = (NewValue >> 8) & 0xFF;
 		ScriptBuffer[WriteOffset+2] = (NewValue >> 16) & 0xFF;
@@ -161,11 +161,11 @@ struct FSkipOffsetEmitter
 
 		//@TODO: Any endian issues?
 #if SCRIPT_LIMIT_BYTECODE_TO_64KB
-		checkAtCompileTime(sizeof(CodeSkipSizeType) == 2, updateThisCodeAsSizeChanged);
+		static_assert(sizeof(CodeSkipSizeType) == 2, "Update this code as size changed.");
 		Script[SkipWriteIndex] = BytesToSkip & 0xFF;
 		Script[SkipWriteIndex+1] = (BytesToSkip >> 8) & 0xFF;
 #else
-		checkAtCompileTime(sizeof(CodeSkipSizeType) == 4, updateThisCodeAsSizeChanged);
+		static_assert(sizeof(CodeSkipSizeType) == 4, "Update this code as size changed.");
 		Script[SkipWriteIndex] = BytesToSkip & 0xFF;
 		Script[SkipWriteIndex+1] = (BytesToSkip >> 8) & 0xFF;
 		Script[SkipWriteIndex+2] = (BytesToSkip >> 16) & 0xFF;
@@ -1275,6 +1275,6 @@ void FKismetCompilerVMBackend::ConstructFunction(FKismetFunctionContext& Functio
 		ScriptArray.Add(EX_EndOfScript);
 	}
 #else
-	checkAtCompileTime(sizeof(CodeSkipSizeType) == 4, updateThisCodeAsSizeChanged);
+	static_assert(sizeof(CodeSkipSizeType) == 4, "Update this code as size changed.");
 #endif
 }
diff --git a/Engine/Source/Editor/LevelEditor/Private/SLevelViewport.cpp b/Engine/Source/Editor/LevelEditor/Private/SLevelViewport.cpp
--- a/Engine/Source/Editor/LevelEditor/Private/SLevelViewport.cpp
+++ b/Engine/Source/Editor/LevelEditor/Private/SLevelViewport.cpp
@@ -3095,17 +3095,17 @@ FText SLevelViewport::GetMouseCaptureLabelText() const
 
 void SLevelViewport::ShowMouseCaptureLabel(ELabelAnchorMode AnchorMode)
 {
-	checkAtCompileTime((EVerticalAlignment)((ELabelAnchorMode::LabelAnchorMode_TopLeft/3)+1) == EVerticalAlignment::VAlign_Top && (EHorizontalAlignment)((ELabelAnchorMode::LabelAnchorMode_TopLeft%3)+1) == EHorizontalAlignment::HAlign_Left, Alignment_From_ELabelAnchorMode_Error);
-	checkAtCompileTime((EVerticalAlignment)((ELabelAnchorMode::LabelAnchorMode_TopCenter/3)+1) == EVerticalAlignment::VAlign_Top && (EHorizontalAlignment)((ELabelAnchorMode::LabelAnchorMode_TopCenter%3)+1) == EHorizontalAlignment::HAlign_Center, Alignment_From_ELabelAnchorMode_Error);
-	checkAtCompileTime((EVerticalAlignment)((ELabelAnchorMode::LabelAnchorMode_TopRight/3)+1) == EVerticalAlignment::VAlign_Top && (EHorizontalAlignment)((ELabelAnchorMode::LabelAnchorMode_TopRight%3)+1) == EHorizontalAlignment::HAlign_Right, Alignment_From_ELabelAnchorMode_Error);
+	static_assert((EVerticalAlignment)((ELabelAnchorMode::LabelAnchorMode_TopLeft / 3) + 1) == EVerticalAlignment::VAlign_Top && (EHorizontalAlignment)((ELabelAnchorMode::LabelAnchorMode_TopLeft % 3) + 1) == EHorizontalAlignment::HAlign_Left, "Alignment from ELabelAnchorMode error.");
+	static_assert((EVerticalAlignment)((ELabelAnchorMode::LabelAnchorMode_TopCenter / 3) + 1) == EVerticalAlignment::VAlign_Top && (EHorizontalAlignment)((ELabelAnchorMode::LabelAnchorMode_TopCenter % 3) + 1) == EHorizontalAlignment::HAlign_Center, "Alignment from ELabelAnchorMode error.");
+	static_assert((EVerticalAlignment)((ELabelAnchorMode::LabelAnchorMode_TopRight / 3) + 1) == EVerticalAlignment::VAlign_Top && (EHorizontalAlignment)((ELabelAnchorMode::LabelAnchorMode_TopRight % 3) + 1) == EHorizontalAlignment::HAlign_Right, "Alignment from ELabelAnchorMode error.");
 	
-	checkAtCompileTime((EVerticalAlignment)((ELabelAnchorMode::LabelAnchorMode_CenterLeft/3)+1) == EVerticalAlignment::VAlign_Center && (EHorizontalAlignment)((ELabelAnchorMode::LabelAnchorMode_CenterLeft%3)+1) == EHorizontalAlignment::HAlign_Left, Alignment_From_ELabelAnchorMode_Error);
-	checkAtCompileTime((EVerticalAlignment)((ELabelAnchorMode::LabelAnchorMode_Centered/3)+1) == EVerticalAlignment::VAlign_Center && (EHorizontalAlignment)((ELabelAnchorMode::LabelAnchorMode_Centered%3)+1) == EHorizontalAlignment::HAlign_Center, Alignment_From_ELabelAnchorMode_Error);
-	checkAtCompileTime((EVerticalAlignment)((ELabelAnchorMode::LabelAnchorMode_CenterRight/3)+1) == EVerticalAlignment::VAlign_Center && (EHorizontalAlignment)((ELabelAnchorMode::LabelAnchorMode_CenterRight%3)+1) == EHorizontalAlignment::HAlign_Right, Alignment_From_ELabelAnchorMode_Error);
+	static_assert((EVerticalAlignment)((ELabelAnchorMode::LabelAnchorMode_CenterLeft / 3) + 1) == EVerticalAlignment::VAlign_Center && (EHorizontalAlignment)((ELabelAnchorMode::LabelAnchorMode_CenterLeft % 3) + 1) == EHorizontalAlignment::HAlign_Left, "Alignment from ELabelAnchorMode error.");
+	static_assert((EVerticalAlignment)((ELabelAnchorMode::LabelAnchorMode_Centered / 3) + 1) == EVerticalAlignment::VAlign_Center && (EHorizontalAlignment)((ELabelAnchorMode::LabelAnchorMode_Centered % 3) + 1) == EHorizontalAlignment::HAlign_Center, "Alignment from ELabelAnchorMode error.");
+	static_assert((EVerticalAlignment)((ELabelAnchorMode::LabelAnchorMode_CenterRight / 3) + 1) == EVerticalAlignment::VAlign_Center && (EHorizontalAlignment)((ELabelAnchorMode::LabelAnchorMode_CenterRight % 3) + 1) == EHorizontalAlignment::HAlign_Right, "Alignment from ELabelAnchorMode error.");
 	
-	checkAtCompileTime((EVerticalAlignment)((ELabelAnchorMode::LabelAnchorMode_BottomLeft/3)+1) == EVerticalAlignment::VAlign_Bottom && (EHorizontalAlignment)((ELabelAnchorMode::LabelAnchorMode_BottomLeft%3)+1) == EHorizontalAlignment::HAlign_Left, Alignment_From_ELabelAnchorMode_Error);
-	checkAtCompileTime((EVerticalAlignment)((ELabelAnchorMode::LabelAnchorMode_BottomCenter/3)+1) == EVerticalAlignment::VAlign_Bottom && (EHorizontalAlignment)((ELabelAnchorMode::LabelAnchorMode_BottomCenter%3)+1) == EHorizontalAlignment::HAlign_Center, Alignment_From_ELabelAnchorMode_Error);
-	checkAtCompileTime((EVerticalAlignment)((ELabelAnchorMode::LabelAnchorMode_BottomRight/3)+1) == EVerticalAlignment::VAlign_Bottom && (EHorizontalAlignment)((ELabelAnchorMode::LabelAnchorMode_BottomRight%3)+1) == EHorizontalAlignment::HAlign_Right, Alignment_From_ELabelAnchorMode_Error);
+	static_assert((EVerticalAlignment)((ELabelAnchorMode::LabelAnchorMode_BottomLeft / 3) + 1) == EVerticalAlignment::VAlign_Bottom && (EHorizontalAlignment)((ELabelAnchorMode::LabelAnchorMode_BottomLeft % 3) + 1) == EHorizontalAlignment::HAlign_Left, "Alignment from ELabelAnchorMode error.");
+	static_assert((EVerticalAlignment)((ELabelAnchorMode::LabelAnchorMode_BottomCenter / 3) + 1) == EVerticalAlignment::VAlign_Bottom && (EHorizontalAlignment)((ELabelAnchorMode::LabelAnchorMode_BottomCenter % 3) + 1) == EHorizontalAlignment::HAlign_Center, "Alignment from ELabelAnchorMode error.");
+	static_assert((EVerticalAlignment)((ELabelAnchorMode::LabelAnchorMode_BottomRight / 3) + 1) == EVerticalAlignment::VAlign_Bottom && (EHorizontalAlignment)((ELabelAnchorMode::LabelAnchorMode_BottomRight % 3) + 1) == EHorizontalAlignment::HAlign_Right, "Alignment from ELabelAnchorMode error.");
 	
 	EVerticalAlignment VAlign = (EVerticalAlignment)((AnchorMode/3)+1);
 	EHorizontalAlignment HAlign = (EHorizontalAlignment)((AnchorMode%3)+1);
diff --git a/Engine/Source/Editor/UnrealEd/Private/ApexClothingUtils.cpp b/Engine/Source/Editor/UnrealEd/Private/ApexClothingUtils.cpp
--- a/Engine/Source/Editor/UnrealEd/Private/ApexClothingUtils.cpp
+++ b/Engine/Source/Editor/UnrealEd/Private/ApexClothingUtils.cpp
@@ -45,7 +45,7 @@ struct FBoneIndices
 		Data[2] = InData2; 
 		Data[3] = InData3;
 
-		checkAtCompileTime(MAX_INFLUENCES_PER_STREAM == sizeof(Data), Mismatched_BoneInfluences_Per_Vertex);
+		static_assert(MAX_INFLUENCES_PER_STREAM == sizeof(Data), "Mismatched bone influences per vertex.");
 	}
 };
 
diff --git a/Engine/Source/Editor/UnrealEd/Private/EditorViewportClient.cpp b/Engine/Source/Editor/UnrealEd/Private/EditorViewportClient.cpp
--- a/Engine/Source/Editor/UnrealEd/Private/EditorViewportClient.cpp
+++ b/Engine/Source/Editor/UnrealEd/Private/EditorViewportClient.cpp
@@ -4047,7 +4047,7 @@ void FEditorViewportStats::SendUsageData()
 {
 	Initialize();
 
-	checkAtCompileTime(FEditorViewportStats::CAT_MAX == 22, "If the number of categories change you need to add more entries below!");
+	static_assert(FEditorViewportStats::CAT_MAX == 22, "If the number of categories change you need to add more entries below!");
 
 	TArray<FAnalyticsEventAttribute> PerspectiveUsage;
 	PerspectiveUsage.Add(FAnalyticsEventAttribute(FString("Keyboard.WASD"),			DataPoints[FEditorViewportStats::CAT_PERSPECTIVE_KEYBOARD_WASD]));
diff --git a/Engine/Source/Editor/UnrealEd/Private/Fbx/FBxLibs.cpp b/Engine/Source/Editor/UnrealEd/Private/Fbx/FBxLibs.cpp
--- a/Engine/Source/Editor/UnrealEd/Private/Fbx/FBxLibs.cpp
+++ b/Engine/Source/Editor/UnrealEd/Private/Fbx/FBxLibs.cpp
@@ -43,7 +43,7 @@ void LoadFBxLibraries()
 		FString RootFBxPath = FPaths::EngineDir() / TEXT("Binaries/ThirdParty/FBx/Win64/");
 		FBxHandle = LoadLibraryW(*(RootFBxPath + "libfbxsdk.dll"));
 	#else
-		checkAtCompileTime( false, TEXT("FBX importing currently not supported in 32 bit versions") );
+		static_assert(false, TEXT("FBX importing currently not supported in 32 bit versions"));
 	#endif // PLATFORM_64BITS
 #endif // PLATFORM_WINDOWS && _MSC_VER == 1700
 
diff --git a/Engine/Source/Editor/UnrealEd/Private/Lightmass/Lightmass.cpp b/Engine/Source/Editor/UnrealEd/Private/Lightmass/Lightmass.cpp
--- a/Engine/Source/Editor/UnrealEd/Private/Lightmass/Lightmass.cpp
+++ b/Engine/Source/Editor/UnrealEd/Private/Lightmass/Lightmass.cpp
@@ -2884,7 +2884,7 @@ void FLightmassProcessor::ImportVolumeSamples()
 	if (VolumeSampleTaskCompleted > 0)
 	{
 		{
-			checkAtCompileTime(sizeof(FDebugVolumeLightingSample) == sizeof(Lightmass::FDebugVolumeLightingSample), DebugTypeSizesMustMatch);
+			static_assert(sizeof(FDebugVolumeLightingSample) == sizeof(Lightmass::FDebugVolumeLightingSample), "Debug type sizes must match.");
 			const FString ChannelName = Lightmass::CreateChannelName(Lightmass::VolumeLightingDebugOutputGuid, Lightmass::LM_VOLUMEDEBUGOUTPUT_VERSION, Lightmass::LM_VOLUMEDEBUGOUTPUT_EXTENSION);
 			const int32 Channel = Swarm.OpenChannel( *ChannelName, LM_VOLUMEDEBUGOUTPUT_CHANNEL_FLAGS );
 			if (Channel >= 0)
@@ -3646,13 +3646,13 @@ void FLightmassProcessor::ReadArray(int32 Channel, TArray<T>& Array)
 /** Fills out GDebugStaticLightingInfo with the output from Lightmass */
 void FLightmassProcessor::ImportDebugOutput()
 {
-	checkAtCompileTime(sizeof(FDebugStaticLightingRay) == sizeof(Lightmass::FDebugStaticLightingRay), DebugTypeSizesMustMatch_FDebugStaticLightingRay);
-	checkAtCompileTime(sizeof(FDebugStaticLightingVertex) == sizeof(Lightmass::FDebugStaticLightingVertex), DebugTypeSizesMustMatch_FDebugStaticLightingVertex);
-	checkAtCompileTime(sizeof(FDebugLightingCacheRecord) == sizeof(Lightmass::FDebugLightingCacheRecord), DebugTypeSizesMustMatch_FDebugLightingCacheRecord);
-	checkAtCompileTime(STRUCT_OFFSET(FDebugLightingCacheRecord,RecordId) == STRUCT_OFFSET(Lightmass::FDebugLightingCacheRecord,RecordId), DebugStructOffsetMustMatch_FDebugLightingCacheRecord_RecordId);
-	checkAtCompileTime(sizeof(FDebugPhoton) == sizeof(Lightmass::FDebugPhoton), DebugTypeSizesMustMatch_FDebugPhoton);
-	checkAtCompileTime(sizeof(FDebugOctreeNode) == sizeof(Lightmass::FDebugOctreeNode), DebugTypeSizesMustMatch_FDebugOctreeNode);
-	checkAtCompileTime(NumTexelCorners == Lightmass::NumTexelCorners, DebugTypeSizesMustMatch_NumTexelCorners);
+	static_assert(sizeof(FDebugStaticLightingRay) == sizeof(Lightmass::FDebugStaticLightingRay), "Debug type sizes must match for FDebugStaticLightingRay.");
+	static_assert(sizeof(FDebugStaticLightingVertex) == sizeof(Lightmass::FDebugStaticLightingVertex), "Debug type sizes must match for FDebugStaticLightingVertex.");
+	static_assert(sizeof(FDebugLightingCacheRecord) == sizeof(Lightmass::FDebugLightingCacheRecord), "Debug type sizes must match for FDebugLightingCacheRecord.");
+	static_assert(STRUCT_OFFSET(FDebugLightingCacheRecord, RecordId) == STRUCT_OFFSET(Lightmass::FDebugLightingCacheRecord, RecordId), "Debug struct offset must match for FDebugLightingCacheRecord::RecordId.");
+	static_assert(sizeof(FDebugPhoton) == sizeof(Lightmass::FDebugPhoton), "Debug type sizes must match for FDebugPhoton.");
+	static_assert(sizeof(FDebugOctreeNode) == sizeof(Lightmass::FDebugOctreeNode), "Debug type sizes must match for FDebugOctreeNode.");
+	static_assert(NumTexelCorners == Lightmass::NumTexelCorners, "Debug type sizes must match for NumTexelCorners.");
 
 	const FString ChannelName = Lightmass::CreateChannelName(Lightmass::DebugOutputGuid, Lightmass::LM_DEBUGOUTPUT_VERSION, Lightmass::LM_DEBUGOUTPUT_EXTENSION);
 	const int32 Channel = Swarm.OpenChannel( *ChannelName, LM_DEBUGOUTPUT_CHANNEL_FLAGS );
@@ -3828,7 +3828,7 @@ bool FLightmassProcessor::ImportSignedDistanceFieldShadowMapData2D(int32 Channel
 		Lightmass::FShadowMapData2DData SMData(0,0);
 		Swarm.ReadChannel(Channel, &SMData, sizeof(Lightmass::FShadowMapData2DData));
 
-		checkAtCompileTime(sizeof(FQuantizedSignedDistanceFieldShadowSample) == sizeof(Lightmass::FQuantizedSignedDistanceFieldShadowSampleData), SampleDataSizesMustMatch);
+		static_assert(sizeof(FQuantizedSignedDistanceFieldShadowSample) == sizeof(Lightmass::FQuantizedSignedDistanceFieldShadowSampleData), "Sample data sizes must match.");
 
 		FQuantizedShadowSignedDistanceFieldData2D* ShadowMapData = new FQuantizedShadowSignedDistanceFieldData2D(SMData.SizeX, SMData.SizeY);
 		check(ShadowMapData);
diff --git a/Engine/Source/Editor/UnrealEd/Private/Lightmass/LightmassRender.cpp b/Engine/Source/Editor/UnrealEd/Private/Lightmass/LightmassRender.cpp
--- a/Engine/Source/Editor/UnrealEd/Private/Lightmass/LightmassRender.cpp
+++ b/Engine/Source/Editor/UnrealEd/Private/Lightmass/LightmassRender.cpp
@@ -831,7 +831,7 @@ bool FLightmassMaterialRenderer::GenerateMaterialData(
 	}
 
 	// Set the blend mode
-	checkAtCompileTime(EBlendMode::BLEND_MAX == (EBlendMode)Lightmass::BLEND_MAX, DebugTypeSizesMustMatch);
+	static_assert(EBlendMode::BLEND_MAX == (EBlendMode)Lightmass::BLEND_MAX, "Debug type sizes must match.");
 	OutMaterialData.BlendMode = (Lightmass::EBlendMode)((int32)BlendMode);
 	// Set the two-sided flag
 	OutMaterialData.bTwoSided = (uint32)InMaterial.IsTwoSided();
diff --git a/Engine/Source/Editor/UnrealEd/Private/ScriptDisassembler.cpp b/Engine/Source/Editor/UnrealEd/Private/ScriptDisassembler.cpp
--- a/Engine/Source/Editor/UnrealEd/Private/ScriptDisassembler.cpp
+++ b/Engine/Source/Editor/UnrealEd/Private/ScriptDisassembler.cpp
@@ -141,7 +141,7 @@ CodeSkipSizeType FKismetBytecodeDisassembler::ReadSkipCount(int32& ScriptIndex)
 #if SCRIPT_LIMIT_BYTECODE_TO_64KB
 	return ReadWORD(ScriptIndex);
 #else
-	checkAtCompileTime(sizeof(CodeSkipSizeType) == 4, updateThisCodeAsSizeChanged);
+	static_assert(sizeof(CodeSkipSizeType) == 4, "Update this code as size changed.");
 	return ReadINT(ScriptIndex);
 #endif
 }
diff --git a/Engine/Source/Programs/UnrealLightmass/Private/ImportExport/Exporter.cpp b/Engine/Source/Programs/UnrealLightmass/Private/ImportExport/Exporter.cpp
--- a/Engine/Source/Programs/UnrealLightmass/Private/ImportExport/Exporter.cpp
+++ b/Engine/Source/Programs/UnrealLightmass/Private/ImportExport/Exporter.cpp
@@ -64,7 +64,7 @@ namespace Lightmass
 			for (TMap<FGuid,TArray<FVolumeLightingSample> >::TConstIterator It(VolumeSamples); It; ++It)
 			{
 				Swarm->Write(&It.Key(), sizeof(It.Key()));
-				checkAtCompileTime(sizeof(FVolumeLightingSample) == sizeof(FVolumeLightingSampleData), VolumeDerivedSizeMustMatch);
+				static_assert(sizeof(FVolumeLightingSample) == sizeof(FVolumeLightingSampleData), "Volume derived size must match.");
 				WriteArray(It.Value());
 			}
 			Swarm->CloseCurrentChannel();
@@ -83,7 +83,7 @@ namespace Lightmass
 		if( ErrorCode >= 0 )
 		{
 			Swarm->Write(&DominantLightShadowInfo, sizeof(FDominantLightShadowInfoData));
-			checkAtCompileTime(sizeof(FDominantLightShadowSample) == sizeof(FDominantLightShadowSampleData), ShadowDerivedSizeMustMatch);
+			static_assert(sizeof(FDominantLightShadowSample) == sizeof(FDominantLightShadowSampleData), "Shadow derived size must match.");
 			WriteArray(DominantLightShadowInfo.ShadowMap);
 			Swarm->CloseCurrentChannel();
 		}
diff --git a/Engine/Source/Runtime/AIModule/Private/BehaviorTree/Decorators/BTDecorator_CompareBBEntries.cpp b/Engine/Source/Runtime/AIModule/Private/BehaviorTree/Decorators/BTDecorator_CompareBBEntries.cpp
--- a/Engine/Source/Runtime/AIModule/Private/BehaviorTree/Decorators/BTDecorator_CompareBBEntries.cpp
+++ b/Engine/Source/Runtime/AIModule/Private/BehaviorTree/Decorators/BTDecorator_CompareBBEntries.cpp
@@ -46,9 +46,9 @@ bool UBTDecorator_CompareBBEntries::CalculateRawConditionValue(class UBehaviorTr
 		const EBlackboardCompare::Type Result = BlackboardComp->CompareKeyValues(BlackboardKeyA.SelectedKeyType, 
 			BlackboardKeyA.GetSelectedKeyID(), BlackboardKeyB.GetSelectedKeyID());
 
-		checkAtCompileTime(int32(EBlackboardCompare::Equal) == int32(EBlackBoardEntryComparison::Equal)
-			&& int32(EBlackboardCompare::NotEqual) == int32(EBlackBoardEntryComparison::NotEqual)
-			, "These values need to be equal");
+		static_assert(int32(EBlackboardCompare::Equal) == int32(EBlackBoardEntryComparison::Equal)
+			&& int32(EBlackboardCompare::NotEqual) == int32(EBlackBoardEntryComparison::NotEqual),
+			"These values need to be equal");
 
 		return int32(Operator) == int32(Result);
 	}
diff --git a/Engine/Source/Runtime/Core/Private/Apple/ApplePlatformCrashContext.cpp b/Engine/Source/Runtime/Core/Private/Apple/ApplePlatformCrashContext.cpp
--- a/Engine/Source/Runtime/Core/Private/Apple/ApplePlatformCrashContext.cpp
+++ b/Engine/Source/Runtime/Core/Private/Apple/ApplePlatformCrashContext.cpp
@@ -133,7 +133,7 @@ void FApplePlatformCrashContext::WriteLine(int ReportFile, const ANSICHAR* Line)
 void FApplePlatformCrashContext::WriteUTF16String(int ReportFile, const TCHAR * UTFString4BytesChar, uint32 NumChars)
 {
 	check(UTFString4BytesChar != NULL || NumChars == 0);
-	checkAtCompileTime(sizeof(TCHAR) == 4, _PlatformTCHARIsNot4BytesRevisitThisFunction);
+	static_assert(sizeof(TCHAR) == 4, "Platform TCHAR is not 4 bytes. Revisit this function.");
 	
 	for (uint32 Idx = 0; Idx < NumChars; ++Idx)
 	{
@@ -143,7 +143,7 @@ void FApplePlatformCrashContext::WriteUTF16String(int ReportFile, const TCHAR *
 
 void FApplePlatformCrashContext::WriteUTF16String(int ReportFile, const TCHAR * UTFString4BytesChar)
 {
-	checkAtCompileTime(sizeof(TCHAR) == 4, _PlatformTCHARIsNot4BytesRevisitThisFunction);
+	static_assert(sizeof(TCHAR) == 4, "Platform TCHAR is not 4 bytes. Revisit this function.");
 	uint32 NumChars = FCString::Strlen(UTFString4BytesChar);
 	WriteUTF16String(ReportFile, UTFString4BytesChar, NumChars);
 }
diff --git a/Engine/Source/Runtime/Core/Private/Async/TaskGraph.cpp b/Engine/Source/Runtime/Core/Private/Async/TaskGraph.cpp
--- a/Engine/Source/Runtime/Core/Private/Async/TaskGraph.cpp
+++ b/Engine/Source/Runtime/Core/Private/Async/TaskGraph.cpp
@@ -194,7 +194,6 @@ public:
 	void ProcessTasks(int32 QueueIndex, bool bAllowStall)
 	{
 #if STATS
-		checkAtCompileTime(ENamedThreads::StatsThread == 0, delete_this_testing_a_build_fail);
 		TStatId StatName;
 		FCycleCounter ProcessingTasks;
 		if (ThreadId == ENamedThreads::GameThread)
@@ -950,7 +949,7 @@ private:
 	**/
 	ENamedThreads::Type GetCurrentThread()
 	{
-		checkAtCompileTime(offsetof(FWorkerThread,TaskGraphWorker)==0,taskgraph_worker_must_be_first_member);
+		static_assert(offsetof(FWorkerThread,TaskGraphWorker)==0, "Task graph worker must be the first member.");
 		ENamedThreads::Type CurrentThreadIfKnown = ENamedThreads::AnyThread;
 		FWorkerThread* TLSPointer = (FWorkerThread*)FPlatformTLS::GetTlsValue(PerThreadIDTLSSlot);
 		if (TLSPointer)
diff --git a/Engine/Source/Runtime/Core/Private/GenericPlatform/GenericPlatformAtomics.cpp b/Engine/Source/Runtime/Core/Private/GenericPlatform/GenericPlatformAtomics.cpp
--- a/Engine/Source/Runtime/Core/Private/GenericPlatform/GenericPlatformAtomics.cpp
+++ b/Engine/Source/Runtime/Core/Private/GenericPlatform/GenericPlatformAtomics.cpp
@@ -1,7 +1,10 @@
 // Copyright 1998-2014 Epic Games, Inc. All Rights Reserved.
 
-#include "CorePrivate.h"
+/*=============================================================================
+	GenericPlatformAtomics.cpp: Generic implementations of atomic functions
+=============================================================================*/
 
+#include "CorePrivate.h"
 
-checkAtCompileTime( sizeof( FInt128 ) == 16, FInt128_sizeof_must_be_16 );
-checkAtCompileTime( ALIGNOF( FInt128 ) == 16, FInt128_alignment_must_be_16 );
+static_assert(sizeof(FInt128) == 16, "FInt128 size must be 16 bytes.");
+static_assert(ALIGNOF(FInt128) == 16, "FInt128 alignment must equals 16.");
\ No newline at end of file
diff --git a/Engine/Source/Runtime/Core/Private/Linux/LinuxPlatformMisc.cpp b/Engine/Source/Runtime/Core/Private/Linux/LinuxPlatformMisc.cpp
--- a/Engine/Source/Runtime/Core/Private/Linux/LinuxPlatformMisc.cpp
+++ b/Engine/Source/Runtime/Core/Private/Linux/LinuxPlatformMisc.cpp
@@ -544,7 +544,7 @@ bool FLinuxCrashContext::GetInfoForAddress(void * Address, const char **OutFunct
 	Dwarf_Unsigned Addr = reinterpret_cast< Dwarf_Unsigned >( Address ), LineNumber = 0;
 	const char * SrcFile = NULL;
 
-	checkAtCompileTime(sizeof(Dwarf_Unsigned) >= sizeof(Address), _Dwarf_Unsigned_type_should_be_long_enough_to_represent_pointers_Check_libdwarf_bitness);
+	static_assert(sizeof(Dwarf_Unsigned) >= sizeof(Address), "Dwarf_Unsigned type should be long enough to represent pointers. Check libdwarf bitness.");
 
 	int ReturnCode = DW_DLV_OK;
 	Dwarf_Error ErrorInfo;
@@ -749,7 +749,7 @@ namespace
 	void WriteUTF16String(FArchive* ReportFile, const TCHAR * UTFString4BytesChar, uint32 NumChars)
 	{
 		check(UTFString4BytesChar != NULL || NumChars == 0);
-		checkAtCompileTime(sizeof(TCHAR) == 4, _PlatformTCHARIsNot4BytesRevisitThisFunction);
+		static_assert(sizeof(TCHAR) == 4, "Platform TCHAR is not 4 bytes. Revisit this function.");
 
 		for (uint32 Idx = 0; Idx < NumChars; ++Idx)
 		{
diff --git a/Engine/Source/Runtime/Core/Private/Mac/MacApplication.cpp b/Engine/Source/Runtime/Core/Private/Mac/MacApplication.cpp
--- a/Engine/Source/Runtime/Core/Private/Mac/MacApplication.cpp
+++ b/Engine/Source/Runtime/Core/Private/Mac/MacApplication.cpp
@@ -1146,7 +1146,7 @@ void FMacApplication::RecordUsage(EGestureEvent::Type Gesture)
 
 void FMacApplication::SendAnalytics(IAnalyticsProvider* Provider)
 {
-	checkAtCompileTime(EGestureEvent::Count == 5, "If the number of gestures changes you need to add more entries below!");
+	static_assert(EGestureEvent::Count == 5, "If the number of gestures changes you need to add more entries below!");
 
 	TArray<FAnalyticsEventAttribute> GestureAttributes;
 	GestureAttributes.Add(FAnalyticsEventAttribute(FString("Scroll"),	GestureUsage[EGestureEvent::Scroll]));
diff --git a/Engine/Source/Runtime/Core/Private/Misc/ConfigCacheIni.cpp b/Engine/Source/Runtime/Core/Private/Misc/ConfigCacheIni.cpp
--- a/Engine/Source/Runtime/Core/Private/Misc/ConfigCacheIni.cpp
+++ b/Engine/Source/Runtime/Core/Private/Misc/ConfigCacheIni.cpp
@@ -2372,8 +2372,8 @@ static bool GenerateDestIniFile(FConfigFile& DestConfigFile, const FString& Dest
 			// Flag indicating whether the user has requested 'Yes/No To All'.
 			static int32 GIniYesNoToAll = -1;
 			// Make sure GIniYesNoToAll's 'uninitialized' value is kosher.
-			checkAtCompileTime( EAppReturnType::YesAll != -1, YesAll_MustNotBeNegOne );
-			checkAtCompileTime( EAppReturnType::NoAll != -1, NoAll_MustNotBeNegOne );
+			static_assert(EAppReturnType::YesAll != -1, "EAppReturnType::YesAll must not be -1.");
+			static_assert(EAppReturnType::NoAll != -1, "EAppReturnType::NoAll must not be -1.");
 
 			// The file exists but is different.
 			// Prompt the user if they haven't already responded with a 'Yes/No To All' answer.
diff --git a/Engine/Source/Runtime/Core/Private/Stats/Stats2.cpp b/Engine/Source/Runtime/Core/Private/Stats/Stats2.cpp
--- a/Engine/Source/Runtime/Core/Private/Stats/Stats2.cpp
+++ b/Engine/Source/Runtime/Core/Private/Stats/Stats2.cpp
@@ -229,7 +229,7 @@ class FStatGroupEnableManager : public IStatGroupEnableManager
 	void EnableStat(FName StatName, FName *DisablePtr)
 	{
 		// This is all complicated to ensure an atomic 8 byte write
-		checkAtCompileTime(sizeof(FName) == sizeof(uint64), assumption_about_size_of_fname);
+		static_assert(sizeof(FName) == sizeof(uint64), "FName should have the same size of uint64.");
 		check(UPTRINT(DisablePtr) % sizeof(FName) == 0);
 		MS_ALIGN(8) struct FAligner
 		{
@@ -244,7 +244,7 @@ class FStatGroupEnableManager : public IStatGroupEnableManager
 
 	void DisableStat(FName *DisablePtr)
 	{
-		checkAtCompileTime(sizeof(FName) == sizeof(uint64), assumption_about_size_of_fname);
+		static_assert(sizeof(FName) == sizeof(uint64), "FName should have the same size of uint64.");
 		check(UPTRINT(DisablePtr) % sizeof(FName) == 0);
 		*(uint64*)DisablePtr = 0;
 	}
@@ -676,7 +676,6 @@ public:
 
 	virtual uint32 Run()
 	{
-		checkAtCompileTime(ENamedThreads::StatsThread == 0, delete_this_testing_a_build_fail);
 		StatsThreadId = FPlatformTLS::GetCurrentThreadId();
 		FTaskGraphInterface::Get().AttachToThread(ENamedThreads::StatsThread);
 		FTaskGraphInterface::Get().ProcessThreadUntilRequestReturn(ENamedThreads::StatsThread);
diff --git a/Engine/Source/Runtime/Core/Private/Tests/HAL/PlatformTest.cpp b/Engine/Source/Runtime/Core/Private/Tests/HAL/PlatformTest.cpp
--- a/Engine/Source/Runtime/Core/Private/Tests/HAL/PlatformTest.cpp
+++ b/Engine/Source/Runtime/Core/Private/Tests/HAL/PlatformTest.cpp
@@ -88,14 +88,14 @@ bool FPlatformVerificationTest::RunTest (const FString& Parameters)
 
 	check(FString(FPlatformProperties::PlatformName()).Len() > 0); 
 
-	checkAtCompileTime(ALIGNOF(int32) == 4, alignof_int_is_not_4); //Hmmm, this would be very strange, ok maybe, but strange
+	static_assert(ALIGNOF(int32) == 4, "Align of int32 is not 4."); //Hmmm, this would be very strange, ok maybe, but strange
 
 	MS_ALIGN(16) struct FTestAlign
 	{
 		uint8 Test;
 	} GCC_ALIGN(16);
 
-	checkAtCompileTime(ALIGNOF(FTestAlign) == 16, alignment_fail); 
+	static_assert(ALIGNOF(FTestAlign) == 16, "Align of FTestAlign is not 16.");
 
 	FName::AutoTest();
 
diff --git a/Engine/Source/Runtime/Core/Private/Windows/WindowsPlatformProcess.cpp b/Engine/Source/Runtime/Core/Private/Windows/WindowsPlatformProcess.cpp
--- a/Engine/Source/Runtime/Core/Private/Windows/WindowsPlatformProcess.cpp
+++ b/Engine/Source/Runtime/Core/Private/Windows/WindowsPlatformProcess.cpp
@@ -451,7 +451,7 @@ FString FWindowsPlatformProcess::GetApplicationName( uint32 ProcessId )
 		TCHAR ProcessNameBuffer[ProcessNameBufferSize];
 		
 		int32 InOutSize = ProcessNameBufferSize;
-		checkAtCompileTime(sizeof(::DWORD) == sizeof(int32), "DWORD size doesn't match int32.  Is it the future or the past?");
+		static_assert(sizeof(::DWORD) == sizeof(int32), "DWORD size doesn't match int32. Is it the future or the past?");
 
 #if WINVER == 0x0502
 		GetProcessImageFileName(ProcessHandle, ProcessNameBuffer, InOutSize);
diff --git a/Engine/Source/Runtime/Core/Private/Windows/WindowsPlatformStackWalk.cpp b/Engine/Source/Runtime/Core/Private/Windows/WindowsPlatformStackWalk.cpp
--- a/Engine/Source/Runtime/Core/Private/Windows/WindowsPlatformStackWalk.cpp
+++ b/Engine/Source/Runtime/Core/Private/Windows/WindowsPlatformStackWalk.cpp
@@ -480,9 +480,9 @@ static void LoadProcessModules()
 		ANSICHAR ModuleName[1024];
 		ANSICHAR ImageName[1024];
 #if PLATFORM_64BITS
-		checkAtCompileTime(sizeof(ModuleInfo) == 24, broken_alignment_for_64bit_windows_include);
+		static_assert(sizeof(ModuleInfo) == 24, "Broken alignment for 64bit Windows include.");
 #else
-		checkAtCompileTime(sizeof(ModuleInfo) == 12, broken_alignment_for_32bit_windows_include);
+		static_assert(sizeof(ModuleInfo) == 12, "Broken alignment for 32bit Windows include.");
 #endif
 		FGetModuleInformation( ProcessHandle, ModuleHandleArray[ModuleIndex], &ModuleInfo,sizeof( ModuleInfo ) );
 		FGetModuleFileNameEx( ProcessHandle, ModuleHandleArray[ModuleIndex], ImageName, 1024 );
@@ -594,9 +594,9 @@ int32 FWindowsPlatformStackWalk::GetProcessModuleSignatures(FStackWalkModuleInfo
 		ANSICHAR ModuleName[1024];
 		ANSICHAR ImageName[1024];
 #if PLATFORM_64BITS
-		checkAtCompileTime(sizeof(ModuleInfo) == 24, broken_alignment_for_64bit_windows_include);
+		static_assert(sizeof(ModuleInfo) == 24, "Broken alignment for 64bit Windows include.");
 #else
-		checkAtCompileTime(sizeof(ModuleInfo) == 12, broken_alignment_for_32bit_windows_include);
+		static_assert(sizeof(ModuleInfo) == 12, "Broken alignment for 32bit Windows include.");
 #endif
 		FGetModuleInformation( ProcessHandle, ModuleHandleArray[ModuleIndex], &ModuleInfo,sizeof( ModuleInfo ) );
 		FGetModuleFileNameEx( ProcessHandle, ModuleHandleArray[ModuleIndex], ImageName, 1024 );
diff --git a/Engine/Source/Runtime/Core/Public/Containers/AllocatorFixedSizeFreeList.h b/Engine/Source/Runtime/Core/Public/Containers/AllocatorFixedSizeFreeList.h
--- a/Engine/Source/Runtime/Core/Public/Containers/AllocatorFixedSizeFreeList.h
+++ b/Engine/Source/Runtime/Core/Public/Containers/AllocatorFixedSizeFreeList.h
@@ -41,7 +41,7 @@ public:
 		,NumLive(0)
 	{
 		// need enough memory to store pointers for the free list
-		checkAtCompileTime(AllocationSize >= sizeof(FreeListNode), AllocationSizeMustBeLargeEnoughToHoldPointer); 
+		static_assert(AllocationSize >= sizeof(FreeListNode), "Allocation size must be large enough to hold pointer.");
 
 		// warm the cache
 		Grow(InitialBlockSize);
diff --git a/Engine/Source/Runtime/Core/Public/Containers/ContainerAllocationPolicies.h b/Engine/Source/Runtime/Core/Public/Containers/ContainerAllocationPolicies.h
--- a/Engine/Source/Runtime/Core/Public/Containers/ContainerAllocationPolicies.h
+++ b/Engine/Source/Runtime/Core/Public/Containers/ContainerAllocationPolicies.h
@@ -80,7 +80,7 @@ struct TAlignedBytes<Size,1>
 		struct TAlignedBytes<Size,Align> \
 		{ \
 			uint8 Pad[Size]; \
-			checkAtCompileTime( Size % Align == 0, CLRInteropTypesMustNotBeAligned ); \
+			static_assert(Size % Align == 0, "CLR interop types must not be aligned."); \
 		};
 #else
 /** A macro that implements TAlignedBytes for a specific alignment. */
diff --git a/Engine/Source/Runtime/Core/Public/Containers/StaticBitArray.h b/Engine/Source/Runtime/Core/Public/Containers/StaticBitArray.h
--- a/Engine/Source/Runtime/Core/Public/Containers/StaticBitArray.h
+++ b/Engine/Source/Runtime/Core/Public/Containers/StaticBitArray.h
@@ -1,7 +1,10 @@
 // Copyright 1998-2014 Epic Games, Inc. All Rights Reserved.
 
-#pragma once
+/*=============================================================================
+	StaticBitArray.h: Static bit array definition
+=============================================================================*/
 
+#pragma once
 
 /** Used to read/write a bit in the static array as a bool. */
 template<typename T>
@@ -102,7 +105,7 @@ public:
 		 the initialization point
 
 		***********************************************************************/
-//		checkAtCompileTime( InBitIndex >= 0 && InBitIndex < NumBits, invalidBitValue );
+//		static_assert(InBitIndex >= 0 && InBitIndex < NumBits, "Invalid bit value.");
 
 		check((NumBits > 0) ? (InBitIndex<NumBits):1);
 
@@ -311,7 +314,7 @@ public:
 
 private:
 
-//	checkAtCompileTime( NumBits > 0, mustHaveAtLeast1Bit );
+//	static_assert(NumBits > 0, "Must have at least 1 bit.");
 	static const uint32 NumBitsPerWord = sizeof(WordType) * 8;
 	static const uint32 NumWords = ((NumBits + NumBitsPerWord - 1) & ~(NumBitsPerWord - 1)) / NumBitsPerWord;
 	WordType Words[NumWords];
diff --git a/Engine/Source/Runtime/Core/Public/Containers/Union.h b/Engine/Source/Runtime/Core/Public/Containers/Union.h
--- a/Engine/Source/Runtime/Core/Public/Containers/Union.h
+++ b/Engine/Source/Runtime/Core/Public/Containers/Union.h
@@ -353,15 +353,14 @@ private:
 		}
 		else
 		{
-			checkAtCompileTime(
+			static_assert(
 				TAreTypesEqual<TEMPLATE_PARAMETERS2(TypeA, Subtype)>::Value ||
 				TAreTypesEqual<TEMPLATE_PARAMETERS2(TypeB, Subtype)>::Value ||
 				TAreTypesEqual<TEMPLATE_PARAMETERS2(TypeC, Subtype)>::Value ||
 				TAreTypesEqual<TEMPLATE_PARAMETERS2(TypeD, Subtype)>::Value ||
 				TAreTypesEqual<TEMPLATE_PARAMETERS2(TypeE, Subtype)>::Value ||
 				TAreTypesEqual<TEMPLATE_PARAMETERS2(TypeF, Subtype)>::Value,
-				TypeIsNotSubtypeOfUnion
-				);
+				"Type is not subtype of union.");
 			OutIndex = -1;
 			OutValuePointer = NULL;
 		}
diff --git a/Engine/Source/Runtime/Core/Public/Delegates/DelegateInstancesImpl.inl b/Engine/Source/Runtime/Core/Public/Delegates/DelegateInstancesImpl.inl
--- a/Engine/Source/Runtime/Core/Public/Delegates/DelegateInstancesImpl.inl
+++ b/Engine/Source/Runtime/Core/Public/Delegates/DelegateInstancesImpl.inl
@@ -267,7 +267,7 @@ template<class UserClass, FUNC_PAYLOAD_TEMPLATE_DECL_TYPENAME>
 class RAW_METHOD_DELEGATE_INSTANCE_CLASS
 	: public DELEGATE_INSTANCE_INTERFACE_CLASS<FUNC_TEMPLATE_ARGS>
 {
-	checkAtCompileTime((!CanConvertPointerFromTo<UserClass, UObjectBase>::Result), You_Cannot_Use_Raw_Method_Delegates_With_UObjects);
+	static_assert((!CanConvertPointerFromTo<UserClass, UObjectBase>::Result), "You cannot use raw method delegates with UObjects.");
 
 public:
 
@@ -431,7 +431,7 @@ template<class UserClass, FUNC_PAYLOAD_TEMPLATE_DECL_TYPENAME>
 class UOBJECT_METHOD_DELEGATE_INSTANCE_CLASS
 	: public DELEGATE_INSTANCE_INTERFACE_CLASS<FUNC_TEMPLATE_ARGS>
 {
-	checkAtCompileTime((CanConvertPointerFromTo<UserClass, UObjectBase>::Result), You_Cannot_Use_UObject_Method_Delegates_With_Raw_Pointers);
+	static_assert((CanConvertPointerFromTo<UserClass, UObjectBase>::Result), "You cannot use UObject method delegates with raw pointers.");
 
 public:
 
@@ -741,7 +741,7 @@ template<class UserClass, FUNC_PAYLOAD_TEMPLATE_DECL_TYPENAME>
 class UFUNCTION_DELEGATE_INSTANCE_CLASS
 	: public DELEGATE_INSTANCE_INTERFACE_CLASS<FUNC_TEMPLATE_ARGS>
 {
-	checkAtCompileTime((CanConvertPointerFromTo<UserClass, UObjectBase>::Result), You_Cannot_Use_UFunction_Delegates_With_Non_UObject_Classes);
+	static_assert((CanConvertPointerFromTo<UserClass, UObjectBase>::Result), "You cannot use UFunction delegates with non UObject classes.");
 
 	// Structure for UFunction call parameter marshaling.
 	struct FParmsWithPayload
diff --git a/Engine/Source/Runtime/Core/Public/GenericPlatform/GenericPlatformCompilerSetup.h b/Engine/Source/Runtime/Core/Public/GenericPlatform/GenericPlatformCompilerSetup.h
--- a/Engine/Source/Runtime/Core/Public/GenericPlatform/GenericPlatformCompilerSetup.h
+++ b/Engine/Source/Runtime/Core/Public/GenericPlatform/GenericPlatformCompilerSetup.h
@@ -14,4 +14,8 @@
 		#define DEPRECATED(VERSION, MESSAGE)
 	#endif // DISABLE_DEPRECATION
 
-#endif // DEPRECATED
\ No newline at end of file
+#endif // DEPRECATED
+
+#ifndef EMIT_DEPRECATED_WARNING_MESSAGE
+	#define EMIT_DEPRECATED_WARNING_MESSAGE(Msg)
+#endif // EMIT_DEPRECATED_WARNING_MESSAGE
\ No newline at end of file
diff --git a/Engine/Source/Runtime/Core/Public/HAL/OutputDevices.h b/Engine/Source/Runtime/Core/Public/HAL/OutputDevices.h
--- a/Engine/Source/Runtime/Core/Public/HAL/OutputDevices.h
+++ b/Engine/Source/Runtime/Core/Public/HAL/OutputDevices.h
@@ -158,8 +158,8 @@ protected:
 template<uint8 InDefaultVerbosity, uint8 InCompileTimeVerbosity>
 struct FLogCategory : public FLogCategoryBase
 {
-	checkAtCompileTime((InDefaultVerbosity & ELogVerbosity::VerbosityMask) < ELogVerbosity::NumVerbosity, bogus_default_verbosity);
-	checkAtCompileTime(InCompileTimeVerbosity < ELogVerbosity::NumVerbosity, bogus_compile_time_verbosity);
+	static_assert((InDefaultVerbosity & ELogVerbosity::VerbosityMask) < ELogVerbosity::NumVerbosity, "Bogus default verbosity.");
+	static_assert(InCompileTimeVerbosity < ELogVerbosity::NumVerbosity, "Bogus compile time verbosity.");
 	enum
 	{
 		CompileTimeVerbosity = InCompileTimeVerbosity
diff --git a/Engine/Source/Runtime/Core/Public/HAL/Platform.h b/Engine/Source/Runtime/Core/Public/HAL/Platform.h
--- a/Engine/Source/Runtime/Core/Public/HAL/Platform.h
+++ b/Engine/Source/Runtime/Core/Public/HAL/Platform.h
@@ -371,7 +371,9 @@
 #endif
 
 #ifndef checkAtCompileTime
-	#define checkAtCompileTime(expr, msg)  static_assert(expr, #msg)
+	#define checkAtCompileTime(expr, msg) \
+		EMIT_DEPRECATED_WARNING_MESSAGE("checkAtCompileTime is deprecated. Please use static_assert instead.") \
+		static_assert(expr, #msg)
 #endif
 
 // DLL export and import definitions
@@ -450,8 +452,8 @@ struct THasOperatorNotEquals
 					public: \
 						FORCEINLINE operator int UE_PrivateBoolType::*() const \
 						{ \
-							checkAtCompileTime(THasOperatorEquals   <decltype(UE_ClassTypeHelper(*this))>::Value, "Class_needs_operator_equals"); \
-							checkAtCompileTime(THasOperatorNotEquals<decltype(UE_ClassTypeHelper(*this))>::Value, "Class_needs_operator_not_equals"); \
+							static_assert(THasOperatorEquals   <decltype(UE_ClassTypeHelper(*this))>::Value, "Class needs operator==()."); \
+							static_assert(THasOperatorNotEquals<decltype(UE_ClassTypeHelper(*this))>::Value, "Class needs operator!=()."); \
 							return UE_OperatorBool() ? &UE_PrivateBoolType::x : 0; \
 						} \
 						FORCEINLINE bool UE_OperatorBool
@@ -576,62 +578,62 @@ namespace TypeTests
 		enum { Value = true };
 	};
 
-	checkAtCompileTime(!PLATFORM_TCHAR_IS_4_BYTES || sizeof(TCHAR) == 4,TypeTests_TCHAR_size);
-	checkAtCompileTime(PLATFORM_TCHAR_IS_4_BYTES || sizeof(TCHAR) == 2,TypeTests_TCHAR_size);
+	static_assert(!PLATFORM_TCHAR_IS_4_BYTES || sizeof(TCHAR) == 4, "TCHAR size must be 4 bytes.");
+	static_assert(PLATFORM_TCHAR_IS_4_BYTES || sizeof(TCHAR) == 2, "TCHAR size must be 2 bytes.");
 
-	checkAtCompileTime(PLATFORM_32BITS || PLATFORM_64BITS,TypeTests_pointer_size);
-	checkAtCompileTime(PLATFORM_32BITS != PLATFORM_64BITS,TypeTests_pointer_exclusive);
-	checkAtCompileTime(!PLATFORM_64BITS || sizeof(void*) == 8,TypeTests_pointer_size_64bit_but_short_pointers);
-	checkAtCompileTime(PLATFORM_64BITS || sizeof(void*) == 4,TypeTests_pointer_size_32bit_but_long_pointers);
+	static_assert(PLATFORM_32BITS || PLATFORM_64BITS, "Type tests pointer size failed.");
+	static_assert(PLATFORM_32BITS != PLATFORM_64BITS, "Type tests pointer exclusive failed.");
+	static_assert(!PLATFORM_64BITS || sizeof(void*) == 8, "Pointer size is 64bit, but pointers are short.");
+	static_assert(PLATFORM_64BITS || sizeof(void*) == 4, "Pointer size is 32bit, but pointers are long.");
 
-	checkAtCompileTime(char(-1) < char(0),TypeTests_char_unsigned);
+	static_assert(char(-1) < char(0), "Unsigned char type test failed.");
 
-	checkAtCompileTime((!TAreTypesEqual<ANSICHAR, WIDECHAR>::Value),TypeTests_ANSICHAR_and_WIDECHAR_should_be_different_types);
-	checkAtCompileTime((!TAreTypesEqual<ANSICHAR, UCS2CHAR>::Value),TypeTests_ANSICHAR_and_CHAR16_should_be_different_types);
-	checkAtCompileTime((!TAreTypesEqual<WIDECHAR, UCS2CHAR>::Value),TypeTests_WIDECHAR_and_CHAR16_should_be_different_types);
-	checkAtCompileTime((TAreTypesEqual<TCHAR, ANSICHAR>::Value || TAreTypesEqual<TCHAR, WIDECHAR>::Value),TypeTests_TCHAR_should_either_be_ANSICHAR_or_WIDECHAR);
+	static_assert((!TAreTypesEqual<ANSICHAR, WIDECHAR>::Value), "ANSICHAR and WIDECHAR should be different types.");
+	static_assert((!TAreTypesEqual<ANSICHAR, UCS2CHAR>::Value), "ANSICHAR and CHAR16 should be different types.");
+	static_assert((!TAreTypesEqual<WIDECHAR, UCS2CHAR>::Value), "WIDECHAR and CHAR16 should be different types.");
+	static_assert((TAreTypesEqual<TCHAR, ANSICHAR>::Value || TAreTypesEqual<TCHAR, WIDECHAR>::Value), "TCHAR should either be ANSICHAR or WIDECHAR.");
 
-	checkAtCompileTime(sizeof(uint8) == 1,TypeTests_BYTE_size);
-	checkAtCompileTime(int32(uint8(-1)) == 0xFF,TypeTests_BYTE_sign);
+	static_assert(sizeof(uint8) == 1, "BYTE type size test failed.");
+	static_assert(int32(uint8(-1)) == 0xFF, "BYTE type sign test failed.");
 
-	checkAtCompileTime(sizeof(uint16) == 2,TypeTests_WORD_size);
-	checkAtCompileTime(int32(uint16(-1)) == 0xFFFF,TypeTests_WORD_sign);
+	static_assert(sizeof(uint16) == 2, "WORD type size test failed.");
+	static_assert(int32(uint16(-1)) == 0xFFFF, "WORD type sign test failed.");
 
-	checkAtCompileTime(sizeof(uint32) == 4,TypeTests_DWORD_size);
-	checkAtCompileTime(int64(uint32(-1)) == int64(0xFFFFFFFF),TypeTests_DWORD_sign);
+	static_assert(sizeof(uint32) == 4, "DWORD type size test failed.");
+	static_assert(int64(uint32(-1)) == int64(0xFFFFFFFF), "DWORD type sign test failed.");
 
-	checkAtCompileTime(sizeof(uint64) == 8,TypeTests_QWORD_size);
-	checkAtCompileTime(uint64(-1) > uint64(0),TypeTests_QWORD_sign);
+	static_assert(sizeof(uint64) == 8, "QWORD type size test failed.");
+	static_assert(uint64(-1) > uint64(0), "QWORD type sign test failed.");
 
 
-	checkAtCompileTime(sizeof(int8) == 1,TypeTests_SBYTE_size);
-	checkAtCompileTime(int32(int8(-1)) == -1,TypeTests_SBYTE_sign);
+	static_assert(sizeof(int8) == 1, "SBYTE type size test failed.");
+	static_assert(int32(int8(-1)) == -1, "SBYTE type sign test failed.");
 
-	checkAtCompileTime(sizeof(int16) == 2,TypeTests_SWORD_size);
-	checkAtCompileTime(int32(int16(-1)) == -1,TypeTests_SWORD_sign);
+	static_assert(sizeof(int16) == 2, "SWORD type size test failed.");
+	static_assert(int32(int16(-1)) == -1, "SWORD type sign test failed.");
 
-	checkAtCompileTime(sizeof(int32) == 4,TypeTests_INT_size);
-	checkAtCompileTime(int64(int32(-1)) == int64(-1),TypeTests_INT_sign);
+	static_assert(sizeof(int32) == 4, "INT type size test failed.");
+	static_assert(int64(int32(-1)) == int64(-1), "INT type sign test failed.");
 
-	checkAtCompileTime(sizeof(int64) == 8,TypeTests_SQWORD_size);
-	checkAtCompileTime(int64(-1) < int64(0),TypeTests_SQWORD_sign);
+	static_assert(sizeof(int64) == 8, "SQWORD type size test failed.");
+	static_assert(int64(-1) < int64(0), "SQWORD type sign test failed.");
 
-	checkAtCompileTime(sizeof(ANSICHAR) == 1,TypeTests_ANSICHAR_size);
-	checkAtCompileTime(int32(ANSICHAR(-1)) == -1,TypeTests_ANSICHAR_sign);
+	static_assert(sizeof(ANSICHAR) == 1, "ANSICHAR type size test failed.");
+	static_assert(int32(ANSICHAR(-1)) == -1, "ANSICHAR type sign test failed.");
 
-	checkAtCompileTime(sizeof(WIDECHAR) == 2 || sizeof(WIDECHAR) == 4,TypeTests_WIDECHAR_size);
+	static_assert(sizeof(WIDECHAR) == 2 || sizeof(WIDECHAR) == 4, "WIDECHAR type size test failed.");
 
-	checkAtCompileTime(sizeof(UCS2CHAR) == 2,TypeTests_UCS2CHAR_size);
+	static_assert(sizeof(UCS2CHAR) == 2, "UCS2CHAR type size test failed.");
 
-	checkAtCompileTime(sizeof(uint32) == 4,TypeTests_BITFIELD_size);
-	checkAtCompileTime(int64(uint32(-1)) == int64(0xFFFFFFFF),TypeTests_BITFIELD_sign);
+	static_assert(sizeof(uint32) == 4, "BITFIELD type size test failed.");
+	static_assert(int64(uint32(-1)) == int64(0xFFFFFFFF), "BITFIELD type sign test failed.");
 
-	checkAtCompileTime(sizeof(PTRINT) == sizeof(void *),TypeTests_PTRINT_size);
-	checkAtCompileTime(PTRINT(-1) < PTRINT(0),TypeTests_PTRINT_sign);
+	static_assert(sizeof(PTRINT) == sizeof(void *), "PTRINT type size test failed.");
+	static_assert(PTRINT(-1) < PTRINT(0), "PTRINT type sign test failed.");
 
-	checkAtCompileTime(sizeof(UPTRINT) == sizeof(void *),TypeTests_UPTRINT_size);
-	checkAtCompileTime(UPTRINT(-1) > UPTRINT(0),TypeTests_UPTRINT_sign);
+	static_assert(sizeof(UPTRINT) == sizeof(void *), "UPTRINT type size test failed.");
+	static_assert(UPTRINT(-1) > UPTRINT(0), "UPTRINT type sign test failed.");
 
-	checkAtCompileTime(sizeof(SIZE_T) == sizeof(void *),TypeTests_SIZE_T_size);
-	checkAtCompileTime(SIZE_T(-1) > SIZE_T(0),TypeTests_SIZE_T_sign);
+	static_assert(sizeof(SIZE_T) == sizeof(void *), "SIZE_T type size test failed.");
+	static_assert(SIZE_T(-1) > SIZE_T(0), "SIZE_T type sign test failed.");
 }
diff --git a/Engine/Source/Runtime/Core/Public/IOS/IOSPlatformMisc.h b/Engine/Source/Runtime/Core/Public/IOS/IOSPlatformMisc.h
--- a/Engine/Source/Runtime/Core/Public/IOS/IOSPlatformMisc.h
+++ b/Engine/Source/Runtime/Core/Public/IOS/IOSPlatformMisc.h
@@ -108,7 +108,7 @@ struct CORE_API FIOSPlatformMisc : public FGenericPlatformMisc
 			L"IPadAir",
 			L"Unknown",
 		};
-		checkAtCompileTime((sizeof(IOSDeviceNames) / sizeof(IOSDeviceNames[0])) == ((int32)IOS_Unknown + 1), Mismatched_IOSDeviceNames_And_EIOSDevice);
+		static_assert((sizeof(IOSDeviceNames) / sizeof(IOSDeviceNames[0])) == ((int32)IOS_Unknown + 1), "Mismatched IOSDeviceNames and EIOSDevice.");
 		
 		// look up into the string array by the enum
 		return IOSDeviceNames[(int32)GetIOSDeviceType()];
diff --git a/Engine/Source/Runtime/Core/Public/Math/BigInt.h b/Engine/Source/Runtime/Core/Public/Math/BigInt.h
--- a/Engine/Source/Runtime/Core/Public/Math/BigInt.h
+++ b/Engine/Source/Runtime/Core/Public/Math/BigInt.h
@@ -17,7 +17,7 @@ class TBigInt
 		BitsPerWord = 32
 	};
 
-	checkAtCompileTime(NumBits >= 64, TBigInt_must_have_at_least_64bits);
+	static_assert(NumBits >= 64, "TBigInt must have at least 64 bits.");
 
 	/** All bits stored as an array of words. */
 	uint32 Bits[NumBits / BitsPerWord];
diff --git a/Engine/Source/Runtime/Core/Public/Misc/AssertionMacros.h b/Engine/Source/Runtime/Core/Public/Misc/AssertionMacros.h
--- a/Engine/Source/Runtime/Core/Public/Misc/AssertionMacros.h
+++ b/Engine/Source/Runtime/Core/Public/Misc/AssertionMacros.h
@@ -178,7 +178,7 @@ struct FTCharArrayTester
 
 #define LowLevelFatalError(Format, ...) \
 	{ \
-		checkAtCompileTime(IS_TCHAR_ARRAY(Format), Formatting_string_must_be_a_TCHAR_array); \
+		static_assert(IS_TCHAR_ARRAY(Format), "Formatting string must be a TCHAR array."); \
 		FError::LowLevelFatal(__FILE__, __LINE__, Format, ##__VA_ARGS__); \
 	}
 
@@ -202,7 +202,7 @@ struct FTCharArrayTester
 	// This will only log Fatal errors
 	#define UE_LOG(CategoryName, Verbosity, Format, ...) \
 	{ \
-		checkAtCompileTime(IS_TCHAR_ARRAY(Format), Formatting_string_must_be_a_TCHAR_array); \
+		static_assert(IS_TCHAR_ARRAY(Format), "Formatting string must be a TCHAR array."); \
 		if (ELogVerbosity::Verbosity == ELogVerbosity::Fatal) \
 		{ \
 			FError::LowLevelFatal(__FILE__, __LINE__, Format, ##__VA_ARGS__); \
@@ -211,7 +211,7 @@ struct FTCharArrayTester
 	// Conditional logging (fatal errors only).
 	#define UE_CLOG(Condition, CategoryName, Verbosity, Format, ...) \
 	{ \
-		checkAtCompileTime(IS_TCHAR_ARRAY(Format), Formatting_string_must_be_a_TCHAR_array); \
+		static_assert(IS_TCHAR_ARRAY(Format), "Formatting string must be a TCHAR array."); \
 		if (ELogVerbosity::Verbosity == ELogVerbosity::Fatal) \
 		{ \
 			if (Condition) \
@@ -269,8 +269,8 @@ struct FTCharArrayTester
 	#else
 		#define UE_LOG(CategoryName, Verbosity, Format, ...) \
 		{ \
-			checkAtCompileTime(IS_TCHAR_ARRAY(Format), Formatting_string_must_be_a_TCHAR_array); \
-			checkAtCompileTime((ELogVerbosity::Verbosity & ELogVerbosity::VerbosityMask) < ELogVerbosity::NumVerbosity && ELogVerbosity::Verbosity > 0, Verbosity_must_be_constant_and_in_range); \
+			static_assert(IS_TCHAR_ARRAY(Format), "Formatting string must be a TCHAR array."); \
+			static_assert((ELogVerbosity::Verbosity & ELogVerbosity::VerbosityMask) < ELogVerbosity::NumVerbosity && ELogVerbosity::Verbosity > 0, "Verbosity must be constant and in range."); \
 			if (UE_LOG_CHECK_COMPILEDIN_VERBOSITY(CategoryName, Verbosity)) \
 			{ \
 				if (!CategoryName.IsSuppressed(ELogVerbosity::Verbosity)) \
@@ -282,8 +282,8 @@ struct FTCharArrayTester
 		// Conditional logging. Will only log if Condition is met.
 		#define UE_CLOG(Condition, CategoryName, Verbosity, Format, ...) \
 		{ \
-			checkAtCompileTime(IS_TCHAR_ARRAY(Format), Formatting_string_must_be_a_TCHAR_array); \
-			checkAtCompileTime((ELogVerbosity::Verbosity & ELogVerbosity::VerbosityMask) < ELogVerbosity::NumVerbosity && ELogVerbosity::Verbosity > 0, Verbosity_must_be_constant_and_in_range); \
+			static_assert(IS_TCHAR_ARRAY(Format), "Formatting string must be a TCHAR array."); \
+			static_assert((ELogVerbosity::Verbosity & ELogVerbosity::VerbosityMask) < ELogVerbosity::NumVerbosity && ELogVerbosity::Verbosity > 0, "Verbosity must be constant and in range."); \
 			if (UE_LOG_CHECK_COMPILEDIN_VERBOSITY(CategoryName, Verbosity)) \
 			{ \
 				if (!CategoryName.IsSuppressed(ELogVerbosity::Verbosity)) \
@@ -307,7 +307,7 @@ struct FTCharArrayTester
 	 ***/
 	#define UE_SUPPRESS(CategoryName, Verbosity, ExecuteIfUnsuppressed) \
 	{ \
-		checkAtCompileTime((ELogVerbosity::Verbosity & ELogVerbosity::VerbosityMask) < ELogVerbosity::NumVerbosity && ELogVerbosity::Verbosity > 0, Verbosity_must_be_constant_and_in_range); \
+		static_assert((ELogVerbosity::Verbosity & ELogVerbosity::VerbosityMask) < ELogVerbosity::NumVerbosity && ELogVerbosity::Verbosity > 0, "Verbosity must be constant and in range."); \
 		if (UE_LOG_CHECK_COMPILEDIN_VERBOSITY(CategoryName, Verbosity)) \
 		{ \
 			if (!CategoryName.IsSuppressed(ELogVerbosity::Verbosity)) \
diff --git a/Engine/Source/Runtime/Core/Public/Misc/Crc.h b/Engine/Source/Runtime/Core/Public/Misc/Crc.h
--- a/Engine/Source/Runtime/Core/Public/Misc/Crc.h
+++ b/Engine/Source/Runtime/Core/Public/Misc/Crc.h
@@ -26,7 +26,7 @@ struct CORE_API FCrc
 		// We ensure that we never try to do a StrCrc32 with a CharType of more than 4 bytes.  This is because
 		// we always want to treat every CRC as if it was based on 4 byte chars, even if it's less, because we
 		// want consistency between equivalent strings with different character types.
-		checkAtCompileTime(sizeof(CharType) <= 4, StrCrc32_only_works_with_CharTypes_up_to_32_bits);
+		static_assert(sizeof(CharType) <= 4, "StrCrc32 only works with CharType up to 32 bits.");
 
 		CRC = ~CRC;
 		while (CharType Ch = *Data++)
diff --git a/Engine/Source/Runtime/Core/Public/Misc/OutputDevice.h b/Engine/Source/Runtime/Core/Public/Misc/OutputDevice.h
--- a/Engine/Source/Runtime/Core/Public/Misc/OutputDevice.h
+++ b/Engine/Source/Runtime/Core/Public/Misc/OutputDevice.h
@@ -193,9 +193,9 @@ namespace ELogVerbosity
 		BreakOnLog		= 0x80
 	};
 }
-checkAtCompileTime(ELogVerbosity::NumVerbosity - 1 < ELogVerbosity::VerbosityMask, bad_ELogVerbosity_VerbosityMask);
-checkAtCompileTime(!(ELogVerbosity::VerbosityMask & ELogVerbosity::BreakOnLog), bad_ELogVerbosity_VerbosityMask2);
-checkAtCompileTime((ELogVerbosity::VerbosityMask | ELogVerbosity::BreakOnLog) < 256, bad_ELogVerbosity_VerbosityMask3); // we use a byte for storage
+static_assert(ELogVerbosity::NumVerbosity - 1 < ELogVerbosity::VerbosityMask, "Bad verbosity mask.");
+static_assert(!(ELogVerbosity::VerbosityMask & ELogVerbosity::BreakOnLog), "Bad verbosity mask.");
+static_assert((ELogVerbosity::VerbosityMask | ELogVerbosity::BreakOnLog) < 256, "Bad verbosity mask."); // we use a byte for storage
 
 /**
  * Enum that defines how the log times are to be displayed.
diff --git a/Engine/Source/Runtime/Core/Public/Misc/Variant.h b/Engine/Source/Runtime/Core/Public/Misc/Variant.h
--- a/Engine/Source/Runtime/Core/Public/Misc/Variant.h
+++ b/Engine/Source/Runtime/Core/Public/Misc/Variant.h
@@ -75,7 +75,7 @@ template<typename T> struct TVariantTraits
 {
 	static int32 GetType( )
 	{
-		checkAtCompileTime(!sizeof(T), Variant_trait_must_be_specialized_for_this_type);
+		static_assert(!sizeof(T), "Variant trait must be specialized for this type.");
 	}
 };
 
diff --git a/Engine/Source/Runtime/Core/Public/Stats/Stats2.h b/Engine/Source/Runtime/Core/Public/Stats/Stats2.h
--- a/Engine/Source/Runtime/Core/Public/Stats/Stats2.h
+++ b/Engine/Source/Runtime/Core/Public/Stats/Stats2.h
@@ -233,7 +233,7 @@ struct EMemoryRegion
 		Shift = EStatMetaFlags::Shift + EStatMetaFlags::NumBits,
 		NumBits = 4
 	};
-	checkAtCompileTime(FPlatformMemory::MCR_MAX < (1 << NumBits), need_to_expand_memoryregion_field);
+	static_assert(FPlatformMemory::MCR_MAX < (1 << NumBits), "Need to expand memory region field.");
 };
 
 
@@ -249,7 +249,7 @@ namespace EStatAllFields
 	};
 }
 
-checkAtCompileTime(EStatAllFields::StartShift > 0, too_many_stat_fields);
+static_assert(EStatAllFields::StartShift > 0, "Too many stat fields.");
 
 FORCEINLINE int64 ToPackedCallCountDuration(uint32 CallCount, uint32 Duration)
 {
@@ -286,7 +286,7 @@ public:
 	FORCEINLINE_STATS FStatNameAndInfo(FStatNameAndInfo const& Other)
 		: NameAndInfo(Other.NameAndInfo)
 	{
-		checkAtCompileTime(EStatAllFields::StartShift >= 0, too_many__fields);
+		static_assert(EStatAllFields::StartShift >= 0, "Too many fields.");
 		CheckInvariants();
 	}
 
@@ -638,7 +638,7 @@ struct FStatMessage
 	*/
 	FORCEINLINE_STATS void Clear()
 	{
-		checkAtCompileTime(sizeof(uint64) == DATA_SIZE, bad_clear);
+		static_assert(sizeof(uint64) == DATA_SIZE, "Bad clear.");
 		*(int64*)&StatData = 0;
 	}
 
@@ -647,7 +647,7 @@ struct FStatMessage
 	*/
 	FORCEINLINE_STATS int64& GetValue_int64()
 	{
-		checkAtCompileTime(sizeof(int64) <= DATA_SIZE && ALIGNOF(int64) <= DATA_ALIGN, bad_data_for_stat_message );
+		static_assert(sizeof(int64) <= DATA_SIZE && ALIGNOF(int64) <= DATA_ALIGN, "Bad data for stat message.");
 		checkStats(NameAndInfo.GetField<EStatDataType>() == EStatDataType::ST_int64);
 		return *(int64*)&StatData;
 	}
@@ -675,7 +675,7 @@ struct FStatMessage
 
 	FORCEINLINE_STATS double& GetValue_double()
 	{
-		checkAtCompileTime(sizeof(double) <= DATA_SIZE && ALIGNOF(double) <= DATA_ALIGN, bad_data_for_stat_message );
+		static_assert(sizeof(double) <= DATA_SIZE && ALIGNOF(double) <= DATA_ALIGN, "Bad data for stat message.");
 		checkStats(NameAndInfo.GetField<EStatDataType>() == EStatDataType::ST_double);
 		return *(double*)&StatData;
 	}
@@ -688,7 +688,7 @@ struct FStatMessage
 
 	FORCEINLINE_STATS FName& GetValue_FName()
 	{
-		checkAtCompileTime(sizeof(FName) <= DATA_SIZE && ALIGNOF(FName) <= DATA_ALIGN, bad_data_for_stat_message );
+		static_assert(sizeof(FName) <= DATA_SIZE && ALIGNOF(FName) <= DATA_ALIGN, "Bad data for stat message.");
 		checkStats(NameAndInfo.GetField<EStatDataType>() == EStatDataType::ST_FName);
 		return *(FName*)&StatData;
 	}
@@ -818,7 +818,7 @@ struct TStatMessage
 	*/
 	FORCEINLINE_STATS void Clear()
 	{
-		checkAtCompileTime(sizeof(uint64) == DATA_SIZE/EnumCount, bad_clear);
+		static_assert(sizeof(uint64) == DATA_SIZE / EnumCount, "Bad clear.");
 
 		for( int32 FieldIndex = 0; FieldIndex < EnumCount; ++FieldIndex )
 		{
@@ -831,7 +831,7 @@ struct TStatMessage
 	*/
 	FORCEINLINE_STATS int64& GetValue_int64( typename TEnum::Type Index )
 	{
-		checkAtCompileTime(sizeof(int64) <= DATA_SIZE && ALIGNOF(int64) <= DATA_ALIGN, bad_data_for_stat_message );
+		static_assert(sizeof(int64) <= DATA_SIZE && ALIGNOF(int64) <= DATA_ALIGN, "Bad data for stat message.");
 		checkStats(NameAndInfo.GetField<EStatDataType>() == EStatDataType::ST_int64);
 		checkStats(Index<EnumCount);
 		int64& Value = *((int64*)&StatData+(uint32)Index);
@@ -868,7 +868,7 @@ struct TStatMessage
 
 	FORCEINLINE_STATS double& GetValue_double( typename TEnum::Type Index )
 	{
-		checkAtCompileTime(sizeof(double) <= DATA_SIZE && ALIGNOF(double) <= DATA_ALIGN, bad_data_for_stat_message );
+		static_assert(sizeof(double) <= DATA_SIZE && ALIGNOF(double) <= DATA_ALIGN, "Bad data for stat message.");
 		checkStats(Index<EnumCount);
 		checkStats(NameAndInfo.GetField<EStatDataType>() == EStatDataType::ST_double);
 		double& Value = *((double*)&StatData+(uint32)Index);
@@ -1312,7 +1312,7 @@ struct FThreadSafeStaticStatInner : public FThreadSafeStaticStatBase
 {
 	FORCEINLINE TStatId GetStatId() const
 	{
-		checkAtCompileTime(sizeof(HighPerformanceEnable) == sizeof(TStatId), unsafe_cast_requires_these_to_be_the_same_thing);
+		static_assert(sizeof(HighPerformanceEnable) == sizeof(TStatId), "Unsafe cast requires these to be the same thing.");
 		if (!HighPerformanceEnable)
 		{
 			DoSetup(TStatData::GetStatName(), TStatData::GetDescription(), TStatData::TGroup::GetGroupName(), TStatData::TGroup::GetGroupCategory(), TStatData::TGroup::GetDescription(), TStatData::TGroup::IsDefaultEnabled(), TStatData::IsClearEveryFrame(), TStatData::GetStatType(), TStatData::IsCycleStat(), TStatData::GetMemoryRegion() );
diff --git a/Engine/Source/Runtime/Core/Public/Templates/SharedPointer.h b/Engine/Source/Runtime/Core/Public/Templates/SharedPointer.h
--- a/Engine/Source/Runtime/Core/Public/Templates/SharedPointer.h
+++ b/Engine/Source/Runtime/Core/Public/Templates/SharedPointer.h
@@ -133,7 +133,7 @@ template< class ObjectType, ESPMode::Type Mode >
 class TSharedRef
 {
 	// TSharedRefs with UObjects are illegal.
-	checkAtCompileTime((!CanConvertPointerFromTo<ObjectType, UObjectBase>::Result), You_Cannot_Use_TSharedRefs_With_UObjects);
+	static_assert((!CanConvertPointerFromTo<ObjectType, UObjectBase>::Result), "You cannot use TSharedRef with UObjects.");
 
 public:
 
@@ -425,7 +425,7 @@ template< class ObjectType, ESPMode::Type Mode >
 class TSharedPtr
 {
 	// TSharedPtrs with UObjects are illegal.
-	checkAtCompileTime((!CanConvertPointerFromTo<ObjectType, UObjectBase>::Result), You_Cannot_Use_TSharedPtrs_With_UObjects);
+	static_assert((!CanConvertPointerFromTo<ObjectType, UObjectBase>::Result), "You cannot use TSharedPtr with UObjects.");
 
 public:
 
diff --git a/Engine/Source/Runtime/Core/Public/Templates/UnrealTemplate.h b/Engine/Source/Runtime/Core/Public/Templates/UnrealTemplate.h
--- a/Engine/Source/Runtime/Core/Public/Templates/UnrealTemplate.h
+++ b/Engine/Source/Runtime/Core/Public/Templates/UnrealTemplate.h
@@ -41,23 +41,23 @@ class CanConvertPointerFromTo_Unrelated
 {
 };
 
-checkAtCompileTime((CanConvertPointerFromTo<bool, bool>::Result), Platform_failed__CanConvertPointerFromTo1);
-checkAtCompileTime((CanConvertPointerFromTo<void, void>::Result), Platform_failed__CanConvertPointerFromTo2);
-checkAtCompileTime((CanConvertPointerFromTo<bool, void>::Result), Platform_failed__CanConvertPointerFromTo2b);
-checkAtCompileTime((CanConvertPointerFromTo<const bool, void>::Result), Platform_failed__CanConvertPointerFromTo3);
-checkAtCompileTime((CanConvertPointerFromTo<CanConvertPointerFromTo_TestDerived, CanConvertPointerFromTo_TestBase>::Result), Platform_failed__CanConvertPointerFromTo4);
-checkAtCompileTime((CanConvertPointerFromTo<CanConvertPointerFromTo_TestDerived, const CanConvertPointerFromTo_TestBase>::Result), Platform_failed__CanConvertPointerFromTo5);
-checkAtCompileTime((CanConvertPointerFromTo<const CanConvertPointerFromTo_TestDerived, CanConvertPointerFromTo_TestBase>::Result), Platform_failed__CanConvertPointerFromTo6);
-checkAtCompileTime((CanConvertPointerFromTo<const CanConvertPointerFromTo_TestDerived, const CanConvertPointerFromTo_TestBase>::Result), Platform_failed__CanConvertPointerFromTo7);
-checkAtCompileTime((CanConvertPointerFromTo<CanConvertPointerFromTo_TestBase, CanConvertPointerFromTo_TestBase>::Result), Platform_failed__CanConvertPointerFromTo8);
-checkAtCompileTime((CanConvertPointerFromTo<CanConvertPointerFromTo_TestBase, void>::Result), Platform_failed__CanConvertPointerFromTo9);
-
-checkAtCompileTime(!(CanConvertPointerFromTo<CanConvertPointerFromTo_TestBase, CanConvertPointerFromTo_TestDerived>::Result), Platform_failed__CanConvertPointerFromTo10);
-checkAtCompileTime(!(CanConvertPointerFromTo<CanConvertPointerFromTo_Unrelated, CanConvertPointerFromTo_TestBase>::Result), Platform_failed__CanConvertPointerFromTo10b);
-checkAtCompileTime(!(CanConvertPointerFromTo<bool, CanConvertPointerFromTo_TestBase>::Result), Platform_failed__CanConvertPointerFromTo11);
-checkAtCompileTime(!(CanConvertPointerFromTo<void, CanConvertPointerFromTo_TestBase>::Result), Platform_failed__CanConvertPointerFromTo12);
-checkAtCompileTime(!(CanConvertPointerFromTo<CanConvertPointerFromTo_TestBase, bool>::Result), Platform_failed__CanConvertPointerFromTo13);
-checkAtCompileTime(!(CanConvertPointerFromTo<void, bool>::Result), Platform_failed__CanConvertPointerFromTo14);
+static_assert((CanConvertPointerFromTo<bool, bool>::Result), "Platform CanConvertPointerFromTo test failed.");
+static_assert((CanConvertPointerFromTo<void, void>::Result), "Platform CanConvertPointerFromTo test failed.");
+static_assert((CanConvertPointerFromTo<bool, void>::Result), "Platform CanConvertPointerFromTo test failed.");
+static_assert((CanConvertPointerFromTo<const bool, void>::Result), "Platform CanConvertPointerFromTo test failed.");
+static_assert((CanConvertPointerFromTo<CanConvertPointerFromTo_TestDerived, CanConvertPointerFromTo_TestBase>::Result), "Platform CanConvertPointerFromTo test failed.");
+static_assert((CanConvertPointerFromTo<CanConvertPointerFromTo_TestDerived, const CanConvertPointerFromTo_TestBase>::Result), "Platform CanConvertPointerFromTo test failed.");
+static_assert((CanConvertPointerFromTo<const CanConvertPointerFromTo_TestDerived, CanConvertPointerFromTo_TestBase>::Result), "Platform CanConvertPointerFromTo test failed.");
+static_assert((CanConvertPointerFromTo<const CanConvertPointerFromTo_TestDerived, const CanConvertPointerFromTo_TestBase>::Result), "Platform CanConvertPointerFromTo test failed.");
+static_assert((CanConvertPointerFromTo<CanConvertPointerFromTo_TestBase, CanConvertPointerFromTo_TestBase>::Result), "Platform CanConvertPointerFromTo test failed.");
+static_assert((CanConvertPointerFromTo<CanConvertPointerFromTo_TestBase, void>::Result), "Platform CanConvertPointerFromTo test failed.");
+
+static_assert(!(CanConvertPointerFromTo<CanConvertPointerFromTo_TestBase, CanConvertPointerFromTo_TestDerived>::Result), "Platform CanConvertPointerFromTo test failed.");
+static_assert(!(CanConvertPointerFromTo<CanConvertPointerFromTo_Unrelated, CanConvertPointerFromTo_TestBase>::Result), "Platform CanConvertPointerFromTo test failed.");
+static_assert(!(CanConvertPointerFromTo<bool, CanConvertPointerFromTo_TestBase>::Result), "Platform CanConvertPointerFromTo test failed.");
+static_assert(!(CanConvertPointerFromTo<void, CanConvertPointerFromTo_TestBase>::Result), "Platform CanConvertPointerFromTo test failed.");
+static_assert(!(CanConvertPointerFromTo<CanConvertPointerFromTo_TestBase, bool>::Result), "Platform CanConvertPointerFromTo test failed.");
+static_assert(!(CanConvertPointerFromTo<void, bool>::Result), "Platform CanConvertPointerFromTo test failed.");
 
 #endif
 
diff --git a/Engine/Source/Runtime/Core/Public/Templates/UnrealTypeTraits.h b/Engine/Source/Runtime/Core/Public/Templates/UnrealTypeTraits.h
--- a/Engine/Source/Runtime/Core/Public/Templates/UnrealTypeTraits.h
+++ b/Engine/Source/Runtime/Core/Public/Templates/UnrealTypeTraits.h
@@ -114,7 +114,7 @@ struct TFormatSpecifier
 		// If the static_assert below were a constant 0 or something not dependent on T, compilers are free to detect this and fail to compile the template.
 		// As specified in the C++ standard s14.6p8. A compiler is free to give a diagnostic here or not. MSVC ignores it, and clang/gcc instantiates the 
 		// template and triggers the static_assert.
-		checkAtCompileTime(sizeof(T) < 0, FormatSpeciferNotSupportedForThisType); // request for a format specifier for a type we do not know about
+		static_assert(sizeof(T) < 0, "Format specifer not supported for this type."); // request for a format specifier for a type we do not know about
 		return TEXT("Unknown");
 	}
 };
diff --git a/Engine/Source/Runtime/Core/Public/Windows/WindowsPlatformCompilerSetup.h b/Engine/Source/Runtime/Core/Public/Windows/WindowsPlatformCompilerSetup.h
--- a/Engine/Source/Runtime/Core/Public/Windows/WindowsPlatformCompilerSetup.h
+++ b/Engine/Source/Runtime/Core/Public/Windows/WindowsPlatformCompilerSetup.h
@@ -22,7 +22,7 @@
 /**
  * We require at least Visual Studio 2012 to compile
  */
-checkAtCompileTime(_MSC_VER >= 1700, Visual_Studio_2012_Required);
+static_assert(_MSC_VER >= 1700, "Visual Studio 2012 or newer required.");
 
 #pragma warning (error:      4001 4002 4003 4004      4006 4007 4008 4009 4010 4011 4012 4013 4014 4015 4016 4017 4018 4019 4020 4021 4022 4023 4024 4025 4026 4027 4028 4029 4030 4031 4032 4033 4034 4035 4036 4037 4038 4039 4040 4041 4042 4043 4044 4045 4046 4047 4048 4049 4050 4051 4052 4053 4054 4055 4056 4057 4058 4059 4060           4063 4064 4065 4066 4067 4068 4069 4070 4071 4072 4073 4074 4075 4076 4077 4078 4079 4080 4081 4082 4083 4084 4085 4086 4087 4088 4089 4090 4091 4092 4093 4094 4095 4096 4097 4098 4099)
 #pragma warning (error: 4100 4101 4102 4103 4104 4105 4106 4107 4108 4109 4110 4111 4112 4113 4114 4115 4116 4117 4118 4119 4120 4121 4122 4123 4124 4125 4126 4127 4128 4129 4130 4131 4132 4133 4134 4135 4136 4137 4138 4139 4140 4141 4142 4143 4144 4145 4146 4147 4148 4149 4150 4151 4152 4153 4154 4155 4156 4157 4158 4159 4160 4161 4162 4163 4164      4166 4167 4168 4169 4170 4171 4172 4173 4174 4175 4176 4177 4178 4179 4180 4181 4182 4183 4184 4185 4186 4187 4188 4189 4190 4191 4192 4193 4194 4195 4196 4197 4198 4199)
diff --git a/Engine/Source/Runtime/CoreUObject/Private/UObject/PropertyArray.cpp b/Engine/Source/Runtime/CoreUObject/Private/UObject/PropertyArray.cpp
--- a/Engine/Source/Runtime/CoreUObject/Private/UObject/PropertyArray.cpp
+++ b/Engine/Source/Runtime/CoreUObject/Private/UObject/PropertyArray.cpp
@@ -354,6 +354,6 @@ IMPLEMENT_CORE_INTRINSIC_CLASS(UArrayProperty, UProperty,
 
 		// Ensure that TArray and FScriptArray are interchangeable, as FScriptArray will be used to access a native array property
 		// from script that is declared as a TArray in C++.
-		checkAtCompileTime(sizeof(FScriptArray) == sizeof(TArray<uint8>),FScriptArrayAndTArrayMustBeInterchangable);
+		static_assert(sizeof(FScriptArray) == sizeof(TArray<uint8>), "FScriptArray and TArray<uint8> must be interchangable.");
 	}
 );
diff --git a/Engine/Source/Runtime/CoreUObject/Public/Templates/Casts.h b/Engine/Source/Runtime/CoreUObject/Public/Templates/Casts.h
--- a/Engine/Source/Runtime/CoreUObject/Public/Templates/Casts.h
+++ b/Engine/Source/Runtime/CoreUObject/Public/Templates/Casts.h
@@ -4,7 +4,7 @@
 
 // it doesn't seem worthwhile to risk changing Cast<>/IsA() to work with interfaces at this time,
 // so we'll just make sure it fails... spectacularly
-#define CATCH_UNSUPPORTED_INTERFACECAST(TheType) checkAtCompileTime((TheType::StaticClassFlags & CLASS_Interface) == 0, __CASTING_TO_INTERFACE_TYPE_##TheType##_IS_UNSUPPORTED)
+#define CATCH_UNSUPPORTED_INTERFACECAST(TheType) static_assert((TheType::StaticClassFlags & CLASS_Interface) == 0, "Casting to interface type " #TheType " is unsupported.")
 
 COREUOBJECT_API void CastLogError(const TCHAR* FromType, const TCHAR* ToType);
 
diff --git a/Engine/Source/Runtime/CoreUObject/Public/UObject/Class.h b/Engine/Source/Runtime/CoreUObject/Public/UObject/Class.h
--- a/Engine/Source/Runtime/CoreUObject/Public/UObject/Class.h
+++ b/Engine/Source/Runtime/CoreUObject/Public/UObject/Class.h
@@ -557,7 +557,7 @@ FORCEINLINE typename TEnableIf<!TStructOpsTypeTraits<CPPSTRUCT>::WithCopy, bool>
 template<class CPPSTRUCT>
 FORCEINLINE typename TEnableIf<TStructOpsTypeTraits<CPPSTRUCT>::WithCopy, bool>::Type CopyOrNot(CPPSTRUCT* Dest, CPPSTRUCT const* Src, int32 ArrayDim)
 {
-	checkAtCompileTime((!TIsPODType<CPPSTRUCT>::Value),you_probably_dont_want_custom_copy_for_a_pod_type); 
+	static_assert((!TIsPODType<CPPSTRUCT>::Value), "You probably don't want custom copy for a POD type.");
 	for (;ArrayDim;--ArrayDim)
 	{
 		*Dest++ = *Src++;
@@ -587,7 +587,7 @@ FORCEINLINE typename TEnableIf<TStructOpsTypeTraits<CPPSTRUCT>::WithAddStructRef
 template<class CPPSTRUCT>
 FORCEINLINE typename TEnableIf<TStructOpsTypeTraits<CPPSTRUCT>::WithIdentical && TStructOpsTypeTraits<CPPSTRUCT>::WithIdenticalViaEquality, bool>::Type IdenticalOrNot(const CPPSTRUCT* A, const CPPSTRUCT* B, uint32 PortFlags, bool& bOutResult)
 {
-	checkAtCompileTime(sizeof(CPPSTRUCT) == 0,should_not_have_both_WithIdenticalViaEquality_and_WithIdentical);
+	static_assert(sizeof(CPPSTRUCT) == 0, "Should not have both WithIdenticalViaEquality and WithIdentical.");
 }
 
 template<class CPPSTRUCT>
diff --git a/Engine/Source/Runtime/CoreUObject/Public/UObject/GarbageCollection.h b/Engine/Source/Runtime/CoreUObject/Public/UObject/GarbageCollection.h
--- a/Engine/Source/Runtime/CoreUObject/Public/UObject/GarbageCollection.h
+++ b/Engine/Source/Runtime/CoreUObject/Public/UObject/GarbageCollection.h
@@ -330,10 +330,10 @@ struct FGCReferenceTokenStream
 	{
 		UPTRINT Result = UPTRINT(Tokens[CurrentIndex++]);
 #if PLATFORM_64BITS // warning C4293: '<<' : shift count negative or too big, undefined behavior, so we needed the ifdef
-		checkAtCompileTime(sizeof(void*) == 8, pointer_size_mismatch);
+		static_assert(sizeof(void*) == 8, "Pointer size mismatch.");
 		Result |= UPTRINT(Tokens[CurrentIndex++]) << 32;
 #else
-		checkAtCompileTime(sizeof(void*) == 4, pointer_size_mismatch);
+		static_assert(sizeof(void*) == 4, "Pointer size mismatch.");
 #endif
 		return (void*)Result;
 	}
@@ -406,11 +406,11 @@ private:
 	FORCEINLINE void StorePointer( uint32* Stream, void const* Ptr )
 	{
 	#if PLATFORM_64BITS // warning C4293: '<<' : shift count negative or too big, undefined behavior, so we needed the ifdef
-		checkAtCompileTime(sizeof(void*) == 8, pointer_size_mismatch);
+		static_assert(sizeof(void*) == 8, "Pointer size mismatch.");
 		Stream[0] = UPTRINT(Ptr) & 0xffffffff;
 		Stream[1] = UPTRINT(Ptr) >> 32;
 	#else
-		checkAtCompileTime(sizeof(void*) == 4, pointer_size_mismatch);
+		static_assert(sizeof(void*) == 4, "Pointer size mismatch.");
 		Stream[0] = PTRINT(Ptr);
 	#endif
 	}
diff --git a/Engine/Source/Runtime/CoreUObject/Public/UObject/UObjectBase.h b/Engine/Source/Runtime/CoreUObject/Public/UObject/UObjectBase.h
--- a/Engine/Source/Runtime/CoreUObject/Public/UObject/UObjectBase.h
+++ b/Engine/Source/Runtime/CoreUObject/Public/UObject/UObjectBase.h
@@ -179,7 +179,7 @@ public:
 	 **/
 	FORCEINLINE bool ThisThreadAtomicallyClearedRFUnreachable()
 	{
-		checkAtCompileTime(sizeof(int32) == sizeof(EObjectFlags),flags_must_be_32bit_for_atomics);
+		static_assert(sizeof(int32) == sizeof(EObjectFlags), "Flags must be 32-bit for atomics.");
 		bool bIChangedIt = false;
 		while (1)
 		{
diff --git a/Engine/Source/Runtime/CoreUObject/Public/UObject/UObjectGlobals.h b/Engine/Source/Runtime/CoreUObject/Public/UObject/UObjectGlobals.h
--- a/Engine/Source/Runtime/CoreUObject/Public/UObject/UObjectGlobals.h
+++ b/Engine/Source/Runtime/CoreUObject/Public/UObject/UObjectGlobals.h
@@ -613,21 +613,21 @@ public:
 	TSubobjectPtr()
 		: FSubobjectPtr((UObject*)FSubobjectPtr::InvalidPtrValue)
 	{
-		checkAtCompileTime(sizeof(TSubobjectPtr) == sizeof(UObject*), TSuobjectPtr_Equals_Pointer_Size_Assumption_Failed);
+		static_assert(sizeof(TSubobjectPtr) == sizeof(UObject*), "TSuobjectPtr should equal pointer size.");
 	}
 	/** Copy constructor */
 	template <class DerivedSubobjectType>
 	TSubobjectPtr(TSubobjectPtr<DerivedSubobjectType>& Other)
 		: FSubobjectPtr(Other.Object)
 	{
-		checkAtCompileTime((CanConvertPointerFromTo<DerivedSubobjectType, SubobjectType>::Result), Subobject_Pointers_Must_Be_Compatible);
+		static_assert((CanConvertPointerFromTo<DerivedSubobjectType, SubobjectType>::Result), "Subobject pointers must be compatible.");
 	}
 	/** Initialization constructor. Can only be used with PCIP.CreateDefaultSubobject(). */
 	template <class DerivedSubobjectType>
 	TSubobjectPtr(const TSubobjectPtrConstructor<DerivedSubobjectType>& Other)
 		: FSubobjectPtr(Other.Object)
 	{
-		checkAtCompileTime((CanConvertPointerFromTo<DerivedSubobjectType, SubobjectType>::Result), Subobject_Pointers_Must_Be_Compatible);
+		static_assert((CanConvertPointerFromTo<DerivedSubobjectType, SubobjectType>::Result), "Subobject pointers must be compatible.");
 	}
 	/** Gets the sub-object pointer. */
 	FORCEINLINE SubobjectType* Get() const
@@ -638,7 +638,7 @@ public:
 	template <class DerivedSubobjectType>
 	FORCEINLINE TSubobjectPtr& operator=(const TSubobjectPtrConstructor<DerivedSubobjectType>& Other)
 	{
-		checkAtCompileTime((CanConvertPointerFromTo<DerivedSubobjectType, SubobjectType>::Result), Subobject_Pointers_Must_Be_Compatible);
+		static_assert((CanConvertPointerFromTo<DerivedSubobjectType, SubobjectType>::Result), "Subobject pointers must be compatible.");
 		Set(Other.Object);
 		return *this;
 	}
diff --git a/Engine/Source/Runtime/CoreUObject/Public/UObject/UnrealType.h b/Engine/Source/Runtime/CoreUObject/Public/UObject/UnrealType.h
--- a/Engine/Source/Runtime/CoreUObject/Public/UObject/UnrealType.h
+++ b/Engine/Source/Runtime/CoreUObject/Public/UObject/UnrealType.h
@@ -3079,7 +3079,7 @@ struct FTestType
 template<>
 struct FTestType<bool>
 {
-	checkAtCompileTime(sizeof(bool)==sizeof(uint8), bool_is_not_one_byte);
+	static_assert(sizeof(bool) == sizeof(uint8), "Bool is not one byte.");
 	typedef uint8 TestType;
 };
 
@@ -3132,7 +3132,7 @@ struct COREUOBJECT_API DetermineBitfieldOffsetAndMask
 	void DoDetermineBitfieldOffsetAndMask(const SIZE_T SizeOf)
 	{
 		typedef typename FTestType<BitfieldType>::TestType TTestType;
-		checkAtCompileTime(sizeof(TTestType)==sizeof(BitfieldType), wrong_size_for_test_type);
+		static_assert(sizeof(TTestType) == sizeof(BitfieldType), "Wrong size for test type.");
 
 		void* Buffer = AllocateBuffer(SizeOf);
 		TTestType* Test = (TTestType*)Buffer;
diff --git a/Engine/Source/Runtime/Engine/Classes/Engine/EngineTypes.h b/Engine/Source/Runtime/Engine/Classes/Engine/EngineTypes.h
--- a/Engine/Source/Runtime/Engine/Classes/Engine/EngineTypes.h
+++ b/Engine/Source/Runtime/Engine/Classes/Engine/EngineTypes.h
@@ -2226,7 +2226,7 @@ struct FReplicationFlags
 	}
 };
 
-checkAtCompileTime( sizeof( FReplicationFlags ) == 4, FReplicationFlags_InvalidSize );
+static_assert(sizeof(FReplicationFlags) == 4, "FReplicationFlags has invalid size.");
 
 /** Struct used to specify the property name of the component to constrain */
 USTRUCT()
diff --git a/Engine/Source/Runtime/Engine/Private/AI/Navigation/PImplRecastNavMesh.cpp b/Engine/Source/Runtime/Engine/Private/AI/Navigation/PImplRecastNavMesh.cpp
--- a/Engine/Source/Runtime/Engine/Private/AI/Navigation/PImplRecastNavMesh.cpp
+++ b/Engine/Source/Runtime/Engine/Private/AI/Navigation/PImplRecastNavMesh.cpp
@@ -26,12 +26,12 @@
 // bunch of compile-time checks to assure types used by Recast and our
 // mid-layer are the same size
 //----------------------------------------------------------------------//
-checkAtCompileTime(sizeof(NavNodeRef) == sizeof(dtPolyRef), NavNodeRef_and_dtPolyRef_should_be_same_size);
-checkAtCompileTime(RECAST_MAX_AREAS <= DT_MAX_AREAS, Number_of_allowed_areas_cannot_exceed_DT_MAX_AREAS);
-checkAtCompileTime(RECAST_STRAIGHTPATH_OFFMESH_CONNECTION == DT_STRAIGHTPATH_OFFMESH_CONNECTION, Path_flags_values_differ);
+static_assert(sizeof(NavNodeRef) == sizeof(dtPolyRef), "NavNodeRef and dtPolyRef should be the same size.");
+static_assert(RECAST_MAX_AREAS <= DT_MAX_AREAS, "Number of allowed areas cannot exceed DT_MAX_AREAS.");
+static_assert(RECAST_STRAIGHTPATH_OFFMESH_CONNECTION == DT_STRAIGHTPATH_OFFMESH_CONNECTION, "Path flags values differ.");
 // @todo ps4 compile issue: FLT_MAX constexpr issue
 #if !PLATFORM_PS4
-checkAtCompileTime(RECAST_UNWALKABLE_POLY_COST == DT_UNWALKABLE_POLY_COST, Unwalkable_poly_cost_differ);
+static_assert(RECAST_UNWALKABLE_POLY_COST == DT_UNWALKABLE_POLY_COST, "Unwalkable poly cost differ.");
 #endif
 
 /// Helper for accessing navigation query from different threads
diff --git a/Engine/Source/Runtime/Engine/Private/AudioDerivedData.cpp b/Engine/Source/Runtime/Engine/Private/AudioDerivedData.cpp
--- a/Engine/Source/Runtime/Engine/Private/AudioDerivedData.cpp
+++ b/Engine/Source/Runtime/Engine/Private/AudioDerivedData.cpp
@@ -681,7 +681,7 @@ void CookSurroundWave( USoundWave* SoundWave, FName FormatName, const IAudioForm
 
 	// Front left channel is the master
 	ChannelCount = 1;
-	checkAtCompileTime(SPEAKER_FrontLeft == 0, frontleft_must_be_first);
+	static_assert(SPEAKER_FrontLeft == 0, "Front-left speaker must be first.");
 	if( WaveInfo.ReadWaveHeader( RawWaveData, SoundWave->ChannelSizes[ SPEAKER_FrontLeft ], SoundWave->ChannelOffsets[ SPEAKER_FrontLeft ] ) )
 	{
 		{
diff --git a/Engine/Source/Runtime/Engine/Private/Collision/PhysXCollision.h b/Engine/Source/Runtime/Engine/Private/Collision/PhysXCollision.h
--- a/Engine/Source/Runtime/Engine/Private/Collision/PhysXCollision.h
+++ b/Engine/Source/Runtime/Engine/Private/Collision/PhysXCollision.h
@@ -17,8 +17,8 @@
 #define OVERLAP_BUFFER_SIZE						1024
 #define OVERLAP_BUFFER_SIZE_MAX_SYNC_QUERIES	992
 
-checkAtCompileTime(HIT_BUFFER_SIZE > 0, InvalidPhysXHitBufferSize);
-checkAtCompileTime(HIT_BUFFER_MAX_SYNC_QUERIES < HIT_BUFFER_SIZE, InvalidPhysXSyncBufferSize);
+static_assert(HIT_BUFFER_SIZE > 0, "Invalid PhysX hit buffer size.");
+static_assert(HIT_BUFFER_MAX_SYNC_QUERIES < HIT_BUFFER_SIZE, "Invalid PhysX sync buffer size.");
 
 
 
diff --git a/Engine/Source/Runtime/Engine/Private/Distributions.cpp b/Engine/Source/Runtime/Engine/Private/Distributions.cpp
--- a/Engine/Source/Runtime/Engine/Private/Distributions.cpp
+++ b/Engine/Source/Runtime/Engine/Private/Distributions.cpp
@@ -26,7 +26,7 @@ ENGINE_API uint32 GDistributionType = 1;
 #define LOOKUP_TABLE_MAX_SAMPLES (128)
 
 // The maximum number of samples must be a power of two.
-checkAtCompileTime( (LOOKUP_TABLE_MAX_SAMPLES & (LOOKUP_TABLE_MAX_SAMPLES-1)) == 0, LookupTableMaxSamplesNotPowerOfTwo );
+static_assert((LOOKUP_TABLE_MAX_SAMPLES & (LOOKUP_TABLE_MAX_SAMPLES - 1)) == 0, "Lookup table max samples is not a power of two.");
 
 // Log category for distributions.
 DEFINE_LOG_CATEGORY_STATIC(LogDistributions,Warning,Warning);
diff --git a/Engine/Source/Runtime/Engine/Private/Materials/MaterialExpressions.cpp b/Engine/Source/Runtime/Engine/Private/Materials/MaterialExpressions.cpp
--- a/Engine/Source/Runtime/Engine/Private/Materials/MaterialExpressions.cpp
+++ b/Engine/Source/Runtime/Engine/Private/Materials/MaterialExpressions.cpp
@@ -6709,7 +6709,7 @@ bool UMaterialExpressionMaterialFunctionCall::IsInputConnectionRequired(int32 In
 
 static FString GetInputDefaultValueString(EFunctionInputType InputType, const FVector4& PreviewValue)
 {
-	checkAtCompileTime(FunctionInput_Scalar < FunctionInput_Vector4, ERROR_EnumValuesOutOfOrder);
+	static_assert(FunctionInput_Scalar < FunctionInput_Vector4, "Enum values out of order.");
 	check(InputType <= FunctionInput_Vector4);
 
 	FString ValueString = FString::Printf(TEXT("DefaultValue = (%.2f"), PreviewValue.X);
diff --git a/Engine/Source/Runtime/Engine/Private/Materials/MaterialShared.cpp b/Engine/Source/Runtime/Engine/Private/Materials/MaterialShared.cpp
--- a/Engine/Source/Runtime/Engine/Private/Materials/MaterialShared.cpp
+++ b/Engine/Source/Runtime/Engine/Private/Materials/MaterialShared.cpp
@@ -26,7 +26,7 @@ FName MaterialQualityLevelNames[] =
 	FName(TEXT("Num"))
 };
 
-checkAtCompileTime(ARRAY_COUNT(MaterialQualityLevelNames) == EMaterialQualityLevel::Num + 1, MissingEntryFromMaterialQualityLevelNames);
+static_assert(ARRAY_COUNT(MaterialQualityLevelNames) == EMaterialQualityLevel::Num + 1, "Missing entry from material quality level names.");
 
 void GetMaterialQualityLevelName(EMaterialQualityLevel::Type InQualityLevel, FString& OutName)
 {
diff --git a/Engine/Source/Runtime/Engine/Private/Particles/ParticleCurveTexture.cpp b/Engine/Source/Runtime/Engine/Private/Particles/ParticleCurveTexture.cpp
--- a/Engine/Source/Runtime/Engine/Private/Particles/ParticleCurveTexture.cpp
+++ b/Engine/Source/Runtime/Engine/Private/Particles/ParticleCurveTexture.cpp
@@ -20,8 +20,8 @@ extern const int32 GParticleCurveTextureSizeX = 256;
 extern const int32 GParticleCurveTextureSizeY = 256;
 
 /** The texel allocator uses 16-bit integers internally. */
-checkAtCompileTime( GParticleCurveTextureSizeX <= 0xffff, CurveTextureWiderThanSixteenBits );
-checkAtCompileTime( GParticleCurveTextureSizeY <= 0xffff, CurveTextureTallerThanSixteenBits );
+static_assert(GParticleCurveTextureSizeX <= 0xffff, "Curve texture wider than sixteen bits.");
+static_assert(GParticleCurveTextureSizeY <= 0xffff, "Curve texture taller than sixteen bits.");
 
 /** The global curve texture resource. */
 TGlobalResource<FParticleCurveTexture> GParticleCurveTexture;
diff --git a/Engine/Source/Runtime/Engine/Private/Particles/ParticleGpuSimulation.cpp b/Engine/Source/Runtime/Engine/Private/Particles/ParticleGpuSimulation.cpp
--- a/Engine/Source/Runtime/Engine/Private/Particles/ParticleGpuSimulation.cpp
+++ b/Engine/Source/Runtime/Engine/Private/Particles/ParticleGpuSimulation.cpp
@@ -40,17 +40,17 @@ const int32 GParticleSimulationTextureSizeX = 1024;
 const int32 GParticleSimulationTextureSizeY = 1024;
 
 /** Texture size must be power-of-two. */
-checkAtCompileTime( (GParticleSimulationTextureSizeX & (GParticleSimulationTextureSizeX-1)) == 0, ParticleSimulationTextureSizeXNotPowerOfTwo );
-checkAtCompileTime( (GParticleSimulationTextureSizeY & (GParticleSimulationTextureSizeY-1)) == 0, ParticleSimulationTextureSizeYNotPowerOfTwo );
+static_assert((GParticleSimulationTextureSizeX & (GParticleSimulationTextureSizeX - 1)) == 0, "Particle simulation texture size X is not a power of two.");
+static_assert((GParticleSimulationTextureSizeY & (GParticleSimulationTextureSizeY - 1)) == 0, "Particle simulation texture size Y is not a power of two.");
 
 /** The tile size. Texture space is allocated in TileSize x TileSize units. */
 const int32 GParticleSimulationTileSize = 4;
 const int32 GParticlesPerTile = GParticleSimulationTileSize * GParticleSimulationTileSize;
 
 /** Tile size must be power-of-two and <= each dimension of the simulation texture. */
-checkAtCompileTime( (GParticleSimulationTileSize & (GParticleSimulationTileSize-1)) == 0, ParticleSimulationTileSizeNotPowerOfTwo );
-checkAtCompileTime( GParticleSimulationTileSize <= GParticleSimulationTextureSizeX, ParticleSimulationTileSizeLargerThanTexture );
-checkAtCompileTime( GParticleSimulationTileSize <= GParticleSimulationTextureSizeY, ParticleSimulationTileSizeLargerThanTexture );
+static_assert((GParticleSimulationTileSize & (GParticleSimulationTileSize - 1)) == 0, "Particle simulation tile size is not a power of two.");
+static_assert(GParticleSimulationTileSize <= GParticleSimulationTextureSizeX, "Particle simulation tile size is larger than texture.");
+static_assert(GParticleSimulationTileSize <= GParticleSimulationTextureSizeY, "Particle simulation tile size is larger than texture.");
 
 /** How many tiles are in the simulation textures. */
 const int32 GParticleSimulationTileCountX = GParticleSimulationTextureSizeX / GParticleSimulationTileSize;
@@ -58,16 +58,16 @@ const int32 GParticleSimulationTileCountY = GParticleSimulationTextureSizeY / GP
 const int32 GParticleSimulationTileCount = GParticleSimulationTileCountX * GParticleSimulationTileCountY;
 
 /** GPU particle rendering code assumes that the number of particles per instanced draw is <= 16. */
-checkAtCompileTime( MAX_PARTICLES_PER_INSTANCE <= 16, MaxParticlesPerInstanceGreaterThan16 );
+static_assert(MAX_PARTICLES_PER_INSTANCE <= 16, "Max particles per instance is greater than 16.");
 /** Also, it must be a power of 2. */
-checkAtCompileTime( (MAX_PARTICLES_PER_INSTANCE & (MAX_PARTICLES_PER_INSTANCE - 1)) == 0, MaxParticlesPerInstanceNotPowerOfTwo );
+static_assert((MAX_PARTICLES_PER_INSTANCE & (MAX_PARTICLES_PER_INSTANCE - 1)) == 0, "Max particles per instance is not a power of two.");
 
 /** Particle tiles are aligned to the same number as when rendering. */
 enum { TILES_PER_INSTANCE = 8 };
 /** The number of tiles per instance must be <= MAX_PARTICLES_PER_INSTANCE. */
-checkAtCompileTime( TILES_PER_INSTANCE <= MAX_PARTICLES_PER_INSTANCE, TilesPerInstanceGreaterThanMaxParticlesPerInstance );
+static_assert(TILES_PER_INSTANCE <= MAX_PARTICLES_PER_INSTANCE, "Tiles per instance is greater than max particles per instance.");
 /** Also, it must be a power of 2. */
-checkAtCompileTime( (TILES_PER_INSTANCE & (TILES_PER_INSTANCE-1)) == 0, TilesPerInstanceNotPowerOfTwo );
+static_assert((TILES_PER_INSTANCE & (TILES_PER_INSTANCE - 1)) == 0, "Tiles per instance is not a power of two.");
 
 /** Maximum number of vector fields that can be evaluated at once. */
 enum { MAX_VECTOR_FIELDS = 4 };
diff --git a/Engine/Source/Runtime/Engine/Private/PhysicsEngine/PhysScene.cpp b/Engine/Source/Runtime/Engine/Private/PhysicsEngine/PhysScene.cpp
--- a/Engine/Source/Runtime/Engine/Private/PhysicsEngine/PhysScene.cpp
+++ b/Engine/Source/Runtime/Engine/Private/PhysicsEngine/PhysScene.cpp
@@ -546,7 +546,7 @@ void FPhysScene::ProcessPhysScene(uint32 SceneType)
 
 	if (FrameLagAsync())
 	{
-		checkAtCompileTime(PST_MAX == 3, Assumtiopns_about_physics_scenes); // Here we assume the PST_Sync is the master and never fame lagged
+		static_assert(PST_MAX == 3, "Physics scene static test failed."); // Here we assume the PST_Sync is the master and never fame lagged
 		if (SceneType == PST_Sync)
 		{
 			// the one frame lagged one should be done by now.
diff --git a/Engine/Source/Runtime/Engine/Private/PhysicsEngine/PhysXSupport.h b/Engine/Source/Runtime/Engine/Private/PhysicsEngine/PhysXSupport.h
--- a/Engine/Source/Runtime/Engine/Private/PhysicsEngine/PhysXSupport.h
+++ b/Engine/Source/Runtime/Engine/Private/PhysicsEngine/PhysXSupport.h
@@ -323,7 +323,7 @@ public:
 	virtual void* allocate(size_t size, const char* typeName, const char* filename, int line) override
 	{
 #if PHYSX_MEMORY_STATS
-		checkAtCompileTime(sizeof(FPhysXAllocationHeader) <= 16,FPhysXAllocationHeaderMustBeLessThan16Bytes);
+		static_assert(sizeof(FPhysXAllocationHeader) <= 16, "FPhysXAllocationHeader size must be less than 16 bytes.");
 
 		INC_DWORD_STAT_BY(STAT_MemoryPhysXTotalAllocationSize, size);
 
diff --git a/Engine/Source/Runtime/Engine/Public/Net/DataChannel.h b/Engine/Source/Runtime/Engine/Public/Net/DataChannel.h
--- a/Engine/Source/Runtime/Engine/Public/Net/DataChannel.h
+++ b/Engine/Source/Runtime/Engine/Public/Net/DataChannel.h
@@ -75,7 +75,7 @@ private:
 		/** sends a message of this type on the specified connection's control channel */ \
 		static void Send(UNetConnection* Conn) \
 		{ \
-			checkAtCompileTime(Index < 256, CONTROL_CHANNEL_MESSAGE_MUST_BE_BYTE); \
+			static_assert(Index < 256, "Control channel message must be a byte."); \
 			checkSlow(!Conn->IsA(UChildConnection::StaticClass())); /** control channel messages can only be sent on the parent connection */ \
 			if (Conn->Channels[0] != NULL && !Conn->Channels[0]->Closing) \
 			{ \
@@ -101,7 +101,7 @@ private:
 		 */ \
 		static void Send(UNetConnection* Conn, TypeA& ParamA) \
 		{ \
-			checkAtCompileTime(Index < 256, CONTROL_CHANNEL_MESSAGE_MUST_BE_BYTE); \
+			static_assert(Index < 256, "Control channel message must be a byte."); \
 			checkSlow(!Conn->IsA(UChildConnection::StaticClass())); /** control channel messages can only be sent on the parent connection */ \
 			if (Conn->Channels[0] != NULL && !Conn->Channels[0]->Closing) \
 			{ \
@@ -139,7 +139,7 @@ private:
 		 */ \
 		static void Send(UNetConnection* Conn, TypeA& ParamA, TypeB& ParamB) \
 		{ \
-			checkAtCompileTime(Index < 256, CONTROL_CHANNEL_MESSAGE_MUST_BE_BYTE); \
+			static_assert(Index < 256, "Control channel message must be a byte."); \
 			checkSlow(!Conn->IsA(UChildConnection::StaticClass())); /** control channel messages can only be sent on the parent connection */ \
 			if (Conn->Channels[0] != NULL && !Conn->Channels[0]->Closing) \
 			{ \
@@ -179,7 +179,7 @@ private:
 		 */ \
 		static void Send(UNetConnection* Conn, TypeA& ParamA, TypeB& ParamB, TypeC& ParamC) \
 		{ \
-			checkAtCompileTime(Index < 256, CONTROL_CHANNEL_MESSAGE_MUST_BE_BYTE); \
+			static_assert(Index < 256, "Control channel message must be a byte."); \
 			checkSlow(!Conn->IsA(UChildConnection::StaticClass())); /** control channel messages can only be sent on the parent connection */ \
 			if (Conn->Channels[0] != NULL && !Conn->Channels[0]->Closing) \
 			{ \
@@ -222,7 +222,7 @@ private:
 		 */ \
 		static void Send(UNetConnection* Conn, TypeA& ParamA, TypeB& ParamB, TypeC& ParamC, TypeD& ParamD) \
 		{ \
-			checkAtCompileTime(Index < 256, CONTROL_CHANNEL_MESSAGE_MUST_BE_BYTE); \
+			static_assert(Index < 256, "Control channel message must be a byte."); \
 			checkSlow(!Conn->IsA(UChildConnection::StaticClass())); /** control channel messages can only be sent on the parent connection */ \
 			if (Conn->Channels[0] != NULL && !Conn->Channels[0]->Closing) \
 			{ \
@@ -269,7 +269,7 @@ private:
 		 */ \
 		static void Send(UNetConnection* Conn, TypeA& ParamA, TypeB& ParamB, TypeC& ParamC, TypeD& ParamD, TypeE& ParamE, TypeF& ParamF, TypeG& ParamG) \
 		{ \
-			checkAtCompileTime(Index < 256, CONTROL_CHANNEL_MESSAGE_MUST_BE_BYTE); \
+			static_assert(Index < 256, "Control channel message must be a byte."); \
 			checkSlow(!Conn->IsA(UChildConnection::StaticClass())); /** control channel messages can only be sent on the parent connection */ \
 			if (Conn->Channels[0] != NULL && !Conn->Channels[0]->Closing) \
 			{ \
@@ -325,7 +325,7 @@ private:
 		 */ \
 		static void Send(UNetConnection* Conn, TypeA& ParamA, TypeB& ParamB, TypeC& ParamC, TypeD& ParamD, TypeE& ParamE, TypeF& ParamF, TypeG& ParamG, TypeH& ParamH) \
 		{ \
-			checkAtCompileTime(Index < 256, CONTROL_CHANNEL_MESSAGE_MUST_BE_BYTE); \
+			static_assert(Index < 256, "Control channel message must be a byte."); \
 			checkSlow(!Conn->IsA(UChildConnection::StaticClass())); /** control channel messages can only be sent on the parent connection */ \
 			if (Conn->Channels[0] != NULL && !Conn->Channels[0]->Closing) \
 			{ \
diff --git a/Engine/Source/Runtime/Engine/Public/PrecomputedLightVolume.h b/Engine/Source/Runtime/Engine/Public/PrecomputedLightVolume.h
--- a/Engine/Source/Runtime/Engine/Public/PrecomputedLightVolume.h
+++ b/Engine/Source/Runtime/Engine/Public/PrecomputedLightVolume.h
@@ -65,7 +65,7 @@ struct FLightVolumeOctreeSemantics
 	{
 		FPlatformMisc::Prefetch( &Sample, CACHE_LINE_SIZE );
 		// here we require that the position and radius are contiguous in memory
-		checkAtCompileTime(STRUCT_OFFSET( FVolumeLightingSample, Position ) + 3 * sizeof(float) ==  STRUCT_OFFSET( FVolumeLightingSample, Radius ), FVolumeLightingSample_Radius_Must_Follow_Position );
+		static_assert(STRUCT_OFFSET(FVolumeLightingSample, Position) + 3 * sizeof(float) == STRUCT_OFFSET(FVolumeLightingSample, Radius), "FVolumeLightingSample radius must follow position.");
 		return &Sample.Position.X;
 	}
 
diff --git a/Engine/Source/Runtime/Engine/Public/RawIndexBuffer.h b/Engine/Source/Runtime/Engine/Public/RawIndexBuffer.h
--- a/Engine/Source/Runtime/Engine/Public/RawIndexBuffer.h
+++ b/Engine/Source/Runtime/Engine/Public/RawIndexBuffer.h
@@ -214,7 +214,7 @@ public:
 	:	Indices(InNeedsCPUAccess)
 	{
 #if DISALLOW_32BIT_INDICES
-		checkAtCompileTime( sizeof(INDEX_TYPE) == sizeof(uint16), DISALLOW_32BIT_INDICESIsDefinedDoNotUse32BitIndices );
+		static_assert(sizeof(INDEX_TYPE) == sizeof(uint16), "DISALLOW_32BIT_INDICES is defined, so you should not use 32-bit indices.");
 #endif
 	}
 
diff --git a/Engine/Source/Runtime/Engine/Public/VisualLog.h b/Engine/Source/Runtime/Engine/Public/VisualLog.h
--- a/Engine/Source/Runtime/Engine/Public/VisualLog.h
+++ b/Engine/Source/Runtime/Engine/Public/VisualLog.h
@@ -196,7 +196,7 @@ struct ENGINE_API FVisLogEntry
 #define UE_VLOG(Actor, CategoryName, Verbosity, Format, ...) \
 { \
 	SCOPE_CYCLE_COUNTER(STAT_VisualLog); \
-	checkAtCompileTime((ELogVerbosity::Verbosity & ELogVerbosity::VerbosityMask) < ELogVerbosity::NumVerbosity && ELogVerbosity::Verbosity > 0, Verbosity_must_be_constant_and_in_range); \
+	static_assert((ELogVerbosity::Verbosity & ELogVerbosity::VerbosityMask) < ELogVerbosity::NumVerbosity && ELogVerbosity::Verbosity > 0, "Verbosity must be constant and in range."); \
 	FVisualLog::Get()->LogLine(Actor, CategoryName.GetCategoryName(), ELogVerbosity::Verbosity, FString::Printf(Format, ##__VA_ARGS__)); \
 	if (UE_LOG_CHECK_COMPILEDIN_VERBOSITY(CategoryName, Verbosity)) \
 	{ \
@@ -215,7 +215,7 @@ struct ENGINE_API FVisLogEntry
 #define UE_VLOG_SEGMENT_THICK(Actor, CategoryName, Verbosity, SegmentStart, SegmentEnd, Color, Thickness, DescriptionFormat, ...) \
 { \
 	SCOPE_CYCLE_COUNTER(STAT_VisualLog); \
-	checkAtCompileTime((ELogVerbosity::Verbosity & ELogVerbosity::VerbosityMask) < ELogVerbosity::NumVerbosity && ELogVerbosity::Verbosity > 0, Verbosity_must_be_constant_and_in_range); \
+	static_assert((ELogVerbosity::Verbosity & ELogVerbosity::VerbosityMask) < ELogVerbosity::NumVerbosity && ELogVerbosity::Verbosity > 0, "Verbosity must be constant and in range."); \
 	if (FVisualLog::Get()->IsRecording() && (!FVisualLog::Get()->IsAllBlocked() || FVisualLog::Get()->InWhitelist(CategoryName.GetCategoryName()))) \
 	{ \
 		FVisualLog::Get()->GetEntryToWrite(Actor)->AddElement(SegmentStart, SegmentEnd, CategoryName.GetCategoryName(), Color, FString::Printf(DescriptionFormat, ##__VA_ARGS__), Thickness); \
@@ -227,7 +227,7 @@ struct ENGINE_API FVisLogEntry
 #define UE_VLOG_LOCATION(Actor, CategoryName, Verbosity, Location, Radius, Color, DescriptionFormat, ...) \
 { \
 	SCOPE_CYCLE_COUNTER(STAT_VisualLog); \
-	checkAtCompileTime((ELogVerbosity::Verbosity & ELogVerbosity::VerbosityMask) < ELogVerbosity::NumVerbosity && ELogVerbosity::Verbosity > 0, Verbosity_must_be_constant_and_in_range); \
+	static_assert((ELogVerbosity::Verbosity & ELogVerbosity::VerbosityMask) < ELogVerbosity::NumVerbosity && ELogVerbosity::Verbosity > 0, "Verbosity must be constant and in range."); \
 	if (FVisualLog::Get()->IsRecording() && (!FVisualLog::Get()->IsAllBlocked() || FVisualLog::Get()->InWhitelist(CategoryName.GetCategoryName()))) \
 	{ \
 		FVisualLog::Get()->GetEntryToWrite(Actor)->AddElement(Location, CategoryName.GetCategoryName(), Color, FString::Printf(DescriptionFormat, ##__VA_ARGS__), Radius); \
@@ -237,7 +237,7 @@ struct ENGINE_API FVisLogEntry
 #define UE_VLOG_BOX(Actor, CategoryName, Verbosity, Box, Color, DescriptionFormat, ...) \
 { \
 	SCOPE_CYCLE_COUNTER(STAT_VisualLog); \
-	checkAtCompileTime((ELogVerbosity::Verbosity & ELogVerbosity::VerbosityMask) < ELogVerbosity::NumVerbosity && ELogVerbosity::Verbosity > 0, Verbosity_must_be_constant_and_in_range); \
+	static_assert((ELogVerbosity::Verbosity & ELogVerbosity::VerbosityMask) < ELogVerbosity::NumVerbosity && ELogVerbosity::Verbosity > 0, "Verbosity must be constant and in range."); \
 	if (FVisualLog::Get()->IsRecording() && (!FVisualLog::Get()->IsAllBlocked() || FVisualLog::Get()->InWhitelist(CategoryName.GetCategoryName()))) \
 	{ \
 		FVisualLog::Get()->GetEntryToWrite(Actor)->AddElement(Box, CategoryName.GetCategoryName(), Color, FString::Printf(DescriptionFormat, ##__VA_ARGS__)); \
@@ -247,7 +247,7 @@ struct ENGINE_API FVisLogEntry
 #define UE_VLOG_HISTOGRAM(Actor, CategoryName, Verbosity, GraphName, DataName, Data) \
 { \
 	SCOPE_CYCLE_COUNTER(STAT_VisualLog); \
-	checkAtCompileTime((ELogVerbosity::Verbosity & ELogVerbosity::VerbosityMask) < ELogVerbosity::NumVerbosity && ELogVerbosity::Verbosity > 0, Verbosity_must_be_constant_and_in_range); \
+	static_assert((ELogVerbosity::Verbosity & ELogVerbosity::VerbosityMask) < ELogVerbosity::NumVerbosity && ELogVerbosity::Verbosity > 0, "Verbosity must be constant and in range."); \
 	if (FVisualLog::Get()->IsRecording() && (!FVisualLog::Get()->IsAllBlocked() || FVisualLog::Get()->InWhitelist(CategoryName.GetCategoryName()))) \
 	{ \
 		FVisualLog::Get()->GetEntryToWrite(Actor)->AddHistogramData(Data, CategoryName.GetCategoryName(), GraphName, DataName); \
diff --git a/Engine/Source/Runtime/GameplayAbilities/Public/GameplayEffect.h b/Engine/Source/Runtime/GameplayAbilities/Public/GameplayEffect.h
--- a/Engine/Source/Runtime/GameplayAbilities/Public/GameplayEffect.h
+++ b/Engine/Source/Runtime/GameplayAbilities/Public/GameplayEffect.h
@@ -80,10 +80,10 @@ namespace EGameplayImmunity
 	};
 }
 
-checkAtCompileTime((int32)EGameplayMod::OutgoingGE == (int32)EGameplayImmunity::OutgoingGE, "EGameplayMod::OutgoingGE and EGameplayImmunity::OutgoingGE must match. Did you forget to modify one of them?");
-checkAtCompileTime((int32)EGameplayMod::IncomingGE == (int32)EGameplayImmunity::IncomingGE, "EGameplayMod::IncomingGE and EGameplayImmunity::IncomingGE must match. Did you forget to modify one of them?");
-checkAtCompileTime((int32)EGameplayMod::ActiveGE == (int32)EGameplayImmunity::ActiveGE, "EGameplayMod::ActiveGE and EGameplayImmunity::ActiveGE must match. Did you forget to modify one of them?");
-checkAtCompileTime((int32)EGameplayMod::Max == (int32)EGameplayImmunity::Max, "EGameplayMod and EGameplayImmunity must cover the same cases. Did you forget to modify one of them?");
+static_assert((int32)EGameplayMod::OutgoingGE == (int32)EGameplayImmunity::OutgoingGE, "EGameplayMod::OutgoingGE and EGameplayImmunity::OutgoingGE must match. Did you forget to modify one of them?");
+static_assert((int32)EGameplayMod::IncomingGE == (int32)EGameplayImmunity::IncomingGE, "EGameplayMod::IncomingGE and EGameplayImmunity::IncomingGE must match. Did you forget to modify one of them?");
+static_assert((int32)EGameplayMod::ActiveGE == (int32)EGameplayImmunity::ActiveGE, "EGameplayMod::ActiveGE and EGameplayImmunity::ActiveGE must match. Did you forget to modify one of them?");
+static_assert((int32)EGameplayMod::Max == (int32)EGameplayImmunity::Max, "EGameplayMod and EGameplayImmunity must cover the same cases. Did you forget to modify one of them?");
 
 /**
  * Tells us what a GameplayEffect modifies when being applied to another GameplayEffect
diff --git a/Engine/Source/Runtime/Messaging/Public/Common/MessageEndpointBuilder.h b/Engine/Source/Runtime/Messaging/Public/Common/MessageEndpointBuilder.h
--- a/Engine/Source/Runtime/Messaging/Public/Common/MessageEndpointBuilder.h
+++ b/Engine/Source/Runtime/Messaging/Public/Common/MessageEndpointBuilder.h
@@ -56,7 +56,7 @@ public:
 	FMessageEndpointBuilder& Handling( HandlerType* Handler, typename TMessageHandlerFunc<MessageType, HandlerType>::Type HandlerFunc )
 	{
 		// @todo gmp: implement proper async message deserialization, so this can be removed
-		checkAtCompileTime(TStructOpsTypeTraits<MessageType>::WithMessageHandling == true, Please_Add_A_WithMessageHandling_TypeTrait);
+		static_assert(TStructOpsTypeTraits<MessageType>::WithMessageHandling == true, "Please add a WithMessageHandling type trait.");
 
 		Handlers.Add(MakeShareable(new TMessageHandler<MessageType, HandlerType>(Handler, HandlerFunc)));
 
diff --git a/Engine/Source/Runtime/RHI/Private/RHI.cpp b/Engine/Source/Runtime/RHI/Private/RHI.cpp
--- a/Engine/Source/Runtime/RHI/Private/RHI.cpp
+++ b/Engine/Source/Runtime/RHI/Private/RHI.cpp
@@ -194,7 +194,7 @@ FName FeatureLevelNames[] =
 	FName(TEXT("SM5")),
 };
 
-checkAtCompileTime(ARRAY_COUNT(FeatureLevelNames) == ERHIFeatureLevel::Num, MissingEntryFromFeatureLevelNames);
+static_assert(ARRAY_COUNT(FeatureLevelNames) == ERHIFeatureLevel::Num, "Missing entry from feature level names.");
 
 RHI_API bool GetFeatureLevelFromName(FName Name, ERHIFeatureLevel::Type& OutFeatureLevel)
 {
diff --git a/Engine/Source/Runtime/RHI/Public/RHI.h b/Engine/Source/Runtime/RHI/Public/RHI.h
--- a/Engine/Source/Runtime/RHI/Public/RHI.h
+++ b/Engine/Source/Runtime/RHI/Public/RHI.h
@@ -623,7 +623,7 @@ public:
 	FBlendStateInitializerRHI(const TStaticArray<FRenderTarget,NumRenderTargets>& InRenderTargetBlendStates)
 	:	bUseIndependentRenderTargetBlendStates(NumRenderTargets > 1)
 	{
-		checkAtCompileTime(NumRenderTargets <= MaxSimultaneousRenderTargets,TooManyRenderTargetBlendStates);
+		static_assert(NumRenderTargets <= MaxSimultaneousRenderTargets, "Too many render target blend states.");
 
 		for(uint32 RenderTargetIndex = 0;RenderTargetIndex < NumRenderTargets;++RenderTargetIndex)
 		{
diff --git a/Engine/Source/Runtime/RenderCore/Public/UniformBuffer.h b/Engine/Source/Runtime/RenderCore/Public/UniformBuffer.h
--- a/Engine/Source/Runtime/RenderCore/Public/UniformBuffer.h
+++ b/Engine/Source/Runtime/RenderCore/Public/UniformBuffer.h
@@ -512,7 +512,7 @@ private:
 		typedef TAlignedTypedef<FShaderResourceViewRHIParamRef,Alignment>::TAlignedType TAlignedType;
 		static const FUniformBufferStruct* GetStruct() { return NULL; }
 	};
-	checkAtCompileTime(sizeof(FShaderResourceViewRHIParamRef) == sizeof(void*),SRVParamRef_NotPointerSized);
+	static_assert(sizeof(FShaderResourceViewRHIParamRef) == sizeof(void*), "FShaderResourceViewRHIParamRef should have size of a pointer.");
 	static_assert(sizeof(TUniformBufferTypeInfo<FShaderResourceViewRHIParamRef>::TAlignedType) == sizeof(void*), "SRV UniformBufferParam is not aligned to pointer type");
 
 	template<>
@@ -528,7 +528,7 @@ private:
 		typedef TAlignedTypedef<FUnorderedAccessViewRHIParamRef,Alignment>::TAlignedType TAlignedType;
 		static const FUniformBufferStruct* GetStruct() { return NULL; }
 	};
-	checkAtCompileTime(sizeof(FUnorderedAccessViewRHIParamRef) == sizeof(void*),UAVParamRef_NotPointerSized);
+	static_assert(sizeof(FUnorderedAccessViewRHIParamRef) == sizeof(void*), "FUnorderedAccessViewRHIParamRef should have size of a pointer.");
 	static_assert(sizeof(TUniformBufferTypeInfo<FUnorderedAccessViewRHIParamRef>::TAlignedType) == sizeof(void*), "UAV UniformBufferParam is not aligned to pointer type");
 
 	template<>
@@ -544,7 +544,7 @@ private:
 		typedef TAlignedTypedef<FSamplerStateRHIParamRef,Alignment>::TAlignedType TAlignedType;
 		static const FUniformBufferStruct* GetStruct() { return NULL; }
 	};
-	checkAtCompileTime(sizeof(FSamplerStateRHIParamRef) == sizeof(void*),SamplerStateParamRef_NotPointerSized);
+	static_assert(sizeof(FSamplerStateRHIParamRef) == sizeof(void*), "FSamplerStateRHIParamRef should have size of a pointer.");
 	static_assert(sizeof(TUniformBufferTypeInfo<FSamplerStateRHIParamRef>::TAlignedType) == sizeof(void*), "SamplerState UniformBufferParam is not aligned to pointer type");
 
 	template<>
@@ -560,7 +560,7 @@ private:
 		typedef TAlignedTypedef<FTextureRHIParamRef,Alignment>::TAlignedType TAlignedType;
 		static const FUniformBufferStruct* GetStruct() { return NULL; }
 	};
-	checkAtCompileTime(sizeof(FTextureRHIParamRef) == sizeof(void*),TextureParamRef_NotPointerSized);
+	static_assert(sizeof(FTextureRHIParamRef) == sizeof(void*), "FTextureRHIParamRef should have size of a pointer.");
 	static_assert(sizeof(TUniformBufferTypeInfo<FTextureRHIParamRef>::TAlignedType) == sizeof(void*), "Texture UniformBufferParam is not aligned to pointer type");
 
 //
@@ -603,14 +603,14 @@ private:
 		typedef MemberType zzA##MemberName ArrayDecl; \
 		typedef TUniformBufferTypeInfo<zzA##MemberName>::TAlignedType zzT##MemberName; \
 		zzT##MemberName MemberName; \
-		checkAtCompileTime(TUniformBufferTypeInfo<zzA##MemberName>::BaseType != UBMT_INVALID,InvalidMemberType_##MemberName_##MemberType); \
-		checkAtCompileTime(TUniformBufferTypeInfo<zzA##MemberName>::BaseType != UBMT_UAV,UAVsNotYetSupportedInResourceTables_##MemberName_##MemberType); \
+		static_assert(TUniformBufferTypeInfo<zzA##MemberName>::BaseType != UBMT_INVALID, "Invalid type " #MemberType " of member " #MemberName "."); \
+		static_assert(TUniformBufferTypeInfo<zzA##MemberName>::BaseType != UBMT_UAV, "UAV is not yet supported in resource tables for " #MemberName " of type " #MemberType "."); \
 	private: \
 		struct zzNextMemberId##MemberName { enum { HasDeclaredResource = zzMemberId##MemberName::HasDeclaredResource || TUniformBufferTypeInfo<zzT##MemberName>::IsResource }; }; \
 		static TArray<FUniformBufferStruct::FMember> zzGetMembersBefore(zzNextMemberId##MemberName) \
 		{ \
-			checkAtCompileTime(TUniformBufferTypeInfo<zzT##MemberName>::IsResource == 1 || zzMemberId##MemberName::HasDeclaredResource == 0, AllResourcesMustBeDeclaredLast_##MemberName); \
-			checkAtCompileTime(TUniformBufferTypeInfo<zzT##MemberName>::IsResource == 0 || IS_TCHAR_ARRAY(OptionalShaderType), NoShaderTypeFor_##MemberName); \
+			static_assert(TUniformBufferTypeInfo<zzT##MemberName>::IsResource == 1 || zzMemberId##MemberName::HasDeclaredResource == 0, "All resources must be declared last for " #MemberName "."); \
+			static_assert(TUniformBufferTypeInfo<zzT##MemberName>::IsResource == 0 || IS_TCHAR_ARRAY(OptionalShaderType), "No shader type for " #MemberName "."); \
 			/* Route the member enumeration on to the function for the member following this. */ \
 			TArray<FUniformBufferStruct::FMember> OutMembers = zzGetMembersBefore(zzMemberId##MemberName()); \
 			/* Add this member. */ \
@@ -625,10 +625,9 @@ private:
 				TUniformBufferTypeInfo<zzA##MemberName>::NumElements, \
 				TUniformBufferTypeInfo<zzA##MemberName>::GetStruct() \
 				)); \
-			checkAtCompileTime( \
+			static_assert( \
 				(STRUCT_OFFSET(zzTThisStruct,MemberName) & (TUniformBufferTypeInfo<zzA##MemberName>::Alignment - 1)) == 0, \
-				MisalignedUniformBufferStructMember##MemberName \
-				); \
+				"Misaligned uniform buffer struct member " #MemberName "."); \
 			return OutMembers; \
 		} \
 		typedef zzNextMemberId##MemberName
diff --git a/Engine/Source/Runtime/Renderer/Private/PostProcess/PostProcessing.cpp b/Engine/Source/Runtime/Renderer/Private/PostProcess/PostProcessing.cpp
--- a/Engine/Source/Runtime/Renderer/Private/PostProcess/PostProcessing.cpp
+++ b/Engine/Source/Runtime/Renderer/Private/PostProcess/PostProcessing.cpp
@@ -351,7 +351,7 @@ static FRenderingCompositeOutputRef AddBloom(FPostprocessContext Context, FRende
 	{
 		static const TCHAR* PassLabels[] =
 			{NULL, TEXT("BloomDownsample1"), TEXT("BloomDownsample2"), TEXT("BloomDownsample3"), TEXT("BloomDownsample4")};
-		checkAtCompileTime(ARRAY_COUNT(PassLabels) == DownSampleStages, PassLabel_count_must_equal_DownSampleStages);
+		static_assert(ARRAY_COUNT(PassLabels) == DownSampleStages, "PassLabel count must be equal to DownSampleStages.");
 		FRenderingCompositePass* Pass = Context.Graph.RegisterPass(new(FMemStack::Get()) FRCPassPostProcessDownsample(PF_Unknown, 1, EPostProcessRectSource::GBS_ViewRect, PassLabels[i]));
 		Pass->SetInput(ePId_Input0, PostProcessDownsamples[i - 1]);
 		PostProcessDownsamples[i] = FRenderingCompositeOutputRef(Pass);
diff --git a/Engine/Source/Runtime/Renderer/Private/TranslucentLighting.cpp b/Engine/Source/Runtime/Renderer/Private/TranslucentLighting.cpp
--- a/Engine/Source/Runtime/Renderer/Private/TranslucentLighting.cpp
+++ b/Engine/Source/Runtime/Renderer/Private/TranslucentLighting.cpp
@@ -1078,7 +1078,7 @@ void FDeferredShadingSceneRenderer::ClearTranslucentVolumeLighting(FRHICommandLi
 
 		// Clear all volume textures in the same draw with MRT, which is faster than individually
 
-		checkAtCompileTime(TVC_MAX == 2, OnlyExpectingTwoTranslucencyLightingCascades);
+		static_assert(TVC_MAX == 2, "Only expecting two translucency lighting cascades.");
 		FTextureRHIParamRef RenderTargets[4];
 		RenderTargets[0] = GSceneRenderTargets.TranslucencyLightingVolumeAmbient[0]->GetRenderTargetItem().TargetableTexture;
 		RenderTargets[1] = GSceneRenderTargets.TranslucencyLightingVolumeDirectional[0]->GetRenderTargetItem().TargetableTexture;
@@ -1359,7 +1359,7 @@ void FDeferredShadingSceneRenderer::ClearTranslucentVolumePerObjectShadowing(FRH
 	{
 		SCOPED_DRAW_EVENT(ClearTranslucentVolumePerLightShadowing, DEC_SCENE_ITEMS);
 
-		checkAtCompileTime(TVC_MAX == 2, OnlyExpectingTwoTranslucencyLightingCascades);
+		static_assert(TVC_MAX == 2, "Only expecting two translucency lighting cascades.");
 		FTextureRHIParamRef RenderTargets[2];
 		RenderTargets[0] = GSceneRenderTargets.GetTranslucencyVolumeAmbient(TVC_Inner)->GetRenderTargetItem().TargetableTexture;
 		RenderTargets[1] = GSceneRenderTargets.GetTranslucencyVolumeDirectional(TVC_Inner)->GetRenderTargetItem().TargetableTexture;
diff --git a/Engine/Source/Runtime/ShaderCore/Public/ShaderCore.h b/Engine/Source/Runtime/ShaderCore/Public/ShaderCore.h
--- a/Engine/Source/Runtime/ShaderCore/Public/ShaderCore.h
+++ b/Engine/Source/Runtime/ShaderCore/Public/ShaderCore.h
@@ -49,7 +49,7 @@ DECLARE_MEMORY_STAT_EXTERN(TEXT("Shader MapMemory"),STAT_Shaders_ShaderMapMemory
 
 inline TStatId GetMemoryStatType(EShaderFrequency ShaderFrequency)
 {
-	checkAtCompileTime(6 == SF_NumFrequencies, Bad_EShaderFrequency);
+	static_assert(6 == SF_NumFrequencies, "EShaderFrequency has a bad size.");
 	switch(ShaderFrequency)
 	{
 		case SF_Pixel:		return GET_STATID(STAT_PixelShaderMemory);
diff --git a/Engine/Source/Runtime/ShaderCore/Public/ShaderParameterUtils.h b/Engine/Source/Runtime/ShaderCore/Public/ShaderParameterUtils.h
--- a/Engine/Source/Runtime/ShaderCore/Public/ShaderParameterUtils.h
+++ b/Engine/Source/Runtime/ShaderCore/Public/ShaderParameterUtils.h
@@ -25,7 +25,7 @@ void SetShaderValue(
 	uint32 ElementIndex = 0
 	)
 {
-	checkAtCompileTime(!TIsPointerType<ParameterType>::Value, ErrorPassByValue);
+	static_assert(!TIsPointerType<ParameterType>::Value, "Passing by value is not valid.");
 
 	const uint32 AlignedTypeSize = Align(sizeof(ParameterType),ShaderArrayElementAlignBytes);
 	const int32 NumBytesToSet = FMath::Min<int32>(sizeof(ParameterType),Parameter.GetNumBytes() - ElementIndex * AlignedTypeSize);
diff --git a/Engine/Source/Runtime/Slate/Public/Framework/Views/TableViewTypeTraits.h b/Engine/Source/Runtime/Slate/Public/Framework/Views/TableViewTypeTraits.h
--- a/Engine/Source/Runtime/Slate/Public/Framework/Views/TableViewTypeTraits.h
+++ b/Engine/Source/Runtime/Slate/Public/Framework/Views/TableViewTypeTraits.h
@@ -22,7 +22,7 @@ template<typename T> struct TIsValidListItem< TWeakObjectPtr<T> >						{ enum {
  */
 template <typename T> struct TListTypeTraits
 {
-	checkAtCompileTime(TIsValidListItem<T>::Value, ItemType_must_be_a_pointer_or_TSharedPtr);
+	static_assert(TIsValidListItem<T>::Value, "Item type T must be a pointer or a TSharedPtr.");
 };
 
 
diff --git a/Engine/Source/Runtime/Slate/Public/Widgets/Views/STableRow.h b/Engine/Source/Runtime/Slate/Public/Widgets/Views/STableRow.h
--- a/Engine/Source/Runtime/Slate/Public/Widgets/Views/STableRow.h
+++ b/Engine/Source/Runtime/Slate/Public/Widgets/Views/STableRow.h
@@ -50,7 +50,7 @@ DECLARE_DELEGATE_RetVal_OneParam(FReply, FOnTableRowDrop, FDragDropEvent const&)
 template<typename ItemType>
 class STableRow : public ITableRow, public SBorder
 {
-	checkAtCompileTime( TIsValidListItem<ItemType>::Value, ItemType_must_be_a_pointer_or_TSharedPtr );
+	static_assert(TIsValidListItem<ItemType>::Value, "Item type T must be a pointer or a TSharedPtr.");
 
 public:
 
diff --git a/Engine/Source/Runtime/Windows/D3D11RHI/Private/D3D11Shaders.cpp b/Engine/Source/Runtime/Windows/D3D11RHI/Private/D3D11Shaders.cpp
--- a/Engine/Source/Runtime/Windows/D3D11RHI/Private/D3D11Shaders.cpp
+++ b/Engine/Source/Runtime/Windows/D3D11RHI/Private/D3D11Shaders.cpp
@@ -234,7 +234,7 @@ FD3D11BoundShaderState::FD3D11BoundShaderState(
 	bShaderNeedsGlobalConstantBuffer[SF_Pixel] = InPixelShader ? InPixelShader->bShaderNeedsGlobalConstantBuffer : false;
 	bShaderNeedsGlobalConstantBuffer[SF_Geometry] = InGeometryShader ? InGeometryShader->bShaderNeedsGlobalConstantBuffer : false;
 
-	checkAtCompileTime(ARRAY_COUNT(bShaderNeedsGlobalConstantBuffer) == SF_NumFrequencies, NumFrequenciesMismatch);
+	static_assert(ARRAY_COUNT(bShaderNeedsGlobalConstantBuffer) == SF_NumFrequencies, "EShaderFrequency size should match with array count of bShaderNeedsGlobalConstantBuffer.");
 }
 
 /**
diff --git a/Engine/Source/Runtime/Windows/D3D11RHI/Private/D3D11State.cpp b/Engine/Source/Runtime/Windows/D3D11RHI/Private/D3D11State.cpp
--- a/Engine/Source/Runtime/Windows/D3D11RHI/Private/D3D11State.cpp
+++ b/Engine/Source/Runtime/Windows/D3D11RHI/Private/D3D11State.cpp
@@ -264,7 +264,7 @@ FBlendStateRHIRef FD3D11DynamicRHI::RHICreateBlendState(const FBlendStateInitial
 	BlendDesc.AlphaToCoverageEnable = false;
 	BlendDesc.IndependentBlendEnable = Initializer.bUseIndependentRenderTargetBlendStates;
 
-	checkAtCompileTime(MaxSimultaneousRenderTargets <= D3D11_SIMULTANEOUS_RENDER_TARGET_COUNT,TooManyMRTs);
+	static_assert(MaxSimultaneousRenderTargets <= D3D11_SIMULTANEOUS_RENDER_TARGET_COUNT, "Too many MRTs.");
 	for(uint32 RenderTargetIndex = 0;RenderTargetIndex < MaxSimultaneousRenderTargets;++RenderTargetIndex)
 	{
 		const FBlendStateInitializerRHI::FRenderTarget& RenderTargetInitializer = Initializer.RenderTargets[RenderTargetIndex];
