Fixed erroneous double-escaping of some characters when replacing characters with their escaped equivalents

This was a regression caused by the re-introduction of the escaping characters in CL#2056465.
This addresses TTP#338005 - Editor: If you copy/paste a comment that has an apostrophe in it, it'll show the escape character in the new comment (\\')

diff --git a/Engine/Source/Runtime/Core/Private/Containers/String.cpp b/Engine/Source/Runtime/Core/Private/Containers/String.cpp
--- a/Engine/Source/Runtime/Core/Private/Containers/String.cpp
+++ b/Engine/Source/Runtime/Core/Private/Containers/String.cpp
@@ -953,18 +953,19 @@ FString FString::ReplaceQuotesWithEscapedQuotes() const
 	return *this;
 }
 
-#define MAX_SUPPORTED_ESCAPE_CHARS 6
-
-static const TCHAR* CharToEscapeSeqMap[MAX_SUPPORTED_ESCAPE_CHARS][2] =
+static const TCHAR* CharToEscapeSeqMap[][2] =
 {
+	// Always replace \\ first to avoid double-escaping characters
+	{ TEXT("\\"), TEXT("\\\\") },
 	{ TEXT("\n"), TEXT("\\n")  },
 	{ TEXT("\r"), TEXT("\\r")  },
 	{ TEXT("\t"), TEXT("\\t")  },
 	{ TEXT("\'"), TEXT("\\'")  },
-	{ TEXT("\\"), TEXT("\\\\") },
 	{ TEXT("\""), TEXT("\\\"") }
 };
 
+static const uint32 MaxSupportedEscapeChars = ARRAY_COUNT(CharToEscapeSeqMap);
+
 /**
  * Replaces certain characters with the "escaped" version of that character (i.e. replaces "\n" with "\\n").
  * The characters supported are: { \n, \r, \t, \', \", \\ }.
@@ -978,7 +979,7 @@ FString FString::ReplaceCharWithEscapedChar( const TArray<TCHAR>* Chars/*=NULL*/
 	if ( Len() > 0 && (Chars == NULL || Chars->Num() > 0) )
 	{
 		FString Result(*this);
-		for ( int32 ChIdx = 0; ChIdx < MAX_SUPPORTED_ESCAPE_CHARS; ChIdx++ )
+		for ( int32 ChIdx = 0; ChIdx < MaxSupportedEscapeChars; ChIdx++ )
 		{
 			if ( Chars == NULL || Chars->Contains(*(CharToEscapeSeqMap[ChIdx][0])) )
 			{
@@ -1001,7 +1002,7 @@ FString FString::ReplaceEscapedCharWithChar( const TArray<TCHAR>* Chars/*=NULL*/
 	{
 		FString Result(*this);
 		// Spin CharToEscapeSeqMap backwards to ensure we're doing the inverse of ReplaceCharWithEscapedChar
-		for ( int32 ChIdx = MAX_SUPPORTED_ESCAPE_CHARS - 1; ChIdx >= 0; ChIdx-- )
+		for ( int32 ChIdx = MaxSupportedEscapeChars - 1; ChIdx >= 0; ChIdx-- )
 		{
 			if ( Chars == NULL || Chars->Contains(*(CharToEscapeSeqMap[ChIdx][0])) )
 			{
diff --git a/Engine/Source/Runtime/CoreUObject/Private/UObject/Property.cpp b/Engine/Source/Runtime/CoreUObject/Private/UObject/Property.cpp
--- a/Engine/Source/Runtime/CoreUObject/Private/UObject/Property.cpp
+++ b/Engine/Source/Runtime/CoreUObject/Private/UObject/Property.cpp
@@ -274,6 +274,11 @@ const TCHAR* UPropertyHelpers::ReadToken( const TCHAR* Buffer, FString& String,
 				String += TCHAR('"');
 				Buffer++;
 			}
+			else if (*Buffer == TCHAR('\'')) // escaped single quote "\'"
+			{
+				String += TCHAR('\'');
+				Buffer++;
+			}
 			else if ( *Buffer == TCHAR('n') ) // escaped newline
 			{
 				String += TCHAR('\n');
