UniquePtr improvements.

#add Construction and assignment from TUniquePtrs of other types.
#change Comparison operators defined outside the class.
#change More MakeUnique overloads for non-variadic compilers.

diff --git a/Engine/Source/Runtime/Core/Public/Templates/UniquePtr.h b/Engine/Source/Runtime/Core/Public/Templates/UniquePtr.h
--- a/Engine/Source/Runtime/Core/Public/Templates/UniquePtr.h
+++ b/Engine/Source/Runtime/Core/Public/Templates/UniquePtr.h
@@ -44,6 +44,16 @@ public:
 		Other.Ptr = NULL;
 	}
 
+	/**
+	 * Constructor from rvalues of other (usually derived) types
+	 */
+	template <typename OtherT>
+	FORCEINLINE TUniquePtr(TUniquePtr<OtherT>&& Other)
+		: Ptr(Other.Ptr)
+	{
+		Other.Ptr = NULL;
+	}
+
 	/**
 	 * Move assignment operator
 	 */
@@ -61,6 +71,24 @@ public:
 		return *this;
 	}
 
+	/**
+	 * Assignment operator for rvalues of other (usually derived) types
+	 */
+	template <typename OtherT>
+	FORCEINLINE TUniquePtr& operator=(TUniquePtr<OtherT>&& Other)
+	{
+		if (this != &Other)
+		{
+			// We delete last, because we don't want odd side effects if the destructor of T relies on the state of this or Other
+			T* OldPtr = Ptr;
+			Ptr = Other.Ptr;
+			Other.Ptr = NULL;
+			delete OldPtr;
+		}
+
+		return *this;
+	}
+
 	/**
 	 * Destructor
 	 */
@@ -109,32 +137,6 @@ public:
 		return *Ptr;
 	}
 
-	/**
-	 * Equality comparison operator
-	 *
-	 * @param Lhs The first TUniquePtr to compare.
-	 * @param Rhs The second TUniquePtr to compare.
-	 *
-	 * @return true if the two TUniquePtrs are logically substitutable for each other, false otherwise.
-	 */
-	friend FORCEINLINE bool operator==(const TUniquePtr& Lhs, const TUniquePtr& Rhs)
-	{
-		return Lhs.Ptr == Rhs.Ptr;
-	}
-
-	/**
-	 * Inequality comparison operator
-	 *
-	 * @param Lhs The first TUniquePtr to compare.
-	 * @param Rhs The second TUniquePtr to compare.
-	 *
-	 * @return false if the two TUniquePtrs are logically substitutable for each other, true otherwise.
-	 */
-	friend FORCEINLINE bool operator!=(const TUniquePtr& Lhs, const TUniquePtr& Rhs)
-	{
-		return Lhs.Ptr != Rhs.Ptr;
-	}
-
 	/**
 	 * Returns a pointer to the owned object without relinquishing ownership.
 	 *
@@ -178,6 +180,34 @@ private:
 	T* Ptr;
 };
 
+/**
+ * Equality comparison operator
+ *
+ * @param Lhs The first TUniquePtr to compare.
+ * @param Rhs The second TUniquePtr to compare.
+ *
+ * @return true if the two TUniquePtrs are logically substitutable for each other, false otherwise.
+ */
+template <typename LhsT, typename RhsT>
+FORCEINLINE bool operator==(const TUniquePtr<LhsT>& Lhs, const TUniquePtr<RhsT>& Rhs)
+{
+	return Lhs.Get() == Rhs.Get();
+}
+
+/**
+ * Inequality comparison operator
+ *
+ * @param Lhs The first TUniquePtr to compare.
+ * @param Rhs The second TUniquePtr to compare.
+ *
+ * @return false if the two TUniquePtrs are logically substitutable for each other, true otherwise.
+ */
+template <typename LhsT, typename RhsT>
+FORCEINLINE bool operator!=(const TUniquePtr<LhsT>& Lhs, const TUniquePtr<RhsT>& Rhs)
+{
+	return Lhs.Get() != Rhs.Get();
+}
+
 template <typename T> struct TIsZeroConstructType<TUniquePtr<T>> { enum { Value = true }; };
 
 #if PLATFORM_COMPILER_HAS_VARIADIC_TEMPLATES
@@ -218,4 +248,22 @@ template <typename T> struct TIsZeroConstructType<TUniquePtr<T>> { enum { Value
 		return TUniquePtr<T>(new T(Forward<TArg0>(Arg0)));
 	}
 
+	template <typename T, typename TArg0, typename TArg1>
+	FORCEINLINE TUniquePtr<T> MakeUnique(TArg0&& Arg0, TArg1&& Arg1)
+	{
+		return TUniquePtr<T>(new T(Forward<TArg0>(Arg0), Forward<TArg1>(Arg1)));
+	}
+
+	template <typename T, typename TArg0, typename TArg1, typename TArg2>
+	FORCEINLINE TUniquePtr<T> MakeUnique(TArg0&& Arg0, TArg1&& Arg1, TArg2&& Arg2)
+	{
+		return TUniquePtr<T>(new T(Forward<TArg0>(Arg0), Forward<TArg1>(Arg1), Forward<TArg2>(Arg2)));
+	}
+
+	template <typename T, typename TArg0, typename TArg1, typename TArg2, typename TArg3>
+	FORCEINLINE TUniquePtr<T> MakeUnique(TArg0&& Arg0, TArg1&& Arg1, TArg2&& Arg2, TArg3&& Arg3)
+	{
+		return TUniquePtr<T>(new T(Forward<TArg0>(Arg0), Forward<TArg1>(Arg1), Forward<TArg2>(Arg2), Forward<TArg3>(Arg3)));
+	}
+
 #endif
