CIS fix.

diff --git a/Engine/Source/Runtime/Core/Public/Templates/UniquePtr.h b/Engine/Source/Runtime/Core/Public/Templates/UniquePtr.h
--- a/Engine/Source/Runtime/Core/Public/Templates/UniquePtr.h
+++ b/Engine/Source/Runtime/Core/Public/Templates/UniquePtr.h
@@ -193,6 +193,11 @@ FORCEINLINE bool operator==(const TUniquePtr<LhsT>& Lhs, const TUniquePtr<RhsT>&
 {
 	return Lhs.Get() == Rhs.Get();
 }
+template <typename T>
+FORCEINLINE bool operator==(const TUniquePtr<T>& Lhs, const TUniquePtr<T>& Rhs)
+{
+	return Lhs.Get() == Rhs.Get();
+}
 
 /**
  * Inequality comparison operator
@@ -207,6 +212,11 @@ FORCEINLINE bool operator!=(const TUniquePtr<LhsT>& Lhs, const TUniquePtr<RhsT>&
 {
 	return Lhs.Get() != Rhs.Get();
 }
+template <typename T>
+FORCEINLINE bool operator!=(const TUniquePtr<T>& Lhs, const TUniquePtr<T>& Rhs)
+{
+	return Lhs.Get() != Rhs.Get();
+}
 
 template <typename T> struct TIsZeroConstructType<TUniquePtr<T>> { enum { Value = true }; };
 
@@ -266,4 +276,4 @@ template <typename T> struct TIsZeroConstructType<TUniquePtr<T>> { enum { Value
 		return TUniquePtr<T>(new T(Forward<TArg0>(Arg0), Forward<TArg1>(Arg1), Forward<TArg2>(Arg2), Forward<TArg3>(Arg3)));
 	}
 
-#endif
+#endif
\ No newline at end of file
