Fix a couple off-by-one string size problems.

diff --git a/Engine/Source/Runtime/Core/Private/Stats/Stats2.cpp b/Engine/Source/Runtime/Core/Private/Stats/Stats2.cpp
--- a/Engine/Source/Runtime/Core/Private/Stats/Stats2.cpp
+++ b/Engine/Source/Runtime/Core/Private/Stats/Stats2.cpp
@@ -397,20 +397,20 @@ public:
 		FName* Result = &PendingFNames[TStatId::INDEX_FNAME];
 
 		// Get the wide stat description.
-		const int32 StatDescLen = FCString::Strlen( InDescription );
+		const int32 StatDescLen = FCString::Strlen( InDescription ) + 1;
 		// We are leaking this. @see STAT_StatDescMemory
-		WIDECHAR* StatDescWide = new WIDECHAR[StatDescLen + 1];
+		WIDECHAR* StatDescWide = new WIDECHAR[StatDescLen];
 		TCString<WIDECHAR>::Strcpy( StatDescWide, StatDescLen, StringCast<WIDECHAR>( InDescription ).Get() );
 
 		// Get the ansi stat description.
 		// We are leaking this. @see STAT_StatDescMemory
-		ANSICHAR* StatDescAnsi = new ANSICHAR[StatDescLen + 1];
+		ANSICHAR* StatDescAnsi = new ANSICHAR[StatDescLen];
 		TCString<ANSICHAR>::Strcpy( StatDescAnsi, StatDescLen, StringCast<ANSICHAR>( InDescription ).Get() );
 
 		*(UPTRINT*)&PendingFNames[TStatId::INDEX_WIDE_STRING] = (UPTRINT)(UPTRINT*)StatDescWide;
 		*(UPTRINT*)&PendingFNames[TStatId::INDEX_ANSI_STRING] = (UPTRINT)(UPTRINT*)StatDescAnsi;
 
-		MemoryCounter.Add( (StatDescLen + 1)*(sizeof( ANSICHAR ) + sizeof( WIDECHAR )) );
+		MemoryCounter.Add( StatDescLen*(sizeof( ANSICHAR ) + sizeof( WIDECHAR )) );
 		
 		PendingFNames += FNAME_INCREMENT;
 
