More information when TextTest fails on the FormattedTestLayer2.

diff --git a/Engine/Source/Runtime/Core/Private/Tests/Internationalization/TextTest.cpp b/Engine/Source/Runtime/Core/Private/Tests/Internationalization/TextTest.cpp
--- a/Engine/Source/Runtime/Core/Private/Tests/Internationalization/TextTest.cpp
+++ b/Engine/Source/Runtime/Core/Private/Tests/Internationalization/TextTest.cpp
@@ -457,11 +457,13 @@ bool FTextTest::RunTest (const FString& Parameters)
 				// When changes are made to FormattedTestLayer2, please pull out the newly translated LEET string and update the below if-statement to keep the test passing!
 				FString LEETTranslatedString = FormattedTestLayer2.ToString();
 
+				FString DesiredOutput = FString(TEXT("\x2021") TEXT("\xAB") TEXT("\xAB") TEXT("L0r3m") TEXT("\xBB") TEXT(" \"L0r3m 1p$um\" ") TEXT("\xAB") TEXT("Ip$um") TEXT("\xBB") TEXT("\xBB") TEXT(" | ") TEXT("\xAB") TEXT("\xAB") TEXT("L0r3m") TEXT("\xBB") TEXT(" \"L0r3m 1p$um\" ") TEXT("\xAB") TEXT("Ip$um") TEXT("\xBB") TEXT("\xBB") TEXT(" | ") TEXT("\xAB") TEXT("5.5421") TEXT("\xBB") TEXT(" | ") TEXT("\xAB") TEXT("5010.89221") TEXT("\xBB") TEXT(" | ") TEXT("\xAB") TEXT("1944 Month7 14 22:05:06") TEXT("\xBB") TEXT(" | ") TEXT("\xAB") TEXT("92%") TEXT("\xBB") TEXT(" | ") TEXT("\xAB") TEXT("\xA4") TEXT("\xA0") TEXT("100.25") TEXT("\xBB") TEXT("\x2021"));
 				// Convert the baked string into an FText, which will be leetified, then compare it to the rebuilt FText
-				if(LEETTranslatedString != FString(TEXT("\x2021") TEXT("\xAB") TEXT("\xAB") TEXT("L0r3m") TEXT("\xBB") TEXT(" \"L0r3m 1p$um\" ") TEXT("\xAB") TEXT("Ip$um") TEXT("\xBB") TEXT("\xBB") TEXT(" | ") TEXT("\xAB") TEXT("\xAB") TEXT("L0r3m") TEXT("\xBB") TEXT(" \"L0r3m 1p$um\" ") TEXT("\xAB") TEXT("Ip$um") TEXT("\xBB") TEXT("\xBB") TEXT(" | ") TEXT("\xAB") TEXT("5.5421") TEXT("\xBB") TEXT(" | ") TEXT("\xAB") TEXT("5010.89221") TEXT("\xBB") TEXT(" | ") TEXT("\xAB") TEXT("1944 Month7 14 22:05:06") TEXT("\xBB") TEXT(" | ") TEXT("\xAB") TEXT("92%") TEXT("\xBB") TEXT(" | ") TEXT("\xAB") TEXT("\xA4") TEXT("\xA0") TEXT("100.25") TEXT("\xBB") TEXT("\x2021")))
+				if(LEETTranslatedString != DesiredOutput)
 				{
 					AddError( TEXT("FormattedTestLayer2 did not rebuild to correctly in LEET!") );
 					AddError( TEXT("Formatted Output=") + LEETTranslatedString );
+					AddError( TEXT("Desired Output=") + DesiredOutput );
 				}
 			}
 
