Urban->Main

diff --git a/Engine/Shaders/BasePassForForwardShadingCommon.usf b/Engine/Shaders/BasePassForForwardShadingCommon.usf
--- a/Engine/Shaders/BasePassForForwardShadingCommon.usf
+++ b/Engine/Shaders/BasePassForForwardShadingCommon.usf
@@ -14,15 +14,21 @@
 #define PACK_INTERPOLANTS (USE_VERTEX_FOG && NUM_VF_PACKED_INTERPOLANTS > 0 && ES2_PROFILE)
 #define LANDSCAPE_BUG_WORKAROUND (IS_FORWARD_BASEPASS_VERTEX_SHADER && PACK_INTERPOLANTS)
 
+#if METAL_PROFILE 
+// @todo-rco: @todo metal: FIXME!
+#define INTERP_TYPE float4
+#else
+#define INTERP_TYPE half4
+#endif
+
 struct FSharedForwardShadingBasePassInterpolants
 {
 #if USE_VERTEX_FOG && !PACK_INTERPOLANTS
-	half4 VertexFog			: TEXCOORD7;
+	INTERP_TYPE VertexFog			: TEXCOORD7;
 #endif
 #if LANDSCAPE_BUG_WORKAROUND
-	half4 DummyInterp		: DUMMY_INTERP;
+	INTERP_TYPE DummyInterp		: DUMMY_INTERP;
 #endif
-
 	float4 PixelPosition	: TEXCOORD8; // xyz = world position, w = clip z
 
 #if USE_WORLD_POSITION_EXCLUDING_SHADER_OFFSETS
diff --git a/Engine/Shaders/BasePassForForwardShadingPixelShader.usf b/Engine/Shaders/BasePassForForwardShadingPixelShader.usf
--- a/Engine/Shaders/BasePassForForwardShadingPixelShader.usf
+++ b/Engine/Shaders/BasePassForForwardShadingPixelShader.usf
@@ -13,8 +13,6 @@
 #include "BRDF.usf"
 #include "SHCommon.usf"
 
-#define MATERIAL_NONMETAL	0
-
 /** Prenormalized capture of the scene that's closest to the object being rendered. */
 #if !MATERIAL_FULLY_ROUGH
 TextureCube ReflectionCubemap;
diff --git a/Engine/Shaders/BasePassForForwardShadingVertexShader.usf b/Engine/Shaders/BasePassForForwardShadingVertexShader.usf
--- a/Engine/Shaders/BasePassForForwardShadingVertexShader.usf
+++ b/Engine/Shaders/BasePassForForwardShadingVertexShader.usf
@@ -6,6 +6,9 @@
 
 #define IS_FORWARD_BASEPASS_VERTEX_SHADER 1
 
+// @todo urban: branch was: #define FORCE_FLOATS 1
+#define FORCE_FLOATS METAL_PROFILE
+
 #include "Common.usf"
 #include "BasePassForForwardShadingCommon.usf"
 #include "Material.usf"
diff --git a/Engine/Shaders/Common.usf b/Engine/Shaders/Common.usf
--- a/Engine/Shaders/Common.usf
+++ b/Engine/Shaders/Common.usf
@@ -4,12 +4,18 @@
 	Common.usf: Common shader code
 =============================================================================*/
   
+#pragma once
+
 #if PS4_PROFILE  
 	// always #include PS4Common.usf so it can #define override anything in any other included file.
 	#include "PS4/PS4Common.usf"
 #endif 
- 
-#if !COMPILER_GLSL_ES2
+
+#ifndef FORCE_FLOATS
+#define FORCE_FLOATS 0
+#endif
+
+#if (!COMPILER_GLSL_ES2 && !METAL_PROFILE) || FORCE_FLOATS
 	// Always use floats when not using the ES2 compiler, because low precision modifiers are currently only tweaked for ES2, 
 	// And we don't want potential side effects on other platforms
 	#define half float
@@ -71,14 +77,15 @@
 
 #include "Definitions.usf"
 
-#define FEATURE_LEVEL_ES2 1
-#define FEATURE_LEVEL_SM3 2
-#define FEATURE_LEVEL_SM4 3 
-#define FEATURE_LEVEL_SM5 4
+#define FEATURE_LEVEL_ES2	1
+#define FEATURE_LEVEL_ES3_1	2
+#define FEATURE_LEVEL_SM3	3
+#define FEATURE_LEVEL_SM4	4	 
+#define FEATURE_LEVEL_SM5	5
 
 // A8 textures when sampled either have their component in R or A. D3D uses A and that is the default.
 #define A8_SAMPLE_MASK .a
-    
+	
 #if PS4_PROFILE
 
 	#define FEATURE_LEVEL FEATURE_LEVEL_SM5
@@ -108,7 +115,9 @@
 	#if GL4_PROFILE
 		#define class struct
 	#endif
-
+#elif METAL_PROFILE
+	#define FEATURE_LEVEL FEATURE_LEVEL_ES2
+	#define FCOLOR_COMPONENT_SWIZZLE .bgra  
 #elif ES2_PROFILE
 
 	#define FEATURE_LEVEL FEATURE_LEVEL_ES2
@@ -122,7 +131,7 @@
 		// A8 textures when sampled have their component in R.
 		#undef A8_SAMPLE_MASK
 		#define A8_SAMPLE_MASK .r
-#endif
+	#endif
 #else
 
 	#error Add your platform here
@@ -211,8 +220,13 @@ MaterialFloat4 TextureCubeSampleBias(TextureCube Tex, SamplerState Sampler, floa
 	#define TANGENTTOWORLD2					TEXCOORD11
 
 	// TangentToWorld0 is float4 only to match D3D11
+#if METAL_PROFILE
+		//@todo-rco: FIXME!
+		#define TANGENTTOWORLD_INTERPOLATOR_BLOCK	float4 TangentToWorld0 : TANGENTTOWORLD0; float4	TangentToWorld2	: TANGENTTOWORLD2;
+#else
 	#define TANGENTTOWORLD_INTERPOLATOR_BLOCK	MaterialFloat4 TangentToWorld0 : TANGENTTOWORLD0; MaterialFloat4	TangentToWorld2	: TANGENTTOWORLD2;
 #endif
+#endif
 
 MaterialFloat Luminance( MaterialFloat3 LinearColor )
 {
@@ -419,7 +433,7 @@ MaterialFloat3 CalcSceneColor(MaterialFloat2 ScreenUV)
 /** return all channels of the scene lighting texture */
 MaterialFloat4 CalcFullSceneColor(MaterialFloat2 ScreenUV)
 {
-      return Texture2DSample(SceneColorTexture, SceneColorTextureSampler,ScreenUV);
+	  return Texture2DSample(SceneColorTexture, SceneColorTextureSampler,ScreenUV);
 }
 
 /** Encodes HDR linear scene color for storage in the 8 bit light attenuation texture. */
@@ -587,6 +601,8 @@ float CalcSceneDepth( float2 ScreenUV )
 			#else 
 				return 65000.0f; 
 			#endif 
+		#elif METAL_PROFILE
+				return FramebufferFetchES2().w;
 		#else
 			return ConvertFromDeviceZ(Texture2DSampleLevel(SceneDepthTexture, SceneDepthTextureSampler, ScreenUV, 0).r);
 		#endif
diff --git a/Engine/Shaders/Definitions.usf b/Engine/Shaders/Definitions.usf
--- a/Engine/Shaders/Definitions.usf
+++ b/Engine/Shaders/Definitions.usf
@@ -88,6 +88,10 @@
 #define ES2_PROFILE										0
 #endif
 
+#ifndef METAL_PROFILE
+#define METAL_PROFILE									0
+#endif
+
 #ifndef IOS
 #define IOS												0
 #endif
diff --git a/Engine/Shaders/HeightFogCommon.usf b/Engine/Shaders/HeightFogCommon.usf
--- a/Engine/Shaders/HeightFogCommon.usf
+++ b/Engine/Shaders/HeightFogCommon.usf
@@ -14,14 +14,14 @@ half4 InscatteringLightDirection;
 half4 DirectionalInscatteringColor;
 float DirectionalInscatteringStartDistance;
 
-/** FogDensity * exp2(-FogHeightFalloff * (CameraWorldPosition.z - FogHeight)) in x, FogHeightFalloff in y, CosTerminatorAngle in z, StartDistance in w. */
-static const float4 ExponentialFogParameters = SharedFogParameter0;
-/** Color to use. */
-static const half3 ExponentialFogColor = SharedFogParameter1.xyz;
-static const half MinFogOpacity = SharedFogParameter1.w;
-
 half4 GetExponentialHeightFog(float3 WorldPosition, float3 InCameraPosition)
 {
+	/** FogDensity * exp2(-FogHeightFalloff * (CameraWorldPosition.z - FogHeight)) in x, FogHeightFalloff in y, CosTerminatorAngle in z, StartDistance in w. */
+	const float4 ExponentialFogParameters = SharedFogParameter0;
+	/** Color to use. */
+	static const half3 ExponentialFogColor = SharedFogParameter1.xyz;
+	static const half MinFogOpacity = SharedFogParameter1.w;
+
 	float3 CameraToReceiver = WorldPosition - InCameraPosition;
 	float CameraToReceiverLengthSqr = dot( CameraToReceiver, CameraToReceiver );
 	float CameraToReceiverLengthInv = rsqrt( CameraToReceiverLengthSqr );
diff --git a/Engine/Shaders/LocalVertexFactory.usf b/Engine/Shaders/LocalVertexFactory.usf
--- a/Engine/Shaders/LocalVertexFactory.usf
+++ b/Engine/Shaders/LocalVertexFactory.usf
@@ -43,10 +43,17 @@ float3 SplineMeshY;
 struct FVertexFactoryInput
 {
 	float4	Position	: ATTRIBUTE0;
+#if METAL_PROFILE
+		//@todo-rco: FIXME!
+		float3	TangentX	: ATTRIBUTE1;
+		float4	TangentZ	: ATTRIBUTE2;
+		float4	Color		: ATTRIBUTE3;
+#else
 	half3	TangentX	: ATTRIBUTE1;
 	// TangentZ.w contains sign of tangent basis determinant
 	half4	TangentZ	: ATTRIBUTE2;
 	half4	Color		: ATTRIBUTE3;
+#endif
 
 #if NUM_MATERIAL_TEXCOORDS_VERTEX
 
diff --git a/Engine/Shaders/ShaderVersion.usf b/Engine/Shaders/ShaderVersion.usf
--- a/Engine/Shaders/ShaderVersion.usf
+++ b/Engine/Shaders/ShaderVersion.usf
@@ -2,4 +2,4 @@
 // Each time the console command is executed it generates a new GUID. As this file is included
 // in common.usf (which should be included in any shader) it allows to invalidate the shader DDC.
 // 
-// GUID = {0C2B0080-9297-4D12-9954-EFA679F1122A}
+// GUID = {0C2B0080-9297-4D12-9954-EFA779F1122A}
diff --git a/Engine/Shaders/SimpleElementPixelShader.usf b/Engine/Shaders/SimpleElementPixelShader.usf
--- a/Engine/Shaders/SimpleElementPixelShader.usf
+++ b/Engine/Shaders/SimpleElementPixelShader.usf
@@ -71,7 +71,11 @@ void Main(
 #endif
 }
 
+#if METAL_PROFILE
+static const half Gamma = 1.0f;
+#else
 half Gamma;
+#endif
 
 void GammaMain(
 	in float2 TextureCoordinate : TEXCOORD0,
diff --git a/Engine/Shaders/VertexFactoryCommon.usf b/Engine/Shaders/VertexFactoryCommon.usf
--- a/Engine/Shaders/VertexFactoryCommon.usf
+++ b/Engine/Shaders/VertexFactoryCommon.usf
@@ -1,5 +1,7 @@
 // Copyright 1998-2014 Epic Games, Inc. All Rights Reserved.
 
+#include "Common.usf"
+
 float4 TransformLocalToTranslatedWorld(float3 LocalPosition)
 {
 	//expand the matrix multiply to avoid a code-gen bug on PS4.
diff --git a/Engine/Source/Developer/ShaderFormatOpenGL/Private/OpenGLShaderCompiler.cpp b/Engine/Source/Developer/ShaderFormatOpenGL/Private/OpenGLShaderCompiler.cpp
--- a/Engine/Source/Developer/ShaderFormatOpenGL/Private/OpenGLShaderCompiler.cpp
+++ b/Engine/Source/Developer/ShaderFormatOpenGL/Private/OpenGLShaderCompiler.cpp
@@ -328,7 +328,7 @@ GLenum GLFrequencyTable[] =
 };
 
 /** Map shader frequency -> string for messages. */
-const TCHAR* GLFrequencyStringTable[] =
+static const TCHAR* GLFrequencyStringTable[] =
 {
 	TEXT("Vertex"),
 	TEXT("Hull"),
diff --git a/Engine/Source/Developer/TextureFormatPVR/Private/TextureFormatPVR.cpp b/Engine/Source/Developer/TextureFormatPVR/Private/TextureFormatPVR.cpp
--- a/Engine/Source/Developer/TextureFormatPVR/Private/TextureFormatPVR.cpp
+++ b/Engine/Source/Developer/TextureFormatPVR/Private/TextureFormatPVR.cpp
@@ -10,6 +10,11 @@
 #include "PixelFormat.h"
 #include "GenericPlatformProcess.h"
 
+
+#define DEFAULT_QUALITY 0
+#define MAX_QUALITY 4
+
+
 DEFINE_LOG_CATEGORY_STATIC(LogTextureFormatPVR, Log, All);
 
 /**
@@ -222,6 +227,39 @@ static void UseOriginal(const FImage& InImage, FCompressedImage2D& OutCompressed
     FMemory::Memcpy(MipData, Image.RawData.GetData(), Image.SizeX * Image.SizeY * 4);
 }
 
+static FString GetPVRTCQualityString()
+{
+	// start at default quality
+	uint32 CompressionModeValue = DEFAULT_QUALITY;
+	FParse::Value(FCommandLine::Get(), TEXT("-pvrtcquality="), CompressionModeValue);
+	CompressionModeValue = FMath::Min<uint32>(CompressionModeValue, MAX_QUALITY);
+
+	// convert to a string
+	FString CompressionMode;
+	switch (CompressionModeValue)
+	{
+		case 0:	CompressionMode = TEXT("fastest"); break;
+		case 1:	CompressionMode = TEXT("fast"); break;
+		case 2:	CompressionMode = TEXT("normal"); break;
+		case 3:	CompressionMode = TEXT("high"); break;
+		case 4:	CompressionMode = TEXT("best"); break;
+		default: UE_LOG(LogTemp, Fatal, TEXT("Max quality higher than expected"));
+	}
+
+	return CompressionMode;
+}
+
+static uint16 GetPVRTCQualityForVersion()
+{
+	// start at default quality
+	uint32 CompressionModeValue = DEFAULT_QUALITY;
+	FParse::Value(FCommandLine::Get(), TEXT("-pvrtcquality="), CompressionModeValue);
+	CompressionModeValue = FMath::Min<uint32>(CompressionModeValue, MAX_QUALITY);
+
+	// top 3 bits for compression value
+	return CompressionModeValue << 13;
+}
+
 /**
  * PVR texture format handler.
  */
@@ -234,7 +272,7 @@ class FTextureFormatPVR : public ITextureFormat
 
 	virtual uint16 GetVersion(FName Format) const override
 	{
-		return 6;
+		return 7 + GetPVRTCQualityForVersion();
 	}
 
 	virtual void GetSupportedFormats(TArray<FName>& OutFormats) const override
@@ -397,30 +435,20 @@ class FTextureFormatPVR : public ITextureFormat
 		delete PVRFile;
 
 		// Compress PVR file to PVRTC
-#if PLATFORM_MAC
-		// Compress PVR file to PVRTC (using apple's commandline tool)
-//		FString Params = FString::Printf(TEXT("-sdk iphoneos texturetool %s-e PVRTC --bits-per-pixel-%d -o %s%s -f PVR %s%s"),
-//			bGenerateMips ? TEXT("-m ") : TEXT(""),
-//			bIsPVRTC2 ? 2 : 4,
-//			FPlatformProcess::BaseDir(), *OutputFilePath,
-//			FPlatformProcess::BaseDir(), *InputFilePath);
-
-		// Use PowerVR's tool
-		FString Params = FString::Printf(TEXT("-i \"%s\" -o \"%s\" %s -legacypvr -q pvrtcbest -f PVRTC1_%d"),
+        FString CompressionMode = GetPVRTCQualityString();
+
+		// Use PowerVR's new CLI tool commandline
+		FString Params = FString::Printf(TEXT("-i \"%s\" -o \"%s\" %s -legacypvr -q pvrtc%s -f PVRTC1_%d"),
 			*InputFilePath, *OutputFilePath,
 			bGenerateMips ? TEXT("-m") : TEXT(""),
+            *CompressionMode,
 			bIsPVRTC2 ? 2 : 4);
-//		FString CompressorPath(TEXT("/usr/bin/xcrun"));
+#if PLATFORM_MAC
 		FString CompressorPath(FPaths::EngineDir() + TEXT("Binaries/ThirdParty/ImgTec/PVRTexToolCL"));
-		UE_LOG(LogTemp, Log, TEXT("Running texturetool with '%s'"), *Params);
 #else
-		// Use PowerVR's tool
-		FString Params = FString::Printf(TEXT("%s%s -fOGLPVRTC%d -yflip0 -i\"%s\" -o\"%s\""),
-			bGenerateMips ? TEXT("-m ") : TEXT(""),                             
-			TEXT("-pvrtciterations8"),
-			bIsPVRTC2 ? 2 : 4, *InputFilePath, *OutputFilePath);
-		FString CompressorPath(FPaths::EngineDir() + TEXT("Binaries/ThirdParty/ImgTec/PVRTexTool.exe"));
+		FString CompressorPath(FPaths::EngineDir() + TEXT("Binaries/ThirdParty/ImgTec/PVRTexToolCLI.exe"));
 #endif
+		UE_LOG(LogTemp, Log, TEXT("Running texturetool with '%s'"), *Params);
 
 		// Give a debug message about the process
 		if (IsRunningCommandlet())
diff --git a/Engine/Source/Developer/iOS/IOSTargetPlatform/Private/IOSTargetPlatform.cpp b/Engine/Source/Developer/iOS/IOSTargetPlatform/Private/IOSTargetPlatform.cpp
--- a/Engine/Source/Developer/iOS/IOSTargetPlatform/Private/IOSTargetPlatform.cpp
+++ b/Engine/Source/Developer/iOS/IOSTargetPlatform/Private/IOSTargetPlatform.cpp
@@ -205,7 +205,10 @@ bool FIOSTargetPlatform::HandleTicker(float DeltaTime )
 void FIOSTargetPlatform::GetAllPossibleShaderFormats( TArray<FName>& OutFormats ) const
 {
 	static FName NAME_OPENGL_ES2_IOS(TEXT("GLSL_ES2_IOS"));
+	static FName NAME_SF_METAL(TEXT("SF_METAL"));
 	OutFormats.AddUnique(NAME_OPENGL_ES2_IOS);
+	// @todo metal enable
+//	OutFormats.AddUnique(NAME_SF_METAL);
 }
 
 
diff --git a/Engine/Source/Developer/iOS/MetalShaderFormat/Private/MetalShaderFormat.cpp b/Engine/Source/Developer/iOS/MetalShaderFormat/Private/MetalShaderFormat.cpp
--- a/Engine/Source/Developer/iOS/MetalShaderFormat/Private/MetalShaderFormat.cpp
+++ b/Engine/Source/Developer/iOS/MetalShaderFormat/Private/MetalShaderFormat.cpp
@@ -14,7 +14,7 @@ class FMetalShaderFormat : public IShaderFormat
 {
 	enum
 	{
-		HEADER_VERSION = 3,
+		HEADER_VERSION = 4,
 	};
 	struct FVersion
 	{
diff --git a/Engine/Source/Programs/UnrealBuildTool/IOS/IOSToolChain.cs b/Engine/Source/Programs/UnrealBuildTool/IOS/IOSToolChain.cs
--- a/Engine/Source/Programs/UnrealBuildTool/IOS/IOSToolChain.cs
+++ b/Engine/Source/Programs/UnrealBuildTool/IOS/IOSToolChain.cs
@@ -39,6 +39,7 @@ namespace UnrealBuildTool
 		/** Which version of the iOS SDK to target at build time */
 		[XmlConfig]
 		public static string IOSSDKVersion = "latest";
+		public static float IOSSDKVersionFloat = 0.0f;
 
 		/** The architecture(s) to compile */
 		[XmlConfig]
@@ -46,9 +47,6 @@ namespace UnrealBuildTool
 		[XmlConfig]
 		public static string ShippingArchitectures = "armv7";
 
-		// In case the SDK checking fails for some reason, use this version
-		private static string BackupVersion = "6.0";
-
 		/** Which version of the iOS to allow at run time */
 		[XmlConfig]
 		public static string IOSVersion = "6.0";
@@ -150,20 +148,26 @@ namespace UnrealBuildTool
 				catch (Exception Ex)
 				{
 					// on any exception, just use the backup version
-					Log.TraceInformation("Triggered an exception while looking for SDK directory in Xcode.app. Using Backup version...");
+					Log.TraceInformation("Triggered an exception while looking for SDK directory in Xcode.app");
 					Log.TraceInformation("{0}", Ex.ToString());
-					IOSSDKVersion = BackupVersion;
 				}
 
-				if (ExternalExecution.GetRuntimePlatform() != UnrealTargetPlatform.Mac)
+				if (IOSSDKVersion == "latest")
 				{
-					Log.TraceInformation("Compiling with IOS SDK {0} on Mac {1}", IOSSDKVersion, RemoteServerName);
-				}
-				else
-				{
-					Log.TraceInformation("Compiling with IOS SDK {0}", IOSSDKVersion);
+					throw new BuildException("Unable to determine SDK version from Xcode, we cannot continue");
 				}
 			}
+
+			IOSSDKVersionFloat = float.Parse(IOSSDKVersion);
+
+			if (ExternalExecution.GetRuntimePlatform() != UnrealTargetPlatform.Mac)
+			{
+				Log.TraceInformation("Compiling with IOS SDK {0} on Mac {1}", IOSSDKVersionFloat, RemoteServerName);
+			}
+			else
+			{
+				Log.TraceInformation("Compiling with IOS SDK {0}", IOSSDKVersionFloat);
+			}
 		}
 
 		static public void AddStubToManifest(ref FileManifest Manifest, UEBuildBinary Binary)
@@ -214,6 +218,10 @@ namespace UnrealBuildTool
 			Result += " -fno-rtti";
 			Result += " -fvisibility=hidden"; // hides the linker warnings with PhysX
 
+// 			if (CompileEnvironment.Config.TargetConfiguration == CPPTargetConfiguration.Shipping)
+// 			{
+// 				Result += " -flto";
+// 			}
 
 			Result += " -Wall -Werror";
 
diff --git a/Engine/Source/Programs/UnrealBuildTool/IOS/UEBuildIOS.cs b/Engine/Source/Programs/UnrealBuildTool/IOS/UEBuildIOS.cs
--- a/Engine/Source/Programs/UnrealBuildTool/IOS/UEBuildIOS.cs
+++ b/Engine/Source/Programs/UnrealBuildTool/IOS/UEBuildIOS.cs
@@ -125,9 +125,20 @@ namespace UnrealBuildTool
 				InBuildTarget.GlobalCompileEnvironment.Config.Definitions.Add("WITH_SIMULATOR=1");
 			}
 
+			// needs IOS8 for Metal
+			if (IOSToolChain.IOSSDKVersionFloat >= 8.0)
+			{
+				InBuildTarget.GlobalCompileEnvironment.Config.Definitions.Add("HAS_METAL=1");
+				InBuildTarget.ExtraModuleNames.Add("MetalRHI");
+			}
+			else
+			{
+				InBuildTarget.GlobalCompileEnvironment.Config.Definitions.Add("HAS_METAL=0");
+			}
+
             InBuildTarget.GlobalLinkEnvironment.Config.AdditionalFrameworks.Add( new UEBuildFramework( "GameKit" ) );
             InBuildTarget.GlobalLinkEnvironment.Config.AdditionalFrameworks.Add( new UEBuildFramework( "StoreKit" ) );
-		}
+        }
 
 		/**
 		*	Whether the editor should be built for this platform or not
@@ -245,10 +256,10 @@ namespace UnrealBuildTool
                     InModule.AddPlatformSpecificDynamicallyLoadedModule("IOSTargetPlatform");
 				}
 
-				if (bBuildShaderFormats)
-				{
-                    //InModule.AddPlatformSpecificDynamicallyLoadedModule("ShaderFormatIOS");
-				}
+ 				if (bBuildShaderFormats)
+ 				{
+					InModule.AddPlatformSpecificDynamicallyLoadedModule("MetalShaderFormat");
+ 				}
 			}
 		}
 
diff --git a/Engine/Source/Programs/UnrealBuildTool/System/XcodeProjectFileGenerator.cs b/Engine/Source/Programs/UnrealBuildTool/System/XcodeProjectFileGenerator.cs
--- a/Engine/Source/Programs/UnrealBuildTool/System/XcodeProjectFileGenerator.cs
+++ b/Engine/Source/Programs/UnrealBuildTool/System/XcodeProjectFileGenerator.cs
@@ -1430,6 +1430,8 @@ namespace UnrealBuildTool
 			List<XcodeContainerItemProxy> ContainerItemProxies = new List<XcodeContainerItemProxy>();
             List<XcodeFramework> Frameworks = new List<XcodeFramework>();
             Frameworks.Add(new XcodeFramework("OpenGLES.framework", "System/Library/Frameworks/OpenGLES.framework", "SDKROOT"));
+			// @todo metal: putting this into the project will make for VERY slow Metal runtime by default...
+//			Frameworks.Add(new XcodeFramework("Metal.framework", "System/Library/Frameworks/Metal.framework", "SDKROOT"));
 
 			XcodeFramework XCTestFramework = new XcodeFramework("XCTest.framework", "Library/Frameworks/XCTest.framework", "DEVELOPER_DIR");
             Frameworks.Add(XCTestFramework);
diff --git a/Engine/Source/Runtime/Core/Private/HAL/PThreadRunnableThread.h b/Engine/Source/Runtime/Core/Private/HAL/PThreadRunnableThread.h
--- a/Engine/Source/Runtime/Core/Private/HAL/PThreadRunnableThread.h
+++ b/Engine/Source/Runtime/Core/Private/HAL/PThreadRunnableThread.h
@@ -69,19 +69,22 @@ protected:
 			case TPri_AboveNormal: return 25;
 			case TPri_Normal: return 15;
 			case TPri_BelowNormal: return 5;
-			default: UE_LOG(LogHAL, Fatal, TEXT("Unknown Priroty passed to FRunnableThreadPThread::TranslateThreadPriority()"));
+			default: UE_LOG(LogHAL, Fatal, TEXT("Unknown Priority passed to FRunnableThreadPThread::TranslateThreadPriority()"));
 		}
 	}
 
 	virtual void SetThreadPriority(pthread_t InThread, EThreadPriority NewPriority)
 	{
-		// initialize the structure
 		struct sched_param Sched;
 		FMemory::Memzero(&Sched, sizeof(struct sched_param));
-		
+		int32 Policy = SCHED_RR;
+
+		// Read the current policy
+		pthread_getschedparam(InThread, &Policy, &Sched);
+
 		// set the priority appropriately
 		Sched.sched_priority = TranslateThreadPriority(NewPriority);
-		pthread_setschedparam(InThread, SCHED_RR, &Sched);
+		pthread_setschedparam(InThread, Policy, &Sched);
 	}
 
 	/**
diff --git a/Engine/Source/Runtime/Core/Private/IOS/BLView.cpp b/Engine/Source/Runtime/Core/Private/IOS/BLView.cpp
new file mode 100644
--- /dev/null
+++ b/Engine/Source/Runtime/Core/Private/IOS/BLView.cpp
@@ -0,0 +1,239 @@
+//
+//  FMetalView.m
+//  UE4
+//
+//  Created by Apple on 3/31/14.
+//  Copyright (c) 2014 EpicGames. All rights reserved.
+//
+
+#import "BLView.h"
+#include "IOSAppDelegate.h"
+#include "IOS/IOSInputInterface.h"
+
+
+#if HAS_METAL
+
+@implementation FMetalView
+
+
+{
+	CAMetalLayer* RenderingLayer;
+	
+}
+
+
+/**
+ * @return The Layer Class for the window
+ */
++ (Class)layerClass
+{
+	return [CAMetalLayer class];
+}
+
+-(BOOL)becomeFirstResponder
+{
+	return YES;
+}
+
+- (id)initWithFrame:(CGRect)Frame
+{
+	if ((self = [super initWithFrame:Frame]))
+	{
+		// grab the CAALayer created by the nib
+		RenderingLayer = (CAMetalLayer*)self.layer;
+		RenderingLayer.presentsWithTransaction = NO;
+		RenderingLayer.drawsAsynchronously = YES;
+		
+		// set a background color to make sure the layer appears
+		CGFloat components[] = {0.0, 0.0, 0.0, 1};
+		RenderingLayer.backgroundColor = CGColorCreate(CGColorSpaceCreateDeviceRGB(),components);
+	
+		Device = MTLCreateSystemDefaultDevice();
+		assert(Device);
+		
+		// set the device on the rendering layer and provide a pixel format
+		RenderingLayer.device = Device;
+		RenderingLayer.pixelFormat = MTLPixelFormatBGRA8Unorm_sRGB;
+		RenderingLayer.framebufferOnly = NO;
+
+//
+//		// Initialize some variables
+//		SwapCount = 0;
+//		
+		FMemory::Memzero(AllTouches, sizeof(AllTouches));
+		
+		bIsInitialized = false;
+	}
+	return self;
+}
+
+-(id<CAMetalDrawable>)MakeDrawable
+{
+	id<CAMetalDrawable> Drawable = [RenderingLayer newDrawable];
+//	check(Drawable != nil);
+	return Drawable;
+}
+
+
+
+
+
+- (bool)CreateFramebuffer:(bool)bIsForOnDevice
+{
+	if (!bIsInitialized)
+	{
+		
+//		// make sure this is current
+//		[self MakeCurrent];
+//		
+		// @todo-mobile
+		const float NativeScale = [[UIScreen mainScreen] scale];
+		
+		// look up the CVar for the scale factor
+		static IConsoleVariable* CVar = IConsoleManager::Get().FindConsoleVariable(TEXT("r.MobileContentScaleFactor"));
+		float RequestedContentScaleFactor = CVar->GetFloat();
+		
+		// for TV screens, always use scale factor of 1
+		self.contentScaleFactor = bIsForOnDevice ? RequestedContentScaleFactor : 1.0f;
+		UE_LOG(LogIOS, Log, TEXT("Setting contentScaleFactor to %0.4f (optimal = %0.4f)"), self.contentScaleFactor, NativeScale);
+		
+//		if (self.contentScaleFactor == 1.0f || self.contentScaleFactor == 2.0f)
+//		{
+//			UE_LOG(LogIOS,Log,TEXT("Setting layer filter to NEAREST"));
+//			CAEAGLLayer *EaglLayer = (CAEAGLLayer *)self.layer;
+//			EaglLayer.magnificationFilter = kCAFilterNearest;
+//		}
+//		
+//#if USE_DETAILED_IPHONE_MEM_TRACKING
+//		//This value is used to allow the engine to track gl allocated memory see GetIPhoneOpenGLBackBufferSize
+//		UINT singleBufferSize = OnScreenWidth * OnScreenHeight *  4/*rgba8*/;
+//		IPhoneBackBufferMemSize = singleBufferSize *3/*iphone back buffer system is tripple buffered*/;
+//#endif
+//		
+//#if USE_DETAILED_IPHONE_MEM_TRACKING
+//		UE_LOG(LogIOS, Log, TEXT("IPhone Back Buffer Size: %i MB"), (IPhoneBackBufferMemSize/1024)/1024.f);
+//#endif
+	
+		bIsInitialized = true;
+	}    
+	return true;
+}
+
+/**
+ * Returns the unique ID for the given touch
+ */
+-(int32) GetTouchIndex:(UITouch*)Touch
+{
+	
+	// look for existing touch
+	for (int Index = 0; Index < ARRAY_COUNT(AllTouches); Index++)
+	{
+		if (AllTouches[Index] == Touch)
+		{
+			return Index;
+		}
+	}
+	
+	// if we get here, it's a new touch, find a slot
+	for (int Index = 0; Index < ARRAY_COUNT(AllTouches); Index++)
+	{
+		if (AllTouches[Index] == nil)
+		{
+			AllTouches[Index] = Touch;
+			return Index;
+		}
+	}
+	
+	// if we get here, that means we are trying to use more than 5 touches at once, which is an error
+	return -1;
+}
+
+/**
+ * Pass touch events to the input queue for slate to pull off of, and trigger the debug console.
+ *
+ * @param View The view the event happened in
+ * @param Touches Set of touch events from the OS
+ */
+-(void) HandleTouches:(NSSet*)Touches ofType:(TouchType)Type
+{
+	CGFloat Scale = self.contentScaleFactor;
+	
+	TArray<TouchInput> TouchesArray;
+	for (UITouch* Touch in Touches)
+	{
+		// get info from the touch
+		CGPoint Loc = [Touch locationInView:self];
+		CGPoint PrevLoc = [Touch previousLocationInView:self];
+		
+		// convert TOuch pointer to a unique 0 based index
+		int32 TouchIndex = [self GetTouchIndex:Touch];
+		if (TouchIndex < 0)
+		{
+			continue;
+		}
+		
+		// make a new touch event struct
+		TouchInput TouchMessage;
+		TouchMessage.Handle = TouchIndex;
+		TouchMessage.Type = Type;
+		TouchMessage.Position = FVector2D(FMath::Min<float>(self.frame.size.width - 1, Loc.x), FMath::Min<float>(self.frame.size.height - 1, Loc.y)) * Scale;
+		TouchMessage.LastPosition = FVector2D(FMath::Min<float>(self.frame.size.width - 1, PrevLoc.x), FMath::Min<float>(self.frame.size.height - 1, PrevLoc.y)) * Scale;
+		TouchesArray.Add(TouchMessage);
+		
+		// clear out the touch when it ends
+		if (Type == TouchEnded)
+		{
+			AllTouches[TouchIndex] = nil;
+		}
+	}
+	
+	FIOSInputInterface::QueueTouchInput(TouchesArray);
+}
+
+/**
+ * Handle the various touch types from the OS
+ *
+ * @param touches Array of touch events
+ * @param event Event information
+ */
+- (void) touchesBegan:(NSSet*)touches withEvent:(UIEvent*)event
+{
+	NumActiveTouches += touches.count;
+	[self HandleTouches:touches ofType:TouchBegan];
+	
+#if !UE_BUILD_SHIPPING
+#if WITH_SIMULATOR
+	// use 2 on the simulator so that Option-Click will bring up console (option-click is for doing pinch gestures, which we don't care about, atm)
+	if( NumActiveTouches >= 2 )
+#else
+		// If there are 3 active touches, bring up the console
+		if( NumActiveTouches >= 4 )
+#endif
+		{
+			// Route the command to the main iOS thread (all UI must go to the main thread)
+			[[IOSAppDelegate GetDelegate] performSelectorOnMainThread:@selector(ShowConsole) withObject:nil waitUntilDone:NO];
+		}
+#endif
+}
+
+- (void) touchesMoved:(NSSet*)touches withEvent:(UIEvent*)event
+{
+	[self HandleTouches:touches ofType:TouchMoved];
+}
+
+- (void) touchesEnded:(NSSet*)touches withEvent:(UIEvent*)event
+{
+	NumActiveTouches -= touches.count;
+	[self HandleTouches:touches ofType:TouchEnded];
+}
+
+- (void) touchesCancelled:(NSSet*)touches withEvent:(UIEvent*)event
+{
+	NumActiveTouches -= touches.count;
+	[self HandleTouches:touches ofType:TouchEnded];
+}
+
+
+@end
+
+#endif // METAL
diff --git a/Engine/Source/Runtime/Core/Private/IOS/EAGLView.cpp b/Engine/Source/Runtime/Core/Private/IOS/EAGLView.cpp
--- a/Engine/Source/Runtime/Core/Private/IOS/EAGLView.cpp
+++ b/Engine/Source/Runtime/Core/Private/IOS/EAGLView.cpp
@@ -5,6 +5,7 @@
 #include "IOSAppDelegate.h"
 #include "IOS/IOSInputInterface.h"
 
+#include <OpenGLES/ES2/gl.h>
 #import <QuartzCore/QuartzCore.h>
 #import <OpenGLES/EAGLDrawable.h>
 #import <UIKit/UIGeometry.h>
@@ -336,7 +337,7 @@ void GetInputStack(/*TArray<FPointerEvent>& TouchEvents*/)
 	CGRect Frame = [[UIScreen mainScreen] bounds];
 	if (![IOSAppDelegate GetDelegate].bDeviceInPortraitMode)
 	{
-		Swap<float>(Frame.size.width, Frame.size.height);
+		Swap(Frame.size.width, Frame.size.height);
 	}
 
 	self.view = [[UIView alloc] initWithFrame:Frame];
diff --git a/Engine/Source/Runtime/Core/Private/IOS/IOSAppDelegate.cpp b/Engine/Source/Runtime/Core/Private/IOS/IOSAppDelegate.cpp
--- a/Engine/Source/Runtime/Core/Private/IOS/IOSAppDelegate.cpp
+++ b/Engine/Source/Runtime/Core/Private/IOS/IOSAppDelegate.cpp
@@ -69,6 +69,9 @@ void InstallSignalHandlers()
 
 @synthesize Window;
 @synthesize GLView;
+#if HAS_METAL
+@synthesize MetalView;
+#endif
 @synthesize IOSController;
 @synthesize SlateController;
 
@@ -168,10 +171,30 @@ void InstallSignalHandlers()
 	self.bCommandLineReady = true;
 }
 
+-(void)AudioInterrupted:(NSNotification*)notification
+{
+	// pull the type of notification out
+	NSDictionary* Info = [notification userInfo];
+	NSNumber* Type = (NSNumber*)[Info valueForKey:AVAudioSessionInterruptionTypeKey];
+
+	NSLog(@"AUDIO INTERRUPTION NOTIFICATION: %d, (began = %d, ended = %d)", (int)[Type integerValue], (int)AVAudioSessionInterruptionTypeBegan, (int)AVAudioSessionInterruptionTypeEnded);
+	if ([Type integerValue] == AVAudioSessionInterruptionTypeBegan)
+	{
+		FAppEntry::Suspend();
+		[self ToggleAudioSession:false];
+	}
+	else if ([Type integerValue] == AVAudioSessionInterruptionTypeEnded)
+	{
+		[self ToggleAudioSession:true];
+	    FAppEntry::Resume();
+	}
+}
+
 - (void)InitializeAudioSession
 {
 	// get notified about interruptions
-	[[AVAudioSession sharedInstance] setDelegate:self];
+	// @todo ios8: Test this (I can't make a notification happen so far)
+	[[NSNotificationCenter defaultCenter] addObserver:self selector:@selector(AudioInterrupted:) name:AVAudioSessionInterruptionNotification object:[AVAudioSession sharedInstance]];
 	
 	self.bUsingBackgroundMusic = [self IsBackgroundAudioPlaying];
 	if (!self.bUsingBackgroundMusic)
@@ -201,10 +224,9 @@ void InstallSignalHandlers()
 			UE_LOG(LogIOSAudioSession, Error, TEXT("Failed to set audio session category to AVAudioSessionCategoryAmbient!"));
 		}
 		ActiveError = nil;
- 	}
-    self.bAudioActive = true;
+	}
 
-	/* TODO::JTM - Jan 16, 2013 06:22PM - Music player support */
+    self.bAudioActive = true;
 }
 
 - (void)ToggleAudioSession:(bool)bActive
@@ -213,61 +235,61 @@ void InstallSignalHandlers()
 	{
         if (!self.bAudioActive)
         {
-		bool bWasUsingBackgroundMusic = self.bUsingBackgroundMusic;
-		self.bUsingBackgroundMusic = [self IsBackgroundAudioPlaying];
-
-		if (bWasUsingBackgroundMusic != self.bUsingBackgroundMusic)
-		{
-			if (!self.bUsingBackgroundMusic)
+			bool bWasUsingBackgroundMusic = self.bUsingBackgroundMusic;
+			self.bUsingBackgroundMusic = [self IsBackgroundAudioPlaying];
+	
+			if (bWasUsingBackgroundMusic != self.bUsingBackgroundMusic)
 			{
-				NSError* ActiveError = nil;
-				[[AVAudioSession sharedInstance] setActive:YES error:&ActiveError];
-				if (ActiveError)
+				if (!self.bUsingBackgroundMusic)
 				{
-					UE_LOG(LogIOSAudioSession, Error, TEXT("Failed to set audio session as active!"));
+					NSError* ActiveError = nil;
+					[[AVAudioSession sharedInstance] setActive:YES error:&ActiveError];
+					if (ActiveError)
+					{
+						UE_LOG(LogIOSAudioSession, Error, TEXT("Failed to set audio session as active! [Error = %s]"), *FString([ActiveError description]));
+					}
+					ActiveError = nil;
+	
+					[[AVAudioSession sharedInstance] setCategory:AVAudioSessionCategorySoloAmbient error:&ActiveError];
+					if (ActiveError)
+					{
+						UE_LOG(LogIOSAudioSession, Error, TEXT("Failed to set audio session category to AVAudioSessionCategorySoloAmbient! [Error = %s]"), *FString([ActiveError description]));
+					}
+					ActiveError = nil;
+	
+					/* TODO::JTM - Jan 16, 2013 05:05PM - Music player support */
 				}
-				ActiveError = nil;
-
-				[[AVAudioSession sharedInstance] setCategory:AVAudioSessionCategorySoloAmbient error:&ActiveError];
-				if (ActiveError)
+				else
 				{
-					UE_LOG(LogIOSAudioSession, Error, TEXT("Failed to set audio session category to AVAudioSessionCategorySoloAmbient!"));
+					/* TODO::JTM - Jan 16, 2013 05:05PM - Music player support */
+	
+					// Allow iPod music to continue playing in the background
+					NSError* ActiveError = nil;
+					[[AVAudioSession sharedInstance] setCategory:AVAudioSessionCategoryAmbient error:&ActiveError];
+					if (ActiveError)
+					{
+						UE_LOG(LogIOSAudioSession, Error, TEXT("Failed to set audio session category to AVAudioSessionCategoryAmbient! [Error = %s]"), *FString([ActiveError description]));
+					}
+					ActiveError = nil;
 				}
-				ActiveError = nil;
-
-				/* TODO::JTM - Jan 16, 2013 05:05PM - Music player support */
 			}
-			else
+			else if (!self.bUsingBackgroundMusic)
 			{
-				/* TODO::JTM - Jan 16, 2013 05:05PM - Music player support */
-
-				// Allow iPod music to continue playing in the background
 				NSError* ActiveError = nil;
-				[[AVAudioSession sharedInstance] setCategory:AVAudioSessionCategoryAmbient error:&ActiveError];
+				[[AVAudioSession sharedInstance] setActive:YES error:&ActiveError];
 				if (ActiveError)
 				{
-					UE_LOG(LogIOSAudioSession, Error, TEXT("Failed to set audio session category to AVAudioSessionCategoryAmbient!"));
+					UE_LOG(LogIOSAudioSession, Error, TEXT("Failed to set audio session as active! [Error = %s]"), *FString([ActiveError description]));
+				}
+				ActiveError = nil;
+				
+				[[AVAudioSession sharedInstance] setCategory:AVAudioSessionCategorySoloAmbient error:&ActiveError];
+				if (ActiveError)
+				{
+					UE_LOG(LogIOSAudioSession, Error, TEXT("Failed to set audio session category to AVAudioSessionCategorySoloAmbient! [Error = %s]"), *FString([ActiveError description]));
 				}
 				ActiveError = nil;
 			}
-		}
-		else if (!self.bUsingBackgroundMusic)
-		{
-			NSError* ActiveError = nil;
-			[[AVAudioSession sharedInstance] setActive:YES error:&ActiveError];
-			if (ActiveError)
-			{
-				UE_LOG(LogIOSAudioSession, Error, TEXT("Failed to set audio session as active!"));
-			}
-			ActiveError = nil;
-			
-			[[AVAudioSession sharedInstance] setCategory:AVAudioSessionCategorySoloAmbient error:&ActiveError];
-			if (ActiveError)
-			{
-				UE_LOG(LogIOSAudioSession, Error, TEXT("Failed to set audio session category to AVAudioSessionCategorySoloAmbient!"));
-			}
-			ActiveError = nil;
-		}
         }
 	}
 	else if (self.bAudioActive && !self.bUsingBackgroundMusic)
@@ -276,7 +298,7 @@ void InstallSignalHandlers()
 		[[AVAudioSession sharedInstance] setActive:NO error:&ActiveError];
 		if (ActiveError)
 		{
-			UE_LOG(LogIOSAudioSession, Error, TEXT("Failed to set audio session as inactive!"));
+			UE_LOG(LogIOSAudioSession, Error, TEXT("Failed to set audio session as inactive! [Error = %s]"), *FString([ActiveError description]));
 		}
 		ActiveError = nil;
         
@@ -284,7 +306,7 @@ void InstallSignalHandlers()
 		[[AVAudioSession sharedInstance] setCategory:AVAudioSessionCategoryAmbient error:&ActiveError];
 		if (ActiveError)
 		{
-			UE_LOG(LogIOSAudioSession, Error, TEXT("Failed to set audio session category to AVAudioSessionCategoryAmbient!"));
+			UE_LOG(LogIOSAudioSession, Error, TEXT("Failed to set audio session category to AVAudioSessionCategoryAmbient! [Error = %s]"), *FString([ActiveError description]));
 		}
 		ActiveError = nil;
  	}
@@ -309,11 +331,28 @@ void InstallSignalHandlers()
 {
     self.bHasSuspended = !bSuspend;
     self.bIsSuspended = bSuspend;
+
+	if (bSuspend)
+	{
+		FAppEntry::Suspend();
+	}
+	else
+	{
+		FAppEntry::Resume();
+	}
     
-    while(!self.bHasSuspended)
-    {
-        FPlatformProcess::Sleep(0.05f);
-    }
+	// if you run this on the main thread super early, it will deadlock
+	if ((GLView && GLView->bIsInitialized) 
+#if HAS_METAL
+		|| (MetalView && MetalView->bIsInitialized)
+#endif
+		)
+	{
+		while(!self.bHasSuspended)
+		{
+			FPlatformProcess::Sleep(0.05f);
+		}
+	}
 }
 
 - (BOOL)application:(UIApplication *)application willFinishLaunchingWithOptions:(NSDictionary *)launchOptions
@@ -354,7 +393,11 @@ void InstallSignalHandlers()
 
 	self.CommandLineParseTimer = [NSTimer scheduledTimerWithTimeInterval:0.01f target:self selector:@selector(NoUrlCommandLine) userInfo:nil repeats:NO];
 #if !UE_BUILD_SHIPPING
-	self.ConsoleHistoryValues = [[NSMutableArray alloc] init];
+	self.ConsoleHistoryValues = [[NSUserDefaults standardUserDefaults] objectForKey:@"ConsoleHistory"];
+	if (self.ConsoleHistoryValues == nil)
+	{
+		self.ConsoleHistoryValues = [[NSMutableArray alloc] init];
+	}
 	self.ConsoleHistoryValuesIndex = -1;
 #endif
 
@@ -390,20 +433,6 @@ void InstallSignalHandlers()
 	return YES;
 }
 
-- (void)beginInterruption
-{
-    FAppEntry::Suspend();
-
-	[self ToggleAudioSession:false];
- }
-
-- (void)endInterruption
-{
-	[self ToggleAudioSession:true];
-
-    FAppEntry::Resume();
-}
-
 - (void)applicationWillResignActive:(UIApplication *)application
 {
 	/*
diff --git a/Engine/Source/Runtime/Core/Private/IOS/IOSPlatformMisc.cpp b/Engine/Source/Runtime/Core/Private/IOS/IOSPlatformMisc.cpp
--- a/Engine/Source/Runtime/Core/Private/IOS/IOSPlatformMisc.cpp
+++ b/Engine/Source/Runtime/Core/Private/IOS/IOSPlatformMisc.cpp
@@ -271,7 +271,27 @@ void FIOSPlatformMisc::ReleaseAutoreleasePool(void *Pool)
 void* FIOSPlatformMisc::GetHardwareWindow()
 {
 	IOSAppDelegate* AppDelegate = [IOSAppDelegate GetDelegate];
-	return AppDelegate.GLView;
+    return AppDelegate.GLView 
+#if HAS_METAL
+		? AppDelegate.GLView : AppDelegate.MetalView
+#endif
+		;
+}
+
+bool FIOSPlatformMisc::HasPlatformFeature(const TCHAR* FeatureName)
+{
+	if (FCString::Stricmp(FeatureName, TEXT("Metal")) == 0)
+	{
+#if HAS_METAL
+		// @todo metal: Obviously not the right way to detect Metal support, but waiting for the right way...
+		FIOSPlatformMisc::EIOSDevice DeviceType = GetIOSDeviceType();
+		return DeviceType == IOS_IPhone5S || DeviceType == IOS_IPadMini2 || DeviceType == IOS_IPadAir;
+#else
+		return false;
+#endif
+	}
+
+	return FGenericPlatformMisc::HasPlatformFeature(FeatureName);
 }
 
 FIOSPlatformMisc::EIOSDevice FIOSPlatformMisc::GetIOSDeviceType()
diff --git a/Engine/Source/Runtime/Core/Private/IOS/IOSPlatformProcess.cpp b/Engine/Source/Runtime/Core/Private/IOS/IOSPlatformProcess.cpp
--- a/Engine/Source/Runtime/Core/Private/IOS/IOSPlatformProcess.cpp
+++ b/Engine/Source/Runtime/Core/Private/IOS/IOSPlatformProcess.cpp
@@ -9,6 +9,11 @@
 #include "IOSAppDelegate.h"
 #include <mach-o/dyld.h>
 
+// numbers recommended by Apple
+#define GAME_THREAD_PRIORITY 47
+#define RENDER_THREAD_PRIORITY 45
+
+
 const TCHAR* FIOSPlatformProcess::ComputerName()
 {
 	static TCHAR Result[256]=TEXT("");
@@ -65,3 +70,19 @@ void FIOSPlatformProcess::SetRealTimeMode()
 		thread_policy_set(pthread_mach_thread_np(pthread_self()), THREAD_TIME_CONSTRAINT_POLICY, (thread_policy_t)&Policy, THREAD_TIME_CONSTRAINT_POLICY_COUNT);
 	}
 }
+
+void FIOSPlatformProcess::SetupGameOrRenderThread(bool bIsRenderThread)
+{
+	// Set the gamethread priority to very high, slightly above the renderthread
+	struct sched_param Sched;
+	FMemory::Memzero(&Sched, sizeof(struct sched_param));
+
+	// Read the current priority and policy
+	int32 CurrentPolicy = SCHED_RR;
+	pthread_getschedparam(pthread_self(), &CurrentPolicy, &Sched);
+
+	// Set the new priority and policy (apple recommended FIFO for the two main non-working threads)
+	int32 Policy = SCHED_FIFO;
+	Sched.sched_priority = bIsRenderThread ? RENDER_THREAD_PRIORITY : GAME_THREAD_PRIORITY;
+	pthread_setschedparam(pthread_self(), Policy, &Sched);
+}
diff --git a/Engine/Source/Runtime/Core/Private/IOS/IOSPlatformString.mm b/Engine/Source/Runtime/Core/Private/IOS/IOSPlatformString.mm
--- a/Engine/Source/Runtime/Core/Private/IOS/IOSPlatformString.mm
+++ b/Engine/Source/Runtime/Core/Private/IOS/IOSPlatformString.mm
@@ -7,9 +7,15 @@
 
 @implementation NSString (FString_Extensions)
 
++ (NSString*) stringWithTCHARString:(const TCHAR*)MyTCHARString
+{
+	return [NSString stringWithCString:TCHAR_TO_UTF8(MyTCHARString) encoding:NSUTF8StringEncoding];
+}
+
 + (NSString*) stringWithFString:(const FString&)InFString
 {
-	return [NSString stringWithCString:TCHAR_TO_UTF8(*InFString) encoding:NSUTF8StringEncoding];
+	return [NSString stringWithTCHARString:*InFString];
 }
 
+
 @end
diff --git a/Engine/Source/Runtime/Core/Private/IOS/IOSWindow.cpp b/Engine/Source/Runtime/Core/Private/IOS/IOSWindow.cpp
--- a/Engine/Source/Runtime/Core/Private/IOS/IOSWindow.cpp
+++ b/Engine/Source/Runtime/Core/Private/IOS/IOSWindow.cpp
@@ -4,6 +4,7 @@
 #include "IOSWindow.h"
 #include "IOSAppDelegate.h"
 #include "EAGLView.h"
+#include "BLView.h"
 
 FIOSWindow::~FIOSWindow()
 {
@@ -46,8 +47,13 @@ FPlatformRect FIOSWindow::GetScreenRect()
 {
 	// get the main view's frame
 	IOSAppDelegate* AppDelegate = (IOSAppDelegate*)[[UIApplication sharedApplication] delegate];
-	CGRect Frame = [AppDelegate.GLView frame];
-	CGFloat Scale = AppDelegate.GLView.contentScaleFactor;
+	UIView* View = AppDelegate.GLView
+#if HAS_METAL
+		? AppDelegate.GLView : AppDelegate.MetalView
+#endif
+		;
+	CGRect Frame = [View frame];
+	CGFloat Scale = View.contentScaleFactor;
 
 	FPlatformRect ScreenRect;
 	ScreenRect.Top = Frame.origin.y * Scale;
diff --git a/Engine/Source/Runtime/Core/Public/GenericPlatform/GenericPlatformMisc.h b/Engine/Source/Runtime/Core/Public/GenericPlatform/GenericPlatformMisc.h
--- a/Engine/Source/Runtime/Core/Public/GenericPlatform/GenericPlatformMisc.h
+++ b/Engine/Source/Runtime/Core/Public/GenericPlatform/GenericPlatformMisc.h
@@ -744,6 +744,14 @@ public:
 		return PLATFORM_HAS_TOUCH_MAIN_SCREEN;
 	}
 
+	/**
+	 * Returns whether the given platform feature is currently available (for instance, Metal is only available in IOS8 and with A7 devices)
+	 */
+	static bool HasPlatformFeature(const TCHAR* FeatureName)
+	{
+		return false;
+	}
+
 #if !UE_BUILD_SHIPPING
 protected:
 	/** Whether the user should be prompted to allow for a remote debugger to be attached */
diff --git a/Engine/Source/Runtime/Core/Public/GenericPlatform/GenericPlatformProcess.h b/Engine/Source/Runtime/Core/Public/GenericPlatform/GenericPlatformProcess.h
--- a/Engine/Source/Runtime/Core/Public/GenericPlatform/GenericPlatformProcess.h
+++ b/Engine/Source/Runtime/Core/Public/GenericPlatform/GenericPlatformProcess.h
@@ -206,6 +206,13 @@ struct CORE_API FGenericPlatformProcess
 	 * @param AffinityMask A bitfield indicating what processors the thread is allowed to run on
 	 */
 	static void SetThreadAffinityMask( uint64 AffinityMask );
+
+	/**
+	 * Allow the platform to do anything it needs for game or render thread (this would use the NamedThreads enum if we could forward declare the enum)
+	 */
+	static void SetupGameOrRenderThread(bool bIsRenderThread)
+	{
+	}
 	
 	/** Get startup directory.  NOTE: Only one return value is valid at a time! **/
 	static const TCHAR* BaseDir();
diff --git a/Engine/Source/Runtime/Core/Public/IOS/BLView.h b/Engine/Source/Runtime/Core/Public/IOS/BLView.h
new file mode 100644
--- /dev/null
+++ b/Engine/Source/Runtime/Core/Public/IOS/BLView.h
@@ -0,0 +1,36 @@
+//
+//  FMTLView.h
+//  UE4
+//
+//  Created by Apple on 3/31/14.
+//  Copyright (c) 2014 EpicGames. All rights reserved.
+//
+
+#import <UIKit/UIKit.h>
+
+#if HAS_METAL
+#import <Metal/Metal.h>
+#import <QuartzCore/CAMetalLayer.h>
+
+@interface FMetalView : UIView
+{
+@public
+	bool bIsInitialized;
+
+	id<MTLDevice> Device;
+
+	// keeps track of the number of active touches
+	// used to bring up the three finger touch debug console after 3 active touches are registered
+	int NumActiveTouches;
+
+	// track the touches by pointer (which will stay the same for a given finger down) - note we don't deref the pointers in this array
+	UITouch* AllTouches[10];
+}
+
+- (bool)CreateFramebuffer:(bool)bIsForOnDevice;
+
+- (id<CAMetalDrawable>)MakeDrawable;
+
+@end
+
+#endif
diff --git a/Engine/Source/Runtime/Core/Public/IOS/IOSAppDelegate.h b/Engine/Source/Runtime/Core/Public/IOS/IOSAppDelegate.h
--- a/Engine/Source/Runtime/Core/Public/IOS/IOSAppDelegate.h
+++ b/Engine/Source/Runtime/Core/Public/IOS/IOSAppDelegate.h
@@ -7,6 +7,7 @@
 #import <GameKit/GKGameCenterViewController.h>
 
 @class EAGLView;
+@class FMetalView;
 @class IOSViewController;
 @class SlateOpenGLESViewController;
 @class IOSAppDelegate;
@@ -30,13 +31,18 @@ namespace FAppEntry
 	UIGestureRecognizerDelegate,
 #endif
 	GKGameCenterControllerDelegate,
-UITextFieldDelegate, AVAudioSessionDelegate>
+UITextFieldDelegate>
 
 /** Window object */
 @property (strong, retain, nonatomic) UIWindow *Window;
 
 /** Main GL View */
-@property (retain) EAGLView *GLView;
+@property (retain) EAGLView* GLView;
+
+#if HAS_METAL
+/** Main GL View */
+@property (retain) FMetalView* MetalView;
+#endif
 
 /** The controller to handle rotation of the view */
 @property (retain) IOSViewController* IOSController;
diff --git a/Engine/Source/Runtime/Core/Public/IOS/IOSPlatformMisc.h b/Engine/Source/Runtime/Core/Public/IOS/IOSPlatformMisc.h
--- a/Engine/Source/Runtime/Core/Public/IOS/IOSPlatformMisc.h
+++ b/Engine/Source/Runtime/Core/Public/IOS/IOSPlatformMisc.h
@@ -40,7 +40,9 @@ struct CORE_API FIOSPlatformMisc : public FGenericPlatformMisc
 			//If you want an immediate break use the trap instruction, continued execuction is halted
 #if WITH_SIMULATOR
             __asm__ ( "int $3" );
-#else   
+#elif PLATFORM_64BITS
+			__asm__ ( "svc 0" );
+#else
             __asm__ ( "trap" );
 #endif
 		}
@@ -60,7 +62,7 @@ struct CORE_API FIOSPlatformMisc : public FGenericPlatformMisc
 	static int32 NumberOfCores();
 	static void LoadPreInitModules();
 	static void SetMemoryWarningHandler(void (* Handler)(const FGenericMemoryWarningContext & Context));
-
+	static bool HasPlatformFeature(const TCHAR* FeatureName);
 
 	//////// Platform specific
 	static void* CreateAutoreleasePool();
diff --git a/Engine/Source/Runtime/Core/Public/IOS/IOSPlatformString.h b/Engine/Source/Runtime/Core/Public/IOS/IOSPlatformString.h
--- a/Engine/Source/Runtime/Core/Public/IOS/IOSPlatformString.h
+++ b/Engine/Source/Runtime/Core/Public/IOS/IOSPlatformString.h
@@ -12,10 +12,14 @@ class FString;
 
 @interface NSString (FString_Extensions)
 
+/**
+* Converts an TCHAR string to an NSString
+*/
++ (NSString*) stringWithTCHARString:(const TCHAR*)MyTCHARString;
+
 /**
  * Converts an FString to an NSString
  */
 + (NSString*) stringWithFString:(const FString&)MyFString;
 
 @end
-
diff --git a/Engine/Source/Runtime/Core/Public/Math/UnrealMathNeon.h b/Engine/Source/Runtime/Core/Public/Math/UnrealMathNeon.h
--- a/Engine/Source/Runtime/Core/Public/Math/UnrealMathNeon.h
+++ b/Engine/Source/Runtime/Core/Public/Math/UnrealMathNeon.h
@@ -850,14 +850,36 @@ FORCEINLINE VectorRegister VectorLoadByte4Reverse( const void * Ptr )
  */
 FORCEINLINE void VectorStoreByte4( VectorRegister Vec, void * Ptr )
 {
-    uint16x8_t u16x8 = (uint16x8_t)vcvtq_u32_f32( VectorMin( Vec, Vec255 ) );
-    uint8x8_t u8x8 = (uint8x8_t)vget_low_u16( vuzpq_u16( u16x8, u16x8 ).val[0] );
-    u8x8 = vuzp_u8( u8x8, u8x8 ).val[0];
+	uint16x8_t u16x8 = (uint16x8_t)vcvtq_u32_f32( VectorMin( Vec, Vec255 ) );
+	uint8x8_t u8x8 = (uint8x8_t)vget_low_u16( vuzpq_u16( u16x8, u16x8 ).val[0] );
+	u8x8 = vuzp_u8( u8x8, u8x8 ).val[0];
 	uint32_t buf[2];
 	vst1_u8( (uint8_t *)buf, u8x8 );
 	*(uint32_t *)Ptr = buf[0]; 
 }
 
+/**
+ * Converts the 4 floats in the vector to 4 fp16 and stores based off bool to [un]aligned memory.
+ *
+ * @param Vec			Vector containing 4 floats
+ * @param Ptr			Memory pointer to store the 4 fp16's.
+ */
+template <bool bAligned>
+FORCEINLINE void VectorStoreHalf4(VectorRegister Vec, void* RESTRICT Ptr)
+{
+	float16x4_t f16x4 = (float16x4_t)vcvt_f16_f32(Vec);
+	if (bAligned)
+	{
+		vst1_u8( (uint8_t *)Ptr, f16x4 );
+	}
+	else
+	{
+		uint32_t buf[2];
+		vst1_u8( (uint8_t *)buf, f16x4 );
+		*(uint32_t *)Ptr = buf[0]; 
+	}
+}
+
 /**
  * Returns non-zero if any element in Vec1 is greater than the corresponding element in Vec2, otherwise 0.
  *
@@ -867,9 +889,9 @@ FORCEINLINE void VectorStoreByte4( VectorRegister Vec, void * Ptr )
  */
 FORCEINLINE int32 VectorAnyGreaterThan( VectorRegister Vec1, VectorRegister Vec2 )
 {
-    uint16x8_t u16x8 = (uint16x8_t)vcgtq_f32( Vec1, Vec2 );
-    uint8x8_t u8x8 = (uint8x8_t)vget_low_u16( vuzpq_u16( u16x8, u16x8 ).val[0] );
-    u8x8 = vuzp_u8( u8x8, u8x8 ).val[0];
+	uint16x8_t u16x8 = (uint16x8_t)vcgtq_f32( Vec1, Vec2 );
+	uint8x8_t u8x8 = (uint8x8_t)vget_low_u16( vuzpq_u16( u16x8, u16x8 ).val[0] );
+	u8x8 = vuzp_u8( u8x8, u8x8 ).val[0];
 	uint32_t buf[2];
 	vst1_u8( (uint8_t *)buf, u8x8 );
 	return (int32)buf[0]; // each byte of output corresponds to a component comparison
diff --git a/Engine/Source/Runtime/Core/Public/Stats/Stats2.h b/Engine/Source/Runtime/Core/Public/Stats/Stats2.h
--- a/Engine/Source/Runtime/Core/Public/Stats/Stats2.h
+++ b/Engine/Source/Runtime/Core/Public/Stats/Stats2.h
@@ -1640,6 +1640,7 @@ DECLARE_STATS_GROUP(TEXT("Memory Allocator"),STATGROUP_MemoryAllocator, STATCAT_
 DECLARE_STATS_GROUP(TEXT("Memory Platform"),STATGROUP_MemoryPlatform, STATCAT_Advanced);
 DECLARE_STATS_GROUP(TEXT("Memory StaticMesh"),STATGROUP_MemoryStaticMesh, STATCAT_Advanced);
 DECLARE_STATS_GROUP(TEXT("Mesh Particles"),STATGROUP_MeshParticles, STATCAT_Advanced);
+DECLARE_STATS_GROUP(TEXT("Metal"),STATGROUP_MetalRHI, STATCAT_Advanced);
 DECLARE_STATS_GROUP(TEXT("Morph"),STATGROUP_MorphTarget, STATCAT_Advanced);
 DECLARE_STATS_GROUP(TEXT("Navigation"),STATGROUP_Navigation, STATCAT_Advanced);
 DECLARE_STATS_GROUP(TEXT("Net"),STATGROUP_Net, STATCAT_Advanced);
diff --git a/Engine/Source/Runtime/Engine/Private/Materials/MaterialShared.cpp b/Engine/Source/Runtime/Engine/Private/Materials/MaterialShared.cpp
--- a/Engine/Source/Runtime/Engine/Private/Materials/MaterialShared.cpp
+++ b/Engine/Source/Runtime/Engine/Private/Materials/MaterialShared.cpp
@@ -770,6 +770,11 @@ bool FMaterialResource::IsFullyRough() const
 	return Material->bFullyRough;
 }
 
+bool FMaterialResource::IsNonmetal() const
+{
+	return !Material->Metallic.IsConnected() && !Material->Specular.IsConnected();
+}
+
 bool FMaterialResource::UseLmDirectionality() const
 {
 	return Material->bUseLightmapDirectionality;
@@ -1150,6 +1155,7 @@ void FMaterial::SetupMaterialEnvironment(
 	OutEnvironment.SetDefine(TEXT("GENERATE_SPHERICAL_PARTICLE_NORMALS"),ShouldGenerateSphericalParticleNormals() ? TEXT("1") : TEXT("0"));
 	OutEnvironment.SetDefine(TEXT("MATERIAL_USES_SCENE_COLOR"), UsesSceneColor() ? TEXT("1") : TEXT("0"));
 	OutEnvironment.SetDefine(TEXT("MATERIAL_FULLY_ROUGH"), IsFullyRough() ? TEXT("1") : TEXT("0"));
+	OutEnvironment.SetDefine(TEXT("MATERIAL_NONMETAL"), IsNonmetal() ? TEXT("1") : TEXT("0"));
 	OutEnvironment.SetDefine(TEXT("MATERIAL_USE_LM_DIRECTIONALITY"), UseLmDirectionality() ? TEXT("1") : TEXT("0"));
 	OutEnvironment.SetDefine(TEXT("MATERIAL_INJECT_EMISSIVE_INTO_LPV"), ShouldInjectEmissiveIntoLPV() ? TEXT("1") : TEXT("0"));
 
diff --git a/Engine/Source/Runtime/Engine/Private/MeshParticleVertexFactory.cpp b/Engine/Source/Runtime/Engine/Private/MeshParticleVertexFactory.cpp
--- a/Engine/Source/Runtime/Engine/Private/MeshParticleVertexFactory.cpp
+++ b/Engine/Source/Runtime/Engine/Private/MeshParticleVertexFactory.cpp
@@ -89,15 +89,16 @@ void FMeshParticleVertexFactory::InitRHI()
 				VertexStream.Offset = 0;
 				Streams.Add(VertexStream);
 	
-				Elements.Add(FVertexElement(0, Data.TransformComponent[0].Offset, Data.TransformComponent[0].Type, 8, Data.TransformComponent[0].bUseInstanceIndex));
-				Elements.Add(FVertexElement(0, Data.TransformComponent[1].Offset, Data.TransformComponent[1].Type, 9, Data.TransformComponent[1].bUseInstanceIndex));
-				Elements.Add(FVertexElement(0, Data.TransformComponent[2].Offset, Data.TransformComponent[2].Type, 10, Data.TransformComponent[2].bUseInstanceIndex));
+				// @todo metal: this will need a valid stride when we get to instanced meshes!
+				Elements.Add(FVertexElement(0, Data.TransformComponent[0].Offset, Data.TransformComponent[0].Type, 8, 0xFFFF, Data.TransformComponent[0].bUseInstanceIndex));
+				Elements.Add(FVertexElement(0, Data.TransformComponent[1].Offset, Data.TransformComponent[1].Type, 9, 0xFFFF, Data.TransformComponent[1].bUseInstanceIndex));
+				Elements.Add(FVertexElement(0, Data.TransformComponent[2].Offset, Data.TransformComponent[2].Type, 10, 0xFFFF, Data.TransformComponent[2].bUseInstanceIndex));
 	
-				Elements.Add(FVertexElement(0, Data.SubUVs.Offset, Data.SubUVs.Type, 11, Data.SubUVs.bUseInstanceIndex));
-				Elements.Add(FVertexElement(0, Data.SubUVLerpAndRelTime.Offset, Data.SubUVLerpAndRelTime.Type, 12, Data.SubUVLerpAndRelTime.bUseInstanceIndex));
+				Elements.Add(FVertexElement(0, Data.SubUVs.Offset, Data.SubUVs.Type, 11, 0, Data.SubUVs.bUseInstanceIndex));
+				Elements.Add(FVertexElement(0, Data.SubUVLerpAndRelTime.Offset, Data.SubUVLerpAndRelTime.Type, 12, 0xFFFF, Data.SubUVLerpAndRelTime.bUseInstanceIndex));
 	
-				Elements.Add(FVertexElement(0, Data.ParticleColorComponent.Offset, Data.ParticleColorComponent.Type, 14, Data.ParticleColorComponent.bUseInstanceIndex));
-				Elements.Add(FVertexElement(0, Data.VelocityComponent.Offset, Data.VelocityComponent.Type, 15, Data.VelocityComponent.bUseInstanceIndex));
+				Elements.Add(FVertexElement(0, Data.ParticleColorComponent.Offset, Data.ParticleColorComponent.Type, 14, 0xFFFF, Data.ParticleColorComponent.bUseInstanceIndex));
+				Elements.Add(FVertexElement(0, Data.VelocityComponent.Offset, Data.VelocityComponent.Type, 15, 0xFFFF, Data.VelocityComponent.bUseInstanceIndex));
 			}
 
 			// Stream 1 - Dynamic parameter
@@ -108,7 +109,7 @@ void FMeshParticleVertexFactory::InitRHI()
 				VertexStream.Offset = 0;
 				Streams.Add(VertexStream);
 	
-				Elements.Add(FVertexElement(1, 0, VET_Float4, 13, true));
+				Elements.Add(FVertexElement(1, 0, VET_Float4, 13, 0xFFFF, true));
 			}
 		}
 
diff --git a/Engine/Source/Runtime/Engine/Private/Particles/ParticleBeamTrailVertexFactory.cpp b/Engine/Source/Runtime/Engine/Private/Particles/ParticleBeamTrailVertexFactory.cpp
--- a/Engine/Source/Runtime/Engine/Private/Particles/ParticleBeamTrailVertexFactory.cpp
+++ b/Engine/Source/Runtime/Engine/Private/Particles/ParticleBeamTrailVertexFactory.cpp
@@ -45,24 +45,25 @@ public:
 
 	virtual void FillDeclElements(FVertexDeclarationElementList& Elements, int32& Offset)
 	{
+		uint16 Stride = sizeof(FParticleBeamTrailVertex);
 		/** The stream to read the vertex position from. */
-		Elements.Add(FVertexElement(0,Offset,VET_Float4,0));
+		Elements.Add(FVertexElement(0, Offset, VET_Float4, 0, Stride));
 		Offset += sizeof(float) * 4;
 		/** The stream to read the vertex old position from. */
-		Elements.Add(FVertexElement(0,Offset,VET_Float3,1));
+		Elements.Add(FVertexElement(0, Offset, VET_Float3, 1, Stride));
 		Offset += sizeof(float) * 4;
 		/** The stream to read the vertex size/rot/subimage from. */
-		Elements.Add(FVertexElement(0,Offset,VET_Float4,2));
+		Elements.Add(FVertexElement(0, Offset, VET_Float4, 2, Stride));
 		Offset += sizeof(float) * 4;
 		/** The stream to read the color from.					*/
-		Elements.Add(FVertexElement(0,Offset,VET_Float4,4));
+		Elements.Add(FVertexElement(0, Offset, VET_Float4, 4, Stride));
 		Offset += sizeof(float) * 4;
 		/** The stream to read the texture coordinates from.	*/
-		Elements.Add(FVertexElement(0,Offset,VET_Float4,3));
+		Elements.Add(FVertexElement(0, Offset, VET_Float4, 3, Stride));
 		Offset += sizeof(float) * 4;
 		
 		/** Dynamic parameters come from a second stream */
-		Elements.Add(FVertexElement(1,0,VET_Float4,5));
+		Elements.Add(FVertexElement(1, 0, VET_Float4, 5, sizeof(FVector4)));
 	}
 
 	virtual void InitDynamicRHI()
diff --git a/Engine/Source/Runtime/Engine/Private/Particles/ParticleCurveTexture.cpp b/Engine/Source/Runtime/Engine/Private/Particles/ParticleCurveTexture.cpp
--- a/Engine/Source/Runtime/Engine/Private/Particles/ParticleCurveTexture.cpp
+++ b/Engine/Source/Runtime/Engine/Private/Particles/ParticleCurveTexture.cpp
@@ -142,7 +142,7 @@ public:
 		{
 			int32 Offset = 0;
 			// Sample.
-			Elements.Add( FVertexElement( 0, Offset, VET_Color, 0, /*bUseInstanceIndex=*/ true ) );
+			Elements.Add(FVertexElement(0, Offset, VET_Color, 0, sizeof(FColor), /*bUseInstanceIndex=*/ true));
 			Offset += sizeof(FColor);
 		}
 
@@ -150,7 +150,7 @@ public:
 		{
 			int32 Offset = 0;
 			// TexCoord.
-			Elements.Add( FVertexElement( 1, Offset, VET_Float2, 1, /*bUseInstanceIndex=*/ false ) );
+			Elements.Add(FVertexElement(1, Offset, VET_Float2, 1, sizeof(FVector2D), /*bUseInstanceIndex=*/ false));
 			Offset += sizeof(FVector2D);
 		}
 
diff --git a/Engine/Source/Runtime/Engine/Private/Particles/ParticleGpuSimulation.cpp b/Engine/Source/Runtime/Engine/Private/Particles/ParticleGpuSimulation.cpp
--- a/Engine/Source/Runtime/Engine/Private/Particles/ParticleGpuSimulation.cpp
+++ b/Engine/Source/Runtime/Engine/Private/Particles/ParticleGpuSimulation.cpp
@@ -541,7 +541,7 @@ public:
 		FVertexDeclarationElementList Elements;
 
 		/** The stream to read the texture coordinates from. */
-		Elements.Add(FVertexElement(0,0,VET_Float2,0,false));
+		Elements.Add(FVertexElement(0, 0, VET_Float2, 0, sizeof(FVector2D), false));
 
 		VertexDeclarationRHI = RHICreateVertexDeclaration(Elements);
 	}
@@ -1094,7 +1094,7 @@ public:
 	{
 		FVertexDeclarationElementList Elements;
 		// TexCoord.
-		Elements.Add( FVertexElement( 0, 0, VET_Float2, 0, /*bUseInstanceIndex=*/ false ) );
+		Elements.Add(FVertexElement(0, 0, VET_Float2, 0, sizeof(FVector2D), /*bUseInstanceIndex=*/ false));
 		VertexDeclarationRHI = RHICreateVertexDeclaration( Elements );
 	}
 
@@ -1465,20 +1465,21 @@ public:
 		// Stream 0.
 		{
 			int32 Offset = 0;
+			uint16 Stride = sizeof(FNewParticle);
 			// InitialPosition.
-			Elements.Add( FVertexElement( 0, Offset, VET_Float4, 0, /*bUseInstanceIndex=*/ true ) );
+			Elements.Add(FVertexElement(0, Offset, VET_Float4, 0, Stride, /*bUseInstanceIndex=*/ true));
 			Offset += sizeof(FVector4);
 			// InitialVelocity.
-			Elements.Add( FVertexElement( 0, Offset, VET_Float4, 1, /*bUseInstanceIndex=*/ true ) );
+			Elements.Add(FVertexElement(0, Offset, VET_Float4, 1, Stride, /*bUseInstanceIndex=*/ true));
 			Offset += sizeof(FVector4);
 			// RenderAttributes.
-			Elements.Add( FVertexElement( 0, Offset, VET_Float4, 2, /*bUseInstanceIndex=*/ true ) );
+			Elements.Add(FVertexElement(0, Offset, VET_Float4, 2, Stride, /*bUseInstanceIndex=*/ true));
 			Offset += sizeof(FVector4);
 			// SimulationAttributes.
-			Elements.Add( FVertexElement( 0, Offset, VET_Float4, 3, /*bUseInstanceIndex=*/ true ) );
+			Elements.Add(FVertexElement(0, Offset, VET_Float4, 3, Stride, /*bUseInstanceIndex=*/ true));
 			Offset += sizeof(FVector4);
 			// ParticleIndex.
-			Elements.Add( FVertexElement( 0, Offset, VET_Float2, 4, /*bUseInstanceIndex=*/ true ) );
+			Elements.Add(FVertexElement(0, Offset, VET_Float2, 4, Stride, /*bUseInstanceIndex=*/ true));
 			Offset += sizeof(FVector2D);
 		}
 
@@ -1486,7 +1487,7 @@ public:
 		{
 			int32 Offset = 0;
 			// TexCoord.
-			Elements.Add( FVertexElement( 1, Offset, VET_Float2, 5, /*bUseInstanceIndex=*/ false ) );
+			Elements.Add(FVertexElement(1, Offset, VET_Float2, 5, sizeof(FVector2D), /*bUseInstanceIndex=*/ false));
 			Offset += sizeof(FVector2D);
 		}
 
@@ -1713,7 +1714,7 @@ public:
 	virtual void InitRHI()
 	{
 		FVertexDeclarationElementList Elements;
-		Elements.Add( FVertexElement( 0, 0, VET_Float2, 0 ) );
+		Elements.Add(FVertexElement(0, 0, VET_Float2, 0, sizeof(FVector2D)));
 		VertexDeclarationRHI = RHICreateVertexDeclaration( Elements );
 	}
 
diff --git a/Engine/Source/Runtime/Engine/Private/Particles/ParticleVertexFactory.cpp b/Engine/Source/Runtime/Engine/Private/Particles/ParticleVertexFactory.cpp
--- a/Engine/Source/Runtime/Engine/Private/Particles/ParticleVertexFactory.cpp
+++ b/Engine/Source/Runtime/Engine/Private/Particles/ParticleVertexFactory.cpp
@@ -70,32 +70,40 @@ public:
 		// Likely need to switch to two decls when enabling runtime ES2 preview.
 		const bool bInstanced = GRHIFeatureLevel >= ERHIFeatureLevel::SM3;
 
+		uint32 InitialStride = sizeof(float) * 2;
+		uint32 PerParticleStride = (sizeof(float) * 4) * 4;
+
 		/** The stream to read the texture coordinates from. */
 		check( Offset == 0 );
-		Elements.Add(FVertexElement(0,Offset,VET_Float2,4,false));
+		uint32 Stride = bInstanced ? InitialStride : InitialStride + PerParticleStride;
+		Elements.Add(FVertexElement(0, Offset, VET_Float2, 4, Stride, false));
 		Offset += sizeof(float) * 2;
 
 		/** The per-particle streams follow. */
 		if(bInstanced) 
 		{
 			Offset = 0;
+			// update stride
+			Stride = PerParticleStride;
 		}
+
 		/** The stream to read the vertex position from. */
-		Elements.Add(FVertexElement(bInstanced ? 1 : 0,Offset,VET_Float4,0,bInstanced));
+		Elements.Add(FVertexElement(bInstanced ? 1 : 0, Offset, VET_Float4, 0, Stride, bInstanced));
 		Offset += sizeof(float) * 4;
 		/** The stream to read the vertex old position from. */
-		Elements.Add(FVertexElement(bInstanced ? 1 : 0,Offset,VET_Float4,1,bInstanced));
+		Elements.Add(FVertexElement(bInstanced ? 1 : 0, Offset, VET_Float4, 1, Stride, bInstanced));
 		Offset += sizeof(float) * 4;
 		/** The stream to read the vertex size/rot/subimage from. */
-		Elements.Add(FVertexElement(bInstanced ? 1 : 0,Offset,VET_Float4,2,bInstanced));
+		Elements.Add(FVertexElement(bInstanced ? 1 : 0, Offset, VET_Float4, 2, Stride, bInstanced));
 		Offset += sizeof(float) * 4;
 		/** The stream to read the color from.					*/
-		Elements.Add(FVertexElement(bInstanced ? 1 : 0,Offset,VET_Float4,3,bInstanced));
+		Elements.Add(FVertexElement(bInstanced ? 1 : 0, Offset, VET_Float4, 3, Stride, bInstanced));
 		Offset += sizeof(float) * 4;
 
 		/** The per-particle dynamic parameter stream */
 		Offset = 0;
-		Elements.Add(FVertexElement(bInstanced ? 2 : 1,Offset,VET_Float4,5,bInstanced));
+		Stride = sizeof(float) * 4;
+		Elements.Add(FVertexElement(bInstanced ? 2 : 1, Offset, VET_Float4, 5, Stride, bInstanced));
 		Offset += sizeof(float) * 4;
 	}
 
diff --git a/Engine/Source/Runtime/Engine/Private/SceneView.cpp b/Engine/Source/Runtime/Engine/Private/SceneView.cpp
--- a/Engine/Source/Runtime/Engine/Private/SceneView.cpp
+++ b/Engine/Source/Runtime/Engine/Private/SceneView.cpp
@@ -336,24 +336,24 @@ void FSceneView::UpdateViewMatrix()
 
 	ViewMatrices.ViewOrigin = StereoViewLocation;
 
- 	ViewMatrices.ViewMatrix = FTranslationMatrix(-StereoViewLocation);
- 	ViewMatrices.ViewMatrix = ViewMatrices.ViewMatrix * FInverseRotationMatrix(ViewRotation);
- 	ViewMatrices.ViewMatrix = ViewMatrices.ViewMatrix * FMatrix(
- 		FPlane(0,	0,	1,	0),
- 		FPlane(1,	0,	0,	0),
- 		FPlane(0,	1,	0,	0),
- 		FPlane(0,	0,	0,	1));
+	ViewMatrices.ViewMatrix = FTranslationMatrix(-StereoViewLocation);
+	ViewMatrices.ViewMatrix = ViewMatrices.ViewMatrix * FInverseRotationMatrix(ViewRotation);
+	ViewMatrices.ViewMatrix = ViewMatrices.ViewMatrix * FMatrix(
+		FPlane(0,	0,	1,	0),
+		FPlane(1,	0,	0,	0),
+		FPlane(0,	1,	0,	0),
+		FPlane(0,	0,	0,	1));
  
 	ViewMatrices.PreViewTranslation = -ViewMatrices.ViewOrigin;
-  	FMatrix TranslatedViewMatrix = FTranslationMatrix(-ViewMatrices.PreViewTranslation) * ViewMatrices.ViewMatrix;
+	FMatrix TranslatedViewMatrix = FTranslationMatrix(-ViewMatrices.PreViewTranslation) * ViewMatrices.ViewMatrix;
 	
 	// Compute a transform from view origin centered world-space to clip space.
- 	ViewMatrices.TranslatedViewProjectionMatrix = TranslatedViewMatrix * ViewMatrices.ProjMatrix;
- 	ViewMatrices.InvTranslatedViewProjectionMatrix = ViewMatrices.TranslatedViewProjectionMatrix.InverseSafe();
+	ViewMatrices.TranslatedViewProjectionMatrix = TranslatedViewMatrix * ViewMatrices.ProjMatrix;
+	ViewMatrices.InvTranslatedViewProjectionMatrix = ViewMatrices.TranslatedViewProjectionMatrix.InverseSafe();
 
- 	ViewProjectionMatrix = ViewMatrices.GetViewProjMatrix();
- 	InvViewMatrix = ViewMatrices.GetInvViewMatrix();
- 	InvViewProjectionMatrix = ViewMatrices.GetInvViewProjMatrix();
+	ViewProjectionMatrix = ViewMatrices.GetViewProjMatrix();
+	InvViewMatrix = ViewMatrices.GetInvViewMatrix();
+	InvViewProjectionMatrix = ViewMatrices.GetInvViewProjMatrix();
 
 }
 
@@ -952,14 +952,15 @@ void FSceneView::EndFinalPostprocessSettings()
 		static const auto CVar = IConsoleManager::Get().FindTConsoleVariableDataInt(TEXT("r.PostProcessAAQuality")); 
 		static auto* MobileHDRCvar = IConsoleManager::Get().FindTConsoleVariableDataInt(TEXT("r.MobileHDR"));
 		static auto* MobileMSAACvar = IConsoleManager::Get().FindTConsoleVariableDataInt(TEXT("r.MobileMSAA"));
+		static uint32 MSAAValue = GRHIShaderPlatform == SP_OPENGL_ES2_IOS ? 1 : MobileMSAACvar->GetValueOnGameThread();
 
 		int32 Quality = FMath::Clamp(CVar->GetValueOnGameThread(), 0, 6);
 		const auto FeatureLevel = GetFeatureLevel();
 
 		if( !Family->EngineShowFlags.PostProcessing || !Family->EngineShowFlags.AntiAliasing || Quality <= 0
 			// Disable antialiasing in GammaLDR mode to avoid jittering.
-			|| (FeatureLevel == ERHIFeatureLevel::ES2 && MobileHDRCvar->GetValueOnGameThread() == 0)
-			|| (FeatureLevel == ERHIFeatureLevel::ES2 && MobileMSAACvar->GetValueOnGameThread()))
+			|| (GRHIFeatureLevel == ERHIFeatureLevel::ES2 && MobileHDRCvar->GetValueOnGameThread() == 0)
+			|| (GRHIFeatureLevel == ERHIFeatureLevel::ES2 && (MSAAValue == 2 || MSAAValue == 4)) )
 		{
 			FinalPostProcessSettings.AntiAliasingMethod = AAM_None;
 		}
diff --git a/Engine/Source/Runtime/Engine/Private/SimpleElementShaders.cpp b/Engine/Source/Runtime/Engine/Private/SimpleElementShaders.cpp
--- a/Engine/Source/Runtime/Engine/Private/SimpleElementShaders.cpp
+++ b/Engine/Source/Runtime/Engine/Private/SimpleElementShaders.cpp
@@ -35,7 +35,7 @@ bool FSimpleElementVS::Serialize(FArchive& Ar)
 void FSimpleElementVS::ModifyCompilationEnvironment(EShaderPlatform Platform, FShaderCompilerEnvironment& OutEnvironment)
 {
 	FGlobalShader::ModifyCompilationEnvironment(Platform, OutEnvironment);
-	OutEnvironment.SetDefine(TEXT("ALLOW_SWITCH_VERTICALAXIS"), 1);
+	OutEnvironment.SetDefine(TEXT("ALLOW_SWITCH_VERTICALAXIS"), (Platform != SP_METAL));
 }
 
 /*------------------------------------------------------------------------------
diff --git a/Engine/Source/Runtime/Engine/Private/Slate/SceneViewport.cpp b/Engine/Source/Runtime/Engine/Private/Slate/SceneViewport.cpp
--- a/Engine/Source/Runtime/Engine/Private/Slate/SceneViewport.cpp
+++ b/Engine/Source/Runtime/Engine/Private/Slate/SceneViewport.cpp
@@ -1022,7 +1022,11 @@ void FSceneViewport::BeginRenderFrame()
 	{
 		// Get the backbuffer render target to render directly to it
 		RenderTargetTextureRHI = RHIGetViewportBackBuffer( ViewportRHI );
-		RHISetRenderTarget( RenderTargetTextureRHI,  FTexture2DRHIRef() );
+		if (GRHIShaderPlatform != SP_METAL)
+		{
+			// unused set render targets are bad on Metal
+			RHISetRenderTarget(RenderTargetTextureRHI, FTexture2DRHIRef());
+		}
 	}
 }
 
diff --git a/Engine/Source/Runtime/Engine/Private/UserInterface/Canvas.cpp b/Engine/Source/Runtime/Engine/Private/UserInterface/Canvas.cpp
--- a/Engine/Source/Runtime/Engine/Private/UserInterface/Canvas.cpp
+++ b/Engine/Source/Runtime/Engine/Private/UserInterface/Canvas.cpp
@@ -478,6 +478,12 @@ void FCanvas::Flush(bool bForce)
 	// current render target set for the canvas
 	check(RenderTarget);	 	
 
+	// no need to set the render target if we aren't going to draw anything to it!
+	if (SortedElements.Num() == 0)
+	{
+		return;
+	}
+
 	// FCanvasSortElement compare class
 	struct FCompareFCanvasSortElement
 	{
diff --git a/Engine/Source/Runtime/Engine/Private/VectorFieldVisualization.cpp b/Engine/Source/Runtime/Engine/Private/VectorFieldVisualization.cpp
--- a/Engine/Source/Runtime/Engine/Private/VectorFieldVisualization.cpp
+++ b/Engine/Source/Runtime/Engine/Private/VectorFieldVisualization.cpp
@@ -63,7 +63,7 @@ public:
 	virtual void InitRHI()
 	{
 		FVertexDeclarationElementList Elements;
-		Elements.Add(FVertexElement(0,0,VET_Float4,0));
+		Elements.Add(FVertexElement(0, 0, VET_Float4, 0, sizeof(FVector4)));
 		VertexDeclarationRHI = RHICreateVertexDeclaration(Elements);
 	}
 
diff --git a/Engine/Source/Runtime/Engine/Public/BatchedElements.h b/Engine/Source/Runtime/Engine/Public/BatchedElements.h
--- a/Engine/Source/Runtime/Engine/Public/BatchedElements.h
+++ b/Engine/Source/Runtime/Engine/Public/BatchedElements.h
@@ -62,10 +62,11 @@ public:
 	virtual void InitRHI()
 	{
 		FVertexDeclarationElementList Elements;
-		Elements.Add(FVertexElement(0,STRUCT_OFFSET(FSimpleElementVertex,Position),VET_Float4,0));
-		Elements.Add(FVertexElement(0,STRUCT_OFFSET(FSimpleElementVertex,TextureCoordinate),VET_Float2,1));
-		Elements.Add(FVertexElement(0,STRUCT_OFFSET(FSimpleElementVertex,Color),VET_Float4,2));
-		Elements.Add(FVertexElement(0,STRUCT_OFFSET(FSimpleElementVertex,HitProxyIdColor),VET_Color,3));
+		uint16 Stride = sizeof(FSimpleElementVertex);
+		Elements.Add(FVertexElement(0,STRUCT_OFFSET(FSimpleElementVertex,Position),VET_Float4,0,Stride));
+		Elements.Add(FVertexElement(0,STRUCT_OFFSET(FSimpleElementVertex,TextureCoordinate),VET_Float2,1,Stride));
+		Elements.Add(FVertexElement(0,STRUCT_OFFSET(FSimpleElementVertex,Color),VET_Float4,2,Stride));
+		Elements.Add(FVertexElement(0,STRUCT_OFFSET(FSimpleElementVertex,HitProxyIdColor),VET_Color,3,Stride));
 		VertexDeclarationRHI = RHICreateVertexDeclaration(Elements);
 	}
 
diff --git a/Engine/Source/Runtime/Engine/Public/MaterialShared.h b/Engine/Source/Runtime/Engine/Public/MaterialShared.h
--- a/Engine/Source/Runtime/Engine/Public/MaterialShared.h
+++ b/Engine/Source/Runtime/Engine/Public/MaterialShared.h
@@ -903,6 +903,7 @@ public:
 	virtual bool IsCrackFreeDisplacementEnabled() const { return false; }
 	virtual bool IsAdaptiveTessellationEnabled() const { return false; }
 	virtual bool IsFullyRough() const { return false; }
+	virtual bool IsNonmetal() const { return false; }
 	virtual bool UseLmDirectionality() const { return true; }
 	virtual bool IsMasked() const = 0;
 	virtual enum EBlendMode GetBlendMode() const = 0;
@@ -1549,6 +1550,7 @@ public:
 	ENGINE_API virtual bool IsCrackFreeDisplacementEnabled() const;
 	ENGINE_API virtual bool IsAdaptiveTessellationEnabled() const;
 	ENGINE_API virtual bool IsFullyRough() const;
+	ENGINE_API virtual bool IsNonmetal() const;
 	ENGINE_API virtual bool UseLmDirectionality() const;
 	ENGINE_API virtual enum EBlendMode GetBlendMode() const;
 	ENGINE_API virtual uint32 GetDecalBlendMode() const;
diff --git a/Engine/Source/Runtime/Engine/Public/ScreenRendering.h b/Engine/Source/Runtime/Engine/Public/ScreenRendering.h
--- a/Engine/Source/Runtime/Engine/Public/ScreenRendering.h
+++ b/Engine/Source/Runtime/Engine/Public/ScreenRendering.h
@@ -29,8 +29,9 @@ public:
 	virtual void InitRHI()
 	{
 		FVertexDeclarationElementList Elements;
-		Elements.Add(FVertexElement(0,STRUCT_OFFSET(FScreenVertex,Position),VET_Float2,0));
-		Elements.Add(FVertexElement(0,STRUCT_OFFSET(FScreenVertex,UV),VET_Float2,1));
+		uint32 Stride = sizeof(FScreenVertex);
+		Elements.Add(FVertexElement(0,STRUCT_OFFSET(FScreenVertex,Position),VET_Float2,0,Stride));
+		Elements.Add(FVertexElement(0,STRUCT_OFFSET(FScreenVertex,UV),VET_Float2,1,Stride));
 		VertexDeclarationRHI = RHICreateVertexDeclaration(Elements);
 	}
 
diff --git a/Engine/Source/Runtime/IOS/IOSAudio/Private/IOSAudioDevice.cpp b/Engine/Source/Runtime/IOS/IOSAudio/Private/IOSAudioDevice.cpp
--- a/Engine/Source/Runtime/IOS/IOSAudio/Private/IOSAudioDevice.cpp
+++ b/Engine/Source/Runtime/IOS/IOSAudio/Private/IOSAudioDevice.cpp
@@ -90,7 +90,7 @@ bool FIOSAudioDevice::InitializeHardware()
 
 	// Linear PCM stream format
 	MixerFormat.mFormatID         = kAudioFormatLinearPCM;
-	MixerFormat.mFormatFlags      = kAudioFormatFlagsCanonical;
+	MixerFormat.mFormatFlags	  = kAudioFormatFlagIsSignedInteger | kAudioFormatFlagsNativeEndian | kAudioFormatFlagIsPacked;
 	MixerFormat.mBytesPerPacket   = SampleSize;
 	MixerFormat.mFramesPerPacket  = 1;
 	MixerFormat.mBytesPerFrame    = SampleSize;
diff --git a/Engine/Source/Runtime/IOS/IOSAudio/Private/IOSAudioSource.cpp b/Engine/Source/Runtime/IOS/IOSAudio/Private/IOSAudioSource.cpp
--- a/Engine/Source/Runtime/IOS/IOSAudio/Private/IOSAudioSource.cpp
+++ b/Engine/Source/Runtime/IOS/IOSAudio/Private/IOSAudioSource.cpp
@@ -701,7 +701,7 @@ OSStatus FIOSAudioSoundSource::IOSAudioRenderCallbackADPCM(AudioSampleType* OutD
 			ADPCM::DecodeBlock(EncodedADPCMBlock, Buffer->CompressedBlockSize, BufferState.StreamBuffer);
 		}
 
-		uint32 NumSamples = FMath::Min(NumFrames, StreamBufferSize / sizeof(AudioSampleType) - BufferState.StreamReadCursor);
+		uint32 NumSamples = FMath::Min<uint32>(NumFrames, StreamBufferSize / sizeof(AudioSampleType)-BufferState.StreamReadCursor);
 		for (uint32 SampleScan = 0; SampleScan < NumSamples; SampleScan++)
 		{
 			*OutData++ = BufferState.StreamBuffer[BufferState.StreamReadCursor++];
diff --git a/Engine/Source/Runtime/IOS/IOSAudio/Public/IOSAudioDevice.h b/Engine/Source/Runtime/IOS/IOSAudio/Public/IOSAudioDevice.h
--- a/Engine/Source/Runtime/IOS/IOSAudio/Public/IOSAudioDevice.h
+++ b/Engine/Source/Runtime/IOS/IOSAudio/Public/IOSAudioDevice.h
@@ -17,6 +17,8 @@
 #include <AudioUnit/AudioUnit.h>
 #include <AVFoundation/AVAudioSession.h>
 
+#define AudioSampleType SInt16
+
 DECLARE_LOG_CATEGORY_EXTERN(LogIOSAudio, Log, All);
 
 /*------------------------------------------------------------------------------------
diff --git a/Engine/Source/Runtime/IOS/MetalRHI/MetalRHI.Build.cs b/Engine/Source/Runtime/IOS/MetalRHI/MetalRHI.Build.cs
new file mode 100644
--- /dev/null
+++ b/Engine/Source/Runtime/IOS/MetalRHI/MetalRHI.Build.cs
@@ -0,0 +1,22 @@
+// Copyright 1998-2014 Epic Games, Inc. All Rights Reserved.
+
+using UnrealBuildTool;
+
+public class MetalRHI : ModuleRules
+{	
+	public MetalRHI(TargetInfo Target)
+	{
+
+		PrivateDependencyModuleNames.AddRange(
+			new string[] {
+				"Core",
+				"Engine",
+				"RHI",
+				"RenderCore",
+				"ShaderCore"
+			}
+			);
+			
+		PublicFrameworks.Add("Metal");
+	}
+}
diff --git a/Engine/Source/Runtime/IOS/MetalRHI/Private/MetalCommands.cpp b/Engine/Source/Runtime/IOS/MetalRHI/Private/MetalCommands.cpp
--- a/Engine/Source/Runtime/IOS/MetalRHI/Private/MetalCommands.cpp
+++ b/Engine/Source/Runtime/IOS/MetalRHI/Private/MetalCommands.cpp
@@ -25,8 +25,6 @@ void FMetalDynamicRHI::RHISetStreamSource(uint32 StreamIndex,FVertexBufferRHIPar
 {
 	DYNAMIC_CAST_METGALRESOURCE(VertexBuffer,VertexBuffer);
 
-//	FMetalManager::Get()->CheckStride(StreamIndex, Stride);
-
 	// queue this zero stride buffer to be grown later if needed
 	HACK_ZeroStrideBuffers[StreamIndex] = VertexBuffer->ZeroStrideElementSize ? VertexBuffer : NULL;
 
@@ -470,7 +468,6 @@ void FMetalDynamicRHI::RHIEndDrawPrimitiveUP()
 	// last minute draw setup
 	HACK_ZeroStrideBuffers[0] = NULL;
 	FMetalManager::Get()->PrepareToDraw(0);
-	FMetalManager::Get()->CheckStride(0, GPendingVertexDataStride);
 	
 #if !NO_DRAW
 	[FMetalManager::GetContext() drawPrimitives:TranslatePrimitiveType(GPendingPrimitiveType)
@@ -517,7 +514,6 @@ void FMetalDynamicRHI::RHIEndDrawIndexedPrimitiveUP()
 	// last minute draw setup
 	HACK_ZeroStrideBuffers[0] = NULL;
 	FMetalManager::Get()->PrepareToDraw(0);
-	FMetalManager::Get()->CheckStride(0, GPendingVertexDataStride);
 	
 #if !NO_DRAW
 	[FMetalManager::GetContext() drawIndexedPrimitives:TranslatePrimitiveType(GPendingPrimitiveType)
diff --git a/Engine/Source/Runtime/IOS/MetalRHI/Private/MetalManager.cpp b/Engine/Source/Runtime/IOS/MetalRHI/Private/MetalManager.cpp
--- a/Engine/Source/Runtime/IOS/MetalRHI/Private/MetalManager.cpp
+++ b/Engine/Source/Runtime/IOS/MetalRHI/Private/MetalManager.cpp
@@ -43,7 +43,6 @@ TMap<id, int32> ClassCounts;
 	}
 #define GET_HASH(Offset, NumBits) ((Hash >> Offset) & ((1ULL << NumBits) - 1))
 
-extern FString GMaterialNameHack;
 DEFINE_STAT(STAT_MetalMakeDrawableTime);
 DEFINE_STAT(STAT_MetalDrawCallTime);
 DEFINE_STAT(STAT_MetalPrepareDrawTime);
@@ -113,10 +112,6 @@ id<MTLRenderPipelineState> FPipelineShadow::CreatePipelineStateForBoundShaderSta
 	// some basic settings
 	Desc.depthWriteEnabled = bIsDepthWriteEnabled;
 	Desc.stencilWriteEnabled = bIsStencilWriteEnabled;
-	if (FParse::Param(FCommandLine::Get(), TEXT("labels")))
-	{
-		Desc.label = GMaterialNameHack.GetNSString();
-	}
 
 	// set per-MRT settings
 	for (uint32 RenderTargetIndex = 0; RenderTargetIndex < MaxMetalRenderTargets; ++RenderTargetIndex)
@@ -490,20 +485,6 @@ void FMetalManager::PrepareToDraw(uint32 NumVertices)
 	}
 }
 
-void FMetalManager::CheckStride(uint32 StreamIndex, uint32 Stride)
-{
-	FVertexDeclarationElementList& Elements = CurrentBoundShaderState->VertexDeclaration->Elements;
-	for (int32 ElementIndex = 0; ElementIndex < Elements.Num(); ElementIndex++)
-	{
-		FVertexElement& Element = Elements[ElementIndex];
-		if (Element.StreamIndex == StreamIndex)
-		{
-			check(Element.Stride != 0xFFFF);
-//			checkf(Element.Stride == Stride, TEXT("Strides didn't match up in Stream %d, expected %d, it was %d"), StreamIndex, Stride, Element.Stride);
-		}
-	}
-}
-
 void FMetalManager::SetDepthStencilWriteEnabled(bool bIsDepthWriteEnabled, bool bIsStencilWriteEnabled)
 {
 	Pipeline.bIsDepthWriteEnabled = bIsDepthWriteEnabled;
@@ -842,4 +823,123 @@ void FMetalManager::SetRasterizerState(const FRasterizerStateInitializerRHI& Sta
 #endif
 
 	bFirstRasterizerState = false;
-}
\ No newline at end of file
+}
+
+
+
+// @todo metal srt
+#if 0 
+template <class ShaderType>
+void FMetalManager::SetResourcesFromTables(const ShaderType* RESTRICT Shader, Gnm::ShaderStage ShaderStage)
+{
+	checkSlow(Shader);
+
+	// Mask the dirty bits by those buffers from which the shader has bound resources.
+	uint32 DirtyBits = Shader->ShaderResourceTable.ResourceTableBits & DirtyUniformBuffers[TRHIShaderToEnum<ShaderType>::ShaderFrequency];
+	uint32 NumSetCalls = 0;
+	while (DirtyBits)
+	{
+		// Scan for the lowest set bit, compute its index, clear it in the set of dirty bits.
+		const uint32 LowestBitMask = (DirtyBits) & (-(int32)DirtyBits);
+		const int32 BufferIndex = FMath::FloorLog2(LowestBitMask); // todo: This has a branch on zero, we know it could never be zero...
+		DirtyBits ^= LowestBitMask;
+		auto* Buffer = (FGnmUniformBuffer*)BoundUniformBuffers[TRHIShaderToEnum<ShaderType>::ShaderFrequency][BufferIndex].GetReference();
+		check(Buffer);
+		check(BufferIndex < Shader->ShaderResourceTable.ResourceTableLayoutHashes.Num());
+		check(Buffer->GetLayout().GetHash() == Shader->ShaderResourceTable.ResourceTableLayoutHashes[BufferIndex]);
+		Buffer->CacheResources(ResourceTableFrameCounter);
+
+		// todo: could make this two pass: gather then set
+		NumSetCalls +=
+		SetShaderResourcesFromBuffer<FGnmSurface>(ShaderStage, Buffer, Shader->ShaderResourceTable.ShaderResourceViewMap.GetData(), BufferIndex);
+		SetShaderResourcesFromBuffer<Gnm::Sampler>(ShaderStage, Buffer, Shader->ShaderResourceTable.SamplerMap.GetData(), BufferIndex);
+	}
+	DirtyUniformBuffers[TRHIShaderToEnum<ShaderType>::ShaderFrequency] = 0;
+	SetTextureInTableCalls += NumSetCalls;
+}
+
+void FMetalManager::CommitGraphicsResourceTables()
+{
+	uint32 Start = FPlatformTime::Cycles();
+
+	GRHICommandList.Verify();
+
+	auto* RESTRICT CurrentBoundShaderState = (FGnmBoundShaderState*)GGnmManager.BoundShaderStateHistory.GetLast();
+	check(CurrentBoundShaderState);
+
+	if (auto* Shader = CurrentBoundShaderState->VertexShader.GetReference())
+	{
+		SetResourcesFromTables(Shader, Shader->ShaderStage);
+	}
+	if (auto* Shader = CurrentBoundShaderState->PixelShader.GetReference())
+	{
+		SetResourcesFromTables(Shader, Gnm::kShaderStagePs);
+	}
+	if (auto* Shader = CurrentBoundShaderState->HullShader.GetReference())
+	{
+		SetResourcesFromTables(Shader, Gnm::kShaderStageHs);
+		auto* DomainShader = CurrentBoundShaderState->DomainShader.GetReference();
+		check(DomainShader);
+		check(0);
+/*
+		SetResourcesFromTables(DomainShader, Gnm::kShaderStageDs);
+*/
+	}
+	if (auto* Shader = CurrentBoundShaderState->GeometryShader.GetReference())
+	{
+		SetResourcesFromTables(Shader, Gnm::kShaderStageGs);
+	}
+
+	CommitResourceTableCycles += (FPlatformTime::Cycles() - Start);
+}
+
+void FMetalManager::CommitNonComputeShaderConstants()
+{
+	GRHICommandList.Verify();
+	auto* CurrentBoundShaderState = (FGnmBoundShaderState*)BoundShaderStateHistory.GetLast();
+	check(CurrentBoundShaderState);
+
+	// update constant buffers from local shadow memory, and send to GPU
+	if (CurrentBoundShaderState->bShaderNeedsGlobalConstantBuffer[ SF_Vertex ] )
+	{
+		VSConstantBuffer->CommitConstantsToDevice(CurrentVertexShaderStage, 0, bDiscardSharedConstants);
+	}
+	if (CurrentBoundShaderState->bShaderNeedsGlobalConstantBuffer[ SF_Pixel ] )
+	{
+		PSConstantBuffer->CommitConstantsToDevice(Gnm::kShaderStagePs, 0, bDiscardSharedConstants);
+	}
+	if (CurrentBoundShaderState->bShaderNeedsGlobalConstantBuffer[ SF_Geometry ] )
+	{
+		GSConstantBuffer->CommitConstantsToDevice(Gnm::kShaderStageGs, 0, bDiscardSharedConstants);
+	}
+/*
+	// Skip HS/DS CB updates in cases where tessellation isn't being used
+	// Note that this is *potentially* unsafe because bDiscardSharedConstants is cleared at the
+	// end of the function, however we're OK for now because bDiscardSharedConstants
+	// is always reset whenever bUsingTessellation changes in SetBoundShaderState()
+	if (bUsingTessellation)
+	{
+		if (CurrentBoundShaderState->bShaderNeedsGlobalConstantBuffer[ SF_Hull ] )
+		{
+			HSConstantBuffer->CommitConstantsToDevice(Gnm::kShaderStageHs, 0, bDiscardSharedConstants);
+		}
+		if (CurrentBoundShaderState->bShaderNeedsGlobalConstantBuffer[ SF_Domain ] )
+		{
+			check(0);
+			//DSConstantBuffer->CommitConstantsToDevice(Gnm::kShaderStageDs, 0, bDiscardSharedConstants);
+		}
+	}
+*/
+
+	// don't need to discard again until next shader is set
+	bDiscardSharedConstants = false;
+}
+
+void FGnmManager::CommitComputeResourceTables( FGnmComputeShader* ComputeShader )
+{
+	GRHICommandList.Verify();
+
+	check(ComputeShader);
+	SetResourcesFromTables(ComputeShader, Gnm::kShaderStageCs);
+}
+#endif
\ No newline at end of file
diff --git a/Engine/Source/Runtime/IOS/MetalRHI/Private/MetalManager.h b/Engine/Source/Runtime/IOS/MetalRHI/Private/MetalManager.h
--- a/Engine/Source/Runtime/IOS/MetalRHI/Private/MetalManager.h
+++ b/Engine/Source/Runtime/IOS/MetalRHI/Private/MetalManager.h
@@ -87,11 +87,6 @@ public:
 	 */
 	void PrepareToDraw(uint32 NumVertices);
 	
-	/**
-	 * Check the stride against what we expected
-	 */
-	void CheckStride(uint32 StreamIndex, uint32 Stride);
-	
 	/**
 	 * Functions to update the pipeline descriptor and/or context
 	 */
diff --git a/Engine/Source/Runtime/IOS/MetalRHI/Private/MetalShaders.cpp b/Engine/Source/Runtime/IOS/MetalRHI/Private/MetalShaders.cpp
--- a/Engine/Source/Runtime/IOS/MetalRHI/Private/MetalShaders.cpp
+++ b/Engine/Source/Runtime/IOS/MetalRHI/Private/MetalShaders.cpp
@@ -248,12 +248,10 @@ FMetalBoundShaderState::FMetalBoundShaderState(
 	GeometryShader = InGeometryShader;
 
 #if !UE_BUILD_SHIPPING
-	extern FString GMaterialNameHack;
-	if (GFrameCounter > 3 || GMaterialNameHack == TEXT("M_Butterfly_OnGround") ||
-		GMaterialNameHack == TEXT("M_BillBoardBGTree") || GMaterialNameHack == TEXT("M_WaterSplash_Particulate"))
+	if (GFrameCounter > 3)
 	{
 		NSLog(@"===============================================================");
-		NSLog(@"Creating a BSS at runtime frame %lld... this may hitch! [%s, this = %p]", GFrameCounter, TCHAR_TO_ANSI(*GMaterialNameHack), this);
+		NSLog(@"Creating a BSS at runtime frame %lld... this may hitch! [this = %p]", GFrameCounter, this);
 		NSLog(@"Vertex declaration:");
 		FVertexDeclarationElementList& Elements = VertexDeclaration->Elements;
 		for (int32 i = 0; i < Elements.Num(); i++)
diff --git a/Engine/Source/Runtime/IOS/MetalRHI/Private/MetalUniformBuffer.cpp b/Engine/Source/Runtime/IOS/MetalRHI/Private/MetalUniformBuffer.cpp
--- a/Engine/Source/Runtime/IOS/MetalRHI/Private/MetalUniformBuffer.cpp
+++ b/Engine/Source/Runtime/IOS/MetalRHI/Private/MetalUniformBuffer.cpp
@@ -182,42 +182,58 @@ FMetalUniformBuffer::FMetalUniformBuffer(const void* Contents, const FRHIUniform
 	, Buffer(nil)
 	, Offset(0)
 {
-	// for single use buffers, allocate from the ring buffer to avoid thrashing memory
-	if (Usage == UniformBuffer_SingleDraw)
+	if (Layout.ConstantBufferSize > 0)
 	{
-		// use a bit of the ring buffer
-		Offset = FMetalManager::Get()->AllocateFromRingBuffer(Layout.ConstantBufferSize);
-		Buffer = FMetalManager::Get()->GetRingBuffer();
-	}
-	else
-	{
-		// Find the appropriate bucket based on size
-		const uint32 BucketIndex = GetPoolBucketIndex(Layout.ConstantBufferSize);
-		TArray<FPooledUniformBuffer>& PoolBucket = UniformBufferPool[BucketIndex];
-		if (PoolBucket.Num() > 0)
+		// for single use buffers, allocate from the ring buffer to avoid thrashing memory
+		if (Usage == UniformBuffer_SingleDraw)
 		{
-			// Reuse the last entry in this size bucket
-			FPooledUniformBuffer FreeBufferEntry = PoolBucket.Pop();
-			DEC_DWORD_STAT(STAT_MetalNumFreeUniformBuffers);
-			DEC_MEMORY_STAT_BY(STAT_MetalFreeUniformBufferMemory, FreeBufferEntry.CreatedSize);
-
-			// reuse the one
-			Buffer = FreeBufferEntry.Buffer;
-			Offset = FreeBufferEntry.Offset;
+			// use a bit of the ring buffer
+			Offset = FMetalManager::Get()->AllocateFromRingBuffer(Layout.ConstantBufferSize);
+			Buffer = FMetalManager::Get()->GetRingBuffer();
 		}
 		else
 		{
-			// Nothing usable was found in the free pool, create a new uniform buffer (full size, not NumBytes)
-			uint32 BufferSize = UniformBufferSizeBuckets[BucketIndex];
-			Buffer = SuballocateUB(BufferSize, Offset);
+			// Find the appropriate bucket based on size
+			const uint32 BucketIndex = GetPoolBucketIndex(Layout.ConstantBufferSize);
+			TArray<FPooledUniformBuffer>& PoolBucket = UniformBufferPool[BucketIndex];
+			if (PoolBucket.Num() > 0)
+			{
+				// Reuse the last entry in this size bucket
+				FPooledUniformBuffer FreeBufferEntry = PoolBucket.Pop();
+				DEC_DWORD_STAT(STAT_MetalNumFreeUniformBuffers);
+				DEC_MEMORY_STAT_BY(STAT_MetalFreeUniformBufferMemory, FreeBufferEntry.CreatedSize);
+
+				// reuse the one
+				Buffer = FreeBufferEntry.Buffer;
+				Offset = FreeBufferEntry.Offset;
+			}
+			else
+			{
+				// Nothing usable was found in the free pool, create a new uniform buffer (full size, not NumBytes)
+				uint32 BufferSize = UniformBufferSizeBuckets[BucketIndex];
+				Buffer = SuballocateUB(BufferSize, Offset);
+			}
 		}
-	}
 
-	// copy the contents
-	FMemory::Memcpy(((uint8*)[Buffer contents]) + Offset, Contents, Layout.ConstantBufferSize);
+		// copy the contents
+		FMemory::Memcpy(((uint8*)[Buffer contents]) + Offset, Contents, Layout.ConstantBufferSize);
+	}
 
 	// @todo metal: Add support for resource tables!
-	check(Layout.Resources.Num() == 0);
+	if (Layout.Resources.Num())
+	{
+		int32 NumResources = Layout.Resources.Num();
+		FRHIResource** InResources = (FRHIResource**)((uint8*)Contents + Layout.ResourceOffset);
+		ResourceTable.Empty(NumResources);
+		ResourceTable.AddZeroed(NumResources);
+		for (int32 i = 0; i < NumResources; ++i)
+		{
+			check(InResources[i]);
+			ResourceTable[i] = InResources[i];
+		}
+		RawResourceTable.Empty(NumResources);
+		RawResourceTable.AddZeroed(NumResources);
+	}
 }
 
 FMetalUniformBuffer::~FMetalUniformBuffer()
@@ -230,6 +246,65 @@ FMetalUniformBuffer::~FMetalUniformBuffer()
 	}
 }
 
+void FMetalUniformBuffer::CacheResourcesInternal()
+{
+	const FRHIUniformBufferLayout& Layout = GetLayout();
+	int32 NumResources = Layout.Resources.Num();
+	const uint8* RESTRICT ResourceTypes = Layout.Resources.GetData();
+	const TRefCountPtr<FRHIResource>* RESTRICT Resources = ResourceTable.GetData();
+	void** RESTRICT RawResources = RawResourceTable.GetData();
+	float CurrentTime = FApp::GetCurrentTime();
+
+	// todo: Immutable resources, i.e. not textures, can be safely cached across frames.
+	// Texture streaming makes textures complicated :)
+	for (int32 i = 0; i < NumResources; ++i)
+	{
+		switch (ResourceTypes[i])
+		{
+			case UBMT_SRV:
+				{
+					FMetalShaderResourceView* SRV = (FMetalShaderResourceView*)Resources[i].GetReference();
+					if (IsValidRef(SRV->SourceTexture))
+					{
+						FMetalSurface& Surface = GetMetalSurfaceFromRHITexture(SRV->SourceTexture);
+						RawResources[i] = Surface.Texture;
+					}
+					else
+					{
+						check(0);
+						RawResources[i] = &SRV->SourceVertexBuffer->Buffer;
+					}
+				}
+				break;
+
+			case UBMT_TEXTURE:
+				{
+					// todo: this does multiple virtual function calls to find the right type to cast to
+					// this is due to multiple inheritance nastiness, NEEDS CLEANUP
+					FRHITexture* TextureRHI = (FRHITexture*)Resources[i].GetReference();
+					TextureRHI->SetLastCachedTime(CurrentTime);
+					FMetalSurface& Surface = GetMetalSurfaceFromRHITexture(TextureRHI);
+					RawResources[i] = &Surface;
+				}
+				break;
+
+			case UBMT_UAV:
+				RawResources[i] = 0;
+				break;
+
+			case UBMT_SAMPLER:
+				RawResources[i] = &((FMetalSamplerState*)Resources[i].GetReference())->State;
+				break;
+
+			default:
+				check(0);
+				break;
+		}
+	}
+}
+
+
+
 FUniformBufferRHIRef FMetalDynamicRHI::RHICreateUniformBuffer(const void* Contents, const FRHIUniformBufferLayout& Layout, EUniformBufferUsage Usage)
 {
 	return new FMetalUniformBuffer(Contents, Layout, Usage);
diff --git a/Engine/Source/Runtime/IOS/MetalRHI/Public/MetalResources.h b/Engine/Source/Runtime/IOS/MetalRHI/Public/MetalResources.h
--- a/Engine/Source/Runtime/IOS/MetalRHI/Public/MetalResources.h
+++ b/Engine/Source/Runtime/IOS/MetalRHI/Public/MetalResources.h
@@ -330,11 +330,43 @@ public:
 	// Destructor 
 	~FMetalUniformBuffer();
 	
+
+	/** Cache resources if needed. */
+	inline void CacheResources(uint32 InFrameCounter)
+	{
+		if (InFrameCounter == INDEX_NONE || LastCachedFrame != InFrameCounter)
+		{
+			CacheResourcesInternal();
+			LastCachedFrame = InFrameCounter;
+		}
+	}
+
+	bool IsConstantBuffer() const
+	{
+		return Buffer.length > 0;
+	}
+
+
 	/** The buffer containing the contents - either a fresh buffer or the ring buffer */
 	id<MTLBuffer> Buffer;
 	
 	/** This offset is used when passing to setVertexBuffer, etc */
 	uint32 Offset;
+
+
+	/** Resource table containing RHI references. */
+	TArray<TRefCountPtr<FRHIResource> > ResourceTable;
+
+	/** Raw resource table, cached once per frame. */
+	TArray<void*> RawResourceTable;
+
+	/** The frame in which RawResourceTable was last cached. */
+	uint32 LastCachedFrame;
+
+private:
+	/** Actually cache resources. */
+	void CacheResourcesInternal();
+
 };
 
 
diff --git a/Engine/Source/Runtime/Launch/Private/IOS/IOSAppDelegateConsoleHandling.cpp b/Engine/Source/Runtime/Launch/Private/IOS/IOSAppDelegateConsoleHandling.cpp
--- a/Engine/Source/Runtime/Launch/Private/IOS/IOSAppDelegateConsoleHandling.cpp
+++ b/Engine/Source/Runtime/Launch/Private/IOS/IOSAppDelegateConsoleHandling.cpp
@@ -74,6 +74,9 @@
 		if (ExistingCommand == NSNotFound)
 		{
 			[self.ConsoleHistoryValues addObject:ConsoleCommand];
+			// save to local storage
+			[[NSUserDefaults standardUserDefaults] setObject:self.ConsoleHistoryValues forKey:@"ConsoleHistory"];
+			[[NSUserDefaults standardUserDefaults] synchronize];
 		}
 	}
 }
diff --git a/Engine/Source/Runtime/Launch/Private/IOS/LaunchIOS.cpp b/Engine/Source/Runtime/Launch/Private/IOS/LaunchIOS.cpp
--- a/Engine/Source/Runtime/Launch/Private/IOS/LaunchIOS.cpp
+++ b/Engine/Source/Runtime/Launch/Private/IOS/LaunchIOS.cpp
@@ -5,21 +5,29 @@
 #include "LaunchPrivatePCH.h"
 #include "IOSAppDelegate.h"
 #include "EAGLView.h"
+#include "BLView.h"
 #include "IOSCommandLineHelper.h"
 #include "GameLaunchDaemonMessageHandler.h"
 #include "AudioDevice.h"
 
+
 FEngineLoop GEngineLoop;
 FGameLaunchDaemonMessageHandler GCommandSystem;
 
 void FAppEntry::Suspend()
 {
-    GEngine->AudioDevice->SuspendContext();
+	if (GEngine && GEngine->AudioDevice)
+	{
+		GEngine->AudioDevice->SuspendContext();
+	}
 }
 
 void FAppEntry::Resume()
 {
-    GEngine->AudioDevice->ResumeContext();
+	if (GEngine && GEngine->AudioDevice)
+	{
+	    GEngine->AudioDevice->ResumeContext();
+	}
 }
 
 void FAppEntry::PreInit(IOSAppDelegate* AppDelegate, UIApplication* Application)
@@ -54,9 +62,9 @@ static void MainThreadInit()
 	// Size the view appropriately for any potentially dynamically attached displays,
 	// prior to creating any framebuffers
 	CGRect MainFrame = [[UIScreen mainScreen] bounds];
-	if (!AppDelegate.bDeviceInPortraitMode)
+	if ([IOSAppDelegate GetDelegate].OSVersion < 8.0f && !AppDelegate.bDeviceInPortraitMode)
 	{
-		Swap<float>(MainFrame.size.width, MainFrame.size.height);
+		Swap(MainFrame.size.width, MainFrame.size.height);
 	}
 	
 	// @todo: use code similar for presizing for secondary screens
@@ -73,16 +81,33 @@ static void MainThreadInit()
 // 		);
 	CGRect FullResolutionRect = MainFrame;
 
-	// create the fullscreen EAGLView
-	AppDelegate.GLView = [[EAGLView alloc] initWithFrame:FullResolutionRect];
-	AppDelegate.GLView.clearsContextBeforeDrawing = NO;
-	AppDelegate.GLView.multipleTouchEnabled = YES;
+#if HAS_METAL
+	if (FPlatformMisc::HasPlatformFeature(TEXT("Metal")) && !FParse::Param(FCommandLine::Get(), TEXT("ES2")))
+	{
+		AppDelegate.MetalView = [[FMetalView alloc] initWithFrame:FullResolutionRect];
+		AppDelegate.MetalView.clearsContextBeforeDrawing = NO;
+		AppDelegate.MetalView.multipleTouchEnabled = YES;
+
+		// add it to the window
+		[AppDelegate.RootView addSubview:AppDelegate.MetalView];
+
+		// initialize the backbuffer of the view (so the RHI can use it)
+		[AppDelegate.MetalView CreateFramebuffer:YES];
+	}
+	else
+#endif
+	{
+		// create the fullscreen EAGLView
+		AppDelegate.GLView = [[EAGLView alloc] initWithFrame:FullResolutionRect];
+		AppDelegate.GLView.clearsContextBeforeDrawing = NO;
+		AppDelegate.GLView.multipleTouchEnabled = YES;
 
-	// add it to the window
-	[AppDelegate.RootView addSubview:AppDelegate.GLView];
+		// add it to the window
+		[AppDelegate.RootView addSubview:AppDelegate.GLView];
 
-	// initialize the backbuffer of the view (so the RHI can use it)
-	[AppDelegate.GLView CreateFramebuffer:YES];
+		// initialize the backbuffer of the view (so the RHI can use it)
+		[AppDelegate.GLView CreateFramebuffer:YES];
+	}
 
 	// Final adjustment to the viewport (this is deferred and won't run until the first engine tick)
 	[FIOSAsyncTask CreateTaskWithBlock:^ bool (void)
@@ -113,14 +138,29 @@ void FAppEntry::PlatformInit()
 
 	// wait until the GLView is fully initialized, so the RHI can be initialized
 	IOSAppDelegate* AppDelegate = [IOSAppDelegate GetDelegate];
-	while (!AppDelegate.GLView || !AppDelegate.GLView->bIsInitialized)
+
+#if HAS_METAL
+	if (!FParse::Param(FCommandLine::Get(), TEXT("ES2")))
 	{
-		FPlatformProcess::Sleep(0.001f);
-	}
+		while (!AppDelegate.MetalView || !AppDelegate.MetalView->bIsInitialized)
+		{
+			FPlatformProcess::Sleep(0.001f);
+		}
 
-	// set the GL context to this thread
-	[AppDelegate.GLView MakeCurrent];
+		// set the Metal context to this thread
+		//		[AppDelegate.MetalView MakeCurrent];		
+	}
+	else
+#endif
+	{
+		while (!AppDelegate.GLView || !AppDelegate.GLView->bIsInitialized)
+		{
+			FPlatformProcess::Sleep(0.001f);
+		}
 
+		// set the GL context to this thread
+		[AppDelegate.GLView MakeCurrent];
+	}
 }
 
 void FAppEntry::Init()
diff --git a/Engine/Source/Runtime/Launch/Private/LaunchEngineLoop.cpp b/Engine/Source/Runtime/Launch/Private/LaunchEngineLoop.cpp
--- a/Engine/Source/Runtime/Launch/Private/LaunchEngineLoop.cpp
+++ b/Engine/Source/Runtime/Launch/Private/LaunchEngineLoop.cpp
@@ -714,16 +714,7 @@ int32 FEngineLoop::PreInit( const TCHAR* CmdLine )
 	GIsGameThreadIdInitialized = true;
 
 	FPlatformProcess::SetThreadAffinityMask(FPlatformAffinity::GetMainGameMask());
-
-#if PLATFORM_XBOXONE
-	HANDLE proc = ::GetCurrentProcess();
-	DWORD64 dwProcessAffinity, dwSystemAffinity;
-	if (GetProcessAffinityMask( proc, &dwProcessAffinity, &dwSystemAffinity))
-	{
-		UE_LOG(LogInit, Log, TEXT("Process affinity %ld System affinity %ld."), dwProcessAffinity , dwSystemAffinity);
-	}
-	UE_LOG(LogInit, Log, TEXT("Runnable thread Main Thread is on Process %d."), static_cast<uint32>(::GetCurrentProcessorNumber()) );	
-#endif // PLATFORM_XBOXONE
+	FPlatformProcess::SetupGameOrRenderThread(false);
 
 	// Figure out whether we're the editor, ucc or the game.
 	const SIZE_T CommandLineSize = FCString::Strlen(CmdLine)+1;
diff --git a/Engine/Source/Runtime/Online/IOS/OnlineSubsystemIOS/Private/OnlineSessionInterfaceIOS.cpp b/Engine/Source/Runtime/Online/IOS/OnlineSubsystemIOS/Private/OnlineSessionInterfaceIOS.cpp
--- a/Engine/Source/Runtime/Online/IOS/OnlineSubsystemIOS/Private/OnlineSessionInterfaceIOS.cpp
+++ b/Engine/Source/Runtime/Online/IOS/OnlineSubsystemIOS/Private/OnlineSessionInterfaceIOS.cpp
@@ -83,7 +83,7 @@
 
 - (void)session:(GKSession *)session connectionWithPeerFailed:(NSString *)peerID withError:(NSError *)error
 {
-	NSString* ErrorString = [NSString stringWithFormat:@"connectionWithPeerFailed - Failed to connect to %@ with error code %i", [session displayNameForPeer:peerID], [error code]];
+	NSString* ErrorString = [NSString stringWithFormat:@"connectionWithPeerFailed - Failed to connect to %@ with error code %d", [session displayNameForPeer:peerID], (uint32)[error code]];
 	const FString ConvertedErrorStr(ErrorString);
 	UE_LOG(LogOnline, Display, TEXT("%s"), *ConvertedErrorStr);
 }
diff --git a/Engine/Source/Runtime/OpenGLDrv/Private/OpenGLDrvPrivate.h b/Engine/Source/Runtime/OpenGLDrv/Private/OpenGLDrvPrivate.h
--- a/Engine/Source/Runtime/OpenGLDrv/Private/OpenGLDrvPrivate.h
+++ b/Engine/Source/Runtime/OpenGLDrv/Private/OpenGLDrvPrivate.h
@@ -466,7 +466,7 @@ public:
 	virtual void InitRHI()
 	{
 		FVertexDeclarationElementList Elements;
-		Elements.Add(FVertexElement(0,0,VET_Float4,0));
+		Elements.Add(FVertexElement(0,0,VET_Float4,0,sizeof(FVector4)));
 		VertexDeclarationRHI = RHICreateVertexDeclaration(Elements);
 	}
 	virtual void ReleaseRHI()
diff --git a/Engine/Source/Runtime/OpenGLDrv/Private/OpenGLUniformBuffer.cpp b/Engine/Source/Runtime/OpenGLDrv/Private/OpenGLUniformBuffer.cpp
--- a/Engine/Source/Runtime/OpenGLDrv/Private/OpenGLUniformBuffer.cpp
+++ b/Engine/Source/Runtime/OpenGLDrv/Private/OpenGLUniformBuffer.cpp
@@ -36,7 +36,7 @@ namespace OpenGLConsoleVariables
 
 #define NUM_SAFE_FRAMES 3
 
-const uint32 RequestedUniformBufferSizeBuckets[NUM_POOL_BUCKETS] = {
+static const uint32 RequestedUniformBufferSizeBuckets[NUM_POOL_BUCKETS] = {
 	16,32,48,64,80,96,112,128,	// 16-byte increments
 	160,192,224,256,			// 32-byte increments
 	320,384,448,512,			// 64-byte increments
@@ -54,7 +54,7 @@ const uint32 RequestedUniformBufferSizeBuckets[NUM_POOL_BUCKETS] = {
 };
 
 // Maps desired size buckets to aligment actually 
-TArray<uint32> UniformBufferSizeBuckets;
+static TArray<uint32> UniformBufferSizeBuckets;
 
 
 static inline bool IsSuballocatingUBOs()
diff --git a/Engine/Source/Runtime/RHI/Private/Apple/AppleDynamicRHI.cpp b/Engine/Source/Runtime/RHI/Private/Apple/AppleDynamicRHI.cpp
--- a/Engine/Source/Runtime/RHI/Private/Apple/AppleDynamicRHI.cpp
+++ b/Engine/Source/Runtime/RHI/Private/Apple/AppleDynamicRHI.cpp
@@ -8,9 +8,20 @@
 FDynamicRHI* PlatformCreateDynamicRHI()
 {
 	FDynamicRHI* DynamicRHI = NULL;
-	// Load the dynamic RHI module.
-	IDynamicRHIModule* DynamicRHIModule = &FModuleManager::LoadModuleChecked<IDynamicRHIModule>(TEXT("OpenGLDrv"));
+	IDynamicRHIModule* DynamicRHIModule = NULL;
 
+	// Load the dynamic RHI module.
+#if PLATFORM_IOS && HAS_METAL
+	if (FPlatformMisc::HasPlatformFeature(TEXT("Metal")) && !FParse::Param(FCommandLine::Get(), TEXT("ES2")))
+	{
+		DynamicRHIModule = &FModuleManager::LoadModuleChecked<IDynamicRHIModule>(TEXT("MetalRHI"));
+	}
+	else
+#endif
+	{
+		DynamicRHIModule = &FModuleManager::LoadModuleChecked<IDynamicRHIModule>(TEXT("OpenGLDrv"));
+	}
+	
 	// Create the dynamic RHI.
 	DynamicRHI = DynamicRHIModule->CreateRHI();
 	return DynamicRHI;
diff --git a/Engine/Source/Runtime/RHI/Private/RHI.cpp b/Engine/Source/Runtime/RHI/Private/RHI.cpp
--- a/Engine/Source/Runtime/RHI/Private/RHI.cpp
+++ b/Engine/Source/Runtime/RHI/Private/RHI.cpp
@@ -235,6 +235,7 @@ static FName NAME_OPENGL_150_ES2(TEXT("GLSL_150_ES2"));
 static FName NAME_OPENGL_ES2(TEXT("GLSL_ES2"));
 static FName NAME_OPENGL_ES2_WEBGL(TEXT("GLSL_ES2_WEBGL"));
 static FName NAME_OPENGL_ES2_IOS(TEXT("GLSL_ES2_IOS"));
+static FName NAME_SF_METAL(TEXT("SF_METAL"));
 
 FName LegacyShaderPlatformToShaderFormat(EShaderPlatform Platform)
 {
@@ -266,6 +267,8 @@ FName LegacyShaderPlatformToShaderFormat(EShaderPlatform Platform)
 		return NAME_OPENGL_ES2_WEBGL;
 	case SP_OPENGL_ES2_IOS:
 		return NAME_OPENGL_ES2_IOS;
+	case SP_METAL:
+		return NAME_SF_METAL;
 	default:
 		check(0);
 		return NAME_PCD3D_SM5;
@@ -286,6 +289,7 @@ EShaderPlatform ShaderFormatToLegacyShaderPlatform(FName ShaderFormat)
 	if (ShaderFormat == NAME_OPENGL_ES2)		return SP_OPENGL_ES2;
 	if (ShaderFormat == NAME_OPENGL_ES2_WEBGL)	return SP_OPENGL_ES2_WEBGL;
 	if (ShaderFormat == NAME_OPENGL_ES2_IOS)	return SP_OPENGL_ES2_IOS;
+	if (ShaderFormat == NAME_SF_METAL)			return SP_METAL;
 	return SP_NumPlatforms;
 }
 
diff --git a/Engine/Source/Runtime/RHI/Public/RHI.h b/Engine/Source/Runtime/RHI/Public/RHI.h
--- a/Engine/Source/Runtime/RHI/Public/RHI.h
+++ b/Engine/Source/Runtime/RHI/Public/RHI.h
@@ -351,14 +351,16 @@ struct FVertexElement
 	uint8 Offset;
 	TEnumAsByte<EVertexElementType> Type;
 	uint8 AttributeIndex;
-	bool bUseInstanceIndex;
+	uint16 Stride;
+	uint16 bUseInstanceIndex;
 
 	FVertexElement() {}
-	FVertexElement(uint8 InStreamIndex,uint8 InOffset,EVertexElementType InType,uint8 InAttributeIndex,bool bInUseInstanceIndex = false):
+	FVertexElement(uint8 InStreamIndex,uint8 InOffset,EVertexElementType InType,uint8 InAttributeIndex,uint16 InStride,bool bInUseInstanceIndex = false):
 		StreamIndex(InStreamIndex),
 		Offset(InOffset),
 		Type(InType),
 		AttributeIndex(InAttributeIndex),
+		Stride(InStride),
 		bUseInstanceIndex(bInUseInstanceIndex)
 	{}
 	/**
@@ -380,6 +382,7 @@ struct FVertexElement
 		Ar << Element.Offset;
 		Ar << Element.Type;
 		Ar << Element.AttributeIndex;
+		Ar << Element.Stride;
 		Ar << Element.bUseInstanceIndex;
 		return Ar;
 	}
diff --git a/Engine/Source/Runtime/RHI/Public/RHIDefinitions.h b/Engine/Source/Runtime/RHI/Public/RHIDefinitions.h
--- a/Engine/Source/Runtime/RHI/Public/RHIDefinitions.h
+++ b/Engine/Source/Runtime/RHI/Public/RHIDefinitions.h
@@ -40,8 +40,9 @@ enum EShaderPlatform
 	SP_OPENGL_ES2		= 8,
 	SP_OPENGL_ES2_WEBGL = 9, 
 	SP_OPENGL_ES2_IOS	= 10,
+	SP_METAL			= 11,
 
-	SP_NumPlatforms		= 11,
+	SP_NumPlatforms		= 12,
 	SP_NumBits			= 4,
 };
 static_assert(SP_NumPlatforms <= (1 << SP_NumBits), "SP_NumPlatforms will not fit on SP_NumBits");
@@ -523,7 +524,7 @@ inline bool IsPCPlatform(const EShaderPlatform Platform)
 /** Whether the shader platform corresponds to the ES2 feature level. */
 inline bool IsES2Platform(const EShaderPlatform Platform)
 {
-	return Platform == SP_PCD3D_ES2 || Platform == SP_OPENGL_PCES2 || Platform == SP_OPENGL_ES2 || Platform == SP_OPENGL_ES2_WEBGL || Platform == SP_OPENGL_ES2_IOS; 
+	return Platform == SP_PCD3D_ES2 || Platform == SP_OPENGL_PCES2 || Platform == SP_OPENGL_ES2 || Platform == SP_OPENGL_ES2_WEBGL || Platform == SP_OPENGL_ES2_IOS || Platform == SP_METAL; 
 }
 
 inline bool IsOpenGLPlatform(const EShaderPlatform Platform)
@@ -552,7 +553,8 @@ inline ERHIFeatureLevel::Type GetMaxSupportedFeatureLevel(EShaderPlatform InShad
 	case SP_OPENGL_PCES2:
 	case SP_OPENGL_ES2:
 	case SP_OPENGL_ES2_WEBGL:
-	case  SP_OPENGL_ES2_IOS:
+	case SP_OPENGL_ES2_IOS:
+	case SP_METAL:
 		return ERHIFeatureLevel::ES2;
 	default:
 		check(0);
@@ -587,6 +589,8 @@ inline bool IsFeatureLevelSupported(EShaderPlatform InShaderPlatform, ERHIFeatur
 		return InFeatureLevel <= ERHIFeatureLevel::SM5;
 	case SP_XBOXONE:
 		return InFeatureLevel <= ERHIFeatureLevel::SM5;
+	case SP_METAL: 
+		return InFeatureLevel <= ERHIFeatureLevel::ES2;
 	default:
 		return false;
 	}	
diff --git a/Engine/Source/Runtime/RHI/Public/RHIUtilities.h b/Engine/Source/Runtime/RHI/Public/RHIUtilities.h
--- a/Engine/Source/Runtime/RHI/Public/RHIUtilities.h
+++ b/Engine/Source/Runtime/RHI/Public/RHIUtilities.h
@@ -247,12 +247,15 @@ inline void RHICreateTargetableShaderResource2D(
 	// Ensure that the targetable texture is either render or depth-stencil targetable.
 	check(TargetableTextureFlags & (TexCreate_RenderTargetable | TexCreate_DepthStencilTargetable | TexCreate_UAV));
 
-	if(NumSamples > 1) 
+	if (NumSamples > 1)
 	{
 		bForceSeparateTargetAndShaderResource = true;
 	}
 
-	if(!bForceSeparateTargetAndShaderResource/* && GSupportsRenderDepthTargetableShaderResources*/)
+	// ES2 doesn't support resolve operations.
+    bForceSeparateTargetAndShaderResource &= (GRHIFeatureLevel > ERHIFeatureLevel::ES2);
+
+	if (!bForceSeparateTargetAndShaderResource/* && GSupportsRenderDepthTargetableShaderResources*/)
 	{
 		// Create a single texture that has both TargetableTextureFlags and TexCreate_ShaderResource set.
 		OutTargetableTexture = OutShaderResourceTexture = RHICreateTexture2D(SizeX,SizeY,Format,NumMips,NumSamples,Flags | TargetableTextureFlags | TexCreate_ShaderResource,CreateInfo);
@@ -297,10 +300,7 @@ inline void RHICreateTargetableShaderResourceCube(
 	check(TargetableTextureFlags & (TexCreate_RenderTargetable | TexCreate_DepthStencilTargetable));
 
 	// ES2 doesn't support resolve operations.
-	if (GRHIFeatureLevel <= ERHIFeatureLevel::ES2)
-	{
-		bForceSeparateTargetAndShaderResource = false;
-	}
+    bForceSeparateTargetAndShaderResource &= (GRHIFeatureLevel > ERHIFeatureLevel::ES2);
 
 	if(!bForceSeparateTargetAndShaderResource/* && GSupportsRenderDepthTargetableShaderResources*/)
 	{
diff --git a/Engine/Source/Runtime/RenderCore/Private/RenderUtils.cpp b/Engine/Source/Runtime/RenderCore/Private/RenderUtils.cpp
--- a/Engine/Source/Runtime/RenderCore/Private/RenderUtils.cpp
+++ b/Engine/Source/Runtime/RenderCore/Private/RenderUtils.cpp
@@ -727,7 +727,7 @@ public:
 	virtual void InitRHI()
 	{
 		FVertexDeclarationElementList Elements;
-		Elements.Add(FVertexElement(0,0,VET_Float4,0));
+		Elements.Add(FVertexElement(0, 0, VET_Float4, 0, sizeof(FVector4)));
 		VertexDeclarationRHI = RHICreateVertexDeclaration(Elements);
 	}
 	virtual void ReleaseRHI()
@@ -750,7 +750,7 @@ public:
 	virtual void InitRHI()
 	{
 		FVertexDeclarationElementList Elements;
-		Elements.Add(FVertexElement(0,0,VET_Float3,0));
+		Elements.Add(FVertexElement(0, 0, VET_Float3, 0, sizeof(FVector)));
 		VertexDeclarationRHI = RHICreateVertexDeclaration(Elements);
 	}
 	virtual void ReleaseRHI()
diff --git a/Engine/Source/Runtime/RenderCore/Private/RenderingThread.cpp b/Engine/Source/Runtime/RenderCore/Private/RenderingThread.cpp
--- a/Engine/Source/Runtime/RenderCore/Private/RenderingThread.cpp
+++ b/Engine/Source/Runtime/RenderCore/Private/RenderingThread.cpp
@@ -323,6 +323,8 @@ public:
 	{
 		GRenderThreadId = FPlatformTLS::GetCurrentThreadId();
 
+		FPlatformProcess::SetupGameOrRenderThread(true);
+
 #if PLATFORM_WINDOWS
 		if ( !FPlatformMisc::IsDebuggerPresent() || GAlwaysReportCrash )
 		{
diff --git a/Engine/Source/Runtime/RenderCore/Public/UniformBuffer.h b/Engine/Source/Runtime/RenderCore/Public/UniformBuffer.h
--- a/Engine/Source/Runtime/RenderCore/Public/UniformBuffer.h
+++ b/Engine/Source/Runtime/RenderCore/Public/UniformBuffer.h
@@ -30,6 +30,7 @@ class TUniformBuffer : public FRenderResource
 public:
 
 	TUniformBuffer()
+		: BufferUsage(UniformBuffer_MultiFrame)
 	{
 		Contents = (uint8*)FMemory::Malloc(sizeof(TBufferStruct),UNIFORM_BUFFER_STRUCT_ALIGNMENT);
 	}
@@ -50,7 +51,7 @@ public:
 	// FRenderResource interface.
 	virtual void InitDynamicRHI()
 	{
-		UniformBufferRHI = RHICreateUniformBuffer(Contents,TBufferStruct::StaticStruct.GetLayout(),UniformBuffer_MultiFrame);
+		UniformBufferRHI = RHICreateUniformBuffer(Contents,TBufferStruct::StaticStruct.GetLayout(),BufferUsage);
 	}
 	virtual void ReleaseDynamicRHI()
 	{
@@ -60,6 +61,7 @@ public:
 	// Accessors.
 	FUniformBufferRHIParamRef GetUniformBufferRHI() const { return UniformBufferRHI; }
 
+	EUniformBufferUsage BufferUsage;
 private:
 
 	FUniformBufferRHIRef UniformBufferRHI;
diff --git a/Engine/Source/Runtime/Renderer/Private/AtmosphereRendering.cpp b/Engine/Source/Runtime/Renderer/Private/AtmosphereRendering.cpp
--- a/Engine/Source/Runtime/Renderer/Private/AtmosphereRendering.cpp
+++ b/Engine/Source/Runtime/Renderer/Private/AtmosphereRendering.cpp
@@ -265,7 +265,7 @@ public:
 	virtual void InitRHI()
 	{
 		FVertexDeclarationElementList Elements;
-		Elements.Add(FVertexElement(0,0,VET_Float2,0));
+		Elements.Add(FVertexElement(0, 0, VET_Float2, 0, sizeof(FVector2D)));
 		VertexDeclarationRHI = RHICreateVertexDeclaration(Elements);
 	}
 
diff --git a/Engine/Source/Runtime/Renderer/Private/FogRendering.cpp b/Engine/Source/Runtime/Renderer/Private/FogRendering.cpp
--- a/Engine/Source/Runtime/Renderer/Private/FogRendering.cpp
+++ b/Engine/Source/Runtime/Renderer/Private/FogRendering.cpp
@@ -183,7 +183,7 @@ public:
 	virtual void InitRHI()
 	{
 		FVertexDeclarationElementList Elements;
-		Elements.Add(FVertexElement(0,0,VET_Float2,0));
+		Elements.Add(FVertexElement(0, 0, VET_Float2, 0, sizeof(FVector2D)));
 		VertexDeclarationRHI = RHICreateVertexDeclaration(Elements);
 	}
 
diff --git a/Engine/Source/Runtime/Renderer/Private/PostProcess/PostProcessHMD.cpp b/Engine/Source/Runtime/Renderer/Private/PostProcess/PostProcessHMD.cpp
--- a/Engine/Source/Runtime/Renderer/Private/PostProcess/PostProcessHMD.cpp
+++ b/Engine/Source/Runtime/Renderer/Private/PostProcess/PostProcessHMD.cpp
@@ -27,12 +27,13 @@ public:
 
 	virtual void InitRHI()
 	{
+		uint16 Stride = sizeof(FDistortionVertex);
 		FVertexDeclarationElementList Elements;
-		Elements.Add(FVertexElement(0, STRUCT_OFFSET(FDistortionVertex, Position),VET_Float2,0));
-		Elements.Add(FVertexElement(0, STRUCT_OFFSET(FDistortionVertex, TexR), VET_Float2, 1));
-		Elements.Add(FVertexElement(0, STRUCT_OFFSET(FDistortionVertex, TexG), VET_Float2, 2));
-		Elements.Add(FVertexElement(0, STRUCT_OFFSET(FDistortionVertex, TexB), VET_Float2, 3));
-		Elements.Add(FVertexElement(0, STRUCT_OFFSET(FDistortionVertex, Color), VET_Float4, 4));
+		Elements.Add(FVertexElement(0, STRUCT_OFFSET(FDistortionVertex, Position),VET_Float2,0,Stride));
+		Elements.Add(FVertexElement(0, STRUCT_OFFSET(FDistortionVertex, TexR), VET_Float2, 1,Stride));
+		Elements.Add(FVertexElement(0, STRUCT_OFFSET(FDistortionVertex, TexG), VET_Float2, 2,Stride));
+		Elements.Add(FVertexElement(0, STRUCT_OFFSET(FDistortionVertex, TexB), VET_Float2, 3,Stride));
+		Elements.Add(FVertexElement(0, STRUCT_OFFSET(FDistortionVertex, Color), VET_Float4, 4,Stride));
 		VertexDeclarationRHI = RHICreateVertexDeclaration(Elements);
 	}
 
diff --git a/Engine/Source/Runtime/Renderer/Private/PostProcess/SceneFilterRendering.cpp b/Engine/Source/Runtime/Renderer/Private/PostProcess/SceneFilterRendering.cpp
--- a/Engine/Source/Runtime/Renderer/Private/PostProcess/SceneFilterRendering.cpp
+++ b/Engine/Source/Runtime/Renderer/Private/PostProcess/SceneFilterRendering.cpp
@@ -155,7 +155,7 @@ void DrawRectangle(
 			/*NumVertices=*/ 3,
 			/*StartIndex=*/ 6,
 			/*NumPrimitives=*/ 1,
-			/*NumInstances=*/ 0
+			/*NumInstances=*/ 1
 			);
 	}
 	else
@@ -168,7 +168,7 @@ void DrawRectangle(
 			/*NumVertices=*/ 4,
 			/*StartIndex=*/ 0,
 			/*NumPrimitives=*/ 2,
-			/*NumInstances=*/ 0
+			/*NumInstances=*/ 1
 			);
 	}
 }
diff --git a/Engine/Source/Runtime/Renderer/Private/PostProcess/SceneFilterRendering.h b/Engine/Source/Runtime/Renderer/Private/PostProcess/SceneFilterRendering.h
--- a/Engine/Source/Runtime/Renderer/Private/PostProcess/SceneFilterRendering.h
+++ b/Engine/Source/Runtime/Renderer/Private/PostProcess/SceneFilterRendering.h
@@ -30,8 +30,9 @@ public:
 	virtual void InitRHI()
 	{
 		FVertexDeclarationElementList Elements;
-		Elements.Add(FVertexElement(0,STRUCT_OFFSET(FFilterVertex,Position),VET_Float4,0));
-		Elements.Add(FVertexElement(0,STRUCT_OFFSET(FFilterVertex,UV),VET_Float2,1));
+		uint32 Stride = sizeof(FFilterVertex);
+		Elements.Add(FVertexElement(0,STRUCT_OFFSET(FFilterVertex,Position),VET_Float4,0,Stride));
+		Elements.Add(FVertexElement(0,STRUCT_OFFSET(FFilterVertex,UV),VET_Float2,1,Stride));
 		VertexDeclarationRHI = RHICreateVertexDeclaration(Elements);
 	}
 
diff --git a/Engine/Source/Runtime/Renderer/Private/PostProcess/SceneRenderTargets.cpp b/Engine/Source/Runtime/Renderer/Private/PostProcess/SceneRenderTargets.cpp
--- a/Engine/Source/Runtime/Renderer/Private/PostProcess/SceneRenderTargets.cpp
+++ b/Engine/Source/Runtime/Renderer/Private/PostProcess/SceneRenderTargets.cpp
@@ -60,7 +60,10 @@ static TAutoConsoleVariable<int32> CVarCustomDepth(
 TAutoConsoleVariable<int32> CVarMobileMSAA(
 	TEXT("r.MobileMSAA"),
 	0,
-	TEXT("Whether to use MSAA instead of Temporal AA on mobile"),
+	TEXT("Use MSAA instead of Temporal AA on mobile:\n")
+    TEXT("1: Use Temporal AA (MSAA disabled)\n")
+    TEXT("2: Use 2x MSAA (Temporal AA disabled)\n")
+    TEXT("4: Use 4x MSAA (Temporal AA disabled)\n"),
 	ECVF_RenderThreadSafe
 	);
 
@@ -170,7 +173,7 @@ void FSceneRenderTargets::Allocate(const FSceneViewFamily& ViewFamily)
 	uint32 Mobile32bpp = !IsMobileHDR() || IsMobileHDR32bpp();
 
 	static const auto CVarMobileMSAA = IConsoleManager::Get().FindTConsoleVariableDataInt(TEXT("r.MobileMSAA"));
-	int32 MobileMSAA = CVarMobileMSAA->GetValueOnRenderThread();
+	int32 MobileMSAA = GRHIShaderPlatform == SP_OPENGL_ES2_IOS ? 1 : CVarMobileMSAA->GetValueOnRenderThread();
 
 	bool bLightPropagationVolume = UseLightPropagationVolumeRT();
 
@@ -318,10 +321,12 @@ void FSceneRenderTargets::AllocSceneColor()
 			Desc.TargetableFlags |= TexCreate_UAV;
 		}
 
-		if (GRHIFeatureLevel == ERHIFeatureLevel::ES2)
+	    uint16 NumSamples = GRHIShaderPlatform == SP_OPENGL_ES2_IOS ? 1 : CVarMobileMSAA.GetValueOnRenderThread();
+        if (NumSamples != 1 && NumSamples != 2 && NumSamples != 4)
 		{
-			Desc.NumSamples = CVarMobileMSAA.GetValueOnRenderThread() ? 4 : 1;
+            NumSamples = 1;
 		}
+		Desc.NumSamples = NumSamples;
 
 		GRenderTargetPool.FindFreeElement(Desc, SceneColor, TEXT("SceneColor"));
 	}
diff --git a/Engine/Source/Runtime/ShaderCore/Private/VertexFactory.cpp b/Engine/Source/Runtime/ShaderCore/Private/VertexFactory.cpp
--- a/Engine/Source/Runtime/ShaderCore/Private/VertexFactory.cpp
+++ b/Engine/Source/Runtime/ShaderCore/Private/VertexFactory.cpp
@@ -248,7 +248,7 @@ FVertexElement FVertexFactory::AccessStreamComponent(const FVertexStreamComponen
 	VertexStream.Stride = Component.Stride;
 	VertexStream.Offset = 0;
 
-	return FVertexElement(Streams.AddUnique(VertexStream),Component.Offset,Component.Type,AttributeIndex,Component.bUseInstanceIndex);
+	return FVertexElement(Streams.AddUnique(VertexStream),Component.Offset,Component.Type,AttributeIndex,VertexStream.Stride,Component.bUseInstanceIndex);
 }
 
 FVertexElement FVertexFactory::AccessPositionStreamComponent(const FVertexStreamComponent& Component,uint8 AttributeIndex)
@@ -258,7 +258,7 @@ FVertexElement FVertexFactory::AccessPositionStreamComponent(const FVertexStream
 	VertexStream.Stride = Component.Stride;
 	VertexStream.Offset = 0;
 
-	return FVertexElement(PositionStream.AddUnique(VertexStream),Component.Offset,Component.Type,AttributeIndex,Component.bUseInstanceIndex);
+	return FVertexElement(PositionStream.AddUnique(VertexStream),Component.Offset,Component.Type,AttributeIndex,VertexStream.Stride,Component.bUseInstanceIndex);
 }
 
 //TTP:51684
@@ -329,6 +329,7 @@ static void PatchVertexStreamOffsetsToBeUnique( FVertexDeclarationElementList& E
 				VertElement.StreamIndex = Elements[ i ].StreamIndex;
 				VertElement.Type		= Elements[ i ].Type;
 				VertElement.AttributeIndex	= Elements[ i ].AttributeIndex;
+				VertElement.Stride = Elements[ i ].Stride;
 				VertElement.bUseInstanceIndex = Elements[i].bUseInstanceIndex;
 
 				// remove the old redundant element
diff --git a/Engine/Source/Runtime/SlateRHIRenderer/Private/SlateShaders.cpp b/Engine/Source/Runtime/SlateRHIRenderer/Private/SlateShaders.cpp
--- a/Engine/Source/Runtime/SlateRHIRenderer/Private/SlateShaders.cpp
+++ b/Engine/Source/Runtime/SlateRHIRenderer/Private/SlateShaders.cpp
@@ -40,10 +40,11 @@ TGlobalResource<FSlateVertexDeclaration> GSlateVertexDeclaration;
 void FSlateVertexDeclaration::InitRHI()
 {
 	FVertexDeclarationElementList Elements;
-	Elements.Add(FVertexElement(0,STRUCT_OFFSET(FSlateVertex,TexCoords),VET_Float4,0));
-	Elements.Add(FVertexElement(0,STRUCT_OFFSET(FSlateVertex,ClipCoords),VET_Short4,1));
-	Elements.Add(FVertexElement(0,STRUCT_OFFSET(FSlateVertex,Position),VET_Short2,2));
-	Elements.Add(FVertexElement(0,STRUCT_OFFSET(FSlateVertex,Color),VET_Color,3));
+	uint32 Stride = sizeof(FSlateVertex);
+	Elements.Add(FVertexElement(0,STRUCT_OFFSET(FSlateVertex,TexCoords),VET_Float4,0,Stride));
+	Elements.Add(FVertexElement(0,STRUCT_OFFSET(FSlateVertex,ClipCoords),VET_Short4,1,Stride));
+	Elements.Add(FVertexElement(0,STRUCT_OFFSET(FSlateVertex,Position),VET_Short2,2,Stride));
+	Elements.Add(FVertexElement(0,STRUCT_OFFSET(FSlateVertex,Color),VET_Color,3,Stride));
 
 	VertexDeclarationRHI = RHICreateVertexDeclaration(Elements);
 }
diff --git a/Engine/Source/Runtime/Windows/D3D11RHI/Private/D3D11RHIPrivate.h b/Engine/Source/Runtime/Windows/D3D11RHI/Private/D3D11RHIPrivate.h
--- a/Engine/Source/Runtime/Windows/D3D11RHI/Private/D3D11RHIPrivate.h
+++ b/Engine/Source/Runtime/Windows/D3D11RHI/Private/D3D11RHIPrivate.h
@@ -733,7 +733,7 @@ public:
 	virtual void InitRHI() override
 	{
 		FVertexDeclarationElementList Elements;
-		Elements.Add(FVertexElement(0,0,VET_Float4,0));
+		Elements.Add(FVertexElement(0,0,VET_Float4,0,sizeof(FVector4)));
 		VertexDeclarationRHI = RHICreateVertexDeclaration(Elements);
 	}
 	virtual void ReleaseRHI() override
