Add method to check if we're running on battery and lock to VSync accordingly
#ttp 337189 Laptops should throttle to 60 hz in editor to reduce battery drain
#doc added "r.DontLimitOnBattery" CVar, 0==Limit performance on device with battery, 1== Do not limit performance.
reviewers Thomas.Sarkanen, Mark.Satterthwaite

diff --git a/Engine/Source/Editor/UnrealEd/Private/Editor.cpp b/Engine/Source/Editor/UnrealEd/Private/Editor.cpp
--- a/Engine/Source/Editor/UnrealEd/Private/Editor.cpp
+++ b/Engine/Source/Editor/UnrealEd/Private/Editor.cpp
@@ -1465,6 +1465,14 @@ float UEditorEngine::GetMaxTickRate( float DeltaTime, bool bAllowFrameRateSmooth
 		}
 	}
 
+	// TTP 337189 - Laptops should throttle to 60 hz in editor to reduce battery drain
+	static const auto CVarDontLimitOnBattery = IConsoleManager::Get().FindTConsoleVariableDataInt(TEXT("r.DontLimitOnBattery"));
+	const bool bLimitOnBattery = (FPlatformMisc::IsRunningOnBattery() && CVarDontLimitOnBattery->GetValueOnGameThread() == 0);
+	if( bLimitOnBattery )
+	{
+		MaxTickRate = 60.0f;
+	}
+
 	return MaxTickRate;
 }
 
diff --git a/Engine/Source/Runtime/Core/Private/GenericPlatform/GenericPlatformMisc.cpp b/Engine/Source/Runtime/Core/Private/GenericPlatform/GenericPlatformMisc.cpp
--- a/Engine/Source/Runtime/Core/Private/GenericPlatform/GenericPlatformMisc.cpp
+++ b/Engine/Source/Runtime/Core/Private/GenericPlatform/GenericPlatformMisc.cpp
@@ -713,3 +713,8 @@ FText FGenericPlatformMisc::GetFileManagerName()
 {
 	return NSLOCTEXT("GenericPlatform", "FileManagerName", "File Manager");
 }
+
+bool FGenericPlatformMisc::IsRunningOnBattery()
+{
+	return false;
+}
diff --git a/Engine/Source/Runtime/Core/Private/HAL/ConsoleManager.cpp b/Engine/Source/Runtime/Core/Private/HAL/ConsoleManager.cpp
--- a/Engine/Source/Runtime/Core/Private/HAL/ConsoleManager.cpp
+++ b/Engine/Source/Runtime/Core/Private/HAL/ConsoleManager.cpp
@@ -1928,6 +1928,12 @@ void CreateConsoleVariables()
 		TEXT("1: VSync is enabled."),
 		ECVF_Scalability | ECVF_RenderThreadSafe);
 
+	IConsoleManager::Get().RegisterConsoleVariable(TEXT("r.DontLimitOnBattery"),
+		0,
+		TEXT("0: Limit performance on devices with a battery.(default)\n")
+		TEXT("1: Do not limit performance due to device having a battery."),
+		ECVF_Scalability | ECVF_RenderThreadSafe);
+
 	IConsoleManager::Get().RegisterConsoleVariable(TEXT("r.FinishCurrentFrame"),
 		0,
 		TEXT("If on, the current frame will be forced to finish and render to the screen instead of being buffered.  This will improve latency, but slow down overall performance."),
diff --git a/Engine/Source/Runtime/Core/Private/Mac/MacPlatformMisc.cpp b/Engine/Source/Runtime/Core/Private/Mac/MacPlatformMisc.cpp
--- a/Engine/Source/Runtime/Core/Private/Mac/MacPlatformMisc.cpp
+++ b/Engine/Source/Runtime/Core/Private/Mac/MacPlatformMisc.cpp
@@ -17,8 +17,11 @@
 #include <IOKit/network/IOEthernetInterface.h>
 #include <IOKit/network/IONetworkInterface.h>
 #include <IOKit/network/IOEthernetController.h>
+#include <IOKit/ps/IOPowerSources.h>
+#include <IOKit/ps/IOPSKeys.h>
 #include <mach-o/dyld.h>
 #include <libproc.h>
+#include <notify.h>
 
 
 /**
@@ -140,8 +143,38 @@ struct MacApplicationInfo
 		AppPath = FPlatformProcess::GenerateApplicationPath(FApp::GetName(), FApp::GetBuildConfiguration());
 		
 		LCID = FString::Printf(TEXT("%d"), FInternationalization::Get().GetCurrentCulture()->GetLCID());
+		
+		// Notification handler to check we are running from a battery - this only applies to MacBook's.
+		notify_handler_t PowerSourceNotifyHandler = ^(int32 Token){
+			RunningOnBattery = false;
+			CFTypeRef PowerSourcesInfo = IOPSCopyPowerSourcesInfo();
+			if (PowerSourcesInfo)
+			{
+				CFArrayRef PowerSourcesArray = IOPSCopyPowerSourcesList(PowerSourcesInfo);
+				for (CFIndex Index = 0; Index < CFArrayGetCount(PowerSourcesArray); Index++)
+				{
+					CFTypeRef PowerSource = CFArrayGetValueAtIndex(PowerSourcesArray, Index);
+					NSDictionary* Description = (NSDictionary*)IOPSGetPowerSourceDescription(PowerSourcesInfo, PowerSource);
+					if ([(NSString*)[Description objectForKey: @kIOPSPowerSourceStateKey] isEqualToString: @kIOPSBatteryPowerValue])
+					{
+						RunningOnBattery = true;
+						break;
+					}
+				}
+				CFRelease(PowerSourcesArray);
+				CFRelease(PowerSourcesInfo);
+			}
+		};
+		
+		// Call now to fetch the status
+		PowerSourceNotifyHandler(0);
+		
+		uint32 Status = notify_register_dispatch(kIOPSNotifyPowerSource, &PowerSourceNotification, dispatch_get_main_queue(), PowerSourceNotifyHandler);
+		check(Status == NOTIFY_STATUS_OK);
 	}
 	
+	bool RunningOnBattery;
+	int32 PowerSourceNotification;
 	char AppNameUTF8[PATH_MAX+1];
 	char AppLogPath[PATH_MAX+1];
 	char CrashReportPath[PATH_MAX+1];
@@ -205,8 +238,10 @@ void FMacPlatformMisc::PlatformInit()
 
 	// Timer resolution.
 	UE_LOG(LogInit, Log, TEXT("High frequency timer resolution =%f MHz"), 0.000001 / FPlatformTime::GetSecondsPerCycle() );
-
+	
 	GMacAppInfo.Init();
+	
+	UE_LOG(LogInit, Log, TEXT("Power Source: %s"), GMacAppInfo.RunningOnBattery ? TEXT(kIOPSBatteryPowerValue) : TEXT(kIOPSACPowerValue) );
 }
 
 void FMacPlatformMisc::PlatformPostInit()
@@ -697,6 +732,10 @@ uint32 FMacPlatformMisc::GetKeyMap( uint16* KeyCodes, FString* KeyNames, uint32
 void FMacPlatformMisc::RequestExit( bool Force )
 {
 	UE_LOG(LogMac, Log,  TEXT("FPlatformMisc::RequestExit(%i)"), Force );
+	
+	notify_cancel(GMacAppInfo.PowerSourceNotification);
+	GMacAppInfo.PowerSourceNotification = 0;
+	
 	if( Force )
 	{
 		// Abort allows signal handler to know we aborted.
@@ -1116,6 +1155,11 @@ FText FMacPlatformMisc::GetFileManagerName()
 	return NSLOCTEXT("MacPlatform", "FileManagerName", "Finder");
 }
 
+bool FMacPlatformMisc::IsRunningOnBattery()
+{
+	return GMacAppInfo.RunningOnBattery;
+}
+
 /** Global pointer to crash handler */
 void (* GCrashHandlerPointer)(const FGenericCrashContext & Context) = NULL;
 
diff --git a/Engine/Source/Runtime/Core/Private/Windows/WindowsPlatformMisc.cpp b/Engine/Source/Runtime/Core/Private/Windows/WindowsPlatformMisc.cpp
--- a/Engine/Source/Runtime/Core/Private/Windows/WindowsPlatformMisc.cpp
+++ b/Engine/Source/Runtime/Core/Private/Windows/WindowsPlatformMisc.cpp
@@ -2248,3 +2248,23 @@ FText FWindowsPlatformMisc::GetFileManagerName()
 {
 	return NSLOCTEXT("WindowsPlatform", "FileManagerName", "Explorer");
 }
+
+bool FWindowsPlatformMisc::IsRunningOnBattery()
+{
+	SYSTEM_POWER_STATUS status;
+	GetSystemPowerStatus(&status);
+	switch(status.BatteryFlag)
+	{
+	case 4://	"Critical—the battery capacity is at less than five percent"
+	case 2://	"Low—the battery capacity is at less than 33 percent"
+	case 1://	"High—the battery capacity is at more than 66 percent"
+	case 8://	"Charging"
+		return true;
+	case 128://	"No system battery" - desktop, NB: UPS don't count as batteries under Windows
+	case 255://	"Unknown status—unable to read the battery flag information"
+	default:
+		return false;
+	}
+
+	return false;
+}
diff --git a/Engine/Source/Runtime/Core/Public/GenericPlatform/GenericPlatformMisc.h b/Engine/Source/Runtime/Core/Public/GenericPlatform/GenericPlatformMisc.h
--- a/Engine/Source/Runtime/Core/Public/GenericPlatform/GenericPlatformMisc.h
+++ b/Engine/Source/Runtime/Core/Public/GenericPlatform/GenericPlatformMisc.h
@@ -752,6 +752,11 @@ public:
 		return false;
 	}
 
+	/**
+	 * Returns whether the platform is running on battery power or not.
+	 */
+	static bool IsRunningOnBattery();
+
 #if !UE_BUILD_SHIPPING
 protected:
 	/** Whether the user should be prompted to allow for a remote debugger to be attached */
diff --git a/Engine/Source/Runtime/Core/Public/Mac/MacPlatformMisc.h b/Engine/Source/Runtime/Core/Public/Mac/MacPlatformMisc.h
--- a/Engine/Source/Runtime/Core/Public/Mac/MacPlatformMisc.h
+++ b/Engine/Source/Runtime/Core/Public/Mac/MacPlatformMisc.h
@@ -123,6 +123,11 @@ struct CORE_API FMacPlatformMisc : public FGenericPlatformMisc
 	static FText GetFileManagerName();
 
 	static void UpdateWindowMenu();
+
+	/**
+	 * Returns whether the platform is running on battery power or not.
+	 */
+	static bool IsRunningOnBattery();
 };
 
 #ifdef __OBJC__
diff --git a/Engine/Source/Runtime/Core/Public/Windows/WindowsPlatformMisc.h b/Engine/Source/Runtime/Core/Public/Windows/WindowsPlatformMisc.h
--- a/Engine/Source/Runtime/Core/Public/Windows/WindowsPlatformMisc.h
+++ b/Engine/Source/Runtime/Core/Public/Windows/WindowsPlatformMisc.h
@@ -210,6 +210,11 @@ struct CORE_API FWindowsPlatformMisc : public FGenericPlatformMisc
 
 	/** @return Get the name of the platform specific file manager (Explorer) */
 	static FText GetFileManagerName();
+
+	/**
+	 * Returns whether the platform is running on battery power or not.
+	 */
+	static bool IsRunningOnBattery();
 };
 
 typedef FWindowsPlatformMisc FPlatformMisc;
