- Missed a file that wasn't checked out right

diff --git a/Engine/Source/Runtime/Core/Public/IOS/IOSPlatformProcess.h b/Engine/Source/Runtime/Core/Public/IOS/IOSPlatformProcess.h
--- a/Engine/Source/Runtime/Core/Public/IOS/IOSPlatformProcess.h
+++ b/Engine/Source/Runtime/Core/Public/IOS/IOSPlatformProcess.h
@@ -31,6 +31,7 @@ struct CORE_API FIOSPlatformProcess : public FGenericPlatformProcess
 	static FRunnableThread* CreateRunnableThread();
 	static void LaunchURL( const TCHAR* URL, const TCHAR* Parms, FString* Error );
     static void SetRealTimeMode();
+	static void SetupGameOrRenderThread(bool bIsRenderThread);
 };
 
 typedef FIOSPlatformProcess FPlatformProcess;
