Fix for TextTest, timezone was not set so it was calculating to a different time based on where the test was run.

#ttp 338786 - CRITICAL: REGRESSION: AUTOMATION: Text Test generates errors: "FormattedTestLayer2 did not rebuild to correctly in LEET!"

diff --git a/Engine/Source/Runtime/Core/Private/Tests/Internationalization/TextTest.cpp b/Engine/Source/Runtime/Core/Private/Tests/Internationalization/TextTest.cpp
--- a/Engine/Source/Runtime/Core/Private/Tests/Internationalization/TextTest.cpp
+++ b/Engine/Source/Runtime/Core/Private/Tests/Internationalization/TextTest.cpp
@@ -427,7 +427,7 @@ bool FTextTest::RunTest (const FString& Parameters)
 			FText AsCurrencyTest1 = FText::AsCurrency(100.25);
 
 			FDateTime DateTimeInfo(2080, 8, 20, 9, 33, 22);
-			FText AsDateTimeTest1 = FText::AsDateTime(DateTimeInfo);
+			FText AsDateTimeTest1 = FText::AsDateTime(DateTimeInfo, EDateTimeStyle::Default, EDateTimeStyle::Default, TEXT("CST"));
 
 			// FormattedTestLayer2 must be updated when adding or removing from this block. Further, below, 
 			// verifying the LEET translated string must be changed to reflect what the new string looks like.
