- Merging EAGLView and MetalView into one class, they share a bunch of stuff, and there are too many if (Delegate.MetalView? Delegate.MetalView : Delegate.EAGLViewl) type things around
- Added better checking for device support for Metal
- Added proper support for WeakFrameworks in UBT

diff --git a/Engine/Source/Developer/StandaloneRenderer/StandaloneRenderer.Build.cs b/Engine/Source/Developer/StandaloneRenderer/StandaloneRenderer.Build.cs
--- a/Engine/Source/Developer/StandaloneRenderer/StandaloneRenderer.Build.cs
+++ b/Engine/Source/Developer/StandaloneRenderer/StandaloneRenderer.Build.cs
@@ -34,7 +34,9 @@ public class StandaloneRenderer : ModuleRules
 
 		if (Target.Platform == UnrealTargetPlatform.IOS)
 		{
-			PublicFrameworks.AddRange(new string[] { "OpenGLES", "QuartzCore", "GLKit" });
+			PublicFrameworks.AddRange(new string[] { "OpenGLES", "GLKit" });
+			// weak for IOS8 support since CAMetalLayer is in QuartzCore
+			PublicWeakFrameworks.AddRange(new string[] { "QuartzCore" });
 		}
 	}
 }
diff --git a/Engine/Source/Programs/UnrealBuildTool/Configuration/UEBuildModule.cs b/Engine/Source/Programs/UnrealBuildTool/Configuration/UEBuildModule.cs
--- a/Engine/Source/Programs/UnrealBuildTool/Configuration/UEBuildModule.cs
+++ b/Engine/Source/Programs/UnrealBuildTool/Configuration/UEBuildModule.cs
@@ -132,6 +132,7 @@ namespace UnrealBuildTool
 		protected readonly List<string> PublicLibraryPaths;
 		protected readonly List<string> PublicAdditionalLibraries;
 		protected readonly List<string> PublicFrameworks;
+		protected readonly List<string> PublicWeakFrameworks;
 		protected readonly List<UEBuildFramework> PublicAdditionalFrameworks;
 		protected readonly List<string> PublicAdditionalShadowFiles;
 
@@ -175,6 +176,7 @@ namespace UnrealBuildTool
 			IEnumerable<string> InPublicLibraryPaths = null,
 			IEnumerable<string> InPublicAdditionalLibraries = null,
 			IEnumerable<string> InPublicFrameworks = null,
+			IEnumerable<string> InPublicWeakFrameworks = null,
 			IEnumerable<UEBuildFramework> InPublicAdditionalFrameworks = null,
 			IEnumerable<string> InPublicAdditionalShadowFiles = null,
 			IEnumerable<string> InPublicIncludePathModuleNames = null,
@@ -199,7 +201,8 @@ namespace UnrealBuildTool
 			PublicLibraryPaths = ListFromOptionalEnumerableStringParameter(InPublicLibraryPaths);
 			PublicAdditionalLibraries = ListFromOptionalEnumerableStringParameter(InPublicAdditionalLibraries);
 			PublicFrameworks = ListFromOptionalEnumerableStringParameter(InPublicFrameworks);
-			PublicAdditionalFrameworks = InPublicAdditionalFrameworks == null ? new List<UEBuildFramework>() : new List<UEBuildFramework>( InPublicAdditionalFrameworks );
+			PublicWeakFrameworks = ListFromOptionalEnumerableStringParameter(InPublicWeakFrameworks);
+			PublicAdditionalFrameworks = InPublicAdditionalFrameworks == null ? new List<UEBuildFramework>() : new List<UEBuildFramework>(InPublicAdditionalFrameworks);
 			PublicAdditionalShadowFiles = ListFromOptionalEnumerableStringParameter(InPublicAdditionalShadowFiles);
 			PublicIncludePathModuleNames = ListFromOptionalEnumerableStringParameter( InPublicIncludePathModuleNames );
 			PublicDependencyModuleNames = ListFromOptionalEnumerableStringParameter(InPublicDependencyModuleNames);
@@ -272,6 +275,19 @@ namespace UnrealBuildTool
 			PublicFrameworks.Add(LibraryName);
 		}
 
+		/** Adds a public weak framework */
+		public virtual void AddPublicWeakFramework(string LibraryName, bool bCheckForDuplicates = true)
+		{
+			if (bCheckForDuplicates == true)
+			{
+				if (PublicWeakFrameworks.Contains(LibraryName))
+				{
+					return;
+				}
+			}
+			PublicWeakFrameworks.Add(LibraryName);
+		}
+
 		/** Adds a public additional framework */
 		public virtual void AddPublicAdditionalFramework(UEBuildFramework Framework, bool bCheckForDuplicates = true)
 		{
@@ -597,6 +613,7 @@ namespace UnrealBuildTool
 			ref List<string> LibraryPaths,
 			ref List<string> AdditionalLibraries,
 			ref List<string> Frameworks,
+			ref List<string> WeakFrameworks,
 			ref List<UEBuildFramework> AdditionalFrameworks,
 			ref List<string> AdditionalShadowFiles,
 			ref List<string> DelayLoadDLLs,
@@ -638,7 +655,7 @@ namespace UnrealBuildTool
 							bool bIsInStaticLibrary = (DependencyModule.Binary != null && DependencyModule.Binary.Config.Type == UEBuildBinaryType.StaticLibrary);
 							if (bIsExternalModule || bIsInStaticLibrary)
 							{
-								DependencyModule.SetupPublicLinkEnvironment(SourceBinary, ref LibraryPaths, ref AdditionalLibraries, ref Frameworks,
+								DependencyModule.SetupPublicLinkEnvironment(SourceBinary, ref LibraryPaths, ref AdditionalLibraries, ref Frameworks, ref WeakFrameworks,
 									ref AdditionalFrameworks, ref AdditionalShadowFiles, ref DelayLoadDLLs, ref BinaryDependencies, ref VisitedModules);
 							}
 						}
@@ -648,6 +665,7 @@ namespace UnrealBuildTool
 					LibraryPaths.AddRange(PublicLibraryPaths);
 					AdditionalLibraries.AddRange(PublicAdditionalLibraries);
 					Frameworks.AddRange(PublicFrameworks);
+					WeakFrameworks.AddRange(PublicWeakFrameworks);
 					// Remember the module so we can refer to it when needed
 					foreach ( var Framework in PublicAdditionalFrameworks )
 					{
@@ -668,7 +686,7 @@ namespace UnrealBuildTool
 			)
 		{
 			// Allow the module's public dependencies to add library paths and additional libraries to the link environment.
-			SetupPublicLinkEnvironment(Binary,ref LinkEnvironment.Config.LibraryPaths,ref LinkEnvironment.Config.AdditionalLibraries,ref LinkEnvironment.Config.Frameworks,
+			SetupPublicLinkEnvironment(Binary,ref LinkEnvironment.Config.LibraryPaths,ref LinkEnvironment.Config.AdditionalLibraries,ref LinkEnvironment.Config.Frameworks,ref LinkEnvironment.Config.WeakFrameworks,
 				ref LinkEnvironment.Config.AdditionalFrameworks,ref LinkEnvironment.Config.AdditionalShadowFiles,ref LinkEnvironment.Config.DelayLoadDLLs,ref BinaryDependencies,ref VisitedModules);
 
 			// Also allow the module's public and private dependencies to modify the link environment.
@@ -678,7 +696,7 @@ namespace UnrealBuildTool
 			foreach (var DependencyName in AllDependencyModuleNames)
 			{
 				var DependencyModule = Target.GetModuleByName(DependencyName);
-				DependencyModule.SetupPublicLinkEnvironment(Binary,ref LinkEnvironment.Config.LibraryPaths,ref LinkEnvironment.Config.AdditionalLibraries,ref LinkEnvironment.Config.Frameworks,
+				DependencyModule.SetupPublicLinkEnvironment(Binary,ref LinkEnvironment.Config.LibraryPaths,ref LinkEnvironment.Config.AdditionalLibraries,ref LinkEnvironment.Config.Frameworks,ref LinkEnvironment.Config.WeakFrameworks,
 					ref LinkEnvironment.Config.AdditionalFrameworks,ref LinkEnvironment.Config.AdditionalShadowFiles,ref LinkEnvironment.Config.DelayLoadDLLs,ref BinaryDependencies,ref VisitedModules);
 			}
 		}
@@ -747,6 +765,7 @@ namespace UnrealBuildTool
 			IEnumerable<string> InPublicLibraryPaths = null,
 			IEnumerable<string> InPublicAdditionalLibraries = null,
 			IEnumerable<string> InPublicFrameworks = null,
+			IEnumerable<string> InPublicWeakFrameworks = null,
 			IEnumerable<UEBuildFramework> InPublicAdditionalFrameworks = null,
 			IEnumerable<string> InPublicAdditionalShadowFiles = null,
 			IEnumerable<string> InPublicDependencyModuleNames = null,
@@ -764,7 +783,8 @@ namespace UnrealBuildTool
 			InPublicLibraryPaths:			InPublicLibraryPaths,
 			InPublicAdditionalLibraries:	InPublicAdditionalLibraries,
 			InPublicFrameworks:             InPublicFrameworks,
-			InPublicAdditionalFrameworks:	InPublicAdditionalFrameworks,
+			InPublicWeakFrameworks:			InPublicWeakFrameworks,
+			InPublicAdditionalFrameworks: 	InPublicAdditionalFrameworks,
 			InPublicAdditionalShadowFiles:  InPublicAdditionalShadowFiles,
 			InPublicDependencyModuleNames:	InPublicDependencyModuleNames,
 			InPublicDelayLoadDLLs:			InPublicDelayLoadDLLs
@@ -939,6 +959,7 @@ namespace UnrealBuildTool
 			IEnumerable<string> InPublicDelayLoadDLLs,
 			IEnumerable<string> InPublicAdditionalLibraries,
 			IEnumerable<string> InPublicFrameworks,
+			IEnumerable<string> InPublicWeakFrameworks,
 			IEnumerable<UEBuildFramework> InPublicAdditionalFrameworks,
 			IEnumerable<string> InPublicAdditionalShadowFiles,
 			IEnumerable<string> InPrivateIncludePaths,
@@ -968,6 +989,7 @@ namespace UnrealBuildTool
 					null, 
 					InPublicAdditionalLibraries,
 					InPublicFrameworks,
+					InPublicWeakFrameworks,
 					InPublicAdditionalFrameworks,
 					InPublicAdditionalShadowFiles,
 					InPublicIncludePathModuleNames,
@@ -1931,6 +1953,7 @@ namespace UnrealBuildTool
 			IEnumerable<string> InPublicDelayLoadDLLs,
 			IEnumerable<string> InPublicAdditionalLibraries,
 			IEnumerable<string> InPublicFrameworks,
+			IEnumerable<string> InPublicWeakFrameworks,
 			IEnumerable<UEBuildFramework> InPublicAdditionalFrameworks,
 			IEnumerable<string> InPublicAdditionalShadowFiles,
 			IEnumerable<string> InPrivateIncludePaths,
@@ -1951,7 +1974,7 @@ namespace UnrealBuildTool
 			)
 		: base(InTarget,InName,InType,InModuleDirectory,InOutputDirectory,InIntelliSenseGatherer,
 			InSourceFiles,InPublicIncludePaths,InPublicSystemIncludePaths,InDefinitions,
-			InPublicIncludePathModuleNames,InPublicDependencyModuleNames,InPublicDelayLoadDLLs,InPublicAdditionalLibraries,InPublicFrameworks,InPublicAdditionalFrameworks,InPublicAdditionalShadowFiles,
+			InPublicIncludePathModuleNames,InPublicDependencyModuleNames,InPublicDelayLoadDLLs,InPublicAdditionalLibraries,InPublicFrameworks,InPublicWeakFrameworks,InPublicAdditionalFrameworks,InPublicAdditionalShadowFiles,
 			InPrivateIncludePaths,InPrivateIncludePathModuleNames,InPrivateDependencyModuleNames,
             InCircularlyReferencedDependentModules, InDynamicallyLoadedModuleNames, InPlatformSpecificDynamicallyLoadedModuleNames, InOptimizeCode,
 			InAllowSharedPCH, InSharedPCHHeaderFile, InUseRTTI, InEnableBufferSecurityChecks, InFasterWithoutUnity, InMinFilesUsingPrecompiledHeaderOverride,
diff --git a/Engine/Source/Programs/UnrealBuildTool/Configuration/UEBuildTarget.cs b/Engine/Source/Programs/UnrealBuildTool/Configuration/UEBuildTarget.cs
--- a/Engine/Source/Programs/UnrealBuildTool/Configuration/UEBuildTarget.cs
+++ b/Engine/Source/Programs/UnrealBuildTool/Configuration/UEBuildTarget.cs
@@ -1686,6 +1686,7 @@ namespace UnrealBuildTool
 					InPublicDelayLoadDLLs: new List<string>(),
 					InPublicAdditionalLibraries: new List<string>(),
 					InPublicFrameworks: new List<string>(),
+					InPublicWeakFrameworks: new List<string>(),
 					InPublicAdditionalShadowFiles: new List<string>(),
 					InPublicAdditionalFrameworks: new List<UEBuildFramework>(),
 					InPrivateIncludePaths: new List<string>(),
@@ -2631,6 +2632,7 @@ namespace UnrealBuildTool
 							InPublicDelayLoadDLLs: RulesObject.PublicDelayLoadDLLs,
 							InPublicAdditionalLibraries: RulesObject.PublicAdditionalLibraries,
 							InPublicFrameworks: RulesObject.PublicFrameworks,
+							InPublicWeakFrameworks: RulesObject.PublicWeakFrameworks,
 							InPublicAdditionalFrameworks: RulesObject.PublicAdditionalFrameworks,
 							InPublicAdditionalShadowFiles: RulesObject.PublicAdditionalShadowFiles,
 							InPrivateIncludePaths: RulesObject.PrivateIncludePaths,
@@ -2667,6 +2669,7 @@ namespace UnrealBuildTool
 							InPublicDelayLoadDLLs: RulesObject.PublicDelayLoadDLLs,
 							InPublicAdditionalLibraries: RulesObject.PublicAdditionalLibraries,
 							InPublicFrameworks: RulesObject.PublicFrameworks,
+							InPublicWeakFrameworks: RulesObject.PublicWeakFrameworks,
 							InPublicAdditionalFrameworks: RulesObject.PublicAdditionalFrameworks,
 							InPublicAdditionalShadowFiles: RulesObject.PublicAdditionalShadowFiles,
 							InPrivateIncludePaths: RulesObject.PrivateIncludePaths,
@@ -2700,6 +2703,7 @@ namespace UnrealBuildTool
 							InPublicLibraryPaths: RulesObject.PublicLibraryPaths,
 							InPublicAdditionalLibraries: RulesObject.PublicAdditionalLibraries,
 							InPublicFrameworks: RulesObject.PublicFrameworks,
+							InPublicWeakFrameworks: RulesObject.PublicWeakFrameworks,
 							InPublicAdditionalFrameworks: RulesObject.PublicAdditionalFrameworks,
 							InPublicAdditionalShadowFiles: RulesObject.PublicAdditionalShadowFiles,
 							InPublicDependencyModuleNames: RulesObject.PublicDependencyModuleNames,
diff --git a/Engine/Source/Programs/UnrealBuildTool/System/RulesCompiler.cs b/Engine/Source/Programs/UnrealBuildTool/System/RulesCompiler.cs
--- a/Engine/Source/Programs/UnrealBuildTool/System/RulesCompiler.cs
+++ b/Engine/Source/Programs/UnrealBuildTool/System/RulesCompiler.cs
@@ -177,6 +177,9 @@ namespace UnrealBuildTool
 		// List of frameworks
 		public List<string> PublicFrameworks = new List<string>();
 
+		// List of weak frameworks (for OS version transitions)
+		public List<string> PublicWeakFrameworks = new List<string>();
+
 		/// List of addition frameworks - typically used for External (third party) modules on Mac and iOS
 		public List<UEBuildFramework> PublicAdditionalFrameworks = new List<UEBuildFramework>();
 
diff --git a/Engine/Source/Runtime/Core/Private/IOS/BLView.cpp b/Engine/Source/Runtime/Core/Private/IOS/BLView.cpp
deleted file mode 100644
diff --git a/Engine/Source/Runtime/Core/Private/IOS/EAGLView.cpp b/Engine/Source/Runtime/Core/Private/IOS/EAGLView.cpp
deleted file mode 100644
diff --git a/Engine/Source/Runtime/Core/Private/IOS/IOSAppDelegate.cpp b/Engine/Source/Runtime/Core/Private/IOS/IOSAppDelegate.cpp
--- a/Engine/Source/Runtime/Core/Private/IOS/IOSAppDelegate.cpp
+++ b/Engine/Source/Runtime/Core/Private/IOS/IOSAppDelegate.cpp
@@ -6,6 +6,7 @@
 #include "ExceptionHandling.h"
 #include "ModuleBoilerplate.h"
 #include "CallbackDevice.h"
+#include "IOSView.h"
 
 #include <AudioToolbox/AudioToolbox.h>
 #include <AVFoundation/AVAudioSession.h>
@@ -68,10 +69,7 @@ void InstallSignalHandlers()
 @synthesize bResetIdleTimer;
 
 @synthesize Window;
-@synthesize GLView;
-#if HAS_METAL
-@synthesize MetalView;
-#endif
+@synthesize IOSView;
 @synthesize IOSController;
 @synthesize SlateController;
 
@@ -342,11 +340,7 @@ void InstallSignalHandlers()
 	}
     
 	// if you run this on the main thread super early, it will deadlock
-	if ((GLView && GLView->bIsInitialized) 
-#if HAS_METAL
-		|| (MetalView && MetalView->bIsInitialized)
-#endif
-		)
+	if (IOSView && IOSView->bIsInitialized) 
 	{
 		while(!self.bHasSuspended)
 		{
diff --git a/Engine/Source/Runtime/Core/Private/IOS/IOSPlatformMisc.cpp b/Engine/Source/Runtime/Core/Private/IOS/IOSPlatformMisc.cpp
--- a/Engine/Source/Runtime/Core/Private/IOS/IOSPlatformMisc.cpp
+++ b/Engine/Source/Runtime/Core/Private/IOS/IOSPlatformMisc.cpp
@@ -271,24 +271,14 @@ void FIOSPlatformMisc::ReleaseAutoreleasePool(void *Pool)
 void* FIOSPlatformMisc::GetHardwareWindow()
 {
 	IOSAppDelegate* AppDelegate = [IOSAppDelegate GetDelegate];
-    return AppDelegate.GLView 
-#if HAS_METAL
-		? AppDelegate.GLView : AppDelegate.MetalView
-#endif
-		;
+    return AppDelegate.IOSView; 
 }
 
 bool FIOSPlatformMisc::HasPlatformFeature(const TCHAR* FeatureName)
 {
 	if (FCString::Stricmp(FeatureName, TEXT("Metal")) == 0)
 	{
-#if HAS_METAL
-		// @todo metal: Obviously not the right way to detect Metal support, but waiting for the right way...
-		FIOSPlatformMisc::EIOSDevice DeviceType = GetIOSDeviceType();
-		return DeviceType == IOS_IPhone5S || DeviceType == IOS_IPadMini2 || DeviceType == IOS_IPadAir;
-#else
-		return false;
-#endif
+		return [IOSAppDelegate GetDelegate].IOSView->bIsUsingMetal;
 	}
 
 	return FGenericPlatformMisc::HasPlatformFeature(FeatureName);
diff --git a/Engine/Source/Runtime/Core/Private/IOS/IOSView.cpp b/Engine/Source/Runtime/Core/Private/IOS/IOSView.cpp
new file mode 100644
--- /dev/null
+++ b/Engine/Source/Runtime/Core/Private/IOS/IOSView.cpp
@@ -0,0 +1,456 @@
+// Copyright 1998-2014 Epic Games, Inc. All Rights Reserved.
+
+#include "CorePrivate.h"
+#include "IOSView.h"
+#include "IOSAppDelegate.h"
+#include "IOS/IOSInputInterface.h"
+
+#include <OpenGLES/ES2/gl.h>
+#import <QuartzCore/QuartzCore.h>
+#import <OpenGLES/EAGLDrawable.h>
+#import <UIKit/UIGeometry.h>
+
+#if HAS_METAL
+id<MTLDevice> GMetalDevice = nil;
+#endif
+
+@implementation FIOSView
+
+@synthesize SwapCount, OnScreenColorRenderBuffer, OnScreenColorRenderBufferMSAA;
+
+/**
+ * @return The Layer Class for the window
+ */
++ (Class)layerClass
+{
+#if HAS_METAL
+	if (!FParse::Param(FCommandLine::Get(), TEXT("ES2")) && MTLCreateSystemDefaultDevice != NULL)
+	{
+		// if the device is unable to run with Metal (pre-A7), this will return nil
+		GMetalDevice = MTLCreateSystemDefaultDevice();
+	}
+
+	if (GMetalDevice != nil)
+	{
+		return [CAMetalLayer class];
+	}
+	else
+#endif
+	{
+		return [CAEAGLLayer class];
+	}
+}
+
+-(BOOL)becomeFirstResponder
+{
+	return YES;
+}
+
+- (id)initWithFrame:(CGRect)Frame
+{
+	if ((self = [super initWithFrame:Frame]))
+	{
+		// figure out if we should start up GL or Metal
+#if HAS_METAL
+		// if the device is valid, we know Metal is usable (see +layerClass)
+		MetalDevice = GMetalDevice;
+		if (MetalDevice != nil)
+		{
+			bIsUsingMetal = true;
+
+			// grab the MetalLayer and typecast it to match what's in layerClass
+			CAMetalLayer* MetalLayer = (CAMetalLayer*)self.layer;
+			MetalLayer.presentsWithTransaction = NO;
+			MetalLayer.drawsAsynchronously = YES;
+
+			// set a background color to make sure the layer appears
+			CGFloat components[] = { 0.0, 0.0, 0.0, 1 };
+			MetalLayer.backgroundColor = CGColorCreate(CGColorSpaceCreateDeviceRGB(), components);
+
+			// set the device on the rendering layer and provide a pixel format
+			MetalLayer.device = MetalDevice;
+			MetalLayer.pixelFormat = MTLPixelFormatBGRA8Unorm_sRGB;
+			MetalLayer.framebufferOnly = NO;
+		
+		}
+		else
+#endif
+		{
+			bIsUsingMetal = false;
+
+			// Get the layer
+			CAEAGLLayer *EaglLayer = (CAEAGLLayer *)self.layer;
+			EaglLayer.opaque = YES;
+			NSMutableDictionary* Dict = [NSMutableDictionary dictionary];
+			[Dict setValue : [NSNumber numberWithBool : NO] forKey : kEAGLDrawablePropertyRetainedBacking];
+			[Dict setValue : kEAGLColorFormatRGBA8 forKey : kEAGLDrawablePropertyColorFormat];
+			EaglLayer.drawableProperties = Dict;
+
+			// Initialize a single, static OpenGL ES 2.0 context, shared by all EAGLView objects
+			Context = [[EAGLContext alloc] initWithAPI:kEAGLRenderingAPIOpenGLES2];
+
+			// delete this on failure
+			if (!Context || ![EAGLContext setCurrentContext : Context])
+			{
+				[self release];
+				return nil;
+			}
+		}
+
+		NSLog(@"::: Created a UIView that will support %@ :::", bIsUsingMetal ? @"Metal" : @"@GLES");
+
+		// Initialize some variables
+		SwapCount = 0;
+
+		FMemory::Memzero(AllTouches, sizeof(AllTouches));
+
+		bIsInitialized = false;
+	}
+	return self;
+}
+
+
+- (bool)CreateFramebuffer:(bool)bIsForOnDevice
+{
+	if (!bIsInitialized)
+	{
+		// look up what the device can support
+		const float NativeScale = [[UIScreen mainScreen] scale];
+
+		// look up the CVar for the scale factor
+		static IConsoleVariable* CVar = IConsoleManager::Get().FindConsoleVariable(TEXT("r.MobileContentScaleFactor"));
+		float RequestedContentScaleFactor = CVar->GetFloat();
+
+		// for TV screens, always use scale factor of 1
+		self.contentScaleFactor = bIsForOnDevice ? RequestedContentScaleFactor : 1.0f;
+		UE_LOG(LogIOS, Log, TEXT("Setting contentScaleFactor to %0.4f (optimal = %0.4f)"), self.contentScaleFactor, NativeScale);
+
+		// metal has no further action needed
+		if (!bIsUsingMetal)
+		{
+			// make sure this is current
+			[self MakeCurrent];
+
+			// This is causing a pretty large slow down on iPad 3 and 4 using iOS 6, for now going to comment it out
+	//		if (self.contentScaleFactor == 1.0f || self.contentScaleFactor == 2.0f)
+	//		{
+	//			UE_LOG(LogIOS,Log,TEXT("Setting layer filter to NEAREST"));
+	//			CAEAGLLayer *EaglLayer = (CAEAGLLayer *)self.layer;
+	//			EaglLayer.magnificationFilter = kCAFilterNearest;
+	//		}
+
+			// Create our standard displayable surface
+			glGenRenderbuffers(1, &OnScreenColorRenderBuffer);
+			check(glGetError() == 0);
+			glBindRenderbuffer(GL_RENDERBUFFER, OnScreenColorRenderBuffer);
+			check(glGetError() == 0);
+			[Context renderbufferStorage:GL_RENDERBUFFER fromDrawable:(CAEAGLLayer*)self.layer];
+
+			// Get the size of the surface
+			GLint OnScreenWidth, OnScreenHeight;
+			glGetRenderbufferParameteriv(GL_RENDERBUFFER, GL_RENDERBUFFER_WIDTH, &OnScreenWidth);
+			glGetRenderbufferParameteriv(GL_RENDERBUFFER, GL_RENDERBUFFER_HEIGHT, &OnScreenHeight);
+
+			// NOTE: This resolve FBO is necessary even if we don't plan on using MSAA because otherwise
+			// the shaders will not warm properly. Future investigation as to why; it seems unnecessary.
+
+			// Create an FBO used to target the resolve surface
+			glGenFramebuffers(1, &ResolveFrameBuffer);
+			check(glGetError() == 0);
+			glBindFramebuffer(GL_FRAMEBUFFER, ResolveFrameBuffer);
+			check(glGetError() == 0);
+			glFramebufferRenderbuffer(GL_FRAMEBUFFER, GL_COLOR_ATTACHMENT0, GL_RENDERBUFFER, OnScreenColorRenderBuffer);
+			check(glGetError() == 0);
+
+#if USE_DETAILED_IPHONE_MEM_TRACKING
+			//This value is used to allow the engine to track gl allocated memory see GetIPhoneOpenGLBackBufferSize
+			uint32 SingleBufferSize = OnScreenWidth * OnScreenHeight * 4/*rgba8*/;
+			IPhoneBackBufferMemSize = singleBufferSize * 3/*iphone back buffer system is tripple buffered*/;
+
+			UE_LOG(LogIOS, Log, TEXT("IPhone Back Buffer Size: %i MB"), (IPhoneBackBufferMemSize / 1024) / 1024.f);
+#endif
+		}
+
+		bIsInitialized = true;
+	}    
+	return true;
+}
+
+#if HAS_METAL
+- (id<CAMetalDrawable>)MakeDrawable
+{
+	return [(CAMetalLayer*)self.layer newDrawable];
+}
+#endif
+
+- (void)DestroyFramebuffer
+{
+	if (bIsInitialized)
+	{
+		// nothing to do here for Metal
+		if (!bIsUsingMetal)
+		{
+			// toss framebuffers
+			if(ResolveFrameBuffer)
+			{
+				glDeleteFramebuffers(1, &ResolveFrameBuffer);
+			}
+			if(OnScreenColorRenderBuffer)
+			{
+				glDeleteRenderbuffers(1, &OnScreenColorRenderBuffer);
+				OnScreenColorRenderBuffer = 0;
+			}
+// 			if( GMSAAAllowed )
+// 			{
+// 				if(OnScreenColorRenderBufferMSAA)
+// 				{
+// 					glDeleteRenderbuffers(1, &OnScreenColorRenderBufferMSAA);
+// 					OnScreenColorRenderBufferMSAA = 0;
+// 				}
+//			}
+ 		}
+
+		// we are ready to be re-initialized
+		bIsInitialized = false;
+	}
+}
+
+- (void)MakeCurrent
+{
+	if (!bIsUsingMetal)
+	{
+		[EAGLContext setCurrentContext:Context];
+	}
+}
+
+- (void)UnmakeCurrent
+{
+	if (!bIsUsingMetal)
+	{
+		[EAGLContext setCurrentContext:nil];
+	}
+}
+
+- (void)SwapBuffers
+{
+	if (!bIsUsingMetal)
+	{
+// 		// We may need this in the MSAA case
+// 		GLint CurrentFramebuffer = 0;
+// 		// @todo-mobile: Fix this when we have MSAA support
+// 		if( GMSAAAllowed && GMSAAEnabled )
+// 		{
+// 			// Get the currently bound FBO
+// 			glGetIntegerv(GL_FRAMEBUFFER_BINDING, &CurrentFramebuffer);
+// 
+// 			// Set up and perform the resolve (the READ is already set)
+// 			glBindFramebuffer(GL_DRAW_FRAMEBUFFER_APPLE, ResolveFrameBuffer);
+// 			glResolveMultisampleFramebufferAPPLE();
+// 
+// 			// After the resolve, we can discard the old attachments
+// 			GLenum attachments[] = { GL_COLOR_ATTACHMENT0, GL_DEPTH_ATTACHMENT, GL_STENCIL_ATTACHMENT };
+// 			glDiscardFramebufferEXT( GL_READ_FRAMEBUFFER_APPLE, 3, attachments );
+// 		}
+// 		else
+// 		{
+// 			// Discard the now-unncessary depth buffer
+// 			GLenum attachments[] = { GL_DEPTH_ATTACHMENT, GL_STENCIL_ATTACHMENT };
+// 			glDiscardFramebufferEXT( GL_READ_FRAMEBUFFER_APPLE, 2, attachments );
+// 		}
+// 
+		// Perform the actual present with the on-screen renderbuffer
+		//glBindRenderbuffer(GL_RENDERBUFFER, OnScreenColorRenderBuffer);
+		[Context presentRenderbuffer:GL_RENDERBUFFER];
+
+// 		if( GMSAAAllowed && GMSAAEnabled )
+// 		{
+// 			// Restore the DRAW framebuffer object
+// 			glBindFramebuffer(GL_DRAW_FRAMEBUFFER_APPLE, CurrentFramebuffer);
+// 		}
+	}
+
+	// increment our swap counter
+	SwapCount++;
+}
+
+/**
+ * Returns the unique ID for the given touch
+ */
+-(int32) GetTouchIndex:(UITouch*)Touch
+{
+	// look for existing touch
+	for (int Index = 0; Index < ARRAY_COUNT(AllTouches); Index++)
+	{
+		if (AllTouches[Index] == Touch)
+		{
+			return Index;
+		}
+	}
+	
+	// if we get here, it's a new touch, find a slot
+	for (int Index = 0; Index < ARRAY_COUNT(AllTouches); Index++)
+	{
+		if (AllTouches[Index] == nil)
+		{
+			AllTouches[Index] = Touch;
+			return Index;
+		}
+	}
+	
+	// if we get here, that means we are trying to use more than 5 touches at once, which is an error
+	return -1;
+}
+
+/**
+ * Pass touch events to the input queue for slate to pull off of, and trigger the debug console.
+ *
+ * @param View The view the event happened in
+ * @param Touches Set of touch events from the OS
+ */
+-(void) HandleTouches:(NSSet*)Touches ofType:(TouchType)Type
+{
+	CGFloat Scale = self.contentScaleFactor;
+
+	TArray<TouchInput> TouchesArray;
+	for (UITouch* Touch in Touches)
+	{
+		// get info from the touch
+		CGPoint Loc = [Touch locationInView:self];
+		CGPoint PrevLoc = [Touch previousLocationInView:self];
+
+		// convert TOuch pointer to a unique 0 based index
+		int32 TouchIndex = [self GetTouchIndex:Touch];
+		if (TouchIndex < 0)
+		{
+			continue;
+		}
+
+		// make a new touch event struct
+		TouchInput TouchMessage;
+		TouchMessage.Handle = TouchIndex;
+		TouchMessage.Type = Type;
+		TouchMessage.Position = FVector2D(FMath::Min<float>(self.frame.size.width - 1, Loc.x), FMath::Min<float>(self.frame.size.height - 1, Loc.y)) * Scale;
+		TouchMessage.LastPosition = FVector2D(FMath::Min<float>(self.frame.size.width - 1, PrevLoc.x), FMath::Min<float>(self.frame.size.height - 1, PrevLoc.y)) * Scale;
+		TouchesArray.Add(TouchMessage);
+
+		// clear out the touch when it ends
+		if (Type == TouchEnded)
+		{
+			AllTouches[TouchIndex] = nil;
+		}
+	}
+
+	FIOSInputInterface::QueueTouchInput(TouchesArray);
+}
+
+/**
+ * Handle the various touch types from the OS
+ *
+ * @param touches Array of touch events
+ * @param event Event information
+ */
+- (void) touchesBegan:(NSSet*)touches withEvent:(UIEvent*)event 
+{
+	NumActiveTouches += touches.count;
+	[self HandleTouches:touches ofType:TouchBegan];
+
+#if !UE_BUILD_SHIPPING
+#if WITH_SIMULATOR
+	// use 2 on the simulator so that Option-Click will bring up console (option-click is for doing pinch gestures, which we don't care about, atm)
+	if( NumActiveTouches >= 2 )
+#else
+	// If there are 3 active touches, bring up the console
+	if( NumActiveTouches >= 4 )
+#endif
+	{
+		// Route the command to the main iOS thread (all UI must go to the main thread)
+		[[IOSAppDelegate GetDelegate] performSelectorOnMainThread:@selector(ShowConsole) withObject:nil waitUntilDone:NO];
+	}
+#endif
+}
+
+- (void) touchesMoved:(NSSet*)touches withEvent:(UIEvent*)event
+{
+	[self HandleTouches:touches ofType:TouchMoved];
+}
+
+- (void) touchesEnded:(NSSet*)touches withEvent:(UIEvent*)event
+{
+	NumActiveTouches -= touches.count;
+	[self HandleTouches:touches ofType:TouchEnded];
+}
+
+- (void) touchesCancelled:(NSSet*)touches withEvent:(UIEvent*)event
+{
+	NumActiveTouches -= touches.count;
+	[self HandleTouches:touches ofType:TouchEnded];
+}
+
+@end
+
+@implementation IOSViewController
+
+/**
+ * The ViewController was created, so now we need to create our view to be controlled (an EAGLView)
+ */
+- (void) loadView
+{
+	// get the landcape size of the screen
+	CGRect Frame = [[UIScreen mainScreen] bounds];
+	if (![IOSAppDelegate GetDelegate].bDeviceInPortraitMode)
+	{
+		Swap(Frame.size.width, Frame.size.height);
+	}
+
+	self.view = [[UIView alloc] initWithFrame:Frame];
+
+	// settings copied from InterfaceBuilder
+#if defined(__IPHONE_7_0)
+	if ([IOSAppDelegate GetDelegate].OSVersion >= 7.0)
+	{
+		self.edgesForExtendedLayout = UIRectEdgeNone;
+	}
+#endif
+
+	self.view.clearsContextBeforeDrawing = NO;
+	self.view.multipleTouchEnabled = NO;
+}
+
+/**
+ * View was unloaded from us
+ */ 
+- (void) viewDidUnload
+{
+	UE_LOG(LogIOS, Log, TEXT("IOSViewController unloaded the view. This is unexpected, tell Josh Adams"));
+}
+
+/**
+ * Tell the OS that our view controller can auto-rotate between the two landscape modes
+ */
+- (BOOL)shouldAutorotateToInterfaceOrientation:(UIInterfaceOrientation)interfaceOrientation
+{
+	//BOOL bIsValidOrientation;
+
+	//// will the app let us rotate?
+	//BOOL bCanRotate = [IPhoneAppDelegate GetDelegate].bCanAutoRotate;
+	//
+	//// is it one of the valid orientations to rotate to?
+	//bIsValidOrientation = bCanRotate && 
+	//	(((interfaceOrientation == UIInterfaceOrientationLandscapeLeft) && ![IOSAppDelegate GetDelegate].bDeviceInPortraitMode) || 
+	//	((interfaceOrientation == UIInterfaceOrientationLandscapeRight) && ![IOSAppDelegate GetDelegate].bDeviceInPortraitMode) || 
+	//	((interfaceOrientation == UIInterfaceOrientationPortrait) && [IOSAppDelegate GetDelegate].bDeviceInPortraitMode) ||
+	//	((interfaceOrientation == UIInterfaceOrientationPortraitUpsideDown) && [IOSAppDelegate GetDelegate].bDeviceInPortraitMode));
+	//	
+	//return bIsValidOrientation;
+	return YES;
+}
+
+/**
+ * Tell the OS to hide the status bar (iOS 7 method for hiding)
+ */
+- (BOOL)prefersStatusBarHidden
+{
+	return YES;
+}
+
+@end
diff --git a/Engine/Source/Runtime/Core/Private/IOS/IOSWindow.cpp b/Engine/Source/Runtime/Core/Private/IOS/IOSWindow.cpp
--- a/Engine/Source/Runtime/Core/Private/IOS/IOSWindow.cpp
+++ b/Engine/Source/Runtime/Core/Private/IOS/IOSWindow.cpp
@@ -3,8 +3,7 @@
 #include "CorePrivate.h"
 #include "IOSWindow.h"
 #include "IOSAppDelegate.h"
-#include "EAGLView.h"
-#include "BLView.h"
+#include "IOSView.h"
 
 FIOSWindow::~FIOSWindow()
 {
@@ -47,11 +46,7 @@ FPlatformRect FIOSWindow::GetScreenRect()
 {
 	// get the main view's frame
 	IOSAppDelegate* AppDelegate = (IOSAppDelegate*)[[UIApplication sharedApplication] delegate];
-	UIView* View = AppDelegate.GLView
-#if HAS_METAL
-		? AppDelegate.GLView : AppDelegate.MetalView
-#endif
-		;
+	UIView* View = AppDelegate.IOSView;
 	CGRect Frame = [View frame];
 	CGFloat Scale = View.contentScaleFactor;
 
diff --git a/Engine/Source/Runtime/Core/Public/IOS/BLView.h b/Engine/Source/Runtime/Core/Public/IOS/BLView.h
deleted file mode 100644
diff --git a/Engine/Source/Runtime/Core/Public/IOS/EAGLView.h b/Engine/Source/Runtime/Core/Public/IOS/EAGLView.h
deleted file mode 100644
diff --git a/Engine/Source/Runtime/Core/Public/IOS/IOSAppDelegate.h b/Engine/Source/Runtime/Core/Public/IOS/IOSAppDelegate.h
--- a/Engine/Source/Runtime/Core/Public/IOS/IOSAppDelegate.h
+++ b/Engine/Source/Runtime/Core/Public/IOS/IOSAppDelegate.h
@@ -6,8 +6,7 @@
 #import <AVFoundation/AVAudioSession.h>
 #import <GameKit/GKGameCenterViewController.h>
 
-@class EAGLView;
-@class FMetalView;
+@class FIOSView;
 @class IOSViewController;
 @class SlateOpenGLESViewController;
 @class IOSAppDelegate;
@@ -37,12 +36,7 @@ UITextFieldDelegate>
 @property (strong, retain, nonatomic) UIWindow *Window;
 
 /** Main GL View */
-@property (retain) EAGLView* GLView;
-
-#if HAS_METAL
-/** Main GL View */
-@property (retain) FMetalView* MetalView;
-#endif
+@property (retain) FIOSView* IOSView;
 
 /** The controller to handle rotation of the view */
 @property (retain) IOSViewController* IOSController;
diff --git a/Engine/Source/Runtime/Core/Public/IOS/IOSView.h b/Engine/Source/Runtime/Core/Public/IOS/IOSView.h
new file mode 100644
--- /dev/null
+++ b/Engine/Source/Runtime/Core/Public/IOS/IOSView.h
@@ -0,0 +1,75 @@
+// Copyright 1998-2014 Epic Games, Inc. All Rights Reserved.
+
+#pragma once
+
+#import <UIKit/UIKit.h>
+#import <OpenGLES/EAGL.h>
+
+#if HAS_METAL
+#import <Metal/Metal.h>
+#import <QuartzCore/CAMetalLayer.h>
+#endif
+
+@interface FIOSView : UIView
+{
+@public
+	// are we initialized?
+	bool bIsInitialized;
+
+//@private
+	// keeps track of the number of active touches
+	// used to bring up the three finger touch debug console after 3 active touches are registered
+	int NumActiveTouches;
+
+	// track the touches by pointer (which will stay the same for a given finger down) - note we don't deref the pointers in this array
+	UITouch* AllTouches[10];
+
+
+	//// GL MEMBERS
+	// the GL context
+	EAGLContext* Context;
+
+	// the internal MSAA FBO used to resolve the color buffer at present-time
+	GLuint ResolveFrameBuffer;
+
+
+	//// METAL MEMBERS
+#if HAS_METAL
+	// global metal device
+	id<MTLDevice> MetalDevice;
+#endif
+
+	// are we using the Metal API?
+	bool bIsUsingMetal;
+}
+
+//// SHARED FUNCTIONALITY
+@property (nonatomic) GLuint SwapCount;
+
+-(bool)CreateFramebuffer:(bool)bIsForOnDevice;
+-(void)DestroyFramebuffer;
+
+//// GL FUNCTIONALITY
+@property (nonatomic) GLuint OnScreenColorRenderBuffer;
+@property (nonatomic) GLuint OnScreenColorRenderBufferMSAA;
+
+- (void)MakeCurrent;
+- (void)UnmakeCurrent;
+- (void)SwapBuffers;
+
+//// METAL FUNCTIONALITY
+#if HAS_METAL
+// Return a drawable object (ie a back buffer texture) for the RHI to render to
+- (id<CAMetalDrawable>)MakeDrawable;
+#endif
+
+@end
+
+/**
+ * A view controller subclass that handles loading our IOS view as well as autorotation
+ */
+@interface IOSViewController : UIViewController
+{
+
+}
+@end
diff --git a/Engine/Source/Runtime/IOS/MetalRHI/MetalRHI.Build.cs b/Engine/Source/Runtime/IOS/MetalRHI/MetalRHI.Build.cs
--- a/Engine/Source/Runtime/IOS/MetalRHI/MetalRHI.Build.cs
+++ b/Engine/Source/Runtime/IOS/MetalRHI/MetalRHI.Build.cs
@@ -17,6 +17,6 @@ public class MetalRHI : ModuleRules
 			}
 			);
 			
-		PublicFrameworks.Add("Metal");
+		PublicWeakFrameworks.Add("Metal");
 	}
 }
diff --git a/Engine/Source/Runtime/IOS/MetalRHI/Private/MetalManager.cpp b/Engine/Source/Runtime/IOS/MetalRHI/Private/MetalManager.cpp
--- a/Engine/Source/Runtime/IOS/MetalRHI/Private/MetalManager.cpp
+++ b/Engine/Source/Runtime/IOS/MetalRHI/Private/MetalManager.cpp
@@ -211,7 +211,7 @@ void FMetalManager::ReleaseObject(id Object)
 MetalFramePacer *FramePacer;
 
 FMetalManager::FMetalManager()
-	: Device([IOSAppDelegate GetDelegate].MetalView->Device)
+	: Device([IOSAppDelegate GetDelegate].IOSView->MetalDevice)
 	, CurrentCommandBuffer(nil)
 	, CurrentDrawable(nil)
 	, CurrentFramebuffer(nil)
@@ -230,7 +230,7 @@ FMetalManager::FMetalManager()
 	CommandQueue = [Device newCommandQueue];
 
 	// get the size of the window
-	CGRect ViewFrame = [[IOSAppDelegate GetDelegate].MetalView frame];
+	CGRect ViewFrame = [[IOSAppDelegate GetDelegate].IOSView frame];
 	FRHIResourceCreateInfo CreateInfo;
 	BackBuffer = (FMetalTexture2D*)(FTexture2DRHIParamRef)RHICreateTexture2D(ViewFrame.size.width, ViewFrame.size.height, PF_B8G8R8A8, 1, 1, TexCreate_RenderTargetable | TexCreate_Presentable, CreateInfo);
 
@@ -311,7 +311,7 @@ void FMetalManager::InitFrame()
 		SCOPE_CYCLE_COUNTER(STAT_MetalMakeDrawableTime);
 
 		// make a drawable object for this frame
-		CurrentDrawable = [[IOSAppDelegate GetDelegate].MetalView MakeDrawable];
+		CurrentDrawable = [[IOSAppDelegate GetDelegate].IOSView MakeDrawable];
 		TRACK_OBJECT(CurrentDrawable);
 		BackBuffer->Surface.Texture = CurrentDrawable.texture;
 	}
@@ -535,7 +535,7 @@ void FMetalManager::SetCurrentRenderTarget(FMetalSurface* RenderSurface)
 				uint32 IdleStart = FPlatformTime::Cycles();
 
 				// make a drawable object for this frame
-				CurrentDrawable = [[IOSAppDelegate GetDelegate].MetalView MakeDrawable];
+				CurrentDrawable = [[IOSAppDelegate GetDelegate].IOSView MakeDrawable];
 				TRACK_OBJECT(CurrentDrawable);
 
 				GRenderThreadIdle[ERenderThreadIdleTypes::WaitingForGPUPresent] += FPlatformTime::Cycles() - IdleStart;
diff --git a/Engine/Source/Runtime/IOS/MetalRHI/Private/MetalManager.h b/Engine/Source/Runtime/IOS/MetalRHI/Private/MetalManager.h
--- a/Engine/Source/Runtime/IOS/MetalRHI/Private/MetalManager.h
+++ b/Engine/Source/Runtime/IOS/MetalRHI/Private/MetalManager.h
@@ -5,7 +5,7 @@
 
 #import <Metal/Metal.h>
 #import <QuartzCore/CAMetalLayer.h>
-#include "BLView.h"
+#include "IOSView.h"
 
 
 struct FPipelineShadow
diff --git a/Engine/Source/Runtime/Launch/Launch.Build.cs b/Engine/Source/Runtime/Launch/Launch.Build.cs
--- a/Engine/Source/Runtime/Launch/Launch.Build.cs
+++ b/Engine/Source/Runtime/Launch/Launch.Build.cs
@@ -145,7 +145,8 @@ public class Launch : ModuleRules
 			DynamicallyLoadedModuleNames.Add("IOSAudio");
 			DynamicallyLoadedModuleNames.Add("IOSRuntimeSettings");
 			PublicFrameworks.Add("OpenGLES");
-			PublicFrameworks.Add("QuartzCore");
+			// this is weak for IOS8 support for CAMetalLayer that is in QuartzCore
+			PublicWeakFrameworks.Add("QuartzCore");
 
 			PrivateDependencyModuleNames.Add("LaunchDaemonMessages");
 		}
diff --git a/Engine/Source/Runtime/Launch/Private/IOS/LaunchIOS.cpp b/Engine/Source/Runtime/Launch/Private/IOS/LaunchIOS.cpp
--- a/Engine/Source/Runtime/Launch/Private/IOS/LaunchIOS.cpp
+++ b/Engine/Source/Runtime/Launch/Private/IOS/LaunchIOS.cpp
@@ -4,8 +4,7 @@
 
 #include "LaunchPrivatePCH.h"
 #include "IOSAppDelegate.h"
-#include "EAGLView.h"
-#include "BLView.h"
+#include "IOSView.h"
 #include "IOSCommandLineHelper.h"
 #include "GameLaunchDaemonMessageHandler.h"
 #include "AudioDevice.h"
@@ -79,53 +78,18 @@ static void MainThreadInit()
 // 		Max<float>(MainFrame.size.height, GSystemSettings.SecondaryDisplayMaximumHeight) :
 // 		MainFrame.size.height
 // 		);
-	CGRect FullResolutionRect = MainFrame;
-
-#if HAS_METAL
-	if (FPlatformMisc::HasPlatformFeature(TEXT("Metal")) && !FParse::Param(FCommandLine::Get(), TEXT("ES2")))
-	{
-		AppDelegate.MetalView = [[FMetalView alloc] initWithFrame:FullResolutionRect];
-		AppDelegate.MetalView.clearsContextBeforeDrawing = NO;
-		AppDelegate.MetalView.multipleTouchEnabled = YES;
-
-		// add it to the window
-		[AppDelegate.RootView addSubview:AppDelegate.MetalView];
 
-		// initialize the backbuffer of the view (so the RHI can use it)
-		[AppDelegate.MetalView CreateFramebuffer:YES];
-	}
-	else
-#endif
-	{
-		// create the fullscreen EAGLView
-		AppDelegate.GLView = [[EAGLView alloc] initWithFrame:FullResolutionRect];
-		AppDelegate.GLView.clearsContextBeforeDrawing = NO;
-		AppDelegate.GLView.multipleTouchEnabled = YES;
+	CGRect FullResolutionRect = MainFrame;
 
-		// add it to the window
-		[AppDelegate.RootView addSubview:AppDelegate.GLView];
+	AppDelegate.IOSView = [[FIOSView alloc] initWithFrame:FullResolutionRect];
+	AppDelegate.IOSView.clearsContextBeforeDrawing = NO;
+	AppDelegate.IOSView.multipleTouchEnabled = YES;
 
-		// initialize the backbuffer of the view (so the RHI can use it)
-		[AppDelegate.GLView CreateFramebuffer:YES];
-	}
+	// add it to the window
+	[AppDelegate.RootView addSubview:AppDelegate.IOSView];
 
-	// Final adjustment to the viewport (this is deferred and won't run until the first engine tick)
-	[FIOSAsyncTask CreateTaskWithBlock:^ bool (void)
-		{
-			// @todo ios: Figure out how to resize the viewport!!!!
-
-// 			if (GEngine->GameViewport &&
-// 				GEngine->GameViewport->Viewport &&
-// 				GEngine->GameViewport->ViewportFrame)
-// 			{
-// 				CGFloat CSF = self.GLView.contentScaleFactor;
-// 				GEngine->GameViewport->ViewportFrame->Resize(
-// 					(INT)(CSF * self.GLView.superview.bounds.size.width),
-// 					(INT)(CSF * self.GLView.superview.bounds.size.height),
-// 					TRUE);
-// 			}
-			return true;
-		}];
+	// initialize the backbuffer of the view (so the RHI can use it)
+	[AppDelegate.IOSView CreateFramebuffer:YES];
 }
 
 
@@ -139,28 +103,13 @@ void FAppEntry::PlatformInit()
 	// wait until the GLView is fully initialized, so the RHI can be initialized
 	IOSAppDelegate* AppDelegate = [IOSAppDelegate GetDelegate];
 
-#if HAS_METAL
-	if (!FParse::Param(FCommandLine::Get(), TEXT("ES2")))
+	while (!AppDelegate.IOSView || !AppDelegate.IOSView->bIsInitialized)
 	{
-		while (!AppDelegate.MetalView || !AppDelegate.MetalView->bIsInitialized)
-		{
-			FPlatformProcess::Sleep(0.001f);
-		}
-
-		// set the Metal context to this thread
-		//		[AppDelegate.MetalView MakeCurrent];		
+		FPlatformProcess::Sleep(0.001f);
 	}
-	else
-#endif
-	{
-		while (!AppDelegate.GLView || !AppDelegate.GLView->bIsInitialized)
-		{
-			FPlatformProcess::Sleep(0.001f);
-		}
 
-		// set the GL context to this thread
-		[AppDelegate.GLView MakeCurrent];
-	}
+	// set the GL context to this thread
+	[AppDelegate.IOSView MakeCurrent];
 }
 
 void FAppEntry::Init()
@@ -177,7 +126,6 @@ void FAppEntry::Init()
 	NSLog(@"%s", "Initializing ULD Communications in game mode\n");
 	GCommandSystem.Init();
 
-	// at this point, we can let the Controller/EAGLView 
 	GLog->SetCurrentThreadAsMasterThread();
 
 	// start up the engine
diff --git a/Engine/Source/Runtime/OpenGLDrv/Private/IOS/IOSOpenGL.cpp b/Engine/Source/Runtime/OpenGLDrv/Private/IOS/IOSOpenGL.cpp
--- a/Engine/Source/Runtime/OpenGLDrv/Private/IOS/IOSOpenGL.cpp
+++ b/Engine/Source/Runtime/OpenGLDrv/Private/IOS/IOSOpenGL.cpp
@@ -2,7 +2,7 @@
 
 #include "OpenGLDrvPrivate.h"
 #include "IOSAppDelegate.h"
-#include "EAGLView.h"
+#include "IOSView.h"
 #include "Slate.h"
 
 // This header file gives access to the thread_policy_set function.
@@ -33,7 +33,7 @@ struct FPlatformOpenGLDevice
 		bSingleContext = !GUseThreadedRendering;
 
 		IOSAppDelegate* AppDelegate = [IOSAppDelegate GetDelegate];
-		EAGLView* GLView = AppDelegate.GLView;
+		FIOSView* GLView = AppDelegate.IOSView;
 
 		// EAGL Context (as the rendering one) on the EAGLView
 		check(GLView->Context);
@@ -113,9 +113,8 @@ void PlatformBlitToViewport( FPlatformOpenGLDevice* Device, FPlatformOpenGLConte
 	// @todo-mobile
 	check(Device->bSingleContext || Context == &Device->RenderingContext);
 	IOSAppDelegate* AppDelegate = [IOSAppDelegate GetDelegate];
-	EAGLView* GLView = AppDelegate.GLView;
 
-	[GLView SwapBuffers];
+	[AppDelegate.IOSView SwapBuffers];
 }
 
 void PlatformFlushIfNeeded()
@@ -128,9 +127,7 @@ void PlatformRebindResources(FPlatformOpenGLDevice* Device)
 	if (!Device->bSingleContext)
 	{
 		IOSAppDelegate* AppDelegate = [IOSAppDelegate GetDelegate];
-		EAGLView* GLView = AppDelegate.GLView;
-
-		glBindRenderbuffer(GL_RENDERBUFFER, GLView.OnScreenColorRenderBuffer);
+		glBindRenderbuffer(GL_RENDERBUFFER, AppDelegate.IOSView.OnScreenColorRenderBuffer);
 	}
 }
 
@@ -186,7 +183,7 @@ EOpenGLCurrentContext PlatformOpenGLCurrentContext(FPlatformOpenGLDevice* Device
 FRHITexture* PlatformCreateBuiltinBackBuffer(FOpenGLDynamicRHI* OpenGLRHI, uint32 SizeX, uint32 SizeY)
 {
 	IOSAppDelegate* AppDelegate = [IOSAppDelegate GetDelegate];
-	EAGLView* GLView = AppDelegate.GLView;
+	FIOSView* GLView = AppDelegate.IOSView;
 
 	uint32 Flags = TexCreate_RenderTargetable;
 	FOpenGLTexture2D* Texture2D = new FOpenGLTexture2D(OpenGLRHI, GLView.OnScreenColorRenderBuffer, GL_RENDERBUFFER, GL_COLOR_ATTACHMENT0, SizeX, SizeY, 0, 1, 1, 1, PF_B8G8R8A8, false, false, Flags);
@@ -266,7 +263,7 @@ bool PlatformContextIsCurrent( uint64 QueryContext )
 void ToggleUpscaleFilter()
 {
 	IOSAppDelegate* AppDelegate = [IOSAppDelegate GetDelegate];
-	EAGLView* GLView = AppDelegate.GLView;
+	FIOSView* GLView = AppDelegate.IOSView;
 	CAEAGLLayer* Layer = (CAEAGLLayer*)GLView.layer;
 	if (Layer.magnificationFilter == kCAFilterNearest)
 	{
