- added a comment

diff --git a/Engine/Source/Runtime/Core/Private/IOS/IOSView.cpp b/Engine/Source/Runtime/Core/Private/IOS/IOSView.cpp
--- a/Engine/Source/Runtime/Core/Private/IOS/IOSView.cpp
+++ b/Engine/Source/Runtime/Core/Private/IOS/IOSView.cpp
@@ -24,6 +24,7 @@ id<MTLDevice> GMetalDevice = nil;
 + (Class)layerClass
 {
 #if HAS_METAL
+	// the check for the function pointer itself is to determine if the Metal framework exists, before calling it
 	if (!FParse::Param(FCommandLine::Get(), TEXT("ES2")) && MTLCreateSystemDefaultDevice != NULL)
 	{
 		// if the device is unable to run with Metal (pre-A7), this will return nil
