fix for splash screen using the wrong image when holding before first render
#ue4
#ios

diff --git a/Engine/Source/Runtime/Core/Private/IOS/IOSAppDelegate.cpp b/Engine/Source/Runtime/Core/Private/IOS/IOSAppDelegate.cpp
--- a/Engine/Source/Runtime/Core/Private/IOS/IOSAppDelegate.cpp
+++ b/Engine/Source/Runtime/Core/Private/IOS/IOSAppDelegate.cpp
@@ -399,10 +399,28 @@ void InstallSignalHandlers()
     // add the default image as a subview
     NSMutableString* path = [[NSMutableString alloc]init];
     [path setString: [[NSBundle mainBundle] resourcePath]];
-    [path setString: [path stringByAppendingPathComponent:@"Default.png"]];
+    UIImageOrientation orient = UIImageOrientationUp;
+    if (MainFrame.size.height == 480)
+    {
+        [path setString: [path stringByAppendingPathComponent:@"Default.png"]];
+    }
+    else if (MainFrame.size.height == 568)
+    {
+        [path setString: [path stringByAppendingPathComponent:@"Default-568h.png"]];
+    }
+    else if (MainFrame.size.height == 1024 && !self.bDeviceInPortraitMode)
+    {
+        [path setString: [path stringByAppendingPathComponent:@"Default-Landscape.png"]];
+        orient = UIImageOrientationRight;
+    }
+    else if (MainFrame.size.height == 1024)
+    {
+        [path setString: [path stringByAppendingPathComponent:@"Default-Portrait.png"]];
+    }
     UIImage* image = [[UIImage alloc] initWithContentsOfFile: path];
     [path release];
-    UIImageView* imageView = [[UIImageView alloc] initWithImage: image];
+    UIImage* imageToDisplay = [UIImage imageWithCGImage: [image CGImage] scale: 1.0 orientation: orient];
+    UIImageView* imageView = [[UIImageView alloc] initWithImage: imageToDisplay];
     imageView.frame = MainFrame;
     imageView.tag = 2;
     [self.Window addSubview: imageView];
