- Renamed a .mm to .cpp

diff --git a/Engine/Source/Runtime/Core/Private/IOS/IOSPlatformString.cpp b/Engine/Source/Runtime/Core/Private/IOS/IOSPlatformString.cpp
new file mode 100644
--- /dev/null
+++ b/Engine/Source/Runtime/Core/Private/IOS/IOSPlatformString.cpp
@@ -0,0 +1,21 @@
+// Copyright 1998-2014 Epic Games, Inc. All Rights Reserved.
+
+/*=============================================================================
+	IOSPlatformString.mm: Additional NSString methods
+=============================================================================*/
+#include "CorePrivate.h"
+
+@implementation NSString (FString_Extensions)
+
++ (NSString*) stringWithTCHARString:(const TCHAR*)MyTCHARString
+{
+	return [NSString stringWithCString:TCHAR_TO_UTF8(MyTCHARString) encoding:NSUTF8StringEncoding];
+}
+
++ (NSString*) stringWithFString:(const FString&)InFString
+{
+	return [NSString stringWithTCHARString:*InFString];
+}
+
+
+@end
diff --git a/Engine/Source/Runtime/Core/Private/IOS/IOSPlatformString.mm b/Engine/Source/Runtime/Core/Private/IOS/IOSPlatformString.mm
deleted file mode 100644
