Modules: Module Manager API cleanup pass

There is now a consistent set of methods for module loading, unloading and lookup that all behave in the same expected way.

- renamed GetModuleInterface to GetModule and made it not assert
- removed GetModuleInterfaceRef
- namespace enums to enum classes, NULL to nullptr
- moved important public methods to the top, so they're easier to discover
- added missing & improved existing documentation
- cleaned up existing use cases

diff --git a/Engine/Source/Developer/ModuleUI/Private/SModuleUI.cpp b/Engine/Source/Developer/ModuleUI/Private/SModuleUI.cpp
--- a/Engine/Source/Developer/ModuleUI/Private/SModuleUI.cpp
+++ b/Engine/Source/Developer/ModuleUI/Private/SModuleUI.cpp
@@ -144,7 +144,7 @@ TSharedRef<ITableRow> SModuleUI::OnGenerateWidgetForModuleListView(TSharedPtr< F
 }
 
 
-void SModuleUI::OnModulesChanged( FName ModuleThatChanged, EModuleChangeReason::Type ReasonForChange )
+void SModuleUI::OnModulesChanged( FName ModuleThatChanged, EModuleChangeReason ReasonForChange )
 {
 	// @todo: Consider using dirty bit for this instead, refresh on demand
 	UpdateModuleListItems();
@@ -154,12 +154,12 @@ void SModuleUI::UpdateModuleListItems()
 {
 	ModuleListItems.Reset();
 
-	TArray< FModuleManager::FModuleStatus > ModuleStatuses;
+	TArray< FModuleStatus > ModuleStatuses;
 	FModuleManager::Get().QueryModules( ModuleStatuses );
 
-	for( TArray< FModuleManager::FModuleStatus >::TConstIterator ModuleIt( ModuleStatuses ); ModuleIt; ++ModuleIt )
+	for( TArray< FModuleStatus >::TConstIterator ModuleIt( ModuleStatuses ); ModuleIt; ++ModuleIt )
 	{
-		const FModuleManager::FModuleStatus& ModuleStatus = *ModuleIt;
+		const FModuleStatus& ModuleStatus = *ModuleIt;
 
 		FName ModuleName(*ModuleStatus.Name);
 
@@ -218,7 +218,7 @@ FReply SModuleUI::FModuleListItem::OnRecompileClicked()
 	GWarn->BeginSlowTask( FText::Format( NSLOCTEXT("ModuleUI", "Recompile_SlowTaskName", "Compiling {ModuleName}..."), Args ), bShowProgressDialog, bShowCancelButton );
 
 	// Does the module have any UObject classes in it?  If so we'll use HotReload to recompile it.
-	FModuleManager::FModuleStatus ModuleStatus;
+	FModuleStatus ModuleStatus;
 	if( ensure( FModuleManager::Get().QueryModule( ModuleName, ModuleStatus ) ) )
 	{
 		//@todo This is for content-only packages that show up in the
@@ -282,7 +282,7 @@ EVisibility SModuleUI::FModuleListItem::GetVisibilityBasedOnLoadedAndShutdownabl
 	const bool bIsHotReloadable = FModuleManager::Get().DoesLoadedModuleHaveUObjects(ModuleName);
 	const bool bCanShutDown = ( FModuleManager::Get().IsModuleLoaded(ModuleName)
 								&& !bIsHotReloadable 
-								&& FModuleManager::Get().GetModuleInterface(ModuleName).SupportsDynamicReloading() );
+								&& FModuleManager::Get().GetModule(ModuleName)->SupportsDynamicReloading() );
 
 	return bCanShutDown ? EVisibility::Visible : EVisibility::Hidden;
 }
@@ -303,7 +303,7 @@ EVisibility SModuleUI::FModuleListItem::GetVisibilityBasedOnRecompilableState()
 
 	const bool bIsHotReloadable = FModuleManager::Get().DoesLoadedModuleHaveUObjects(ModuleName);
 	const bool bCanReload = ( !FModuleManager::Get().IsModuleLoaded(ModuleName)
-							|| FModuleManager::Get().GetModuleInterface(ModuleName).SupportsDynamicReloading()
+							|| FModuleManager::Get().GetModule(ModuleName)->SupportsDynamicReloading()
 							|| bIsHotReloadable );
 
 	return bCanReload ? EVisibility::Visible : EVisibility::Hidden;
diff --git a/Engine/Source/Developer/ModuleUI/Private/SModuleUI.h b/Engine/Source/Developer/ModuleUI/Private/SModuleUI.h
--- a/Engine/Source/Developer/ModuleUI/Private/SModuleUI.h
+++ b/Engine/Source/Developer/ModuleUI/Private/SModuleUI.h
@@ -70,7 +70,7 @@ private:
 	 * Called by the engine's module manager when a module is loaded, unloaded, or the list of known
 	 * modules has changed
 	 */
-	void OnModulesChanged( FName ModuleThatChanged, EModuleChangeReason::Type ReasonForChange );
+	void OnModulesChanged( FName ModuleThatChanged, EModuleChangeReason ReasonForChange );
 
 	/**
 	 * Updates our module list items
diff --git a/Engine/Source/Developer/TargetPlatform/Private/TargetPlatformManagerModule.cpp b/Engine/Source/Developer/TargetPlatform/Private/TargetPlatformManagerModule.cpp
--- a/Engine/Source/Developer/TargetPlatform/Private/TargetPlatformManagerModule.cpp
+++ b/Engine/Source/Developer/TargetPlatform/Private/TargetPlatformManagerModule.cpp
@@ -517,7 +517,7 @@ protected:
 
 private:
 
-	void ModulesChangesCallback(FName ModuleName, EModuleChangeReason::Type ReasonForChange)
+	void ModulesChangesCallback(FName ModuleName, EModuleChangeReason ReasonForChange)
 	{
 		if (ModuleName.ToString().Contains(TEXT("TargetPlatform")) )
 		{
diff --git a/Engine/Source/Editor/ClassViewer/Private/SClassViewer.cpp b/Engine/Source/Editor/ClassViewer/Private/SClassViewer.cpp
--- a/Engine/Source/Editor/ClassViewer/Private/SClassViewer.cpp
+++ b/Engine/Source/Editor/ClassViewer/Private/SClassViewer.cpp
@@ -1511,7 +1511,7 @@ private:
 	FOnClassItemDoubleClickDelegate OnDoubleClicked;
 };
 
-static void OnModulesChanged(FName ModuleThatChanged, EModuleChangeReason::Type ReasonForChange)
+static void OnModulesChanged(FName ModuleThatChanged, EModuleChangeReason ReasonForChange)
 {
 	ClassViewer::Helpers::RequestPopulateClassHierarchy();
 }
diff --git a/Engine/Source/Editor/LevelEditor/Private/LevelEditorActions.cpp b/Engine/Source/Editor/LevelEditor/Private/LevelEditorActions.cpp
--- a/Engine/Source/Editor/LevelEditor/Private/LevelEditorActions.cpp
+++ b/Engine/Source/Editor/LevelEditor/Private/LevelEditorActions.cpp
@@ -956,12 +956,12 @@ void FLevelEditorActionCallbacks::MapCheck_Execute()
 bool FLevelEditorActionCallbacks::CanShowSourceCodeActions()
 {
 	// Ask the module manager for a list of currently-loaded gameplay modules
-	TArray< FModuleManager::FModuleStatus > ModuleStatuses;
+	TArray< FModuleStatus > ModuleStatuses;
 	FModuleManager::Get().QueryModules( ModuleStatuses );
 
 	for( auto ModuleStatusIt = ModuleStatuses.CreateConstIterator(); ModuleStatusIt; ++ModuleStatusIt )
 	{
-		const FModuleManager::FModuleStatus& ModuleStatus = *ModuleStatusIt;
+		const FModuleStatus& ModuleStatus = *ModuleStatusIt;
 
 		// We only care about game modules that are currently loaded
 		if( ModuleStatus.bIsLoaded && ModuleStatus.bIsGameModule )
@@ -992,12 +992,12 @@ void FLevelEditorActionCallbacks::RecompileGameCode_Clicked()
 		TArray< FString > GameModuleNames;
 		{
 			// Ask the module manager for a list of currently-loaded gameplay modules
-			TArray< FModuleManager::FModuleStatus > ModuleStatuses;
+			TArray< FModuleStatus > ModuleStatuses;
 			FModuleManager::Get().QueryModules( ModuleStatuses );
 
-			for( TArray< FModuleManager::FModuleStatus >::TConstIterator ModuleStatusIt( ModuleStatuses ); ModuleStatusIt; ++ModuleStatusIt )
+			for( TArray< FModuleStatus >::TConstIterator ModuleStatusIt( ModuleStatuses ); ModuleStatusIt; ++ModuleStatusIt )
 			{
-				const FModuleManager::FModuleStatus& ModuleStatus = *ModuleStatusIt;
+				const FModuleStatus& ModuleStatus = *ModuleStatusIt;
 
 				// We only care about game modules that are currently loaded
 				if( ModuleStatus.bIsLoaded && ModuleStatus.bIsGameModule )
diff --git a/Engine/Source/Editor/PackagesDialog/Private/PackagesDialog.cpp b/Engine/Source/Editor/PackagesDialog/Private/PackagesDialog.cpp
--- a/Engine/Source/Editor/PackagesDialog/Private/PackagesDialog.cpp
+++ b/Engine/Source/Editor/PackagesDialog/Private/PackagesDialog.cpp
@@ -194,7 +194,7 @@ bool FPackagesDialogModule::IsWindowInitialized() const
 /**
  * Called when a module has been changed (unloaded, loaded, etc)
  */
-void FPackagesDialogModule::OnModulesChanged( FName ModuleThatChanged, EModuleChangeReason::Type ReasonForChange )
+void FPackagesDialogModule::OnModulesChanged( FName ModuleThatChanged, EModuleChangeReason ReasonForChange )
 {
 
 }
diff --git a/Engine/Source/Editor/PackagesDialog/Public/PackagesDialog.h b/Engine/Source/Editor/PackagesDialog/Public/PackagesDialog.h
--- a/Engine/Source/Editor/PackagesDialog/Public/PackagesDialog.h
+++ b/Engine/Source/Editor/PackagesDialog/Public/PackagesDialog.h
@@ -124,7 +124,7 @@ protected:
 	/**
 	 * Called when a module has been changed (unloaded, loaded, etc)
 	 */
-	void OnModulesChanged( FName ModuleThatChanged, EModuleChangeReason::Type ReasonForChange );
+	void OnModulesChanged( FName ModuleThatChanged, EModuleChangeReason ReasonForChange );
 
 private:
 	/** A default window size for the package dialog */
diff --git a/Engine/Source/Editor/UnrealEd/Private/SourceCodeNavigation.cpp b/Engine/Source/Editor/UnrealEd/Private/SourceCodeNavigation.cpp
--- a/Engine/Source/Editor/UnrealEd/Private/SourceCodeNavigation.cpp
+++ b/Engine/Source/Editor/UnrealEd/Private/SourceCodeNavigation.cpp
@@ -1281,7 +1281,7 @@ bool FSourceCodeNavigation::FindClassModuleName( UClass* InClass, FString& Modul
 				// Because the module loaded into memory may have a slightly mutated file name (for
 				// hot reload, etc), we ask the module manager for the actual file name being used.  This
 				// is important as we need to be sure to get the correct symbols.
-				FModuleManager::FModuleStatus ModuleStatus;
+				FModuleStatus ModuleStatus;
 				if( ensure( FModuleManager::Get().QueryModule( ShortClassPackageName, ModuleStatus ) ) )
 				{
 					// Use the base file name (no path, no extension) as the module name for symbol look up!
diff --git a/Engine/Source/Editor/UnrealEd/Private/UnrealEdMisc.cpp b/Engine/Source/Editor/UnrealEd/Private/UnrealEdMisc.cpp
--- a/Engine/Source/Editor/UnrealEd/Private/UnrealEdMisc.cpp
+++ b/Engine/Source/Editor/UnrealEd/Private/UnrealEdMisc.cpp
@@ -367,7 +367,7 @@ void FUnrealEdMisc::InitEngineAnalytics()
 		}
 
 		// Record known modules' compilation methods
-		TArray<FModuleManager::FModuleStatus> Modules;
+		TArray<FModuleStatus> Modules;
 		FModuleManager::Get().QueryModules(Modules);
 		for (auto& Module : Modules)
 		{
diff --git a/Engine/Source/Runtime/Core/Private/Modules/ModuleManager.cpp b/Engine/Source/Runtime/Core/Private/Modules/ModuleManager.cpp
--- a/Engine/Source/Runtime/Core/Private/Modules/ModuleManager.cpp
+++ b/Engine/Source/Runtime/Core/Private/Modules/ModuleManager.cpp
@@ -264,16 +264,16 @@ void FModuleManager::AddModule( const FName InModuleName )
  */
 TSharedPtr<IModuleInterface> FModuleManager::LoadModule( const FName InModuleName, bool bWasReloaded /*=false*/ )
 {
-	ELoadModuleFailureReason::Type FailureReason;
+	EModuleLoadResult FailureReason;
 	return LoadModuleWithFailureReason(InModuleName, FailureReason, bWasReloaded );
 }
 
-TSharedPtr<IModuleInterface> FModuleManager::LoadModuleWithFailureReason( const FName InModuleName, ELoadModuleFailureReason::Type& OutFailureReason, bool bWasReloaded /*=false*/ )
+TSharedPtr<IModuleInterface> FModuleManager::LoadModuleWithFailureReason( const FName InModuleName, EModuleLoadResult& OutFailureReason, bool bWasReloaded /*=false*/ )
 {
 	DECLARE_SCOPE_CYCLE_COUNTER(TEXT("Module Load"), STAT_ModuleLoad, STATGROUP_LoadTime);
 
 	TSharedPtr<IModuleInterface> LoadedModule;
-	OutFailureReason = ELoadModuleFailureReason::Success;
+	OutFailureReason = EModuleLoadResult::Success;
 
 	// Update our set of known modules, in case we don't already know about this module
 	AddModule( InModuleName );
@@ -319,7 +319,7 @@ TSharedPtr<IModuleInterface> FModuleManager::LoadModuleWithFailureReason( const
 			else
 			{
 				UE_LOG(LogModuleManager, Warning, TEXT( "ModuleManager: Unable to load module '%s' because InitializeModule function failed (returned NULL pointer.)" ), *InModuleName.ToString() );
-				OutFailureReason = ELoadModuleFailureReason::FailedToInitialize;
+				OutFailureReason = EModuleLoadResult::FailedToInitialize;
 			}
 		}
 #if IS_MONOLITHIC
@@ -328,7 +328,7 @@ TSharedPtr<IModuleInterface> FModuleManager::LoadModuleWithFailureReason( const
 			// Monolithic builds that do not have the initializer were *not found* during the build step, so return FileNotFound
 			// (FileNotFound is an acceptable error in some case - ie loading a content only project)
 			UE_LOG(LogModuleManager, Warning, TEXT( "ModuleManager: Module '%s' not found - it's StaticallyLinkedModuleInitializers function is null." ), *InModuleName.ToString() );
-			OutFailureReason = ELoadModuleFailureReason::FileNotFound;
+			OutFailureReason = EModuleLoadResult::FileNotFound;
 		}
 #endif
 #if !IS_MONOLITHIC
@@ -399,7 +399,7 @@ TSharedPtr<IModuleInterface> FModuleManager::LoadModuleWithFailureReason( const
 
 								FPlatformProcess::FreeDllHandle( ModuleInfo->Handle );
 								ModuleInfo->Handle = NULL;
-								OutFailureReason = ELoadModuleFailureReason::FailedToInitialize;
+								OutFailureReason = EModuleLoadResult::FailedToInitialize;
 							}
 						}
 						else
@@ -408,25 +408,25 @@ TSharedPtr<IModuleInterface> FModuleManager::LoadModuleWithFailureReason( const
 
 							FPlatformProcess::FreeDllHandle( ModuleInfo->Handle );
 							ModuleInfo->Handle = NULL;
-							OutFailureReason = ELoadModuleFailureReason::FailedToInitialize;
+							OutFailureReason = EModuleLoadResult::FailedToInitialize;
 						}
 					}
 					else
 					{
 						UE_LOG(LogModuleManager, Warning, TEXT( "ModuleManager: Unable to load module '%s' because the file couldn't be loaded by the OS." ), *ModuleFileToLoad );
-						OutFailureReason = ELoadModuleFailureReason::CouldNotBeLoadedByOS;
+						OutFailureReason = EModuleLoadResult::CouldNotBeLoadedByOS;
 					}
 				}
 				else
 				{
 					// The log warning about this failure reason is within CheckModuleCompatibility
-					OutFailureReason = ELoadModuleFailureReason::FileIncompatible;
+					OutFailureReason = EModuleLoadResult::FileIncompatible;
 				}
 			}
 			else
 			{
 				UE_LOG(LogModuleManager, Warning, TEXT( "ModuleManager: Unable to load module '%s' because the file '%s' was not found." ), *InModuleName.ToString(), *ModuleFileToLoad );
-				OutFailureReason = ELoadModuleFailureReason::FileNotFound;
+				OutFailureReason = EModuleLoadResult::FileNotFound;
 			}
 		}
 #endif
@@ -574,53 +574,20 @@ void FModuleManager::UnloadModulesAtShutdown()
 }
 
 
-/**
- * Given a module name, returns a shared reference to that module's interface.  If the module is unknown or not loaded, this will assert!
- *
- * @param	InModuleName	Name of the module to return
- *
- * @return	A shared reference to the module
- */
-TSharedRef< IModuleInterface > FModuleManager::GetModuleInterfaceRef( const FName InModuleName )
+TSharedPtr<IModuleInterface> FModuleManager::GetModule( const FName InModuleName )
 {
 	// Do we even know about this module?
-	TSharedRef< FModuleInfo > ModuleInfo = Modules.FindChecked( InModuleName );
-
-	// Module must be loaded!
-	checkf( ModuleInfo->Module.IsValid(), TEXT("Tried to use invalid module %s!"), *ModuleInfo->Filename );
-	return ModuleInfo->Module.ToSharedRef();
-}
-
+	const TSharedRef<FModuleInfo>* ModuleInfo = Modules.Find(InModuleName);
 
+	if (ModuleInfo == nullptr)
+	{
+		return nullptr;
+	}
 
-/**
- * Given a module name, returns a reference to that module's interface.  If the module is unknown or not loaded, this will assert!
- *
- * @param	InModuleName	Name of the module to return
- *
- * @return	A reference to the module
- */
-IModuleInterface& FModuleManager::GetModuleInterface( const FName InModuleName )
-{
-	// Do we even know about this module?
-	TSharedRef< FModuleInfo > ModuleInfo = Modules.FindChecked( InModuleName );
-
-	// Module must be loaded!
-	checkf( ModuleInfo->Module.IsValid(), TEXT("Tried to use invalid module %s!"), *ModuleInfo->Filename );
-	return *ModuleInfo->Module;
+	return (*ModuleInfo)->Module;
 }
 
 
-
-/**
- * Exec command
- *
- * @param	InWorld		World context
- * @param	Cmd		Command to execute
- * @param	Ar		Output device
- *
- * @return	True if command was handled
- */
 bool FModuleManager::Exec( UWorld* Inworld, const TCHAR* Cmd, FOutputDevice& Ar )
 {
 #if !UE_BUILD_SHIPPING
@@ -764,15 +731,6 @@ bool FModuleManager::Exec( UWorld* Inworld, const TCHAR* Cmd, FOutputDevice& Ar
 }
 
 
-
-/**
- * Queries information about a specific module name
- *
- * @param	InModuleName	Module to query status for
- * @param	OutModuleStatus	Status of the specified module
- *
- * @return	True if the module was found the OutModuleStatus is valid
- */
 bool FModuleManager::QueryModule( const FName InModuleName, FModuleStatus& OutModuleStatus )
 {
 	// Do we even know about this module?
@@ -814,12 +772,7 @@ bool FModuleManager::QueryModule( const FName InModuleName, FModuleStatus& OutMo
 	return false;
 }
 
-	
-/**
- * Queries information about all of the currently known modules
- *
- * @param	OutModuleStatuses	Status of all modules
- */
+
 void FModuleManager::QueryModules( TArray< FModuleStatus >& OutModuleStatuses )
 {
 	OutModuleStatuses.Reset();
@@ -861,11 +814,13 @@ void FModuleManager::QueryModules( TArray< FModuleStatus >& OutModuleStatuses )
 	}
 }
 
+
 bool FModuleManager::IsSolutionFilePresent()
 {
 	return !GetSolutionFilepath().IsEmpty();
 }
 
+
 FString FModuleManager::GetSolutionFilepath()
 {
 #if PLATFORM_MAC
@@ -899,6 +854,7 @@ FString FModuleManager::GetSolutionFilepath()
 	return SolutionFilepath;
 }
 
+
 bool FModuleManager::RecompileModule( const FName InModuleName, const bool bReloadAfterRecompile, FOutputDevice &Ar )
 {
 #if !IS_MONOLITHIC
@@ -982,7 +938,6 @@ bool FModuleManager::RecompileModule( const FName InModuleName, const bool bRelo
 }
 
 
-
 bool FModuleManager::IsCurrentlyCompiling() const
 {
 	return ModuleCompileProcessHandle.IsValid();
@@ -1016,6 +971,7 @@ FString FModuleManager::MakeUBTArgumentsForModuleCompiling()
 	return AdditionalArguments;
 }
 
+
 void FModuleManager::OnModuleCompileSucceeded(FName ModuleName, TSharedRef<FModuleManager::FModuleInfo> ModuleInfo)
 {
 #if !IS_MONOLITHIC && WITH_EDITOR
@@ -1042,6 +998,7 @@ void FModuleManager::OnModuleCompileSucceeded(FName ModuleName, TSharedRef<FModu
 #endif
 }
 
+
 void FModuleManager::UpdateModuleCompileData(FName ModuleName, TSharedRef<FModuleManager::FModuleInfo> ModuleInfo)
 {
 	// reset the compile data before updating it
@@ -1089,6 +1046,7 @@ void FModuleManager::UpdateModuleCompileData(FName ModuleName, TSharedRef<FModul
 #endif
 }
 
+
 void FModuleManager::ReadModuleCompilationInfoFromConfig(FName ModuleName, TSharedRef<FModuleManager::FModuleInfo> ModuleInfo)
 {
 	FString DateTimeString;
@@ -1116,6 +1074,7 @@ void FModuleManager::ReadModuleCompilationInfoFromConfig(FName ModuleName, TShar
 	}
 }
 
+
 void FModuleManager::WriteModuleCompilationInfoToConfig(FName ModuleName, TSharedRef<FModuleManager::FModuleInfo> ModuleInfo)
 {
 	if (ensure(ModuleInfo->CompileData.bIsValid))
@@ -1142,6 +1101,7 @@ void FModuleManager::WriteModuleCompilationInfoToConfig(FName ModuleName, TShare
 	}
 }
 
+
 bool FModuleManager::GetModuleFileTimeStamp(TSharedRef<const FModuleInfo> ModuleInfo, FDateTime& OutFileTimeStamp)
 {
 	if (IFileManager::Get().FileSize(*ModuleInfo->Filename) > 0)
@@ -1152,6 +1112,7 @@ bool FModuleManager::GetModuleFileTimeStamp(TSharedRef<const FModuleInfo> Module
 	return false;
 }
 
+
 void FModuleManager::FindModulePaths(const TCHAR* NamePattern, TMap<FName, FString> &OutModulePaths) const
 {
 	// Search through the engine directory
@@ -1170,6 +1131,7 @@ void FModuleManager::FindModulePaths(const TCHAR* NamePattern, TMap<FName, FStri
 	}
 }
 
+
 void FModuleManager::FindModulePathsInDirectory(const FString& InDirectoryName, bool bIsGameDirectory, const TCHAR* NamePattern, TMap<FName, FString> &OutModulePaths) const
 {
 	// Get the module configuration for this directory type
@@ -1229,6 +1191,7 @@ void FModuleManager::FindModulePathsInDirectory(const FString& InDirectoryName,
 	}
 }
 
+
 bool FModuleManager::RecompileModulesAsync( const TArray< FName > ModuleNames, const FRecompileModulesCallback& InRecompileModulesCallback, const bool bWaitForCompletion, FOutputDevice &Ar )
 {
 #if !IS_MONOLITHIC
@@ -1289,6 +1252,7 @@ bool FModuleManager::RecompileModulesAsync( const TArray< FName > ModuleNames, c
 #endif // !IS_MONOLITHIC
 }
 
+
 bool FModuleManager::CompileGameProject( const FString& ProjectFilename, FOutputDevice &Ar )
 {
 	bool bCompileSucceeded = false;
@@ -1347,6 +1311,7 @@ bool FModuleManager::CompileGameProjectEditor( const FString& GameProjectFilenam
 	return bCompileSucceeded;
 }
 
+
 bool FModuleManager::GenerateCodeProjectFiles( const FString& ProjectFilename, FOutputDevice &Ar )
 {
 	bool bCompileSucceeded = false;
@@ -1421,6 +1386,7 @@ bool FModuleManager::GenerateCodeProjectFiles( const FString& ProjectFilename, F
 	return bCompileSucceeded;
 }
 
+
 bool FModuleManager::IsUnrealBuildToolAvailable()
 {
 	// If using Rocket and the Rocket unreal build tool executable exists, then UBT is available
@@ -1466,6 +1432,7 @@ void FModuleManager::UnloadOrAbandonModuleWithCallback( const FName InModuleName
 	}
 }
 
+
 bool FModuleManager::LoadModuleWithCallback( const FName InModuleName, FOutputDevice &Ar )
 {
 	TSharedPtr<IModuleInterface> LoadedModule = LoadModule( InModuleName, true );
@@ -1483,6 +1450,7 @@ bool FModuleManager::LoadModuleWithCallback( const FName InModuleName, FOutputDe
 	return bWasSuccessful;
 }
 
+
 void FModuleManager::MakeUniqueModuleFilename( const FName InModuleName, FString& UniqueSuffix, FString& UniqueModuleFileName )
 {
 	TSharedRef< FModuleInfo > Module = Modules.FindChecked( InModuleName );
@@ -1505,6 +1473,7 @@ void FModuleManager::MakeUniqueModuleFilename( const FName InModuleName, FString
 	while (IFileManager::Get().GetFileAgeSeconds(*UniqueModuleFileName) != -1.0);
 }
 
+
 const TCHAR *FModuleManager::GetUBTConfiguration()
 {
 #if UE_BUILD_DEBUG
@@ -1517,6 +1486,7 @@ const TCHAR *FModuleManager::GetUBTConfiguration()
 #endif
 }
 
+
 bool FModuleManager::StartCompilingModuleDLLs(const FString& GameName, const TArray< FModuleToRecompile >& ModuleNames, 
 	const FRecompileModulesCallback& InRecompileModulesCallback, FOutputDevice& Ar, bool bInFailIfGeneratedCodeChanges, 
 	const FString& InAdditionalCmdLineArgs )
@@ -1593,11 +1563,13 @@ bool FModuleManager::StartCompilingModuleDLLs(const FString& GameName, const TAr
 #endif
 }
 
+
 FString FModuleManager::GetUnrealBuildToolSourceCodePath()
 {
 	return FPaths::Combine(*FPaths::EngineDir(), TEXT("Source"), TEXT("Programs"), TEXT("UnrealBuildTool"));
 }
 
+
 FString FModuleManager::GetUnrealBuildToolExecutableFilename()
 {
 	const FString UBTPath = FString::Printf(TEXT("%sBinaries/DotNET/"), *FPaths::EngineDir());
@@ -1608,6 +1580,7 @@ FString FModuleManager::GetUnrealBuildToolExecutableFilename()
 	return ExecutableFileName;
 }
 
+
 bool FModuleManager::BuildUnrealBuildTool(FOutputDevice &Ar)
 {
 #if !IS_MONOLITHIC
@@ -1695,6 +1668,7 @@ bool FModuleManager::BuildUnrealBuildTool(FOutputDevice &Ar)
 	}
 }
 
+
 bool FModuleManager::InvokeUnrealBuildTool( const FString& InCmdLineParams, FOutputDevice &Ar )
 {
 #if PLATFORM_DESKTOP && !IS_MONOLITHIC
@@ -1945,6 +1919,7 @@ void FModuleManager::CheckForFinishedModuleDLLCompile( const bool bWaitForComple
 #endif // PLATFORM_DESKTOP && !IS_MONOLITHIC
 }
 
+
 bool FModuleManager::CheckModuleCompatibility(const TCHAR* Filename)
 {
 	static FBinaryFileVersion ExeVersion = FPlatformProcess::GetBinaryFileVersion(FPlatformProcess::ExecutableName(/*bRemoveExtension=*/false));
@@ -1981,6 +1956,7 @@ bool FModuleManager::RecompileModuleDLLs( const TArray< FModuleToRecompile >& Mo
 	return bCompileSucceeded;
 }
 
+
 void FModuleManager::StartProcessingNewlyLoadedObjects()
 {
 	// Only supposed to be called once
@@ -1988,6 +1964,7 @@ void FModuleManager::StartProcessingNewlyLoadedObjects()
 	bCanProcessNewlyLoadedObjects = true;
 }
 
+
 void FModuleManager::AddBinariesDirectory(const TCHAR *InDirectory, bool bIsGameDirectory)
 {
 	if (bIsGameDirectory)
@@ -2000,6 +1977,7 @@ void FModuleManager::AddBinariesDirectory(const TCHAR *InDirectory, bool bIsGame
 	}
 }
 
+
 void FModuleManager::SetGameBinariesDirectory(const TCHAR* InDirectory)
 {
 #if !IS_MONOLITHIC
@@ -2013,6 +1991,7 @@ void FModuleManager::SetGameBinariesDirectory(const TCHAR* InDirectory)
 #endif
 }
 
+
 bool FModuleManager::DoesLoadedModuleHaveUObjects( const FName ModuleName )
 {
 	if( IsModuleLoaded( ModuleName ) && IsPackageLoaded.IsBound() )
diff --git a/Engine/Source/Runtime/Core/Private/Windows/WindowsPlatformStackWalk.cpp b/Engine/Source/Runtime/Core/Private/Windows/WindowsPlatformStackWalk.cpp
--- a/Engine/Source/Runtime/Core/Private/Windows/WindowsPlatformStackWalk.cpp
+++ b/Engine/Source/Runtime/Core/Private/Windows/WindowsPlatformStackWalk.cpp
@@ -633,7 +633,7 @@ int32 FWindowsPlatformStackWalk::GetProcessModuleSignatures(FStackWalkModuleInfo
 /**
  * Callback from the modules system that the loaded modules have changed and we need to reload symbols.
  */ 
-static void OnModulesChanged( FName ModuleThatChanged, EModuleChangeReason::Type ReasonForChange )
+static void OnModulesChanged( FName ModuleThatChanged, EModuleChangeReason ReasonForChange )
 {
 	GNeedToRefreshSymbols = true;
 }
diff --git a/Engine/Source/Runtime/Core/Public/Modules/ModuleManager.h b/Engine/Source/Runtime/Core/Public/Modules/ModuleManager.h
--- a/Engine/Source/Runtime/Core/Public/Modules/ModuleManager.h
+++ b/Engine/Source/Runtime/Core/Public/Modules/ModuleManager.h
@@ -2,32 +2,46 @@
 
 #pragma once
 
-
 #include "ModuleInterface.h"
 #include "Delegate.h"
 #include "Boilerplate/ModuleBoilerplate.h"
 
+
 #if !IS_MONOLITHIC
 	/** If true, we are reloading a class for HotReload */
-	extern CORE_API bool		GIsHotReload;
+	extern CORE_API bool GIsHotReload;
 #endif
 
-namespace ELoadModuleFailureReason
+
+/**
+ * Enumerates reasons for failed module loads.
+ */
+enum class EModuleLoadResult
 {
-	enum Type
-	{
-		Success,
-		FileNotFound,
-		FileIncompatible,
-		CouldNotBeLoadedByOS,
-		FailedToInitialize
-	};
-}
+	/** Module loaded successfully. */
+	Success,
+
+	/** The specified module file could not be found. */
+	FileNotFound,
+
+	/** The specified module file is incompatible with the module system. */
+	FileIncompatible,
+
+	/** The operating system failed to load the module file. */
+	CouldNotBeLoadedByOS,
+
+	/** Module initialization failed. */
+	FailedToInitialize
+};
+
 
-// This enum has to be compatible with the one defined in the
-// UE4\Engine\Source\Programs\UnrealBuildTool\System\ExternalExecution.cs
-// to keep communication between UHT, UBT and Editor compiling
-// processes valid.
+/**
+ * Enumerates possible results of a compilation operation.
+ *
+ * This enum has to be compatible with the one defined in the
+ * UE4\Engine\Source\Programs\UnrealBuildTool\System\ExternalExecution.cs file
+ * to keep communication between UHT, UBT and Editor compiling processes valid.
+ */
 namespace ECompilationResult
 {
 	enum Type
@@ -38,216 +52,243 @@ namespace ECompilationResult
 	};
 }
 
-namespace EModuleCompileMethod
-{
-	enum Type
-	{
-		Runtime,
-		External,
-		Unknown
-	};
-}
 
-namespace EModuleChangeReason
+/**
+ * Enumerates compilation methods for modules.
+ */
+enum class EModuleCompileMethod
 {
-	/** Reasons that an FModuleManager::OnModuleChanged() delegate will be called */
-	enum Type
-	{
-		/** A module has been loaded and is ready to be used */
-		ModuleLoaded,
-
-		/* A module has been unloaded and should no longer be used */
-		ModuleUnloaded,
+	Runtime,
+	External,
+	Unknown
+};
 
-		/** The paths controlling which plugins are loaded have been changed and the given module has been found, but not yet loaded */
-		PluginDirectoryChanged
-	};
-}
 
 /**
- * Module Manager is used to load and unload modules, as well as to keep track of all of the
- * modules that are currently loaded.  You can access this singleton using FModuleManager::Get().
+ * Enumerates reasons for modules to change.
+ *
+ * Values of this type will be passed into OnModuleChanged() delegates.
  */
-class CORE_API FModuleManager : private FSelfRegisteringExec
+enum class EModuleChangeReason
 {
+	/** A module has been loaded and is ready to be used. */
+	ModuleLoaded,
 
-public:
-
-	/**
-	 * Status of a module.  Used only for reporting information to external objects.
-	 */
-	struct FModuleStatus
-	{
-		/** Default constructor. */
-		FModuleStatus()
-			: bIsLoaded( false )
-			, bIsGameModule( false )
-		{}
-
-		/** Short name for this module */
-		FString Name;
+	/* A module has been unloaded and should no longer be used. */
+	ModuleUnloaded,
 
-		/** Full path to this module file on disk */
-		FString FilePath;
-
-		/** Whether the module is currently loaded or not */
-		bool bIsLoaded;
+	/** The paths controlling which plug-ins are loaded have been changed and the given module has been found, but not yet loaded. */
+	PluginDirectoryChanged
+};
 
-		/** Whether this module contains gameplay code */
-		bool bIsGameModule;
 
-		/** The compilation method of this module */
-		FString CompilationMethod;
-	};
+/**
+ * Structure for reporting module statuses.
+ */
+struct FModuleStatus
+{
+	/** Default constructor. */
+	FModuleStatus()
+		: bIsLoaded(false)
+		, bIsGameModule(false)
+	{ }
 
-	/** Declare delegate that for binding functions to be called when modules are loaded, or unloaded, or
-	    our set of known modules changes. Passes in the name of the module that changed, and the reason. */
-	DECLARE_EVENT_TwoParams( FModuleManager, FModulesChangedEvent, FName, EModuleChangeReason::Type );
-	FModulesChangedEvent& OnModulesChanged() { return ModulesChangedEvent; }
+	/** Short name for this module. */
+	FString Name;
 
-	/** Delegate for binding functions to be called when starting the module compiling */
-	DECLARE_EVENT( FModuleManager, FModuleCompilerStartedEvent );
-	FModuleCompilerStartedEvent& OnModuleCompilerStarted() { return ModuleCompilerStartedEvent; }
+	/** Full path to this module file on disk. */
+	FString FilePath;
 
-	/** Delegate for binding functions to be called when the module compiler finishes,
-		passing in the compiler output and a enumeration with compilation result
-		as well as a second boolean which determines if the log should be shown */
-	DECLARE_EVENT_ThreeParams(FModuleManager, FModuleCompilerFinishedEvent, const FString&, ECompilationResult::Type, bool);
-	FModuleCompilerFinishedEvent& OnModuleCompilerFinished() { return ModuleCompilerFinishedEvent; }
+	/** Whether the module is currently loaded or not. */
+	bool bIsLoaded;
 
-	/** Called after a module recompile finishes.  First argument specifies whether the compilation has finished, 
-		Second argument specifies whether the compilation was successful or not */
-	DECLARE_DELEGATE_TwoParams( FRecompileModulesCallback, bool, bool );
+	/** Whether this module contains game play code. */
+	bool bIsGameModule;
 
-	/** Called when any loaded objects need to be processed. */
-	FSimpleMulticastDelegate& OnProcessLoadedObjectsCallback() { return ProcessLoadedObjectsCallback; }
+	/** The compilation method of this module. */
+	FString CompilationMethod;
+};
 
-	/** When module manager is linked against an application that supports UObjects, this delegate will be primed
-	    at startup to provide information about whether a UObject package is loaded into memory. */
-	DECLARE_DELEGATE_RetVal_OneParam( bool, FIsPackageLoadedCallback, FName );
-	FIsPackageLoadedCallback& IsPackageLoadedCallback() { return IsPackageLoaded; }
 
+/**
+ * Implements the module manager.
+ *
+ * The module manager is used to load and unload modules, as well as to keep track of all of the
+ * modules that are currently loaded. You can access this singleton using FModuleManager::Get().
+ */
+class CORE_API FModuleManager
+	: private FSelfRegisteringExec
+{
+public:
 
 	/**
-	 * Static: Access single instance of module manager
-	 *
-	 * @return	Reference to the module manager singleton object
+	 * Destructor.
 	 */
-	static FModuleManager& Get();
-
-
-	/** Destructor */
 	~FModuleManager();
 
-
 	/**
-	 * Module manager ticking is only used to check for asynchronously compiled modules that may need to be reloaded
+	 * Gets the singleton instance of the module manager.
+	 *
+	 * @return The module manager instance.
 	 */
-	void Tick();
+	static FModuleManager& Get( );
 
+public:
 
 	/**
-	 * Looks on the disk for loadable modules matching a wildcard
+	 * Abandons a loaded module, leaving it loaded in memory but no longer tracking it in the module manager.
 	 *
-	 * @param	WildcardWithoutExtension		Filename part (no path, no extension, no build config info) to search for
-	 * @param	OutModules					List of modules found
+	 * @param InModuleName The name of the module to abandon.  Should not include path, extension or platform/configuration info.  This is just the "module name" part of the module file name.
+	 * @see IsModuleLoaded, LoadModule, LoadModuleWithFailureReason, UnloadModule
 	 */
-	void FindModules(const TCHAR* WildcardWithoutExtension, TArray<FName>& OutModules);
+	void AbandonModule( const FName InModuleName );
 
 	/**
-	 * Checks to see if the specified module is currently loaded.  This is an O(1) operation.
+	 * Adds a module to our list of modules, unless it's already known.
 	 *
-	 * @param	InModuleName		The base name of the module file.  Should not include path, extension or platform/configuration info.  This is just the "module name" part of the module file name.  Names should be globally unique.
+	 * This method is used by the plug-in manager to register a plug-in module.
 	 *
-	 * @return	True if module is currently loaded, otherwise false
+	 * @param InModuleName The base name of the module file.  Should not include path, extension or platform/configuration info.  This is just the "name" part of the module file name.  Names should be globally unique.
+	 * @param InBinariesDirectory The directory where to find this file, or an empty string to search in the default locations.  This parameter is used by the plugin system to locate plugin binaries.
 	 */
-	bool IsModuleLoaded( const FName InModuleName ) const;
-
+	void AddModule( const FName InModuleName );
 
 	/**
-	 * Adds a module to our list of modules and tries to load it immediately
-	 *
-	 * @param	InModuleName	The base name of the module file.  Should not include path, extension or platform/configuration info.  This is just the "module name" part of the module file name.  Names should be globally unique.
-	 * @param bWasReloaded	Indicates that the modue has been reloaded.
+	 * Gets the specified module.
 	 *
-	 * @return	The loaded module (null if the load operation failed)
+	 * @param InModuleName Name of the module to return.
+	 * @return 	The module, or nullptr if the module is not loaded.
 	 */
-	TSharedPtr<IModuleInterface> LoadModule( const FName InModuleName, const bool bWasReloaded=false );
-	TSharedPtr<IModuleInterface> LoadModuleWithFailureReason( const FName InModuleName, ELoadModuleFailureReason::Type& OutFailureReason, const bool bWasReloaded=false );
-
+	TSharedPtr<IModuleInterface> GetModule( const FName InModuleName );
 
 	/**
-	 * Unloads a specific module
+	 * Checks whether the specified module is currently loaded.
 	 *
-	 * @param	InModuleName	The name of the module to unload.  Should not include path, extension or platform/configuration info.  This is just the "module name" part of the module file name.
-	 * @param   bIsShutdown		Is this unload module call occurring at shutdown
+	 * This is an O(1) operation.
 	 *
-	 * @return	True if module was unloaded successfully
+	 * @param InModuleName The base name of the module file.  Should not include path, extension or platform/configuration info.  This is just the "module name" part of the module file name.  Names should be globally unique.
+	 * @return true if module is currently loaded, false otherwise.
+	 * @see AbandonModule, LoadModule, LoadModuleWithFailureReason, UnloadModule
 	 */
-	bool UnloadModule( const FName InModuleName, bool bIsShutdown=false );
-
+	bool IsModuleLoaded( const FName InModuleName ) const;
 
 	/**
-	 * Abandons a loaded module, leaving it loaded in memory but no longer tracking it in the module manager
+	 * Loads the specified module.
 	 *
-	 * @param	InModuleName	The name of the module to abandon.  Should not include path, extension or platform/configuration info.  This is just the "module name" part of the module file name.
+	 * @param InModuleName The base name of the module file.  Should not include path, extension or platform/configuration info.  This is just the "module name" part of the module file name.  Names should be globally unique.
+	 * @param bWasReloaded Indicates that the module has been reloaded (default = false).
+	 * @return The loaded module, or nullptr if the load operation failed.
+	 * @see AbandonModule, IsModuleLoaded, LoadModuleWithFailureReason, UnloadModule
 	 */
-	void AbandonModule( const FName InModuleName );
-
+	TSharedPtr<IModuleInterface> LoadModule( const FName InModuleName, const bool bWasReloaded = false );
 
 	/**
-	 * Unloads modules during the shutdown process.  Usually called at various points while
-	 * while exiting the application.
+	 * Loads a module in memory then calls PostLoad.
+	 *
+	 * @param InModuleName The name of the module to load.
+	 * @param Ar The archive to receive error messages, if any.
+	 * @return true on success, false otherwise.
+	 * @see UnloadOrAbandonModuleWithCallback
 	 */
-	void UnloadModulesAtShutdown();
-
+	bool LoadModuleWithCallback( const FName InModuleName, FOutputDevice &Ar );
 
 	/**
-	 * Given a module name, returns a shared reference to that module's interface.  If the module is unknown or not loaded, this will assert!
+	 * Loads the specified module and returns a result.
 	 *
-	 * @param	InModuleName	Name of the module to return
+	 * @param InModuleName The base name of the module file.  Should not include path, extension or platform/configuration info.  This is just the "module name" part of the module file name.  Names should be globally unique.
+	 * @param OutFailureReason Will contain the result.
+	 * @param bWasReloaded Indicates that the module has been reloaded (default = false).
+	 * @return The loaded module (null if the load operation failed).
+	 * @see AbandonModule, IsModuleLoaded, LoadModule, UnloadModule
+	 */
+	TSharedPtr<IModuleInterface> LoadModuleWithFailureReason( const FName InModuleName, EModuleLoadResult& OutFailureReason, const bool bWasReloaded = false );
+
+	/**
+	 * Queries information about a specific module name.
 	 *
-	 * @return	A shared reference to the module
+	 * @param InModuleName Module to query status for.
+	 * @param OutModuleStatus Status of the specified module.
+	 * @return true if the module was found and the OutModuleStatus is valid, false otherwise.
+	 * @see QueryModules
 	 */
-	TSharedRef< IModuleInterface > GetModuleInterfaceRef( const FName InModuleName );
+	bool QueryModule( const FName InModuleName, FModuleStatus& OutModuleStatus );
 
+	/**
+	 * Queries information about all of the currently known modules.
+	 *
+	 * @param OutModuleStatuses Status of all modules.
+	 * @see QueryModule
+	 */
+	void QueryModules( TArray<FModuleStatus>& OutModuleStatuses );
 
 	/**
-	 * Given a module name, returns a reference to that module's interface.  If the module is unknown or not loaded, this will assert!
+	 * Unloads a specific module
 	 *
-	 * @param	InModuleName	Name of the module to return
+	 * @param InModuleName The name of the module to unload.  Should not include path, extension or platform/configuration info.  This is just the "module name" part of the module file name.
+	 * @param bIsShutdown Is this unload module call occurring at shutdown (default = false).
+	 * @return true if module was unloaded successfully, false otherwise.
+	 * @see AbandonModule, IsModuleLoaded, LoadModule, LoadModuleWithFailureReason
+	 */
+	bool UnloadModule( const FName InModuleName, bool bIsShutdown = false );
+
+	/**
+	 * Calls PreUnload then either unloads or abandons a module in memory, depending on whether the module supports unloading.
 	 *
-	 * @return	A reference to the module
+	 * @param InModuleName The name of the module to unload.
+	 * @param Ar The archive to receive error messages, if any.
+	 * @see LoadModuleWithCallback
 	 */
-	IModuleInterface& GetModuleInterface( const FName InModuleName );
+	void UnloadOrAbandonModuleWithCallback( const FName InModuleName, FOutputDevice &Ar );
 
+public:
 
 	/**
-	  * Static: Finds a module by name, checking to ensure it exists.
-	  * The return value is the loaded module, casted to the template parameter type (this is an unchecked cast)
+	  * Gets a module by name, checking to ensure it exists.
 	  *
-	  * @param	ModuleName	The module to find
+	  * This method checks whether the module actually exists. If the module does not exist, an assertion will be triggered.
 	  *
-	  * @return	Returns the module interface, casted to the specified typename
+	  * @param ModuleName The module to get.
+	  * @return The interface to the module.
+	  * @see GetModulePtr, LoadModulePtr, LoadModuleChecked
 	  */
 	template<typename TModuleInterface>
 	static TModuleInterface& GetModuleChecked( const FName ModuleName )
 	{
 		FModuleManager& ModuleManager = FModuleManager::Get();
 
-		// Make sure the module is loaded now.
-		checkf( ModuleManager.IsModuleLoaded(ModuleName), TEXT("Tried to get module interface for unloaded module: '%s'"), *(ModuleName.ToString()));
-		return (TModuleInterface&)(ModuleManager.GetModuleInterface(ModuleName));
+		checkf(ModuleManager.IsModuleLoaded(ModuleName), TEXT("Tried to get module interface for unloaded module: '%s'"), *(ModuleName.ToString()));
+		return (TModuleInterface&)(*ModuleManager.GetModule(ModuleName));
 	}
 
+	/**
+	  * Gets a module by name.
+	  *
+	  * @param ModuleName The module to get.
+	  * @return The interface to the module, or nullptr if the module was not found.
+	  * @see GetModuleChecked, LoadModulePtr, LoadModuleChecked
+	  */
+	template<typename TModuleInterface>
+	static TModuleInterface* GetModulePtr( const FName ModuleName )
+	{
+		FModuleManager& ModuleManager = FModuleManager::Get();
+
+		if (!ModuleManager.IsModuleLoaded(ModuleName))
+		{
+			return nullptr;
+		}
+
+		return (TModuleInterface*)(ModuleManager.GetModule(ModuleName).Get());
+	}
 
 	/**
-	  * Static: Finds a module by name, checking to ensure it exists and loading it if not already loaded.
-	  * The return value is the loaded module, casted to the template parameter type (this is an unchecked cast)
+	  * Loads a module by name, checking to ensure it exists.
 	  *
-	  * @param	ModuleName	The module to find and load
+	  * This method checks whether the module actually exists. If the module does not exist, an assertion will be triggered.
+	  * If the module was already loaded previously, the existing instance will be returned.
 	  *
+	  * @param ModuleName The module to find and load
 	  * @return	Returns the module interface, casted to the specified typename
+	  * @see GetModulePtr, LoadModulePtr, LoadModuleChecked
 	  */
 	template<typename TModuleInterface>
 	static TModuleInterface& LoadModuleChecked( const FName ModuleName )
@@ -256,20 +297,18 @@ public:
 
 		if (!ModuleManager.IsModuleLoaded(ModuleName))
 		{
-			// Ensure the module is loaded.
 			ModuleManager.LoadModule(ModuleName);
 		}
 
-		return (TModuleInterface&)(ModuleManager.GetModuleInterface(ModuleName));
+		return GetModuleChecked<TModuleInterface>(ModuleName);
 	}
 
 	/**
-	  * Static: Finds a module by name, checking to ensure it exists and loading it if not already loaded.
-	  * The return value is the loaded module, casted to the template parameter pointer type (this is an unchecked cast)
-	  *
-	  * @param	ModuleName	The module to find and load
+	  * Loads a module by name.
 	  *
-	  * @return	Returns the module interface, casted to the specified typename, or NULL if the module could not be loaded
+	  * @param ModuleName The module to find and load.
+	  * @return The interface to the module, or nullptr if the module was not found.
+	  * @see GetModulePtr, GetModuleChecked, LoadModuleChecked
 	  */
 	template<typename TModuleInterface>
 	static TModuleInterface* LoadModulePtr( const FName ModuleName )
@@ -279,38 +318,42 @@ public:
 		if (!ModuleManager.IsModuleLoaded(ModuleName))
 		{
 			ModuleManager.LoadModule(ModuleName);
-			if (!ModuleManager.IsModuleLoaded(ModuleName))
-			{
-				return NULL;
-			}
 		}
-		return (TModuleInterface*)(&ModuleManager.GetModuleInterface(ModuleName));
+
+		return GetModulePtr<TModuleInterface>(ModuleName);
 	}
 
+public:
+
 	/**
-	 * Queries information about a specific module name
-	 *
-	 * @param	InModuleName	Module to query status for
-	 * @param	OutModuleStatus	Status of the specified module
+	 * Finds module files on the disk for loadable modules matching the specified wildcard.
 	 *
-	 * @return	True if the module was found the OutModuleStatus is valid
+	 * @param WildcardWithoutExtension Filename part (no path, no extension, no build config info) to search for.
+	 * @param OutModules List of modules found.
 	 */
-	bool QueryModule( const FName InModuleName, FModuleStatus& OutModuleStatus );
-
+	void FindModules( const TCHAR* WildcardWithoutExtension, TArray<FName>& OutModules );
 
 	/**
-	 * Queries information about all of the currently known modules
+	 * Gets the number of loaded modules.
 	 *
-	 * @param	OutModuleStatuses	Status of all modules
+	 * @return The number of modules.
 	 */
-	void QueryModules( TArray< FModuleStatus >& OutModuleStatuses );
+	int32 GetModuleCount( ) const
+	{
+		return Modules.Num();
+	}
 
 	/**
-	 * Query to determine current module count
+	 * Module manager ticking is only used to check for asynchronously compiled modules that may need to be reloaded
+	 */
+	void Tick( );
+
+	/**
+	 * Unloads modules during the shutdown process.
 	 *
-	 * @return	the number of modules
+	 * This method is Usually called at various points while exiting an application.
 	 */
-	int32 GetModuleCount() const { return Modules.Num(); }
+	void UnloadModulesAtShutdown( );
 
 	/**
 	 * Checks for the solution file using the hard-coded location on disk
@@ -331,33 +374,39 @@ public:
 	 * Tries to recompile the specified module.  If the module is loaded, it will first be unloaded (then reloaded after,
 	 * if the recompile was successful.)
 	 *
-	 * @param	InModuleName			Name of the module to recompile
-	 * @param	bReloadAfterRecompile	If true, the module will automatically be reloaded after a successful compile.  Otherwise, you'll need to load it yourself after.
-	 * @param	Ar						Output device for logging compilation status
-	 *
-	 * @return	Returns true if the module was successfully recompiled (and reloaded, if it was previously loaded.)
+	 * @param InModuleName Name of the module to recompile.
+	 * @param bReloadAfterRecompile If true, the module will automatically be reloaded after a successful compile.  Otherwise, you'll need to load it yourself after.
+	 * @param Ar Output device for logging compilation status.
+	 * @return	Returns true if the module was successfully recompiled (and reloaded, if it was previously loaded).
 	 */
 	bool RecompileModule( const FName InModuleName, const bool bReloadAfterRecompile, FOutputDevice &Ar );
 
 	/** @return	Returns true if an asynchronous compile is currently in progress */
 	bool IsCurrentlyCompiling() const;
 
+	/**
+	 * Declares a type of delegates that is executed after a module recompile has finished.
+	 *
+	 * The first argument signals whether compilation has finished.
+	 * The second argument shows whether compilation was successful or not.
+	 */
+	DECLARE_DELEGATE_TwoParams( FRecompileModulesCallback, bool, bool );
+
 	/**
 	 * Tries to recompile the specified modules in the background.  When recompiling finishes, the specified callback
 	 * delegate will be triggered, passing along a bool that tells you whether the compile action succeeded.  This
 	 * function never tries to unload modules or to reload the modules after they finish compiling.  You should do
 	 * that yourself in the recompile completion callback!
 	 *
-	 * @param	ModuleNames					Names of the modules to recompile
-	 * @param	RecompileModulesCallback	Callback function to execute after compilation finishes (whether successful or not.)
-	 * @param	bWaitForCompletion			True if the function should not return until recompilation attempt has finished and callbacks have fired
-	 * @param	Ar							Output device for logging compilation status
-	 *
+	 * @param ModuleNames Names of the modules to recompile
+	 * @param RecompileModulesCallback Callback function to execute after compilation finishes (whether successful or not.)
+	 * @param bWaitForCompletion True if the function should not return until recompilation attempt has finished and callbacks have fired
+	 * @param Ar Output device for logging compilation status
 	 * @return	True if the recompile action was kicked off successfully.  If this returns false, then the recompile callback will never fire.  In the case where bWaitForCompletion=false, this will also return false if the compilation failed for any reason.
 	 */
 	bool RecompileModulesAsync( const TArray< FName > ModuleNames, const FRecompileModulesCallback& RecompileModulesCallback, const bool bWaitForCompletion, FOutputDevice &Ar );
 
-	/** Request that any current compilation operation be abandoned */
+	/** Request that any current compilation operation be abandoned. */
 	void RequestStopCompilation()
 	{
 		bRequestCancelCompilation = true;
@@ -366,19 +415,17 @@ public:
 	/**
 	 * Tries to compile the specified game project. Not used for recompiling modules that are already loaded.
 	 *
-	 * @param	GameProjectFilename		The filename (including path) of the game project to compile
-	 * @param	Ar						Output device for logging compilation status
-	 *
-	 * @return	Returns true if the project was successfully compiled
+	 * @param GameProjectFilename The filename (including path) of the game project to compile.
+	 * @param Ar Output device for logging compilation status.
+	 * @return Returns true if the project was successfully compiled.
 	 */
 	bool CompileGameProject( const FString& GameProjectFilename, FOutputDevice &Ar );
 
 	/**
 	 * Tries to compile the specified game projects editor. Not used for recompiling modules that are already loaded.
 	 *
-	 * @param	GameProjectFilename		The filename (including path) of the game project to compile
-	 * @param	Ar						Output device for logging compilation status
-	 *
+	 * @param GameProjectFilename The filename (including path) of the game project to compile.
+	 * @param Ar Output device for logging compilation status.
 	 * @return	Returns true if the project was successfully compiled
 	 */
 	bool CompileGameProjectEditor( const FString& GameProjectFilename, FOutputDevice &Ar );
@@ -386,10 +433,9 @@ public:
 	/**
 	 * Tries to compile the specified game project. Not used for recompiling modules that are already loaded.
 	 *
-	 * @param	GameProjectFilename		The filename (including path) of the game project for which to generate code projects
-	 * @param	Ar						Output device for logging compilation status
-	 *
-	 * @return	Returns true if the project was successfully compiled
+	 * @param GameProjectFilename The filename (including path) of the game project for which to generate code projects.
+	 * @param Ar Output device for logging compilation status.
+	 * @return	Returns true if the project was successfully compiled.
 	 */
 	bool GenerateCodeProjectFiles( const FString& GameProjectFilename, FOutputDevice &Ar );
 
@@ -398,24 +444,24 @@ public:
 	 */
 	bool IsUnrealBuildToolAvailable();
 
-
 	/** Delegate that's used by the module manager to initialize a registered module that we statically linked with (monolithic only) */
 	DECLARE_DELEGATE_RetVal( IModuleInterface*, FInitializeStaticallyLinkedModule )
 
 	/**
-	 * Registers an initializer for a module that is statically linked
+	 * Registers an initializer for a module that is statically linked.
 	 *
-	 * @param	InModuleName	The name of this module
-	 * @param	InInitializerDelegate	The delegate that will be called to initialize an instance of this module
+	 * @param InModuleName The name of this module.
+	 * @param InInitializerDelegate The delegate that will be called to initialize an instance of this module.
 	 */
 	void RegisterStaticallyLinkedModule( const FName InModuleName, const FInitializeStaticallyLinkedModule& InInitializerDelegate )
 	{
 		StaticallyLinkedModuleInitializers.Add( InModuleName, InInitializerDelegate );
 	}
 
-
-	/** Called by the engine at startup to let the Module Manager know that it's now safe to process
-	    new UObjects discovered by loading C++ modules */
+	/**
+	 * Called by the engine at startup to let the Module Manager know that it's now
+	 * safe to process new UObjects discovered by loading C++ modules.
+	 */
 	void StartProcessingNewlyLoadedObjects();
 
 	/** Adds an engine binaries directory. */
@@ -424,34 +470,24 @@ public:
 	/**
 	 *	Set the game binaries directory
 	 *
-	 *	@param	InDirectory		The game binaries directory
+	 *	@param InDirectory The game binaries directory.
 	 */
 	void SetGameBinariesDirectory(const TCHAR* InDirectory);
 
 	/**
 	 * Checks to see if the specified module exists and is compatible with the current engine version. 
 	 *
-	 * @param	InModuleName		The base name of the module file.
-	 *
-	 * @return	True if module exists and is up to date.
+	 * @param InModuleName The base name of the module file.
+	 * @return true if module exists and is up to date, false otherwise.
 	 */
 	bool IsModuleUpToDate( const FName InModuleName ) const;
 
-	/**
-	 * Adds a module to our list of modules, unless it's already known
-	 *
-	 * @param	InModuleName			The base name of the module file.  Should not include path, extension or platform/configuration info.  This is just the "name" part of the module file name.  Names should be globally unique.
-	 * @param	InBinariesDirectory		The directory where to find this file, or an empty string to search in the default locations.  This parameter is used by the plugin system to locate plugin binaries.
-	 */
-	void AddModule( const FName InModuleName );
-
 	/**
 	 * Determines whether the specified module contains UObjects.  The module must already be loaded into
 	 * memory before calling this function.
 	 *
-	 * @param	ModuleName	Name of the loaded module to check
-	 *
-	 * @return	True if the module was found to contain UObjects, or false if it did not (or wasn't loaded.)
+	 * @param ModuleName Name of the loaded module to check.
+	 * @return True if the module was found to contain UObjects, or false if it did not (or wasn't loaded.)
 	 */
 	bool DoesLoadedModuleHaveUObjects( const FName ModuleName );
 
@@ -462,17 +498,97 @@ public:
 	 */
 	static const TCHAR *GetUBTConfiguration( );
 
-private:
+public:
 
-	/** Private constructor.  Singleton instance is always constructed on demand. */
-	FModuleManager()
-		: bCanProcessNewlyLoadedObjects( false )
-//		,  ModuleCompileProcessHandle( NULL )
-		, ModuleCompileReadPipe( NULL )
-		, bRequestCancelCompilation(false)
+	/**
+	 * Gets an event delegate that is executed when the set of known modules changed, i.e. upon module load or unload.
+	 *
+	 * The first parameter is the name of the module that changed.
+	 * The second parameter is the reason for the change.
+	 *
+	 * @return The event delegate.
+	 */
+	DECLARE_EVENT_TwoParams(FModuleManager, FModulesChangedEvent, FName, EModuleChangeReason);
+	FModulesChangedEvent& OnModulesChanged( )
+	{
+		return ModulesChangedEvent;
+	}
+
+	/**
+	 * Gets an event delegate that is executed when compilation of a module has started.
+	 *
+	 * @return The event delegate.
+	 */
+	DECLARE_EVENT(FModuleManager, FModuleCompilerStartedEvent);
+	FModuleCompilerStartedEvent& OnModuleCompilerStarted( )
+	{
+		return ModuleCompilerStartedEvent;
+	}
+
+	/**
+	 * Gets an event delegate that is executed when compilation of a module has finished.
+	 *
+	 * The first parameter is the result of the compilation operation.
+	 * The second parameter determines whether the log should be shown.
+	 *
+	 * @return The event delegate.
+	 */
+	DECLARE_EVENT_ThreeParams(FModuleManager, FModuleCompilerFinishedEvent, const FString&, ECompilationResult::Type, bool);
+	FModuleCompilerFinishedEvent& OnModuleCompilerFinished()
+	{
+		return ModuleCompilerFinishedEvent;
+	}
+
+	/**
+	 * Gets a multicast delegate that is executed when any UObjects need processing after a module was loaded.
+	 *
+	 * @return The delegate.
+	 */
+	FSimpleMulticastDelegate& OnProcessLoadedObjectsCallback()
 	{
+		return ProcessLoadedObjectsCallback;
 	}
 
+	/**
+	 * Gets a delegate that is executed when a module containing UObjects has been loaded.
+	 *
+	 * The first parameter is the name of the loaded module.
+	 *
+	 * @return The event delegate.
+	 */
+	DECLARE_DELEGATE_RetVal_OneParam(bool, FIsPackageLoadedCallback, FName);
+	FIsPackageLoadedCallback& IsPackageLoadedCallback()
+	{
+		return IsPackageLoaded;
+	}
+
+public:
+
+	// FSelfRegisteringExec interface.
+
+	virtual bool Exec( UWorld* Inworld, const TCHAR* Cmd, FOutputDevice& Ar ) override;
+
+public:
+
+	/** @returns Static: Returns arguments to pass to UnrealBuildTool when compiling modules */
+	static FString MakeUBTArgumentsForModuleCompiling();
+
+protected:
+
+	/**
+	 * Hidden constructor.
+	 *
+	 * Use the static Get function to return the singleton instance.
+	 */
+	FModuleManager( )
+		: bCanProcessNewlyLoadedObjects(false)
+//		,  ModuleCompileProcessHandle(nullptr)
+		, ModuleCompileReadPipe(nullptr)
+		, bRequestCancelCompilation(false)
+	{ }
+
+protected:
+
 	/**
 	 * Helper structure to hold on to module state while asynchronously recompiling DLLs
 	 */
@@ -503,14 +619,13 @@ private:
 		FDateTime FileTimeStamp;
 
 		/** Last known compilation method of the .dll file */
-		EModuleCompileMethod::Type CompileMethod;
+		EModuleCompileMethod CompileMethod;
 
 		FModuleCompilationData()
 			: bIsValid(false)
 			, bHasFileTimeStamp(false)
 			, CompileMethod(EModuleCompileMethod::Unknown)
-		{
-		}
+		{ }
 	};
 
 	/**
@@ -519,6 +634,7 @@ private:
 	class FModuleInfo
 	{
 	public:
+
 		/** The original file name of the module, without any suffixes added */
 		FString OriginalFilename;
 
@@ -528,8 +644,7 @@ private:
 		/** Handle to this module (DLL handle), if it's currently loaded */
 		void* Handle;
 
-		/** The module object for this module.  We actually *own* this module, so it's lifetime is controlled
-		    by the scope of this shared pointer. */
+		/** The module object for this module.  We actually *own* this module, so it's lifetime is controlled by the scope of this shared pointer. */
 		TSharedPtr< IModuleInterface > Module;
 
 		/** True if this module was unloaded at shutdown time, and we never want it to be loaded again */
@@ -548,43 +663,38 @@ private:
 
 		/** Constructor */
 		FModuleInfo()
-			: Handle( NULL ),
+			: Handle( nullptr ),
 			  bWasUnloadedAtShutdown( false ),
 			  LoadOrder(CurrentLoadOrder++)
-		{
-		}
+		{ }
 	};
 
+	/** Type definition for maps of module names to module infos. */
+	typedef TMap<FName, TSharedRef<FModuleInfo>> FModuleMap;
+
 	/**
 	 * Tries to recompile the specified DLL using UBT. Does not interact with modules. This is a low level routine.
 	 *
-	 * @param	ModuleNames			List of modules to recompile, including the module name and optional file suffix
-	 * @param	Ar					Output device for logging compilation status
+	 * @param ModuleNames List of modules to recompile, including the module name and optional file suffix.
+	 * @param Ar Output device for logging compilation status.
 	 */
 	bool RecompileModuleDLLs( const TArray< FModuleToRecompile >& ModuleNames, FOutputDevice &Ar );
 
-public:
-	/** Calls PreUnload then either unloads or abandons a module in memory, depending on whether the module supports unloading */;
-	void UnloadOrAbandonModuleWithCallback( const FName InModuleName, FOutputDevice &Ar );
-
-	/** Loads a module in memory then calls PostLoad */;
-	bool LoadModuleWithCallback( const FName InModuleName, FOutputDevice &Ar );
-
-private:
-	/** Generates a unique file name for the specified module name by adding a random suffix and checking for file collisions */
+	/**
+	 * Generates a unique file name for the specified module name by adding a random suffix and checking for file collisions.
+	 */
 	void MakeUniqueModuleFilename( const FName InModuleName, FString& UniqueSuffix, FString& UniqueModuleFileName );
 
 	/** 
-	 *	Starts compiling DLL files for one or more modules 
+	 *	Starts compiling DLL files for one or more modules.
 	 *
-	 *	@param	GameName							The name of the game
-	 *	@param	ModuleNames							The list of modules to compile
-	 *	@param	InRecompileModulesCallback			Callback function to make when module recompiles
-	 *	@param	Ar									
-	 *	@param	bInFailIfGeneratedCodeChanges		If true, fail the compilation if generated headers change
-	 *	@param	InAdditionalCmdLineArgs				Additional arguments to pass to UBT
-	 *
-	 *	@return	bool								true if successful, false if not
+	 *	@param GameName The name of the game.
+	 *	@param ModuleNames The list of modules to compile.
+	 *	@param InRecompileModulesCallback Callback function to make when module recompiles.
+	 *	@param Ar
+	 *	@param bInFailIfGeneratedCodeChanges If true, fail the compilation if generated headers change.
+	 *	@param InAdditionalCmdLineArgs Additional arguments to pass to UBT.
+	 *	@return true if successful, false otherwise.
 	 */
 	bool StartCompilingModuleDLLs( const FString& GameName, const TArray< FModuleToRecompile >& ModuleNames, 
 		const FRecompileModulesCallback& InRecompileModulesCallback, FOutputDevice& Ar, bool bInFailIfGeneratedCodeChanges, 
@@ -608,15 +718,6 @@ private:
 	/** Compares file versions between the current executing engine version and the specified dll */
 	static bool CheckModuleCompatibility(const TCHAR *Filename);
 
-	// Begin FExec interface.
-	virtual bool Exec( UWorld* Inworld, const TCHAR* Cmd, FOutputDevice& Ar ) override;
-	// End of FExec interface.
-
-public:
-	/** @returns Static: Returns arguments to pass to UnrealBuildTool when compiling modules */
-	static FString MakeUBTArgumentsForModuleCompiling();
-
-private:
 	/** Called during CheckForFinishedModuleDLLCompile() for each successfully recomplied module */
 	void OnModuleCompileSucceeded(FName ModuleName, TSharedRef<FModuleInfo> ModuleInfo);
 
@@ -641,10 +742,8 @@ private:
 private:
 
 	/** Map of all modules.  Maps the case-insensitive module name to information about that module, loaded or not. */
-	typedef TMap< FName, TSharedRef< FModuleInfo > > FModuleMap;
 	FModuleMap Modules;
 
-
 	/** Map of module names to a delegate that can initialize each respective statically linked module */
 	typedef TMap< FName, FInitializeStaticallyLinkedModule > FStaticallyLinkedModuleInitializerMap;
 	FStaticallyLinkedModuleInitializerMap StaticallyLinkedModuleInitializers;
@@ -665,10 +764,6 @@ private:
 	/** Multicast delegate called to process any new loaded objects. */
 	FSimpleMulticastDelegate ProcessLoadedObjectsCallback;
 
-	///
-	/// Async module recompile
-	///
-
 	/** When compiling a module using an external application, stores the handle to the process that is running */
 	FProcHandle ModuleCompileProcessHandle;
 
@@ -701,13 +796,13 @@ private:
 	TArray<FString> GameBinariesDirectories;
 };
 
+
 /**
- * Utility class for registering modules that are statically linked
+ * Utility class for registering modules that are statically linked.
  */
 template< class ModuleClass >
 class FStaticallyLinkedModuleRegistrant
 {
-
 public:
 
 	/**
@@ -724,10 +819,16 @@ public:
 			FName( InModuleName ),	// Module name
 			InitializerDelegate );	// Initializer delegate
 	}
-
 	
-	/** Called by the module manager (via delegate above) to initialize this statically-linked module */
-	IModuleInterface* InitializeModule()
+	/**
+	 * Creates and initializes this statically linked module.
+	 *
+	 * The module manager calls this function through the delegate that was created
+	 * in the @see FStaticallyLinkedModuleRegistrant constructor.
+	 *
+	 * @return A pointer to a new instance of the module.
+	 */
+	IModuleInterface* InitializeModule( )
 	{
 		return new ModuleClass();
 	}
@@ -735,12 +836,12 @@ public:
 
 
 /**
- * Function pointer type for InitializeModule().  All modules must have an InitializeModule() function.
- * Usually this is declared automatically using the IMPLEMENT_MODULE macro below.
+ * Function pointer type for InitializeModule().
  *
- * The function must be declared using as 'extern "C"' so that the name remains undecorated.
- *
- * The object returned will be "owned" by the caller, and will be deleted by the caller before the module is unloaded.
+ * All modules must have an InitializeModule() function. Usually this is declared automatically using
+ * the IMPLEMENT_MODULE macro below. The function must be declared using as 'extern "C"' so that the
+ * name remains undecorated. The object returned will be "owned" by the caller, and will be deleted
+ * by the caller before the module is unloaded.
  */
 typedef IModuleInterface* ( *FInitializeModuleFunctionPtr )( void );
 
@@ -750,8 +851,7 @@ typedef IModuleInterface* ( *FInitializeModuleFunctionPtr )( void );
  */
 class FDefaultModuleImpl
 	: public IModuleInterface
-{
-};
+{ };
 
 
 /**
@@ -763,7 +863,7 @@ class FDefaultGameModuleImpl
 	/**
 	 * Returns true if this module hosts gameplay code
 	 *
-	 * @return True for "gameplay modules", or false for engine code modules, plugins, etc.
+	 * @return True for "gameplay modules", or false for engine code modules, plug-ins, etc.
 	 */
 	virtual bool IsGameModule() const override
 	{
@@ -771,18 +871,25 @@ class FDefaultGameModuleImpl
 	}
 };
 
+
 /**
- * The IMPLEMENT_MODULE macro is used to expose your module's main class to the rest of the engine.
+ * Module implementation boilerplate for regular modules.
+ *
+ * This macro is used to expose a module's main class to the rest of the engine.
  * You must use this macro in one of your modules C++ modules, in order for the 'InitializeModule'
- * function to be declared in such a way that the engine can find it.  Also, this macro will handle
+ * function to be declared in such a way that the engine can find it. Also, this macro will handle
  * the case where a module is statically linked with the engine instead of dynamically loaded.
  *
- * Usage:   IMPLEMENT_MODULE( <My Module Class>, <Module name string> )
+ * This macro is intended for modules that do NOT contain gameplay code.
+ * If your module does contain game classes, use IMPLEMENT_GAME_MODULE instead.
+ *
+ * Usage:   IMPLEMENT_MODULE(<My Module Class>, <Module name string>)
+ *
+ * @see IMPLEMENT_GAME_MODULE
  */
-
-// If we're linking monolithically we assume all modules are linked in with the main binary.
 #if IS_MONOLITHIC
 
+	// If we're linking monolithically we assume all modules are linked in with the main binary.
 	#define IMPLEMENT_MODULE( ModuleImplClass, ModuleName ) \
 		/** Global registrant object for this module when linked statically */ \
 		static FStaticallyLinkedModuleRegistrant< ModuleImplClass > ModuleRegistrant##ModuleName( #ModuleName ); \
@@ -791,7 +898,7 @@ class FDefaultGameModuleImpl
 		void EmptyLinkFunctionForStaticInitialization##ModuleName(){} \
 		PER_MODULE_BOILERPLATE_ANYLINK(ModuleImplClass, ModuleName)
 
-#else	// IS_MONOLITHIC
+#else
 
 	#define IMPLEMENT_MODULE( ModuleImplClass, ModuleName ) \
 		\
@@ -807,27 +914,40 @@ class FDefaultGameModuleImpl
 		PER_MODULE_BOILERPLATE \
 		PER_MODULE_BOILERPLATE_ANYLINK(ModuleImplClass, ModuleName)
 
-#endif	// !IS_MONOLITHIC
+#endif //IS_MONOLITHIC
 
-/** IMPLEMENT_GAME_MODULE is used for modules that contain gameplay code */
+
+/**
+ * Module implementation boilerplate for game play code modules.
+ *
+ * This macro works like IMPLEMENT_MODULE but is specifically used for modules that contain game play code.
+ * If your module does not contain game classes, use IMPLEMENT_MODULE instead.
+ *
+ * Usage:   IMPLEMENT_GAME_MODULE(<My Game Module Class>, <Game Module name string>)
+ *
+ * @see IMPLEMENT_MODULE
+ */
 #define IMPLEMENT_GAME_MODULE( ModuleImplClass, ModuleName ) \
 	IMPLEMENT_MODULE( ModuleImplClass, ModuleName )
 
-/** IMPLEMENT_ENGINE_DIR declares the engine directory to check for foreign or nested projects */
+
+/**
+ * Macro for declaring the engine directory to check for foreign or nested projects.
+ */
 #if PLATFORM_DESKTOP
 	#ifdef UE_ENGINE_DIRECTORY
 		#define IMPLEMENT_FOREIGN_ENGINE_DIR() const TCHAR *GForeignEngineDir = TEXT( PREPROCESSOR_TO_STRING(UE_ENGINE_DIRECTORY) );
 	#else
-		#define IMPLEMENT_FOREIGN_ENGINE_DIR() const TCHAR *GForeignEngineDir = NULL;
+		#define IMPLEMENT_FOREIGN_ENGINE_DIR() const TCHAR *GForeignEngineDir = nullptr;
 	#endif
 #else
 	#define IMPLEMENT_FOREIGN_ENGINE_DIR() 
 #endif
 
+
 #if IS_PROGRAM
 
 	#if IS_MONOLITHIC
-
 		#define IMPLEMENT_APPLICATION( ModuleName, GameName ) \
 			/* For monolithic builds, we must statically define the game's name string (See Core.h) */ \
 			TCHAR GGameName[64] = TEXT( GameName ); \
diff --git a/Engine/Source/Runtime/Engine/Private/LevelTick.cpp b/Engine/Source/Runtime/Engine/Private/LevelTick.cpp
--- a/Engine/Source/Runtime/Engine/Private/LevelTick.cpp
+++ b/Engine/Source/Runtime/Engine/Private/LevelTick.cpp
@@ -972,8 +972,8 @@ static class FFileProfileWrapperExec: private FSelfRegisteringExec
 		static FName TaskGraphModule(TEXT("TaskGraph"));
 		if (FModuleManager::Get().IsModuleLoaded(TaskGraphModule))
 		{
-			IProfileVisualizerModule* ProfileVisualizer = (IProfileVisualizerModule*)&FModuleManager::Get().GetModuleInterface(TaskGraphModule);
-			ProfileVisualizer->DisplayProfileVisualizer( RootEvent, TEXT("I/O") );
+			IProfileVisualizerModule& ProfileVisualizer = FModuleManager::GetModuleChecked<IProfileVisualizerModule>(TaskGraphModule);
+			ProfileVisualizer.DisplayProfileVisualizer( RootEvent, TEXT("I/O") );
 		}
 	}
 
diff --git a/Engine/Source/Runtime/Online/OnlineSubsystem/Private/OnlineSubsystemModule.cpp b/Engine/Source/Runtime/Online/OnlineSubsystem/Private/OnlineSubsystemModule.cpp
--- a/Engine/Source/Runtime/Online/OnlineSubsystem/Private/OnlineSubsystemModule.cpp
+++ b/Engine/Source/Runtime/Online/OnlineSubsystem/Private/OnlineSubsystemModule.cpp
@@ -47,11 +47,7 @@ static TSharedPtr<IModuleInterface> LoadSubsystemModule(const FString& Subsystem
 			ModuleManager.LoadModule(ModuleName);
 		}
 
-		if (ModuleManager.IsModuleLoaded(ModuleName))
-		{
-			// Now return the module itself
-			return ModuleManager.GetModuleInterfaceRef(ModuleName);
-		}
+		return ModuleManager.GetModule(ModuleName);
 	}
 
 	return NULL;
diff --git a/Engine/Source/Runtime/Projects/Private/PluginManager.cpp b/Engine/Source/Runtime/Projects/Private/PluginManager.cpp
--- a/Engine/Source/Runtime/Projects/Private/PluginManager.cpp
+++ b/Engine/Source/Runtime/Projects/Private/PluginManager.cpp
@@ -300,7 +300,7 @@ void FPluginManager::LoadModulesForEnabledPlugins( const ELoadingPhase::Type Loa
 	{
 		if ( Plugin->IsEnabled() )
 		{
-			TMap<FName, ELoadModuleFailureReason::Type> ModuleLoadFailures;
+			TMap<FName, EModuleLoadResult> ModuleLoadFailures;
 			Plugin->LoadModules( LoadingPhase, ModuleLoadFailures );
 
 			FText FailureMessage;
@@ -309,24 +309,24 @@ void FPluginManager::LoadModulesForEnabledPlugins( const ELoadingPhase::Type Loa
 				const auto ModuleNameThatFailedToLoad = FailureIt.Key();
 				const auto FailureReason = FailureIt.Value();
 
-				if( FailureReason != ELoadModuleFailureReason::Success )
+				if( FailureReason != EModuleLoadResult::Success )
 				{
 					const FText PluginNameText = FText::FromString(Plugin->GetPluginName());
 					const FText TextModuleName = FText::FromName(FailureIt.Key());
 
-					if ( FailureReason == ELoadModuleFailureReason::FileNotFound )
+					if ( FailureReason == EModuleLoadResult::FileNotFound )
 					{
 						FailureMessage = FText::Format( LOCTEXT("PluginModuleNotFound", "Plugin '{0}' failed to load because module '{1}' could not be found.  This plugin's functionality will not be available.  Please ensure the plugin is properly installed, otherwise consider disabling the plugin for this project."), PluginNameText, TextModuleName );
 					}
-					else if ( FailureReason == ELoadModuleFailureReason::FileIncompatible )
+					else if ( FailureReason == EModuleLoadResult::FileIncompatible )
 					{
 						FailureMessage = FText::Format( LOCTEXT("PluginModuleIncompatible", "Plugin '{0}' failed to load because module '{1}' does not appear to be compatible with the current version of the engine.  This plugin's functionality will not be available.  The plugin may need to be recompiled."), PluginNameText, TextModuleName );
 					}
-					else if ( FailureReason == ELoadModuleFailureReason::CouldNotBeLoadedByOS )
+					else if ( FailureReason == EModuleLoadResult::CouldNotBeLoadedByOS )
 					{
 						FailureMessage = FText::Format( LOCTEXT("PluginModuleCouldntBeLoaded", "Plugin '{0}' failed to load because module '{1}' could not be loaded.  This plugin's functionality will not be available.  There may be an operating system error or the module may not be properly set up."), PluginNameText, TextModuleName );
 					}
-					else if ( FailureReason == ELoadModuleFailureReason::FailedToInitialize )
+					else if ( FailureReason == EModuleLoadResult::FailedToInitialize )
 					{
 						FailureMessage = FText::Format( LOCTEXT("PluginModuleFailedToInitialize", "Plugin '{0}' failed to load because module '{1}' could be initialized successfully after it was loaded.  This plugin's functionality will not be available."), PluginNameText, TextModuleName );
 					}
diff --git a/Engine/Source/Runtime/Projects/Private/PluginManagerShared.cpp b/Engine/Source/Runtime/Projects/Private/PluginManagerShared.cpp
--- a/Engine/Source/Runtime/Projects/Private/PluginManagerShared.cpp
+++ b/Engine/Source/Runtime/Projects/Private/PluginManagerShared.cpp
@@ -7,7 +7,7 @@
 #define LOCTEXT_NAMESPACE "PluginManagerShared"
 
 
-void FProjectOrPlugin::LoadModules( const ELoadingPhase::Type LoadingPhase, TMap<FName, ELoadModuleFailureReason::Type>& ModuleLoadErrors )
+void FProjectOrPlugin::LoadModules( const ELoadingPhase::Type LoadingPhase, TMap<FName, EModuleLoadResult>& ModuleLoadErrors )
 {
 	const FProjectOrPluginInfo& ProjectOrPluginInfo = GetProjectOrPluginInfo();
 
@@ -24,7 +24,7 @@ void FProjectOrPlugin::LoadModules( const ELoadingPhase::Type LoadingPhase, TMap
 
 				// NOTE: Loading this module may cause other modules to become loaded, both in the engine or game, or other modules 
 				//       that are part of this project or plugin.  That's totally fine.
-				ELoadModuleFailureReason::Type FailureReason;
+				EModuleLoadResult FailureReason;
 				const TSharedPtr<IModuleInterface>& ModuleInterface = FModuleManager::Get().LoadModuleWithFailureReason( ModuleInfo.Name, FailureReason );
 				if( ModuleInterface.IsValid() )
 				{
diff --git a/Engine/Source/Runtime/Projects/Private/PluginManagerShared.h b/Engine/Source/Runtime/Projects/Private/PluginManagerShared.h
--- a/Engine/Source/Runtime/Projects/Private/PluginManagerShared.h
+++ b/Engine/Source/Runtime/Projects/Private/PluginManagerShared.h
@@ -195,7 +195,7 @@ public:
 	 *								loaded at the specified loading phase will be loaded during this call.
 	 * @param	ModuleLoadErrors	If there were load errors, this map has an entry for every failed module.
 	 */
-	void LoadModules( const ELoadingPhase::Type LoadingPhase, TMap<FName, ELoadModuleFailureReason::Type>& ModuleLoadErrors );
+	void LoadModules( const ELoadingPhase::Type LoadingPhase, TMap<FName, EModuleLoadResult>& ModuleLoadErrors );
 
 	/** Populates ProjectInfo from disk */
 	bool LoadFromFile( const FString& FileToLoad, FText& OutFailureReason );
diff --git a/Engine/Source/Runtime/Projects/Private/ProjectManager.cpp b/Engine/Source/Runtime/Projects/Private/ProjectManager.cpp
--- a/Engine/Source/Runtime/Projects/Private/ProjectManager.cpp
+++ b/Engine/Source/Runtime/Projects/Private/ProjectManager.cpp
@@ -119,7 +119,7 @@ bool FProjectManager::LoadModulesForProject( const ELoadingPhase::Type LoadingPh
 
 	if ( CurrentlyLoadedProject.IsValid() )
 	{
-		TMap<FName, ELoadModuleFailureReason::Type> ModuleLoadFailures;
+		TMap<FName, EModuleLoadResult> ModuleLoadFailures;
 		CurrentlyLoadedProject->LoadModules(LoadingPhase, ModuleLoadFailures);
 
 		if ( ModuleLoadFailures.Num() > 0 )
@@ -127,25 +127,25 @@ bool FProjectManager::LoadModulesForProject( const ELoadingPhase::Type LoadingPh
 			FText FailureMessage;
 			for ( auto FailureIt = ModuleLoadFailures.CreateConstIterator(); FailureIt; ++FailureIt )
 			{
-				const ELoadModuleFailureReason::Type FailureReason = FailureIt.Value();
+				const EModuleLoadResult FailureReason = FailureIt.Value();
 
-				if( FailureReason != ELoadModuleFailureReason::Success )
+				if( FailureReason != EModuleLoadResult::Success )
 				{
 					const FText TextModuleName = FText::FromName(FailureIt.Key());
 
-					if ( FailureReason == ELoadModuleFailureReason::FileNotFound )
+					if ( FailureReason == EModuleLoadResult::FileNotFound )
 					{
 						FailureMessage = FText::Format( LOCTEXT("PrimaryGameModuleNotFound", "The game module '{0}' could not be found. Please ensure that this module exists and that it is compiled."), TextModuleName );
 					}
-					else if ( FailureReason == ELoadModuleFailureReason::FileIncompatible )
+					else if ( FailureReason == EModuleLoadResult::FileIncompatible )
 					{
 						FailureMessage = FText::Format( LOCTEXT("PrimaryGameModuleIncompatible", "The game module '{0}' does not appear to be up to date. This may happen after updating the engine. Please recompile this module and try again."), TextModuleName );
 					}
-					else if ( FailureReason == ELoadModuleFailureReason::FailedToInitialize )
+					else if ( FailureReason == EModuleLoadResult::FailedToInitialize )
 					{
 						FailureMessage = FText::Format( LOCTEXT("PrimaryGameModuleFailedToInitialize", "The game module '{0}' could not be successfully initialized after it was loaded."), TextModuleName );
 					}
-					else if ( FailureReason == ELoadModuleFailureReason::CouldNotBeLoadedByOS )
+					else if ( FailureReason == EModuleLoadResult::CouldNotBeLoadedByOS )
 					{
 						FailureMessage = FText::Format( LOCTEXT("PrimaryGameModuleCouldntBeLoaded", "The game module '{0}' could not be loaded. There may be an operating system error or the module may not be properly set up."), TextModuleName );
 					}
diff --git a/Engine/Source/Runtime/RHI/Private/GPUProfiler.cpp b/Engine/Source/Runtime/RHI/Private/GPUProfiler.cpp
--- a/Engine/Source/Runtime/RHI/Private/GPUProfiler.cpp
+++ b/Engine/Source/Runtime/RHI/Private/GPUProfiler.cpp
@@ -226,8 +226,8 @@ void FGPUProfilerEventNodeFrame::DumpEventTree()
 						static FName TaskGraphModule(TEXT("TaskGraph"));			
 						if (FModuleManager::Get().IsModuleLoaded(TaskGraphModule))
 						{
-							IProfileVisualizerModule* ProfileVisualizer = (IProfileVisualizerModule*)&FModuleManager::Get().GetModuleInterface(TaskGraphModule);
-							ProfileVisualizer->DisplayProfileVisualizer( InVisualizerData, TEXT("GPU") );
+							IProfileVisualizerModule& ProfileVisualizer = FModuleManager::GetModuleChecked<IProfileVisualizerModule>(TaskGraphModule);
+							ProfileVisualizer.DisplayProfileVisualizer( InVisualizerData, TEXT("GPU") );
 						}
 					}
 				} DisplayProfilerVisualizer;
diff --git a/Engine/Source/Runtime/Sockets/Private/SocketSubsystem.cpp b/Engine/Source/Runtime/Sockets/Private/SocketSubsystem.cpp
--- a/Engine/Source/Runtime/Sockets/Private/SocketSubsystem.cpp
+++ b/Engine/Source/Runtime/Sockets/Private/SocketSubsystem.cpp
@@ -55,11 +55,7 @@ static TSharedPtr<IModuleInterface> LoadSubsystemModule(const FString& Subsystem
 			ModuleManager.LoadModule(ModuleName);
 		}
 
-		if (ModuleManager.IsModuleLoaded(ModuleName))
-		{
-			// Now return the module itself
-			return ModuleManager.GetModuleInterfaceRef(ModuleName);
-		}
+		return ModuleManager.GetModule(ModuleName);
 	}
 
 	return NULL;
