Include stdint.h on Windows for consistency with other platforms. Fixes build issue with ICU disabled.

diff --git a/Engine/Source/Runtime/Core/Public/Windows/WindowsPlatformIncludes.h b/Engine/Source/Runtime/Core/Public/Windows/WindowsPlatformIncludes.h
--- a/Engine/Source/Runtime/Core/Public/Windows/WindowsPlatformIncludes.h
+++ b/Engine/Source/Runtime/Core/Public/Windows/WindowsPlatformIncludes.h
@@ -25,6 +25,7 @@ extern "C" CORE_API HINSTANCE hInstance;
 // SIMD intrinsics
 #include <intrin.h>
 
+#include <stdint.h>
 #include <tchar.h>
 
 #if USING_CODE_ANALYSIS
