Docs: added documentation links

diff --git a/Engine/Source/Runtime/Core/Public/Modules/ModuleManager.h b/Engine/Source/Runtime/Core/Public/Modules/ModuleManager.h
--- a/Engine/Source/Runtime/Core/Public/Modules/ModuleManager.h
+++ b/Engine/Source/Runtime/Core/Public/Modules/ModuleManager.h
@@ -158,6 +158,7 @@ public:
 	 *
 	 * @param InModuleName Name of the module to return.
 	 * @return 	The module, or nullptr if the module is not loaded.
+	 * @see GetModuleChecked, GetModulePtr
 	 */
 	TSharedPtr<IModuleInterface> GetModule( const FName InModuleName );
 
@@ -178,7 +179,7 @@ public:
 	 * @param InModuleName The base name of the module file.  Should not include path, extension or platform/configuration info.  This is just the "module name" part of the module file name.  Names should be globally unique.
 	 * @param bWasReloaded Indicates that the module has been reloaded (default = false).
 	 * @return The loaded module, or nullptr if the load operation failed.
-	 * @see AbandonModule, IsModuleLoaded, LoadModuleWithFailureReason, UnloadModule
+	 * @see AbandonModule, IsModuleLoaded, LoadModuleChecked, LoadModulePtr, LoadModuleWithFailureReason, UnloadModule
 	 */
 	TSharedPtr<IModuleInterface> LoadModule( const FName InModuleName, const bool bWasReloaded = false );
 
@@ -199,7 +200,7 @@ public:
 	 * @param OutFailureReason Will contain the result.
 	 * @param bWasReloaded Indicates that the module has been reloaded (default = false).
 	 * @return The loaded module (null if the load operation failed).
-	 * @see AbandonModule, IsModuleLoaded, LoadModule, UnloadModule
+	 * @see AbandonModule, IsModuleLoaded, LoadModule, LoadModuleChecked, LoadModulePtr, UnloadModule
 	 */
 	TSharedPtr<IModuleInterface> LoadModuleWithFailureReason( const FName InModuleName, EModuleLoadResult& OutFailureReason, const bool bWasReloaded = false );
 
