Submitting some new automation tests for FText.

diff --git a/Engine/Source/Runtime/Core/Private/Tests/Internationalization/TextTest.cpp b/Engine/Source/Runtime/Core/Private/Tests/Internationalization/TextTest.cpp
--- a/Engine/Source/Runtime/Core/Private/Tests/Internationalization/TextTest.cpp
+++ b/Engine/Source/Runtime/Core/Private/Tests/Internationalization/TextTest.cpp
@@ -407,6 +407,8 @@ bool FTextTest::RunTest (const FString& Parameters)
 		TArray<uint8> InvariantFTextData;
 
 		FString InvariantString = TEXT("This is a culture invariant string.");
+		FString FormattedTestLayer2_OriginalLanguageSourceString;
+		FText FormattedTestLayer2;
 
 		// Scoping to allow all locals to leave scope after we serialize at the end
 		{
@@ -439,7 +441,7 @@ bool FTextTest::RunTest (const FString& Parameters)
 			ArgsLayer2.Add("DateTime", AsDateTimeTest1);
 			ArgsLayer2.Add("Percent", AsPercentTest1);
 			ArgsLayer2.Add("Currency", AsCurrencyTest1);
-			FText FormattedTestLayer2 = FText::Format(LOCTEXT("RebuildTextLayer2", "{NamedLayer1} | {OrderedLayer1} | {FTextNumber} | {Number} | {DateTime} | {Percent} | {Currency}"), ArgsLayer2);
+			FormattedTestLayer2 = FText::Format(LOCTEXT("RebuildTextLayer2", "{NamedLayer1} | {OrderedLayer1} | {FTextNumber} | {Number} | {DateTime} | {Percent} | {Currency}"), ArgsLayer2);
 
 			{
 				// Serialize the full, bulky FText that is a composite of most of the other FTextHistories.
@@ -448,7 +450,8 @@ bool FTextTest::RunTest (const FString& Parameters)
 				Ar.Close();
 			}
 
-			FString FormattedTestLayer2_Baked = FormattedTestLayer2.ToString();
+			// The original string in the native language.
+			FormattedTestLayer2_OriginalLanguageSourceString = FormattedTestLayer2.BuildSourceString();
 
 			{
 				// Swap to "LEET" culture to check if rebuilding works (verify the whole)
@@ -542,6 +545,17 @@ bool FTextTest::RunTest (const FString& Parameters)
 				AddError( TEXT("French-Canadian Output=") + FormattedFrenchCanadianTextHistoryAsLeet.ToString() );
 			}
 
+			// Confirm the two FText's source strings for the serialized FTexts are the same.
+			if(FormattedEnglishTextHistoryAsLeet.BuildSourceString() != FormattedFrenchCanadianTextHistoryAsLeet.BuildSourceString())
+			{
+				AddError( TEXT("Serialization of text histories from source English and source French-Canadian to LEET did not produce the same source results!") );
+				AddError( TEXT("English Output=") + FormattedEnglishTextHistoryAsLeet.BuildSourceString() );
+				AddError( TEXT("French-Canadian Output=") + FormattedFrenchCanadianTextHistoryAsLeet.BuildSourceString() );
+			}
+
+			// Rebuild in LEET so that when we build the source string the DisplayString is still in LEET. 
+			FormattedTestLayer2.ToString();
+
 			{
 				I18N.SetCurrentCulture(OriginalCulture);
 
@@ -556,6 +570,17 @@ bool FTextTest::RunTest (const FString& Parameters)
 					AddError( TEXT("Invariant FText did not match the original FString after serialization!") );
 					AddError( TEXT("Invariant Output=") + InvariantFText.ToString() );
 				}
+
+
+				FString FormattedTestLayer2_SourceString = FormattedTestLayer2.BuildSourceString();
+
+				// Compare the source string of the LEETified version of FormattedTestLayer2 to ensure it is correct.
+				if(FormattedTestLayer2_OriginalLanguageSourceString != FormattedTestLayer2_SourceString)
+				{
+					AddError( TEXT("FormattedTestLayer2's source string was incorrect!") );
+					AddError( TEXT("Output=") + FormattedTestLayer2_SourceString );
+					AddError( TEXT("Desired Output=") + FormattedTestLayer2_OriginalLanguageSourceString );
+				}
 			}
 		}
 	}
