Moved the code responsible for increasing max open files limit on Mac to PlatformPreInit() to solve problem with ICU opening too many files befoew PlatformInit() was called.

diff --git a/Engine/Source/Runtime/Core/Private/Mac/MacPlatformMisc.cpp b/Engine/Source/Runtime/Core/Private/Mac/MacPlatformMisc.cpp
--- a/Engine/Source/Runtime/Core/Private/Mac/MacPlatformMisc.cpp
+++ b/Engine/Source/Runtime/Core/Private/Mac/MacPlatformMisc.cpp
@@ -199,7 +199,7 @@ struct MacApplicationInfo
 };
 static MacApplicationInfo GMacAppInfo;
 
-void FMacPlatformMisc::PlatformInit()
+void FMacPlatformMisc::PlatformPreInit()
 {
 	// Increase the maximum number of simultaneously open files
 	uint32 MaxFilesPerProc = OPEN_MAX;
@@ -228,7 +228,10 @@ void FMacPlatformMisc::PlatformInit()
 	{
 		UE_LOG(LogInit, Warning, TEXT("Failed to change open file limit, UE4 may be unstable."));
 	}
+}
 
+void FMacPlatformMisc::PlatformInit()
+{
 	// Identity.
 	UE_LOG(LogInit, Log, TEXT("Computer: %s"), FPlatformProcess::ComputerName() );
 	UE_LOG(LogInit, Log, TEXT("User: %s"), FPlatformProcess::UserName() );
diff --git a/Engine/Source/Runtime/Core/Public/Mac/MacPlatformMisc.h b/Engine/Source/Runtime/Core/Public/Mac/MacPlatformMisc.h
--- a/Engine/Source/Runtime/Core/Public/Mac/MacPlatformMisc.h
+++ b/Engine/Source/Runtime/Core/Public/Mac/MacPlatformMisc.h
@@ -11,6 +11,7 @@
 **/
 struct CORE_API FMacPlatformMisc : public FGenericPlatformMisc
 {
+	static void PlatformPreInit();
 	static void PlatformInit();
 	static void PlatformPostInit(bool IsMoviePlaying = false);
 	static class GenericApplication* CreateApplication();
