Likely fix for running out of file handles in Mac builds due to ICU file loading.

diff --git a/Engine/Source/Runtime/Core/Private/Internationalization/Internationalization.cpp b/Engine/Source/Runtime/Core/Private/Internationalization/Internationalization.cpp
--- a/Engine/Source/Runtime/Core/Private/Internationalization/Internationalization.cpp
+++ b/Engine/Source/Runtime/Core/Private/Internationalization/Internationalization.cpp
@@ -210,7 +210,9 @@ namespace
 		}
 #endif
 
-		FileAr->Serialize(*contents, FileSize);
+		FileAr->Serialize(*contents, FileSize); 
+
+		FileAr->Close(); // Close the file archive to free up the handle - it is unneeded since the data is in memory.
 
 		return TRUE;
 	}
@@ -219,8 +221,6 @@ namespace
 	{
 		if(fileContext)
 		{
-			FArchive* FileAr = reinterpret_cast<FArchive*>(fileContext);
-
 			if(contents)
 			{
 #if STATS
@@ -233,10 +233,6 @@ namespace
 #endif
 				FICUOverrides::Free(nullptr, contents);
 			}
-			if(FileAr)
-			{
-				FileAr->Close();
-			}
 		}
 	}
 }
