Core: Made FBox2D Min and Max public to match how FBox works

diff --git a/Engine/Source/Runtime/Core/Public/Math/Box2D.h b/Engine/Source/Runtime/Core/Public/Math/Box2D.h
--- a/Engine/Source/Runtime/Core/Public/Math/Box2D.h
+++ b/Engine/Source/Runtime/Core/Public/Math/Box2D.h
@@ -12,6 +12,7 @@
  */
 class FBox2D
 {
+public:
 	/** Holds the box's minimum point. */
 	FVector2D Min;
 
