Add concept of an 'API version' for modules, which is used for determining compatibility of DLLs. By default, the API version is the current changelist number, but can be set to previous changelists when making hotfixes. The API version number is stored in a special resource section for DLLs on Windows, and in the current-version field for dylibs on Mac.

In contrast to the previous system, this does not require special-casing for Rocket builds and will work correctly for nightlies and other out-of-band releases. It also removes the requirement for games to adhere to a strict format of version string on Windows for the editor to be able to load them, since the new data is injected indepdendently of the default resource by UBT.

diff --git a/Engine/Source/Programs/UnrealBuildTool/Configuration/UEBuildBinary.cs b/Engine/Source/Programs/UnrealBuildTool/Configuration/UEBuildBinary.cs
--- a/Engine/Source/Programs/UnrealBuildTool/Configuration/UEBuildBinary.cs
+++ b/Engine/Source/Programs/UnrealBuildTool/Configuration/UEBuildBinary.cs
@@ -504,7 +504,7 @@ namespace UnrealBuildTool
 				bool bHasResFile = false;
 				foreach (var InputFile in BinaryLinkEnvironment.InputFiles)
 				{
-					if (InputFile.AbsolutePath.EndsWith(".res")) 
+					if (InputFile.AbsolutePath.EndsWith(".res") && !InputFile.AbsolutePath.EndsWith(".inl.res")) 
 					{ 
 						bHasResFile = true;
 						break;
@@ -513,7 +513,7 @@ namespace UnrealBuildTool
 
 				foreach( var LinkInputFile in LinkInputFiles )
 				{
-					bool bIsResourceFile = LinkInputFile.AbsolutePath.EndsWith(".res");
+					bool bIsResourceFile = LinkInputFile.AbsolutePath.EndsWith(".res") && !LinkInputFile.AbsolutePath.EndsWith(".inl.res");
 					if( !BinaryLinkEnvironment.InputFiles.Contains( LinkInputFile ) && (!bHasResFile || !bIsResourceFile))
 					{
 						BinaryLinkEnvironment.InputFiles.Add( LinkInputFile );
diff --git a/Engine/Source/Programs/UnrealBuildTool/Configuration/UEBuildModule.cs b/Engine/Source/Programs/UnrealBuildTool/Configuration/UEBuildModule.cs
--- a/Engine/Source/Programs/UnrealBuildTool/Configuration/UEBuildModule.cs
+++ b/Engine/Source/Programs/UnrealBuildTool/Configuration/UEBuildModule.cs
@@ -1080,6 +1080,11 @@ namespace UnrealBuildTool
 					FileItem Item = FileItem.GetItemByFullPath(DefRC);
 					SourceFilesToBuild.RCFiles.Add(Item);
 				}
+
+				// Always compile in the API version resource separately. This is required for the module manager to detect compatible API versions.
+				string ModuleVersionRC = Path.GetFullPath(Path.Combine(Directory.GetCurrentDirectory(), "Runtime/Core/Resources/Windows/ModuleVersionResource.rc.inl"));
+				FileItem ModuleVersionItem = FileItem.GetItemByFullPath(ModuleVersionRC);
+				if(!SourceFilesToBuild.RCFiles.Contains(ModuleVersionItem)) SourceFilesToBuild.RCFiles.Add(ModuleVersionItem);
 			}
 
 			// Check to see if this is an Engine module (including program or plugin modules).  That is, the module is located under the "Engine" folder
diff --git a/Engine/Source/Programs/UnrealBuildTool/Mac/MacToolChain.cs b/Engine/Source/Programs/UnrealBuildTool/Mac/MacToolChain.cs
--- a/Engine/Source/Programs/UnrealBuildTool/Mac/MacToolChain.cs
+++ b/Engine/Source/Programs/UnrealBuildTool/Mac/MacToolChain.cs
@@ -460,27 +460,17 @@ namespace UnrealBuildTool
 			Writer.Write(PreLine + "\n");
 		}
 
-		private string LoadEngineCLVersion()
+		private int LoadEngineCL()
 		{
 			string[] VersionHeader = Utils.ReadAllText("../Source/Runtime/Launch/Resources/Version.h").Replace("\r\n", "\n").Replace("\t", " ").Split('\n');
-			string EngineVersion = "0";
 			foreach (string Line in VersionHeader)
 			{
 				if (Line.StartsWith("#define ENGINE_VERSION "))
 				{
-					EngineVersion = Line.Split(' ')[2];
+					return int.Parse(Line.Split(' ')[2]);
 				}
 			}
-			if (EngineVersion.Length > 6)
-			{
-				EngineVersion = EngineVersion.Insert(EngineVersion.Length - 2, ".");
-				EngineVersion = EngineVersion.Insert(EngineVersion.Length - 5, ".");
-			}
-			else
-			{
-				EngineVersion = LoadEngineDisplayVersion(true);
-			}
-			return EngineVersion;
+			return 0;
 		}
 
 		private string LoadEngineDisplayVersion(bool bIgnorePatchVersion = false)
@@ -531,6 +521,28 @@ namespace UnrealBuildTool
 			return LauncherVersionMajor + "." + LauncherVersionMinor + "." + LauncherVersionPatch;
 		}
 
+		private string LoadEngineAPIVersion()
+		{
+			int CL = 0;
+			foreach (string Line in File.ReadAllLines("../Source/Runtime/Core/Public/Modules/ModuleVersion.h"))
+			{
+				string[] Tokens = Line.Split(' ', '\t');
+				if (Tokens[0] == "#define" && Tokens[1] == "MODULE_API_VERSION")
+				{
+					if(Tokens[2] == "BUILT_FROM_CHANGELIST")
+					{
+						CL = LoadEngineCL();
+					}
+					else
+					{
+						CL = int.Parse(Tokens[2]);
+					}
+					break;
+				}
+			}
+			return String.Format("{0}.{1}.{2}", CL / (100 * 100), (CL / 100) % 100, CL % 100);
+		}
+
 		public override FileItem LinkFiles(LinkEnvironment LinkEnvironment, bool bBuildImportLibraryOnly)
 		{
 			bool bIsBuildingLibrary = LinkEnvironment.Config.bIsBuildingLibrary || bBuildImportLibraryOnly;
@@ -547,9 +559,9 @@ namespace UnrealBuildTool
 			LinkAction.CommandPath = "/bin/sh";
 			LinkAction.CommandDescription = "Link";
 
-			string EngineCLVersion = LoadEngineCLVersion();
+			string EngineAPIVersion = LoadEngineAPIVersion();
 			string EngineDisplayVersion = LoadEngineDisplayVersion(true);
-			string VersionArg = LinkEnvironment.Config.bIsBuildingDLL ? " -current_version " + EngineCLVersion + " -compatibility_version " + EngineDisplayVersion : "";
+			string VersionArg = LinkEnvironment.Config.bIsBuildingDLL ? " -current_version " + EngineAPIVersion + " -compatibility_version " + EngineDisplayVersion : "";
 
 			string Linker = bIsBuildingLibrary ? MacArchiver : MacLinker;
 			string LinkCommand = "xcrun " + Linker + VersionArg + " " + (bIsBuildingLibrary ? GetArchiveArguments_Global(LinkEnvironment) : GetLinkArguments_Global(LinkEnvironment));
diff --git a/Engine/Source/Runtime/Core/Private/GenericPlatform/GenericPlatformProcess.cpp b/Engine/Source/Runtime/Core/Private/GenericPlatform/GenericPlatformProcess.cpp
--- a/Engine/Source/Runtime/Core/Private/GenericPlatform/GenericPlatformProcess.cpp
+++ b/Engine/Source/Runtime/Core/Private/GenericPlatform/GenericPlatformProcess.cpp
@@ -1,6 +1,7 @@
 // Copyright 1998-2014 Epic Games, Inc. All Rights Reserved.
 
 #include "CorePrivate.h"
+#include "../../Public/Modules/ModuleVersion.h"
 
 
 #if PLATFORM_HAS_BSD_TIME 
@@ -26,10 +27,10 @@ void* FGenericPlatformProcess::GetDllExport( void* DllHandle, const TCHAR* ProcN
 	return NULL;
 }
 
-FBinaryFileVersion FGenericPlatformProcess::GetBinaryFileVersion( const TCHAR* Filename )
+int32 FGenericPlatformProcess::GetDllApiVersion( const TCHAR* Filename )
 {
 	UE_LOG(LogHAL, Fatal, TEXT("FPlatformProcess::GetBinaryFileVersion not implemented on this platform"));
-	return FBinaryFileVersion(0, 0, 0, 0);
+	return MODULE_API_VERSION;
 }
 
 uint32 FGenericPlatformProcess::GetCurrentProcessId()
@@ -401,11 +402,6 @@ bool FGenericPlatformProcess::SupportsMultithreading()
 	return bSupportsMultithreading;
 }
 
-void FBinaryFileVersion::ToString(FString& OutString) const
-{
-	OutString = FString::Printf(TEXT("%d %d %d %d"), A, B, C, D);
-}
-
 FGenericPlatformProcess::FSemaphore::FSemaphore(const FString & InName)
 {
 	FCString::Strcpy(Name, sizeof(Name)-1, *InName);
diff --git a/Engine/Source/Runtime/Core/Private/Linux/LinuxPlatformProcess.cpp b/Engine/Source/Runtime/Core/Private/Linux/LinuxPlatformProcess.cpp
--- a/Engine/Source/Runtime/Core/Private/Linux/LinuxPlatformProcess.cpp
+++ b/Engine/Source/Runtime/Core/Private/Linux/LinuxPlatformProcess.cpp
@@ -36,10 +36,10 @@ void* FLinuxPlatformProcess::GetDllExport( void* DllHandle, const TCHAR* ProcNam
 	return dlsym( DllHandle, TCHAR_TO_ANSI(ProcName) );
 }
 
-FBinaryFileVersion FLinuxPlatformProcess::GetBinaryFileVersion( const TCHAR* Filename )
+int32 FLinuxPlatformProcess::GetDllApiVersion( const TCHAR* Filename )
 {
 	check(Filename);
-	return FBinaryFileVersion(ENGINE_MAJOR_VERSION, ENGINE_MINOR_VERSION, 0, FRocketSupport::IsRocket() ? 0 : ENGINE_VERSION);
+	return MODULE_API_VERSION;
 }
 
 const TCHAR* FLinuxPlatformProcess::GetModulePrefix()
diff --git a/Engine/Source/Runtime/Core/Private/Mac/MacPlatformProcess.cpp b/Engine/Source/Runtime/Core/Private/Mac/MacPlatformProcess.cpp
--- a/Engine/Source/Runtime/Core/Private/Mac/MacPlatformProcess.cpp
+++ b/Engine/Source/Runtime/Core/Private/Mac/MacPlatformProcess.cpp
@@ -94,92 +94,75 @@ void* FMacPlatformProcess::GetDllExport( void* DllHandle, const TCHAR* ProcName
 	return dlsym( DllHandle, TCHAR_TO_ANSI(ProcName) );
 }
 
-FBinaryFileVersion FMacPlatformProcess::GetBinaryFileVersion( const TCHAR* Filename )
+int32 FMacPlatformProcess::GetDllApiVersion( const TCHAR* Filename )
 {
 	check(Filename);
 
-	// On Mac the executable itself cannot have a version number. But since we have a compatibility version number set in Core dylib,
-	// OS X makes sure that loaded Core dylib has the same version the executable was linked against. Thanks to this, for the exe
-	// we can simply return FBinaryFileVersion(ENGINE_MAJOR_VERSION, ENGINE_MINOR_VERSION, 0, ENGINE_VERSION)
-
-	if( FCString::Strstr( Filename, TEXT( ".dylib" ) ) )
-	{
-		uint32 CurrentVersion = 0;
-		uint32 CompatibilityVersion = 0;
+	uint32 CurrentVersion = 0;
 		
-		CFStringRef CFStr = FPlatformString::TCHARToCFString(Filename);
+	CFStringRef CFStr = FPlatformString::TCHARToCFString(Filename);
 		
-		NSString* Path = (NSString*)CFStr;
+	NSString* Path = (NSString*)CFStr;
 		
-		if([Path isAbsolutePath] == NO)
-		{
-			NSString* CurDir = [[NSFileManager defaultManager] currentDirectoryPath];
-			NSString* FullPath = [NSString stringWithFormat:@"%@/%@", CurDir, Path];
-			Path = [FullPath stringByResolvingSymlinksInPath];
-		}
+	if([Path isAbsolutePath] == NO)
+	{
+		NSString* CurDir = [[NSFileManager defaultManager] currentDirectoryPath];
+		NSString* FullPath = [NSString stringWithFormat:@"%@/%@", CurDir, Path];
+		Path = [FullPath stringByResolvingSymlinksInPath];
+	}
 		
-		if([[NSFileManager defaultManager] fileExistsAtPath:Path] == NO)
-		{
-			Path = [[[[NSBundle mainBundle] executablePath] stringByDeletingLastPathComponent] stringByAppendingPathComponent: [Path lastPathComponent]];
-		}
+	if([[NSFileManager defaultManager] fileExistsAtPath:Path] == NO)
+	{
+		Path = [[[[NSBundle mainBundle] executablePath] stringByDeletingLastPathComponent] stringByAppendingPathComponent: [Path lastPathComponent]];
+	}
 		
-		BOOL bIsDirectory = NO;
+	BOOL bIsDirectory = NO;
 		
-		int32 File = -1;
+	int32 File = -1;
 		
-		if([[NSFileManager defaultManager] fileExistsAtPath:Path isDirectory:&bIsDirectory] && bIsDirectory == YES && [[NSWorkspace sharedWorkspace] isFilePackageAtPath:Path])
-		{
-			// Try inside the bundle's MacOS folder
-			NSString *FullPath = [[[[NSBundle mainBundle] executablePath] stringByDeletingLastPathComponent] stringByAppendingPathComponent:[Path lastPathComponent]];
-			File = open([FullPath fileSystemRepresentation], O_RDONLY);
-		}
-		else
-		{
-			File = open([Path fileSystemRepresentation], O_RDONLY);
-		}
+	if([[NSFileManager defaultManager] fileExistsAtPath:Path isDirectory:&bIsDirectory] && bIsDirectory == YES && [[NSWorkspace sharedWorkspace] isFilePackageAtPath:Path])
+	{
+		// Try inside the bundle's MacOS folder
+		NSString *FullPath = [[[[NSBundle mainBundle] executablePath] stringByDeletingLastPathComponent] stringByAppendingPathComponent:[Path lastPathComponent]];
+		File = open([FullPath fileSystemRepresentation], O_RDONLY);
+	}
+	else
+	{
+		File = open([Path fileSystemRepresentation], O_RDONLY);
+	}
 		
-		CFRelease(CFStr);
+	CFRelease(CFStr);
 		
-		if(File > -1)
+	if(File > -1)
+	{
+		struct mach_header_64 Header;
+		ssize_t Bytes = read( File, &Header, sizeof( Header ) );
+		if( Bytes == sizeof( Header ) && Header.filetype == MH_DYLIB )
 		{
-			struct mach_header_64 Header;
-			ssize_t Bytes = read( File, &Header, sizeof( Header ) );
-			if( Bytes == sizeof( Header ) && Header.filetype == MH_DYLIB )
-			{
-				struct load_command* Commands = ( struct load_command* )FMemory::Malloc( Header.sizeofcmds );
-				Bytes = read( File, Commands, Header.sizeofcmds );
+			struct load_command* Commands = ( struct load_command* )FMemory::Malloc( Header.sizeofcmds );
+			Bytes = read( File, Commands, Header.sizeofcmds );
 				
-				if( Bytes == Header.sizeofcmds )
+			if( Bytes == Header.sizeofcmds )
+			{
+				struct load_command* Command = Commands;
+				for( int32 Index = 0; Index < Header.ncmds; Index++ )
 				{
-					struct load_command* Command = Commands;
-					for( int32 Index = 0; Index < Header.ncmds; Index++ )
+					if( Command->cmd == LC_ID_DYLIB )
 					{
-						if( Command->cmd == LC_ID_DYLIB )
-						{
-							CurrentVersion = ( ( struct dylib_command* )Command )->dylib.current_version;
-							CompatibilityVersion = ( ( struct dylib_command* )Command )->dylib.compatibility_version;
-							break;
-						}
-						
-						Command = ( struct load_command* )( ( uint8* )Command + Command->cmdsize );
+						CurrentVersion = ( ( struct dylib_command* )Command )->dylib.current_version;
+						break;
 					}
+						
+					Command = ( struct load_command* )( ( uint8* )Command + Command->cmdsize );
 				}
-				
-				FMemory::Free( Commands );
 			}
-			close(File);
+				
+			FMemory::Free( Commands );
 		}
-
-		const uint32 Major = (CompatibilityVersion >> 16) & 0xffff;
-		const uint32 Minor = (CompatibilityVersion >> 8) & 0xff;
-		const uint32 Patch = CompatibilityVersion & 0xff;
-		const uint32 Build = (FRocketSupport::IsRocket() || CurrentVersion == CompatibilityVersion) ? 0 : ((CurrentVersion & 0xff) + ((CurrentVersion >> 8) & 0xff) * 100 + ((CurrentVersion >> 16) & 0xffff) * 10000);
-		return FBinaryFileVersion(Major, Minor, Patch, Build);
-	}
-	else
-	{
-		return FBinaryFileVersion(ENGINE_MAJOR_VERSION, ENGINE_MINOR_VERSION, 0, FRocketSupport::IsRocket() ? 0 : ENGINE_VERSION);
+		close(File);
 	}
+
+	return ((CurrentVersion & 0xff) + ((CurrentVersion >> 8) & 0xff) * 100 + ((CurrentVersion >> 16) & 0xffff) * 10000);
 }
 
 void FMacPlatformProcess::LaunchURL( const TCHAR* URL, const TCHAR* Parms, FString* Error )
diff --git a/Engine/Source/Runtime/Core/Private/Modules/ModuleManager.cpp b/Engine/Source/Runtime/Core/Private/Modules/ModuleManager.cpp
--- a/Engine/Source/Runtime/Core/Private/Modules/ModuleManager.cpp
+++ b/Engine/Source/Runtime/Core/Private/Modules/ModuleManager.cpp
@@ -2,6 +2,7 @@
 
 #include "CorePrivate.h"
 #include "ModuleManager.h"
+#include "ModuleVersion.h"
 
 DEFINE_LOG_CATEGORY_STATIC(LogModuleManager, Log, All);
 
@@ -1922,25 +1923,15 @@ void FModuleManager::CheckForFinishedModuleDLLCompile( const bool bWaitForComple
 
 bool FModuleManager::CheckModuleCompatibility(const TCHAR* Filename)
 {
-	static FBinaryFileVersion ExeVersion = FPlatformProcess::GetBinaryFileVersion(FPlatformProcess::ExecutableName(/*bRemoveExtension=*/false));
-	FBinaryFileVersion DLLVersion = FPlatformProcess::GetBinaryFileVersion(Filename);
-
-	if (ExeVersion.IsValid() && DLLVersion.IsValid() && ExeVersion == DLLVersion)
-	{
-		return true;
-	}
-	else
+	int32 ModuleApiVersion = FPlatformProcess::GetDllApiVersion(Filename);
+	if(ModuleApiVersion != MODULE_API_VERSION)
 	{
-		FString ExeVersionString;
-		FString DLLVersionString;
-		ExeVersion.ToString(ExeVersionString);
-		DLLVersion.ToString(DLLVersionString);
-		UE_LOG(LogModuleManager, Warning, TEXT("Found module file %s (version %s), but it was incompatible with the current engine version (%s). This is likely a stale module that must be recompiled."), Filename, *DLLVersionString, *ExeVersionString);
+		UE_LOG(LogModuleManager, Warning, TEXT("Found module file %s (API version %d), but it was incompatible with the current engine API version (%d). This is likely a stale module that must be recompiled."), Filename, ModuleApiVersion, MODULE_API_VERSION);
 		return false;
 	}
+	return true;
 }
 
-
 bool FModuleManager::RecompileModuleDLLs( const TArray< FModuleToRecompile >& ModuleNames, FOutputDevice& Ar )
 {
 	bool bCompileSucceeded = false;
diff --git a/Engine/Source/Runtime/Core/Private/Windows/WindowsPlatformProcess.cpp b/Engine/Source/Runtime/Core/Private/Windows/WindowsPlatformProcess.cpp
--- a/Engine/Source/Runtime/Core/Private/Windows/WindowsPlatformProcess.cpp
+++ b/Engine/Source/Runtime/Core/Private/Windows/WindowsPlatformProcess.cpp
@@ -6,6 +6,7 @@
 
 #include "CorePrivate.h"
 #include "EngineVersion.h"
+#include "Resources/Windows/ModuleVersionResource.h"
 
 #include "AllowWindowsPlatformTypes.h"
 	#include <shellapi.h>
@@ -70,58 +71,46 @@ void* FWindowsPlatformProcess::GetDllExport( void* DllHandle, const TCHAR* ProcN
 	return (void*)::GetProcAddress( (HMODULE)DllHandle, TCHAR_TO_ANSI(ProcName) );
 }
 
-FBinaryFileVersion FWindowsPlatformProcess::GetBinaryFileVersion( const TCHAR* Filename )
+int32 FWindowsPlatformProcess::GetDllApiVersion( const TCHAR* Filename )
 {
-	check(Filename);
-
-	::DWORD dwSize = GetFileVersionInfoSize( Filename, NULL );
-	if ( dwSize == 0 )
-	{
-		// Failed to get the version info size
-		return FBinaryFileVersion(0, 0, 0, 0);
-	}
-
-	uint8* VersionInfo = new uint8[ dwSize ];
-	if ( !GetFileVersionInfo(Filename, NULL, dwSize, VersionInfo) )
-	{
-		// Failed to get the version info
-		delete[] VersionInfo;
-		return FBinaryFileVersion(0, 0, 0, 0);
-	}
-
-	VS_FIXEDFILEINFO* FileInfo = NULL;
-	::UINT pLenFileInfo = 0;
-	if ( !VerQueryValue( VersionInfo, TEXT("\\"), (LPVOID*) &FileInfo, &pLenFileInfo ) )
-	{
-		// Failed to query the version info for a value
-		delete[] VersionInfo;
-		return FBinaryFileVersion(0, 0, 0, 0);
-	}
+	int32 Result = -1;
 
-	// Get the major/minor/patch fields from the product version
-	int32 Major = (FileInfo->dwProductVersionMS >> 16) & 0xFFFF;
-	int32 Minor = (FileInfo->dwProductVersionMS) & 0xFFFF;
-	int32 Patch = 0;	// We never want to report back the 'patch' version, because the engine has no business using it for determining compatibility.  All hotfix versions are supposed to be binary and API-compatible.
-	int32 Build = (FileInfo->dwProductVersionLS) & 0xFFFF;
-
-	// Try to get the build number from the ProductVersion string. It's too large to fit in a 16-bit ProductVersion field
-	if (Build == 0 && !FRocketSupport::IsRocket())
+	// Retrieves the embedded API version from a DLL
+	check(Filename);
+	HMODULE hModule = LoadLibraryEx(Filename, NULL, LOAD_LIBRARY_AS_DATAFILE);
+	if(hModule != NULL)
 	{
-		TCHAR *ProductVersion;
-		::UINT ProductVersionLen;
-		if (VerQueryValue(VersionInfo, TEXT("\\StringFileInfo\\040904b0\\ProductVersion"), (LPVOID*) &ProductVersion, &ProductVersionLen))
+		HRSRC hResInfo = FindResource(hModule, MAKEINTRESOURCE(ID_MODULE_API_VERSION_RESOURCE), RT_RCDATA);
+		if(hResInfo != NULL)
 		{
-			FEngineVersion DllEngineVersion;
-			if (FEngineVersion::Parse(ProductVersion, DllEngineVersion))
+			HGLOBAL hResGlobal = LoadResource(hModule, hResInfo);
+			if(hResGlobal != NULL)
 			{
-				Build = DllEngineVersion.GetChangelist();
+				void *pResData = LockResource(hResGlobal);
+				if(pResData != NULL)
+				{
+					::DWORD Length = SizeofResource(hModule, hResInfo);
+					if(Length > 0)
+					{
+						char *Str = (char*)pResData;
+						if(Str[Length - 1] == 0)
+						{
+							char *End = Str;
+							uint64 Value = FCStringAnsi::Strtoui64(Str, &End, 10);
+							if(*End == 0 && Value <= INT_MAX)
+							{
+								Result = (int32)Value;
+							}
+						}
+					}
+					UnlockResource(pResData);
+				}
 			}
 		}
+		FreeLibrary(hModule);
 	}
 
-	delete[] VersionInfo;
-
-	return FBinaryFileVersion(Major, Minor, Patch, Build);
+	return Result;
 }
 
 void FWindowsPlatformProcess::PushDllDirectory(const TCHAR* Directory)
diff --git a/Engine/Source/Runtime/Core/Public/GenericPlatform/GenericPlatformProcess.h b/Engine/Source/Runtime/Core/Public/GenericPlatform/GenericPlatformProcess.h
--- a/Engine/Source/Runtime/Core/Public/GenericPlatform/GenericPlatformProcess.h
+++ b/Engine/Source/Runtime/Core/Public/GenericPlatform/GenericPlatformProcess.h
@@ -33,31 +33,6 @@ namespace ELaunchVerb
 	};
 }
 
-struct FBinaryFileVersion
-{
-	FBinaryFileVersion(int32 InA, int32 InB, int32 InC, int32 InD)
-		: A(InA), B(InB), C(InC), D(InD)
-	{}
-
-	bool operator==(const FBinaryFileVersion& Other)
-	{
-		return A == Other.A && B == Other.B && C == Other.C && D == Other.D;
-	}
-
-	bool IsValid() const
-	{
-		return (A | B | C | D) != 0;
-	}
-
-	/**
-	 * Print as formatted string
-	 */
-	void ToString(class FString& OutString) const;
-
-private:
-	int32 A, B, C, D;
-};
-
 /** Generic implementation for the process handle. */
 template< typename T, T InvalidHandleValue >
 struct TProcHandle
@@ -173,8 +148,8 @@ struct CORE_API FGenericPlatformProcess
 	/** Lookup the address of a DLL function. **/
 	static void* GetDllExport( void* DllHandle, const TCHAR* ProcName );
 
-	/** Gets a version number from the specified DLL or EXE **/
-	static FBinaryFileVersion GetBinaryFileVersion( const TCHAR* Filename );
+	/** Gets the API version from the specified DLL **/
+	static int32 GetDllApiVersion( const TCHAR* Filename );
 
 	/** Set a directory to look for DLL files. NEEDS to have a Pop call when complete */
 	FORCEINLINE static void PushDllDirectory(const TCHAR* Directory)
diff --git a/Engine/Source/Runtime/Core/Public/Linux/LinuxPlatformProcess.h b/Engine/Source/Runtime/Core/Public/Linux/LinuxPlatformProcess.h
--- a/Engine/Source/Runtime/Core/Public/Linux/LinuxPlatformProcess.h
+++ b/Engine/Source/Runtime/Core/Public/Linux/LinuxPlatformProcess.h
@@ -129,7 +129,7 @@ struct CORE_API FLinuxPlatformProcess : public FGenericPlatformProcess
 	static void* GetDllHandle( const TCHAR* Filename );
 	static void FreeDllHandle( void* DllHandle );
 	static void* GetDllExport( void* DllHandle, const TCHAR* ProcName );
-	static FBinaryFileVersion GetBinaryFileVersion( const TCHAR* Filename );
+	static int32 GetDllApiVersion( const TCHAR* Filename );
 	static const TCHAR* ComputerName();
 	static const TCHAR* BaseDir();
 	static FString GetApplicationName( uint32 ProcessId );
diff --git a/Engine/Source/Runtime/Core/Public/Mac/MacPlatformProcess.h b/Engine/Source/Runtime/Core/Public/Mac/MacPlatformProcess.h
--- a/Engine/Source/Runtime/Core/Public/Mac/MacPlatformProcess.h
+++ b/Engine/Source/Runtime/Core/Public/Mac/MacPlatformProcess.h
@@ -30,7 +30,7 @@ struct CORE_API FMacPlatformProcess : public FGenericPlatformProcess
 	static void* GetDllHandle( const TCHAR* Filename );
 	static void FreeDllHandle( void* DllHandle );
 	static void* GetDllExport( void* DllHandle, const TCHAR* ProcName );
-	static FBinaryFileVersion GetBinaryFileVersion( const TCHAR* Filename );
+	static int32 GetDllApiVersion( const TCHAR* Filename );
 	static void CleanFileCache();
 	static uint32 GetCurrentProcessId();
 	static const TCHAR* BaseDir();
diff --git a/Engine/Source/Runtime/Core/Public/Modules/ModuleVersion.h b/Engine/Source/Runtime/Core/Public/Modules/ModuleVersion.h
new file mode 100644
--- /dev/null
+++ b/Engine/Source/Runtime/Core/Public/Modules/ModuleVersion.h
@@ -0,0 +1,14 @@
+// Copyright 1998-2014 Epic Games, Inc. All Rights Reserved.
+
+#pragma once
+
+#include "../../../Launch/Resources/Version.h"
+
+// This number identifies a particular API revision, and is used to determine module compatibility. Hotfixes should retain the API version of the original release.
+// This define is parsed by the build tools, and should be a number or BUILT_FROM_CHANGELIST.
+#define MODULE_API_VERSION BUILT_FROM_CHANGELIST
+
+// Check that the API version has been set manually for a hotfix release
+#if ENGINE_PATCH_VERSION > 0 && BUILT_FROM_CHANGELIST > 0 && MODULE_API_VERSION == BUILT_FROM_CHANGELIST
+	#error MODULE_API_VERSION must be manually defined for hotfix builds
+#endif
diff --git a/Engine/Source/Runtime/Core/Public/Windows/WindowsPlatformProcess.h b/Engine/Source/Runtime/Core/Public/Windows/WindowsPlatformProcess.h
--- a/Engine/Source/Runtime/Core/Public/Windows/WindowsPlatformProcess.h
+++ b/Engine/Source/Runtime/Core/Public/Windows/WindowsPlatformProcess.h
@@ -70,7 +70,7 @@ struct CORE_API FWindowsPlatformProcess : public FGenericPlatformProcess
 	static void* GetDllHandle( const TCHAR* Filename );
 	static void FreeDllHandle( void* DllHandle );
 	static void* GetDllExport( void* DllHandle, const TCHAR* ProcName );
-	static FBinaryFileVersion GetBinaryFileVersion( const TCHAR* Filename );
+	static int32 GetDllApiVersion( const TCHAR* Filename );
 	static void PushDllDirectory(const TCHAR* Directory);
 	static void PopDllDirectory(const TCHAR* Directory);
 	static void CleanFileCache();
diff --git a/Engine/Source/Runtime/Core/Resources/Windows/ModuleVersionResource.h b/Engine/Source/Runtime/Core/Resources/Windows/ModuleVersionResource.h
new file mode 100644
--- /dev/null
+++ b/Engine/Source/Runtime/Core/Resources/Windows/ModuleVersionResource.h
@@ -0,0 +1,5 @@
+// Copyright 1998-2014 Epic Games, Inc. All Rights Reserved.
+
+#pragma once
+
+#define ID_MODULE_API_VERSION_RESOURCE 191
diff --git a/Engine/Source/Runtime/Core/Resources/Windows/ModuleVersionResource.rc.inl b/Engine/Source/Runtime/Core/Resources/Windows/ModuleVersionResource.rc.inl
new file mode 100644
--- /dev/null
+++ b/Engine/Source/Runtime/Core/Resources/Windows/ModuleVersionResource.rc.inl
@@ -0,0 +1,19 @@
+// Copyright 1998-2014 Epic Games, Inc. All Rights Reserved.
+
+#include <windows.h>
+#include "../../Public/Modules/ModuleVersion.h"
+#include "ModuleVersionResource.h"
+
+#define XSTR(s) #s
+#define STRINGIFY(s) XSTR(s)
+
+// Netral language
+LANGUAGE LANG_NEUTRAL,SUBLANG_NEUTRAL
+
+// Declare the engine API version string
+#if !IS_MONOLITHIC
+ID_MODULE_API_VERSION_RESOURCE RCDATA
+{
+	STRINGIFY(MODULE_API_VERSION) 0
+}
+#endif
\ No newline at end of file
diff --git a/Engine/Source/Runtime/Launch/Private/LaunchEngineLoop.cpp b/Engine/Source/Runtime/Launch/Private/LaunchEngineLoop.cpp
--- a/Engine/Source/Runtime/Launch/Private/LaunchEngineLoop.cpp
+++ b/Engine/Source/Runtime/Launch/Private/LaunchEngineLoop.cpp
@@ -11,6 +11,7 @@
 #include "ExceptionHandling.h"
 #include "FileManagerGeneric.h"
 #include "TaskGraphInterfaces.h"
+#include "Runtime/Core/Public/Modules/ModuleVersion.h"
 
 
 #include "Projects.h"
@@ -2393,6 +2394,7 @@ bool FEngineLoop::AppInit( )
 
 	//// Command line.
 	UE_LOG(LogInit, Log, TEXT("Version: %s"), *GEngineVersion.ToString());
+	UE_LOG(LogInit, Log, TEXT("API Version: %u"), MODULE_API_VERSION);
 
 #if PLATFORM_64BITS
 	UE_LOG(LogInit, Log, TEXT("Compiled (64-bit): %s %s"), ANSI_TO_TCHAR(__DATE__), ANSI_TO_TCHAR(__TIME__));
