Fixed pressed state reporting for gamepad analog triggers on Mac

diff --git a/Engine/Source/Runtime/Core/Private/Mac/HIDInputInterface.cpp b/Engine/Source/Runtime/Core/Private/Mac/HIDInputInterface.cpp
--- a/Engine/Source/Runtime/Core/Private/Mac/HIDInputInterface.cpp
+++ b/Engine/Source/Runtime/Core/Private/Mac/HIDInputInterface.cpp
@@ -314,17 +314,23 @@ void HIDInputInterface::SendControllerEvents()
 							CurrentButtonStates[21] = FloatValue > 0.2f;	// EControllerButtons::RightStickRight
 							ControllerState.RightAnalogY = NewValue;
 						}
-						else if (Element.Usage == ControllerState.Device.LeftTriggerAnalogMapping && ControllerState.LeftTriggerAnalog != NewValue)
+						else if (Element.Usage == ControllerState.Device.LeftTriggerAnalogMapping)
 						{
-							MessageHandler->OnControllerAnalog(EControllerButtons::LeftTriggerAnalog, ControllerState.ControllerId, FloatValue);
+							if (ControllerState.LeftTriggerAnalog != NewValue)
+							{
+								MessageHandler->OnControllerAnalog(EControllerButtons::LeftTriggerAnalog, ControllerState.ControllerId, FloatValue);
+								ControllerState.LeftTriggerAnalog = NewValue;
+							}
 							CurrentButtonStates[10] = FloatValue > 0.01f;
-							ControllerState.LeftTriggerAnalog = NewValue;
 						}
-						else if (Element.Usage == ControllerState.Device.RightTriggerAnalogMapping && ControllerState.RightTriggerAnalog != NewValue)
+						else if (Element.Usage == ControllerState.Device.RightTriggerAnalogMapping)
 						{
-							MessageHandler->OnControllerAnalog(EControllerButtons::RightTriggerAnalog, ControllerState.ControllerId, FloatValue);
+							if (ControllerState.RightTriggerAnalog != NewValue)
+							{
+								MessageHandler->OnControllerAnalog(EControllerButtons::RightTriggerAnalog, ControllerState.ControllerId, FloatValue);
+								ControllerState.RightTriggerAnalog = NewValue;
+							}
 							CurrentButtonStates[11] = FloatValue > 0.01f;
-							ControllerState.RightTriggerAnalog = NewValue;
 						}
 						else if (Element.Usage == kHIDUsage_GD_Hatswitch)
 						{
