Stats2 - Increased stack size for the stats thread, from 128k to 512k

diff --git a/Engine/Source/Runtime/Core/Private/Stats/Stats2.cpp b/Engine/Source/Runtime/Core/Private/Stats/Stats2.cpp
--- a/Engine/Source/Runtime/Core/Private/Stats/Stats2.cpp
+++ b/Engine/Source/Runtime/Core/Private/Stats/Stats2.cpp
@@ -743,7 +743,7 @@ public:
 
 	void Start()
 	{
-		Thread = FRunnableThread::Create(this, TEXT("StatsThread"), 128 * 1024, TPri_BelowNormal, FPlatformAffinity::GetStatsThreadMask());
+		Thread = FRunnableThread::Create(this, TEXT("StatsThread"), 512 * 1024, TPri_BelowNormal, FPlatformAffinity::GetStatsThreadMask());
 		check(Thread != NULL);
 	}
 
