Corrected use of platform physical file abstraction rather than file manager abstraction in ICU data directory check.

diff --git a/Engine/Source/Runtime/Core/Private/Internationalization/Internationalization.cpp b/Engine/Source/Runtime/Core/Private/Internationalization/Internationalization.cpp
--- a/Engine/Source/Runtime/Core/Private/Internationalization/Internationalization.cpp
+++ b/Engine/Source/Runtime/Core/Private/Internationalization/Internationalization.cpp
@@ -263,21 +263,21 @@ void FInternationalization::Initialize()
 
 	FString DataDirectory;
 	const FString DataDirectoryRelativeToContent = FString(TEXT("Localization")) / FString(TEXT("ICU"));
-	IPlatformFile& PlatformFile = IPlatformFile::GetPlatformPhysical();
+	IFileManager& FileManager = IFileManager::Get();
 	DataDirectory = FPaths::GameContentDir() / DataDirectoryRelativeToContent; // Try game content directory.
-	if( !PlatformFile.DirectoryExists(*DataDirectory) )
+	if( !FileManager.DirectoryExists(*DataDirectory) )
 	{
 		DataDirectory = FString(FPlatformProcess::BaseDir()) / FPaths::GameContentDir() / DataDirectoryRelativeToContent; // Try game content directory with appropriate base.
 	}
-	if( !PlatformFile.DirectoryExists(*DataDirectory) )
+	if( !FileManager.DirectoryExists(*DataDirectory) )
 	{
 		FString DataDirectory = FPaths::EngineContentDir() / DataDirectoryRelativeToContent; // Try engine content directory.
 	}
-	if( !PlatformFile.DirectoryExists(*DataDirectory) )
+	if( !FileManager.DirectoryExists(*DataDirectory) )
 	{
 		DataDirectory = FString(FPlatformProcess::BaseDir()) / FPaths::EngineContentDir() / DataDirectoryRelativeToContent; // Try engine content directory with appropriate base.
 	}
-	check( PlatformFile.DirectoryExists(*DataDirectory) );
+	check( FileManager.DirectoryExists(*DataDirectory) );
 	u_setDataDirectory(StringCast<char>(*DataDirectory).Get());
 
 	u_setDataFileFunctions(nullptr, &OpenDataFile, &CloseDataFile, &(ICUStatus));
