#ue4 Adding script call stack dumping when asserting.  Every FFrame will now push and pop on it's info to a global stack record, and dump them with the native stack when assertions fail

diff --git a/Engine/Source/Runtime/Core/Private/Misc/Core.cpp b/Engine/Source/Runtime/Core/Private/Misc/Core.cpp
--- a/Engine/Source/Runtime/Core/Private/Misc/Core.cpp
+++ b/Engine/Source/Runtime/Core/Private/Misc/Core.cpp
@@ -228,6 +228,9 @@ bool					GPumpingMessagesOutsideOfMainLoop = false;
 /** Total blueprint compile time.																			*/
 double GBlueprintCompileTime = 0.0;
 
+/** Stack names from the VM to be unrolled when we assert */
+TArray<FScriptTraceStackNode> GScriptStack;
+
 DEFINE_STAT(STAT_AudioMemory);
 DEFINE_STAT(STAT_TextureMemory);
 DEFINE_STAT(STAT_MemoryPhysXTotalAllocationSize);
diff --git a/Engine/Source/Runtime/Core/Private/Misc/OutputDevice.cpp b/Engine/Source/Runtime/Core/Private/Misc/OutputDevice.cpp
--- a/Engine/Source/Runtime/Core/Private/Misc/OutputDevice.cpp
+++ b/Engine/Source/Runtime/Core/Private/Misc/OutputDevice.cpp
@@ -183,6 +183,20 @@ void VARARGS FDebug::AssertFailed( const ANSICHAR* Expr, const ANSICHAR* File, i
 		ANSICHAR* StackTrace = (ANSICHAR*) FMemory::SystemMalloc( StackTraceSize );
 		if( StackTrace != NULL )
 		{
+#if !(UE_BUILD_SHIPPING || UE_BUILD_TEST)
+			// Walk the script stack, if any
+			if (GScriptStack.Num() > 0)
+			{
+				FString ScriptStack = TEXT("\n\nScript Stack:\n");
+				while (GScriptStack.Num())
+				{
+					ScriptStack += GScriptStack.Pop().GetStackDescription() + TEXT("\n");
+				}
+
+				UE_LOG(LogOutputDevice, Warning, TEXT("%s"), *ScriptStack);
+			}
+#endif
+
 			StackTrace[0] = 0;
 			// Walk the stack and dump it to the allocated memory.
 			FPlatformStackWalk::StackWalkAndDump( StackTrace, StackTraceSize, CALLSTACK_IGNOREDEPTH );
diff --git a/Engine/Source/Runtime/Core/Public/Core.h b/Engine/Source/Runtime/Core/Public/Core.h
--- a/Engine/Source/Runtime/Core/Public/Core.h
+++ b/Engine/Source/Runtime/Core/Public/Core.h
@@ -363,6 +363,9 @@ extern CORE_API bool GPumpingMessagesOutsideOfMainLoop;
 /** Total blueprint compile time. */
 extern CORE_API double GBlueprintCompileTime;
 
+/** Stack names from the VM to be unrolled when we assert */
+extern CORE_API TArray<struct FScriptTraceStackNode> GScriptStack;
+
 /*----------------------------------------------------------------------------
 	Metadata macros.
 ----------------------------------------------------------------------------*/
@@ -406,6 +409,7 @@ extern CORE_API double GBlueprintCompileTime;
 #include "Bitstreams.h"					// Bit stream archiver.
 #include "SparseArray.h"				// Sparse array definitions.
 #include "UnrealString.h"				// Dynamic string definitions.
+#include "NameTypes.h"					// Global name subsystem.
 #include "CoreMisc.h"					// Low level utility code.
 #include "Paths.h"						// Path helper functions
 #include "StaticArray.h"                // Static array definition.
@@ -416,7 +420,6 @@ extern CORE_API double GBlueprintCompileTime;
 #include "List.h"						// Dynamic list definitions.
 #include "ResourceArray.h"				// Resource array definitions.
 #include "RefCounting.h"				// Reference counting definitions.
-#include "NameTypes.h"					// Global name subsystem.
 #include "ScriptDelegates.h"
 #include "Delegate.h"					// C++ delegate system
 #include "ThreadingBase.h"				// Non-platform specific multi-threaded support.
diff --git a/Engine/Source/Runtime/Core/Public/Misc/CoreMisc.h b/Engine/Source/Runtime/Core/Public/Misc/CoreMisc.h
--- a/Engine/Source/Runtime/Core/Public/Misc/CoreMisc.h
+++ b/Engine/Source/Runtime/Core/Public/Misc/CoreMisc.h
@@ -341,3 +341,26 @@ bool CORE_API StringHasBadDashes(const TCHAR* Str);
 
 /** Helper function to generate a set of windowed resolutions which are convenient for the current primary display size */
 void CORE_API GenerateConvenientWindowedResolutions(const struct FDisplayMetrics& InDisplayMetrics, TArray<FIntPoint>& OutResolutions);
+
+/**
+ * Helper for script stack traces
+ */
+struct FScriptTraceStackNode
+{
+	FName	Scope;
+	FName	FunctionName;
+
+	FScriptTraceStackNode(FName InScope, FName InFunctionName)
+		: Scope(InScope)
+		, FunctionName(InFunctionName)
+	{
+	}
+
+	FString GetStackDescription() const
+	{
+		return Scope.ToString() + TEXT(".") + FunctionName.ToString();
+	}
+
+private:
+	FScriptTraceStackNode() {};
+};
\ No newline at end of file
diff --git a/Engine/Source/Runtime/CoreUObject/Public/UObject/Stack.h b/Engine/Source/Runtime/CoreUObject/Public/UObject/Stack.h
--- a/Engine/Source/Runtime/CoreUObject/Public/UObject/Stack.h
+++ b/Engine/Source/Runtime/CoreUObject/Public/UObject/Stack.h
@@ -102,7 +102,14 @@ public:
 	FFrame( UObject* InObject, UFunction* InNode, void* InLocals, FFrame* InPreviousFrame = NULL, UField* InPropertyChainForCompiledIn = NULL );
 
 	virtual ~FFrame()
-	{}
+	{
+#if !(UE_BUILD_SHIPPING || UE_BUILD_TEST)
+		if (GScriptStack.Num())
+		{
+			GScriptStack.Pop();
+		}
+#endif
+	}
 
 	// Functions.
 	COREUOBJECT_API void Step( UObject* Context, RESULT_DECL );
@@ -168,7 +175,12 @@ inline FFrame::FFrame( UObject* InObject, UFunction* InNode, void* InLocals, FFr
 	, PreviousFrame(InPreviousFrame)
 	, OutParms(NULL)
 	, PropertyChainForCompiledIn(InPropertyChainForCompiledIn)
-{}
+{
+#if !(UE_BUILD_SHIPPING || UE_BUILD_TEST)
+	FScriptTraceStackNode StackNode(InNode->GetOuter()->GetFName(), InNode->GetFName());
+ 	GScriptStack.Push(StackNode);
+#endif
+}
 
 inline int32 FFrame::ReadInt()
 {
