The editor now sends a Machine ID and the last used Epic Account ID when started

TTP# 337511 - ANALYTICS: Store Epic Account ID and Machine Unique GUIDs on Users Machine. Send with Every Event.

This works differently to what the TTP originally asked for, but I've checked with MaxP/MattK/WesH and we're okay with the way this works, as it avoids an annoyingly tight dependency between the analytics and the editor.

On the request of MattK, I also changed FPlatformMisc::GetRegistryString to be FPlatformMisc::GetStoredValue and added a corresponding FPlatformMisc::SetStoredValue. These will store generic key->value pairs using a platform specific store (the registry on Windows, and an INI file on other platforms).

DesktopPlatform uses these to provide access to the Machine ID and Epic Account ID (the Epic Account ID is set when the user logs into the Launcher).

The Editor sends an "Editor.ProgramStarted" event containing this data. This can then be tied to the current analytics session via the analytics session ID.

Tested on Windows and OS X.

diff --git a/Engine/Source/Developer/DesktopPlatform/Private/DesktopPlatformBase.cpp b/Engine/Source/Developer/DesktopPlatform/Private/DesktopPlatformBase.cpp
--- a/Engine/Source/Developer/DesktopPlatform/Private/DesktopPlatformBase.cpp
+++ b/Engine/Source/Developer/DesktopPlatform/Private/DesktopPlatformBase.cpp
@@ -421,6 +421,40 @@ FString FDesktopPlatformBase::GetDefaultProjectCreationPath()
 	return FString(FPlatformProcess::UserDir()) + DefaultProjectSubFolder;
 }
 
+FGuid FDesktopPlatformBase::GetMachineId()
+{
+	FGuid MachineId;
+	FString MachineIdStr;
+
+	// Check to see if we already have a valid machine ID to use
+	if(!FPlatformMisc::GetStoredValue(TEXT("Epic Games"), TEXT("Unreal Engine/Identifiers"), TEXT("MachineId"), MachineIdStr) || !FGuid::Parse(MachineIdStr, MachineId))
+	{
+		// No valid machine ID, generate and save a new one
+		MachineId = FGuid::NewGuid();
+		MachineIdStr = MachineId.ToString(EGuidFormats::Digits);
+
+		if(!FPlatformMisc::SetStoredValue(TEXT("Epic Games"), TEXT("Unreal Engine/Identifiers"), TEXT("MachineId"), MachineIdStr))
+		{
+			// Failed to persist the machine ID - reset it to zero to avoid returning a transient value
+			MachineId = FGuid();
+		}
+	}
+
+	return MachineId;
+}
+
+FString FDesktopPlatformBase::GetEpicAccountId()
+{
+	FString AccountId;
+	FPlatformMisc::GetStoredValue(TEXT("Epic Games"), TEXT("Unreal Engine/Identifiers"), TEXT("AccountId"), AccountId);
+	return AccountId;
+}
+
+void FDesktopPlatformBase::SetEpicAccountId(const FString& AccountId)
+{
+	FPlatformMisc::SetStoredValue(TEXT("Epic Games"), TEXT("Unreal Engine/Identifiers"), TEXT("AccountId"), AccountId);
+}
+
 void FDesktopPlatformBase::ReadLauncherInstallationList()
 {
 	FString InstalledListFile = FString(FPlatformProcess::ApplicationSettingsDir()) / TEXT("UnrealEngineLauncher/LauncherInstalled.dat");
diff --git a/Engine/Source/Developer/DesktopPlatform/Private/DesktopPlatformBase.h b/Engine/Source/Developer/DesktopPlatform/Private/DesktopPlatformBase.h
--- a/Engine/Source/Developer/DesktopPlatform/Private/DesktopPlatformBase.h
+++ b/Engine/Source/Developer/DesktopPlatform/Private/DesktopPlatformBase.h
@@ -42,6 +42,10 @@ public:
 	virtual bool EnumerateProjectsKnownByEngine(const FString &Identifier, bool bIncludeNativeProjects, TArray<FString> &OutProjectFileNames) override;
 	virtual FString GetDefaultProjectCreationPath() override;
 
+	virtual FGuid GetMachineId() override;
+	virtual FString GetEpicAccountId() override;
+	virtual void SetEpicAccountId(const FString& AccountId) override;
+
 private:
 	FString CurrentEngineIdentifier;
 	FDateTime LauncherInstallationTimestamp;
diff --git a/Engine/Source/Developer/DesktopPlatform/Public/IDesktopPlatform.h b/Engine/Source/Developer/DesktopPlatform/Public/IDesktopPlatform.h
--- a/Engine/Source/Developer/DesktopPlatform/Public/IDesktopPlatform.h
+++ b/Engine/Source/Developer/DesktopPlatform/Public/IDesktopPlatform.h
@@ -355,4 +355,19 @@ public:
 	* @return Path to the folder
 	*/
 	virtual FString GetDefaultProjectCreationPath() = 0;
+
+	/** 
+	 * Get (or create) the unique ID used to identify this computer
+	 */
+	virtual FGuid GetMachineId() = 0;
+
+	/**
+	 * Get the Epic account ID for the user who last used the Launcher
+	 */
+	virtual FString GetEpicAccountId() = 0;
+
+	/**
+	 * Set the Epic account ID for the user who last used the Launcher
+	 */
+	virtual void SetEpicAccountId(const FString& AccountId) = 0;
 };
diff --git a/Engine/Source/Editor/UnrealEd/Private/UnrealEd.cpp b/Engine/Source/Editor/UnrealEd/Private/UnrealEd.cpp
--- a/Engine/Source/Editor/UnrealEd/Private/UnrealEd.cpp
+++ b/Engine/Source/Editor/UnrealEd/Private/UnrealEd.cpp
@@ -17,6 +17,7 @@
 #include "EngineAnalytics.h"
 #include "IAnalyticsProvider.h"
 
+#include "DesktopPlatformModule.h"
 #include "GameProjectGenerationModule.h"
 #include "AssetEditorManager.h"
 
@@ -61,6 +62,20 @@ int32 EditorInit( IEngineLoop& EngineLoop )
 		return 0;
 	}
 
+	// Let the analytics know that the editor has started
+	if ( FEngineAnalytics::IsAvailable() )
+	{
+		IDesktopPlatform* const DesktopPlatform = FDesktopPlatformModule::Get();
+		if ( DesktopPlatform )
+		{
+			TArray<FAnalyticsEventAttribute> EventAttributes;
+			EventAttributes.Add(FAnalyticsEventAttribute(TEXT("MachineID"), DesktopPlatform->GetMachineId().ToString(EGuidFormats::Digits).ToLower()));
+			EventAttributes.Add(FAnalyticsEventAttribute(TEXT("AccountID"), DesktopPlatform->GetEpicAccountId()));
+
+			FEngineAnalytics::GetProvider().RecordEvent(TEXT("Editor.ProgramStarted"), EventAttributes);
+		}
+	}
+
 	// Initialize the misc editor
 	FUnrealEdMisc::Get().OnInit();
 
diff --git a/Engine/Source/Runtime/Core/Private/GenericPlatform/GenericPlatformMisc.cpp b/Engine/Source/Runtime/Core/Private/GenericPlatform/GenericPlatformMisc.cpp
--- a/Engine/Source/Runtime/Core/Private/GenericPlatform/GenericPlatformMisc.cpp
+++ b/Engine/Source/Runtime/Core/Private/GenericPlatform/GenericPlatformMisc.cpp
@@ -238,9 +238,52 @@ void FGenericPlatformMisc::HandleIOFailure( const TCHAR* Filename )
 	UE_LOG(LogGenericPlatformMisc, Fatal,TEXT("I/O failure operating on '%s'"), Filename ? Filename : TEXT("Unknown file"));
 }
 
-bool FGenericPlatformMisc::GetRegistryString(const FString& InRegistryKey, const FString& InValueName, bool bPerUserSetting, FString& OutValue)
+bool FGenericPlatformMisc::SetStoredValue(const FString& InStoreId, const FString& InSectionName, const FString& InKeyName, const FString& InValue)
 {
-	// By default, fail.
+	check(!InStoreId.IsEmpty());
+	check(!InSectionName.IsEmpty());
+	check(!InKeyName.IsEmpty());
+
+	// This assumes that FPlatformProcess::ApplicationSettingsDir() returns a user-specific directory; it doesn't on Windows, but Windows overrides this behavior to use the registry
+	const FString ConfigPath = FString(FPlatformProcess::ApplicationSettingsDir()) / InStoreId / FString(TEXT("KeyValueStore.ini"));
+		
+	FConfigFile ConfigFile;
+	ConfigFile.Read(ConfigPath);
+
+	FConfigSection& Section = ConfigFile.FindOrAdd(InSectionName);
+
+	FString& KeyValue = Section.FindOrAdd(*InKeyName);
+	KeyValue = InValue;
+
+	ConfigFile.Dirty = true;
+	ConfigFile.Write(ConfigPath);
+
+	return true;
+}
+
+bool FGenericPlatformMisc::GetStoredValue(const FString& InStoreId, const FString& InSectionName, const FString& InKeyName, FString& OutValue)
+{
+	check(!InStoreId.IsEmpty());
+	check(!InSectionName.IsEmpty());
+	check(!InKeyName.IsEmpty());
+
+	// This assumes that FPlatformProcess::ApplicationSettingsDir() returns a user-specific directory; it doesn't on Windows, but Windows overrides this behavior to use the registry
+	const FString ConfigPath = FString(FPlatformProcess::ApplicationSettingsDir()) / InStoreId / FString(TEXT("KeyValueStore.ini"));
+		
+	FConfigFile ConfigFile;
+	ConfigFile.Read(ConfigPath);
+
+	const FConfigSection* const Section = ConfigFile.Find(InSectionName);
+	if(Section)
+	{
+		const FString* const KeyValue = Section->Find(*InKeyName);
+		if(KeyValue)
+		{
+			OutValue = *KeyValue;
+			return true;
+		}
+	}
+
 	return false;
 }
 
diff --git a/Engine/Source/Runtime/Core/Private/Windows/WindowsPlatformMisc.cpp b/Engine/Source/Runtime/Core/Private/Windows/WindowsPlatformMisc.cpp
--- a/Engine/Source/Runtime/Core/Private/Windows/WindowsPlatformMisc.cpp
+++ b/Engine/Source/Runtime/Core/Private/Windows/WindowsPlatformMisc.cpp
@@ -1810,6 +1810,46 @@ void FWindowsPlatformMisc::RaiseException( uint32 ExceptionCode )
 	::RaiseException( ExceptionCode, 0, 0, NULL );
 }
 
+bool FWindowsPlatformMisc::SetStoredValue(const FString& InStoreId, const FString& InSectionName, const FString& InKeyName, const FString& InValue)
+{
+	check(!InStoreId.IsEmpty());
+	check(!InSectionName.IsEmpty());
+	check(!InKeyName.IsEmpty());
+
+	FString FullRegistryKey = FString(TEXT("Software")) / InStoreId / InSectionName;
+	FullRegistryKey = FullRegistryKey.Replace(TEXT("/"), TEXT("\\")); // we use forward slashes, but the registry needs back slashes
+
+	HKEY hKey;
+	HRESULT Result = ::RegCreateKeyEx(HKEY_CURRENT_USER, *FullRegistryKey, 0, nullptr, REG_OPTION_NON_VOLATILE, KEY_WRITE, nullptr, &hKey, nullptr);
+	if(Result == ERROR_SUCCESS)
+	{
+		Result = ::RegSetValueEx(hKey, *InKeyName, 0, REG_SZ, (const BYTE*)*InValue, (InValue.Len() + 1) * sizeof(TCHAR));
+		::RegCloseKey(hKey);
+	}
+	
+	if(Result != ERROR_SUCCESS)
+	{
+		TCHAR ErrorBuffer[1024];
+		::FormatMessageW(FORMAT_MESSAGE_FROM_SYSTEM, nullptr, Result, 0, ErrorBuffer, 1024, nullptr);
+		GWarn->Logf(TEXT("FWindowsPlatformMisc::SetStoredValue: ERROR: Could not store value for '%s'. Error Code %u: %s"), *InKeyName, Result, ErrorBuffer);
+		return false;
+	}
+
+	return true;
+}
+
+bool FWindowsPlatformMisc::GetStoredValue(const FString& InStoreId, const FString& InSectionName, const FString& InKeyName, FString& OutValue)
+{
+	check(!InStoreId.IsEmpty());
+	check(!InSectionName.IsEmpty());
+	check(!InKeyName.IsEmpty());
+
+	FString FullRegistryKey = FString(TEXT("Software")) / InStoreId / InSectionName;
+	FullRegistryKey = FullRegistryKey.Replace(TEXT("/"), TEXT("\\")); // we use forward slashes, but the registry needs back slashes
+
+	return QueryRegKey(HKEY_CURRENT_USER, *FullRegistryKey, *InKeyName, OutValue);
+}
+
 uint32 FWindowsPlatformMisc::GetLastError()
 {
 	return (uint32)::GetLastError();
@@ -2161,20 +2201,6 @@ int32 FWindowsPlatformMisc::GetCacheLineSize()
 	return FCPUIDQueriedData::GetCacheLineSize();
 }
 
-bool FWindowsPlatformMisc::GetRegistryString(const FString& InRegistryKey, const FString& InValueName, bool bPerUserSetting, FString& OutValue)
-{
-	FString RegistryValue;
-	FString FullRegistryKey = FString(TEXT("Software")) / InRegistryKey;
-	FullRegistryKey = FullRegistryKey.Replace(TEXT("/"), TEXT("\\"));
-	if (QueryRegKey(bPerUserSetting? HKEY_CURRENT_USER : HKEY_LOCAL_MACHINE, *FullRegistryKey, *InValueName, RegistryValue) == true)
-	{
-		OutValue = RegistryValue;
-		return true;
-	}
-
-	return false;
-}
-
 bool FWindowsPlatformMisc::QueryRegKey( const HKEY InKey, const TCHAR* InSubKey, const TCHAR* InValueName, FString& OutData )
 {
 	bool bSuccess = false;
diff --git a/Engine/Source/Runtime/Core/Public/GenericPlatform/GenericPlatformMisc.h b/Engine/Source/Runtime/Core/Public/GenericPlatform/GenericPlatformMisc.h
--- a/Engine/Source/Runtime/Core/Public/GenericPlatform/GenericPlatformMisc.h
+++ b/Engine/Source/Runtime/Core/Public/GenericPlatform/GenericPlatformMisc.h
@@ -430,17 +430,30 @@ public:
 	}
 
 	/** 
-	 *	Retrieve the entry for the given string from the platforms equivalent to the Windows Registry
+	 *	Set the value for the given section and key in the platform specific key->value store
+	 *  Note: The key->value store is user-specific, but may be used to share data between different applications for the same user
 	 *
-	 *	@param	InRegistryKey		The registry key to query
-	 *	@param	InValueName			The name of the value from that key
-	 *  @param	bPerUserSetting		Whether to get the per-user registry key rather than the local machine's registry key
+	 *  @param	InStoreId			The name used to identify the store you want to use (eg, MyGame)
+	 *	@param	InSectionName		The section that this key->value pair is placed within (can contain / separators, eg UserDetails/AccountInfo)
+	 *	@param	InKeyName			The name of the key to set the value for
+	 *	@param	InValue				The value to set
+	 *
+	 *	@return	bool				true if the value was set correctly, false if not
+	 */
+	static bool SetStoredValue(const FString& InStoreId, const FString& InSectionName, const FString& InKeyName, const FString& InValue);
+
+	/** 
+	 *	Get the value for the given section and key from the platform specific key->value store
+	 *  Note: The key->value store is user-specific, but may be used to share data between different applications for the same user
+	 *
+	 *  @param	InStoreId			The name used to identify the store you want to use (eg, MyGame)
+	 *	@param	InSectionName		The section that this key->value pair is placed within (can contain / separators, eg UserDetails/AccountInfo)
+	 *	@param	InKeyName			The name of the key to get the value for
 	 *	@param	OutValue			The value found
 	 *
 	 *	@return	bool				true if the entry was found (and OutValue contains the result), false if not
 	 */
-	static bool GetRegistryString(const FString& InRegistryKey, const FString& InValueName, bool bPerUserSetting, FString& OutValue);
-
+	static bool GetStoredValue(const FString& InStoreId, const FString& InSectionName, const FString& InKeyName, FString& OutValue);
 
 	 /** Sends a message to a remote tool, and debugger consoles */
 	static void LowLevelOutputDebugString(const TCHAR *Message);
diff --git a/Engine/Source/Runtime/Core/Public/Misc/ConfigCacheIni.h b/Engine/Source/Runtime/Core/Public/Misc/ConfigCacheIni.h
--- a/Engine/Source/Runtime/Core/Public/Misc/ConfigCacheIni.h
+++ b/Engine/Source/Runtime/Core/Public/Misc/ConfigCacheIni.h
@@ -66,18 +66,17 @@ public:
 
 	CORE_API bool Combine( const FString& Filename);
 	CORE_API void CombineFromBuffer(const FString& Filename,const FString& Buffer);
-	void Read( const FString& Filename );
-	bool Write( const FString& Filename, bool bDoRemoteWrite=true, const FString& InitialText=FString() );
+	CORE_API void Read( const FString& Filename );
+	CORE_API bool Write( const FString& Filename, bool bDoRemoteWrite=true, const FString& InitialText=FString() );
 	CORE_API void Dump(FOutputDevice& Ar);
 
 	CORE_API bool GetString( const TCHAR* Section, const TCHAR* Key, FString& Value ) const;
 	CORE_API bool GetText( const TCHAR* Section, const TCHAR* Key, FText& Value ) const;
-	bool GetInt64( const TCHAR* Section, const TCHAR* Key, int64& Value ) const;
+	CORE_API bool GetInt64( const TCHAR* Section, const TCHAR* Key, int64& Value ) const;
 
-	void SetString( const TCHAR* Section, const TCHAR* Key, const TCHAR* Value );
-	void SetText( const TCHAR* Section, const TCHAR* Key, const FText& Value );
-
-	void SetInt64( const TCHAR* Section, const TCHAR* Key, const int64 Value );
+	CORE_API void SetString( const TCHAR* Section, const TCHAR* Key, const TCHAR* Value );
+	CORE_API void SetText( const TCHAR* Section, const TCHAR* Key, const FText& Value );
+	CORE_API void SetInt64( const TCHAR* Section, const TCHAR* Key, const int64 Value );
 	
 	/**
 	 * Process the contents of an .ini file that has been read into an FString
diff --git a/Engine/Source/Runtime/Core/Public/Windows/WindowsPlatformMisc.h b/Engine/Source/Runtime/Core/Public/Windows/WindowsPlatformMisc.h
--- a/Engine/Source/Runtime/Core/Public/Windows/WindowsPlatformMisc.h
+++ b/Engine/Source/Runtime/Core/Public/Windows/WindowsPlatformMisc.h
@@ -72,6 +72,9 @@ struct CORE_API FWindowsPlatformMisc : public FGenericPlatformMisc
 
 	static void RaiseException( uint32 ExceptionCode );
 
+	static bool SetStoredValue(const FString& InStoreId, const FString& InSectionName, const FString& InKeyName, const FString& InValue);
+	static bool GetStoredValue(const FString& InStoreId, const FString& InSectionName, const FString& InKeyName, FString& OutValue);
+
 	static bool CoInitialize();
 	static void CoUninitialize();
 
@@ -172,8 +175,6 @@ struct CORE_API FWindowsPlatformMisc : public FGenericPlatformMisc
 	 */
 	static uint32 GetCPUInfo();
 
-	static bool GetRegistryString(const FString& InRegistryKey, const FString& InValueName, bool bPerUserSetting, FString& OutValue);
-
 	/** 
 	 * Provides a simpler interface for fetching and cleanup of registry value queries
 	 *
