TBigInt fixes.

Originated from here: https://answers.unrealengine.com/questions/57349/tbigint-has-issues-shifting-by-multiples-of-32-bit.html

#add Added some tests for multiple-of-32-bits shifts.
#add Added compile-time NumWords constant for everywhere that was duplicating the same calculation.
#change Changed the IsEqual function to be less weird.
#fix Implemented a rotating implementation for shifts of exactly a multiple of the word bits.
#fix ToString was writing out incorrectly for zeroes in the lower-order words.

diff --git a/Engine/Source/Runtime/Core/Private/GenericPlatform/GenericPlatformMath.cpp b/Engine/Source/Runtime/Core/Private/GenericPlatform/GenericPlatformMath.cpp
--- a/Engine/Source/Runtime/Core/Private/GenericPlatform/GenericPlatformMath.cpp
+++ b/Engine/Source/Runtime/Core/Private/GenericPlatform/GenericPlatformMath.cpp
@@ -74,4 +74,13 @@ void FGenericPlatformMath::AutoTest()
 		check(Dividend.ToInt() == 4788146LL);
 		check(Remainder.ToInt() == 334LL);
 	}
+
+	{
+		// Shift test with multiple of 32
+		int256 Value(1);
+		Value <<= 32;
+		check(Value.ToInt() == 4294967296LL);
+		Value >>= 32;
+		check(Value.ToInt() == 1LL);
+	}
 }
diff --git a/Engine/Source/Runtime/Core/Public/Math/BigInt.h b/Engine/Source/Runtime/Core/Public/Math/BigInt.h
--- a/Engine/Source/Runtime/Core/Public/Math/BigInt.h
+++ b/Engine/Source/Runtime/Core/Public/Math/BigInt.h
@@ -14,13 +14,14 @@ class TBigInt
 	enum
 	{
 		/** Word size. */
-		BitsPerWord = 32
+		BitsPerWord = 32,
+		NumWords    = NumBits / BitsPerWord
 	};
 
 	static_assert(NumBits >= 64, "TBigInt must have at least 64 bits.");
 
 	/** All bits stored as an array of words. */
-	uint32 Bits[NumBits / BitsPerWord];
+	uint32 Bits[NumWords];
 	
 	/**
 	 * Makes sure both factors are positive integers and stores their original signs
@@ -119,18 +120,29 @@ public:
 	FORCEINLINE void ShiftLeftInternal(const int32 BitCount)
 	{
 		checkSlow(BitCount > 0);
+
 		TBigInt<NumBits> Result;
-		const int32 NumWords = NumBits / BitsPerWord;
-		const int32 HiWordOffset = (BitCount - 1) / BitsPerWord + 1;
-		const int32 LoWordOffset = HiWordOffset - 1;
-		const int32 LoWordShift = BitCount - LoWordOffset * BitsPerWord;
-		const int32 HiWordShift = BitsPerWord - LoWordShift;
-		Result.Bits[NumWords - 1] |= Bits[NumWords - LoWordOffset - 1] << LoWordShift;
-		for (int32 WordIndex = NumWords - HiWordOffset - 1; WordIndex >= 0; --WordIndex)
-		{			
-			uint32 Value = Bits[WordIndex];
-			Result.Bits[WordIndex + LoWordOffset] |= Value << LoWordShift;
-			Result.Bits[WordIndex + HiWordOffset] |= Value >> HiWordShift;
+		if (BitCount & (BitsPerWord - 1))
+		{
+			const int32 LoWordOffset = (BitCount - 1) / BitsPerWord;
+			const int32 HiWordOffset = LoWordOffset + 1;
+			const int32 LoWordShift  = BitCount - LoWordOffset * BitsPerWord;
+			const int32 HiWordShift  = BitsPerWord - LoWordShift;
+			Result.Bits[NumWords - 1] |= Bits[NumWords - HiWordOffset] << LoWordShift;
+			for (int32 WordIndex = (NumWords - 1) - HiWordOffset; WordIndex >= 0; --WordIndex)
+			{
+				uint32 Value = Bits[WordIndex];
+				Result.Bits[WordIndex + LoWordOffset] |= Value << LoWordShift;
+				Result.Bits[WordIndex + HiWordOffset] |= Value >> HiWordShift;
+			}
+		}
+		else
+		{
+			const int32 ShiftWords = BitCount / BitsPerWord;
+			for (int32 WordIndex = NumWords - 1; WordIndex >= ShiftWords; --WordIndex)
+			{
+				Result.Bits[WordIndex] = Bits[WordIndex - ShiftWords];
+			}
 		}
 		*this = Result;
 	}
@@ -143,18 +155,29 @@ public:
 	FORCEINLINE void ShiftRightInternal(const int32 BitCount)
 	{
 		checkSlow(BitCount > 0);
+
 		TBigInt<NumBits> Result;
-		const int32 NumWords = NumBits / BitsPerWord;
-		const int32 LoWordOffset = (BitCount - 1) / BitsPerWord + 1;
-		const int32 HiWordOffset = LoWordOffset - 1;
-		const int32 HiWordShift = BitCount - HiWordOffset * BitsPerWord;
-		const int32 LoWordShift = BitsPerWord - HiWordShift;
-		Result.Bits[0] |= Bits[HiWordOffset] >> HiWordShift;
-		for (int32 WordIndex = LoWordOffset; WordIndex < NumWords; ++WordIndex)
-		{			
-			uint32 Value = Bits[WordIndex];
-			Result.Bits[WordIndex - HiWordOffset] |= Value >> HiWordShift;
-			Result.Bits[WordIndex - LoWordOffset] |= Value << LoWordShift;
+		if (BitCount & (BitsPerWord - 1))
+		{
+			const int32 HiWordOffset = (BitCount - 1) / BitsPerWord;
+			const int32 LoWordOffset = HiWordOffset + 1;
+			const int32 HiWordShift  = BitCount - HiWordOffset * BitsPerWord;
+			const int32 LoWordShift  = BitsPerWord - HiWordShift;
+			Result.Bits[0] |= Bits[HiWordOffset] >> HiWordShift;
+			for (int32 WordIndex = LoWordOffset; WordIndex < NumWords; ++WordIndex)
+			{
+				uint32 Value = Bits[WordIndex];
+				Result.Bits[WordIndex - HiWordOffset] |= Value >> HiWordShift;
+				Result.Bits[WordIndex - LoWordOffset] |= Value << LoWordShift;
+			}
+		}
+		else
+		{
+			const int32 ShiftWords = BitCount / BitsPerWord;
+			for (int32 WordIndex = NumWords - 1; WordIndex >= ShiftWords; --WordIndex)
+			{
+				Result.Bits[WordIndex - ShiftWords] = Bits[WordIndex];
+			}
 		}
 		*this = Result;
 	}
@@ -165,7 +188,7 @@ public:
 	FORCEINLINE void Add(const TBigInt<NumBits>& Other)
 	{
 		int64 CarryOver = 0;
-		for (int32 WordIndex = 0; WordIndex < NumBits / BitsPerWord; ++WordIndex)
+		for (int32 WordIndex = 0; WordIndex < NumWords; ++WordIndex)
 		{
 			int64 WordSum = (int64)Other.Bits[WordIndex] + (int64)Bits[WordIndex] + CarryOver;
 			CarryOver = WordSum >> BitsPerWord;
@@ -189,7 +212,7 @@ public:
 	 */
 	FORCEINLINE bool IsNegative() const
 	{
-		return !!(Bits[NumBits / BitsPerWord - 1] & (1U << (BitsPerWord - 1)));
+		return !!(Bits[NumWords - 1] & (1U << (BitsPerWord - 1)));
 	}
 
 	/**
@@ -206,7 +229,7 @@ public:
 	void Negate()
 	{
 		static const TBigInt<NumBits> One(1LL);
-		for (int32 WordIndex = 0; WordIndex < NumBits / BitsPerWord; ++WordIndex)
+		for (int32 WordIndex = 0; WordIndex < NumWords; ++WordIndex)
 		{
 			Bits[WordIndex] = ~Bits[WordIndex];
 		}
@@ -228,7 +251,7 @@ public:
 		MakePositiveFactors(Temp, ThisSign, Other, OtherSign);
 
 		int32 ShiftCount = 0;
-		for (int32 WordIndex = 0; WordIndex < NumBits / BitsPerWord; ++WordIndex)
+		for (int32 WordIndex = 0; WordIndex < NumWords; ++WordIndex)
 		{			
 			if (Other.Bits[WordIndex])
 			{
@@ -381,7 +404,7 @@ public:
 	FORCEINLINE int32 GetHighestNonZeroWord() const
 	{
 		int32 WordIndex;
-		for (WordIndex = NumBits / BitsPerWord - 1; WordIndex >= 0 && Bits[WordIndex] == 0; --WordIndex);
+		for (WordIndex = NumWords - 1; WordIndex >= 0 && Bits[WordIndex] == 0; --WordIndex);
 		return WordIndex;
 	}
 
@@ -390,7 +413,7 @@ public:
 	 */
 	FORCEINLINE int32 GetHighestNonZeroBit() const
 	{
-		for (int32 WordIndex = NumBits / BitsPerWord - 1; WordIndex >= 0; --WordIndex)
+		for (int32 WordIndex = NumWords - 1; WordIndex >= 0; --WordIndex)
 		{
 			if (!!Bits[WordIndex])
 			{
@@ -484,8 +507,7 @@ public:
 	 */
 	FORCEINLINE void BitwiseOr(const TBigInt<NumBits>& Other)
 	{
-		const int32 WordCount = NumBits / BitsPerWord;
-		for (int32 WordIndex = 0; WordIndex < WordCount; ++WordIndex)
+		for (int32 WordIndex = 0; WordIndex < NumWords; ++WordIndex)
 		{
 			Bits[WordIndex] |= Other.Bits[WordIndex];
 		}
@@ -496,8 +518,7 @@ public:
 	 */
 	FORCEINLINE void BitwiseAnd(const TBigInt<NumBits>& Other)
 	{
-		const int32 WordCount = NumBits / BitsPerWord;
-		for (int32 WordIndex = 0; WordIndex < WordCount; ++WordIndex)
+		for (int32 WordIndex = 0; WordIndex < NumWords; ++WordIndex)
 		{
 			Bits[WordIndex] &= Other.Bits[WordIndex];
 		}
@@ -508,8 +529,7 @@ public:
 	 */
 	FORCEINLINE void BitwiseNot()
 	{
-		const int32 WordCount = NumBits / BitsPerWord;
-		for (int32 WordIndex = 0; WordIndex < WordCount; ++WordIndex)
+		for (int32 WordIndex = 0; WordIndex < NumWords; ++WordIndex)
 		{
 			Bits[WordIndex] = ~Bits[WordIndex];
 		}
@@ -520,10 +540,15 @@ public:
 	 */
 	FORCEINLINE bool IsEqual(const TBigInt<NumBits>& Other) const
 	{
-		const int32 WordCount = NumBits / BitsPerWord;
-		int32 WordIndex;
-		for (WordIndex = 0; WordIndex < WordCount && Bits[WordIndex] == Other.Bits[WordIndex]; ++WordIndex);
-		return WordIndex == WordCount;
+		for (int32 WordIndex = 0; WordIndex < NumWords; ++WordIndex)
+		{
+			if (Bits[WordIndex] != Other.Bits[WordIndex])
+			{
+				return false;
+			}
+		}
+
+		return true;
 	}
 
 	/**
@@ -543,7 +568,7 @@ public:
 			}
 		}
 		int32 WordIndex;
-		for (WordIndex = NumBits / BitsPerWord - 1; WordIndex >= 0 && Other.Bits[WordIndex] == Bits[WordIndex]; --WordIndex);
+		for (WordIndex = NumWords - 1; WordIndex >= 0 && Other.Bits[WordIndex] == Bits[WordIndex]; --WordIndex);
 		return WordIndex >= 0 && Bits[WordIndex] < Other.Bits[WordIndex];
 	}
 
@@ -564,7 +589,7 @@ public:
 			}
 		}
 		int32 WordIndex;
-		for (WordIndex = NumBits / BitsPerWord - 1; WordIndex >= 0 && Other.Bits[WordIndex] == Bits[WordIndex]; --WordIndex);
+		for (WordIndex = NumWords - 1; WordIndex >= 0 && Other.Bits[WordIndex] == Bits[WordIndex]; --WordIndex);
 		return WordIndex < 0 || Bits[WordIndex] < Other.Bits[WordIndex];
 	}
 
@@ -585,7 +610,7 @@ public:
 			}
 		}
 		int32 WordIndex;
-		for (WordIndex = NumBits / BitsPerWord - 1; WordIndex >= 0 && Other.Bits[WordIndex] == Bits[WordIndex]; --WordIndex);
+		for (WordIndex = NumWords - 1; WordIndex >= 0 && Other.Bits[WordIndex] == Bits[WordIndex]; --WordIndex);
 		return WordIndex >= 0 && Bits[WordIndex] > Other.Bits[WordIndex];
 	}
 
@@ -606,7 +631,7 @@ public:
 			}
 		}
 		int32 WordIndex;
-		for (WordIndex = NumBits / BitsPerWord - 1; WordIndex >= 0 && Other.Bits[WordIndex] == Bits[WordIndex]; --WordIndex);
+		for (WordIndex = NumWords - 1; WordIndex >= 0 && Other.Bits[WordIndex] == Bits[WordIndex]; --WordIndex);
 		return WordIndex < 0 || Bits[WordIndex] > Other.Bits[WordIndex];
 	}
 
@@ -616,7 +641,7 @@ public:
 	FORCEINLINE bool IsZero() const
 	{
 		int32 WordIndex;
-		for (WordIndex = NumBits / BitsPerWord - 1; WordIndex >= 0 && !Bits[WordIndex]; --WordIndex);
+		for (WordIndex = NumWords - 1; WordIndex >= 0 && !Bits[WordIndex]; --WordIndex);
 		return WordIndex < 0;
 	}
 
@@ -838,14 +863,19 @@ public:
 	 */
 	FString ToString() const
 	{
-		FString Text(TEXT("0x"));		
-		for (int32 WordIndex = NumBits / BitsPerWord - 1; WordIndex >= 0; --WordIndex)
+		FString Text(TEXT("0x"));
+		int32 WordIndex;
+		for (WordIndex = NumWords - 1; WordIndex > 0; --WordIndex)
 		{
 			if (Bits[WordIndex])
 			{
-				Text += FString::Printf(TEXT("%08x"), Bits[WordIndex]);
+				break;
 			}
 		}
+		for (; WordIndex >= 0; --WordIndex)
+		{
+			Text += FString::Printf(TEXT("%08x"), Bits[WordIndex]);
+		}
 		return Text;
 	}
 
@@ -878,7 +908,7 @@ public:
 	 */
 	friend FArchive& operator<<(FArchive& Ar, TBigInt<NumBits>& Value)
 	{
-		for (int32 Index = 0; Index < NumBits / BitsPerWord; ++Index)
+		for (int32 Index = 0; Index < NumWords; ++Index)
 		{
 			Ar << Value.Bits[Index];
 		}
