Fix how project files get generated for foreign vs non-foreign projects. Previous logic detected projects under the root directory as being native and everything else as foreign, which broke GitHub builds that do not have the appropriate .uprojectdirs file. Now checks against the .uprojectdirs list to determine whether to use the UE4 solution or a game-specific solution.

diff --git a/Engine/Source/Programs/UnrealBuildTool/System/VCProject.cs b/Engine/Source/Programs/UnrealBuildTool/System/VCProject.cs
--- a/Engine/Source/Programs/UnrealBuildTool/System/VCProject.cs
+++ b/Engine/Source/Programs/UnrealBuildTool/System/VCProject.cs
@@ -896,13 +896,6 @@ namespace UnrealBuildTool
 
 			string UProjectPath = "";
 			bool bIsProjectTarget = UnrealBuildTool.HasUProjectFile() && Utils.IsFileUnderDirectory(TargetFilePath, UnrealBuildTool.GetUProjectPath());
-
-			// @todo Rocket: HACK: Only use long project names on the UBT command-line for out-of-root projects for now.  We need to revisit all uses of HasUProjectFile and short names in general to fix this.
-			if (bIsProjectTarget && !UnrealBuildTool.RunningRocket() && Utils.IsFileUnderDirectory(UnrealBuildTool.GetUProjectFile(), ProjectFileGenerator.RootRelativePath))
-			{
-				bIsProjectTarget = false;
-			}
-
 			if (bIsProjectTarget)
 			{
 				UProjectPath = "\"$(SolutionDir)$(SolutionName).uproject\"";
diff --git a/Engine/Source/Runtime/Core/Private/Misc/EngineBuildSettings.cpp b/Engine/Source/Runtime/Core/Private/Misc/EngineBuildSettings.cpp
--- a/Engine/Source/Runtime/Core/Private/Misc/EngineBuildSettings.cpp
+++ b/Engine/Source/Runtime/Core/Private/Misc/EngineBuildSettings.cpp
@@ -9,4 +9,15 @@ bool FEngineBuildSettings::IsInternalBuild()
 bool FEngineBuildSettings::IsPerforceBuild()
 {
 	return FPaths::FileExists( FPaths::EngineDir() / TEXT("Build/PerforceBuild.txt") );
-}
\ No newline at end of file
+}
+
+bool FEngineBuildSettings::IsSourceDistribution()
+{
+	return IsSourceDistribution( FPaths::RootDir() );
+}
+
+bool FEngineBuildSettings::IsSourceDistribution(const FString& RootDir)
+{
+	return FPaths::FileExists( RootDir / TEXT("Engine/Build/SourceDistribution.txt") );
+}
+
diff --git a/Engine/Source/Runtime/Core/Private/Modules/ModuleManager.cpp b/Engine/Source/Runtime/Core/Private/Modules/ModuleManager.cpp
--- a/Engine/Source/Runtime/Core/Private/Modules/ModuleManager.cpp
+++ b/Engine/Source/Runtime/Core/Private/Modules/ModuleManager.cpp
@@ -3,6 +3,7 @@
 #include "CorePrivate.h"
 #include "ModuleManager.h"
 #include "ModuleVersion.h"
+#include "EngineBuildSettings.h"
 
 DEFINE_LOG_CATEGORY_STATIC(LogModuleManager, Log, All);
 
@@ -1327,43 +1328,26 @@ bool FModuleManager::GenerateCodeProjectFiles( const FString& ProjectFilename, F
 #endif
 	if ( !ProjectFilename.IsEmpty() )
 	{
-		bool bIsInRootFolder = false;
-		if ( !FRocketSupport::IsRocket() )
+		// Normalize the project filename
+		FString NormalizedProjectFilename = ProjectFilename;
+		FPaths::NormalizeFilename(NormalizedProjectFilename);
+
+		// Figure out if this is a foreign project, and specify the project filename on the command line if it is. Don't use the cached project dictionary; the project may have been newly added.
+		FUProjectDictionary Dictionary(FPaths::RootDir());
+		if(Dictionary.IsForeignProject(NormalizedProjectFilename))
 		{
-			// This is a non-rocket GenerateCodeProjectFiles call. If we are in the UE4 root, just generate all project files.
-			FString AbsoluteProjectParentFolder = IFileManager::Get().ConvertToAbsolutePathForExternalAppForRead(*FPaths::GetPath(FPaths::GetPath(ProjectFilename)));
-			FString AbsoluteRootPath = IFileManager::Get().ConvertToAbsolutePathForExternalAppForRead(*FPaths::RootDir());
-			
-			if ( AbsoluteProjectParentFolder.Right(1) != TEXT("/") )
-			{
-				AbsoluteProjectParentFolder += TEXT("/");
-			}
+			CmdLineParams += FString::Printf(TEXT(" \"%s\" -game"), *IFileManager::Get().ConvertToAbsolutePathForExternalAppForRead(*NormalizedProjectFilename));
 
-			if ( AbsoluteRootPath.Right(1) != TEXT("/") )
+			// Determine whether we want engine source in the generated project files
+			if(FEngineBuildSettings::IsSourceDistribution())
 			{
-				AbsoluteRootPath += TEXT("/");
+				CmdLineParams += TEXT(" -engine");
 			}
-
-			bIsInRootFolder = (AbsoluteProjectParentFolder == AbsoluteRootPath);
 		}
 
-		// @todo Rocket: New projects created under the root don't need a full path to the .uproject name (but eventually, they probably should always have this for consistency.)
-		const bool bNeedsUProjectFilePath = FRocketSupport::IsRocket() || !bIsInRootFolder;
-		if( bNeedsUProjectFilePath )
+		// Add the Rocket parameter
+		if ( FRocketSupport::IsRocket() )
 		{
-			CmdLineParams += FString::Printf(TEXT(" \"%s\""), *IFileManager::Get().ConvertToAbsolutePathForExternalAppForRead(*ProjectFilename));
-		}
-
-		if ( !FRocketSupport::IsRocket() )
-		{
-			// This is a non-rocket GenerateCodeProjectFiles call. If we are in the UE4 root, just generate all project files.
-			if ( !bIsInRootFolder )
-			{
-				CmdLineParams += TEXT(" -game -engine");
-			}
-		}
-        else
-        {
             CmdLineParams += TEXT(" -rocket");
         }
 	}
diff --git a/Engine/Source/Runtime/Core/Public/Misc/EngineBuildSettings.h b/Engine/Source/Runtime/Core/Public/Misc/EngineBuildSettings.h
--- a/Engine/Source/Runtime/Core/Public/Misc/EngineBuildSettings.h
+++ b/Engine/Source/Runtime/Core/Public/Misc/EngineBuildSettings.h
@@ -14,4 +14,14 @@ public:
 	 * @return True if the build is for internal projects only
 	 */
 	static bool IsInternalBuild();
+
+	/**
+	 * @return True if the current installation is a source distribution.
+	 */
+	static bool IsSourceDistribution();
+
+	/**
+	 * @return True if a given engine distribution contains source (as opposed to, say, Rocket builds)
+	 */
+	static bool IsSourceDistribution(const FString& RootDir);
 };
\ No newline at end of file
