Elimimate heap allocations when calling GetModularFeatureImplementation() each time

diff --git a/Engine/Source/Runtime/Core/Private/Features/ModularFeatures.cpp b/Engine/Source/Runtime/Core/Private/Features/ModularFeatures.cpp
--- a/Engine/Source/Runtime/Core/Private/Features/ModularFeatures.cpp
+++ b/Engine/Source/Runtime/Core/Private/Features/ModularFeatures.cpp
@@ -20,7 +20,8 @@ int32 FModularFeatures::GetModularFeatureImplementationCount( const FName Type )
 
 IModularFeature* FModularFeatures::GetModularFeatureImplementation( const FName Type, const int32 Index )
 {
-	TArray< IModularFeature* > FeatureImplementations;
+	static TArray< IModularFeature* > FeatureImplementations;	// Static, so that we don't have to reallocate entries every time that MultiFind is called
+	FeatureImplementations.Reset();
 	ModularFeaturesMap.MultiFind( Type, FeatureImplementations );
 	return FeatureImplementations[ Index ];
 }
