Stats - Fixed FStatsUtils::DebugPrint not printing the group name, a stat desc is not logged if it's the same as the stat name

diff --git a/Engine/Source/Runtime/Core/Private/Stats/StatsData.cpp b/Engine/Source/Runtime/Core/Private/Stats/StatsData.cpp
--- a/Engine/Source/Runtime/Core/Private/Stats/StatsData.cpp
+++ b/Engine/Source/Runtime/Core/Private/Stats/StatsData.cpp
@@ -1347,36 +1347,31 @@ FString FStatsUtils::DebugPrint(FStatMessage const& Item)
 
 	Result = FString(FCString::Spc(FMath::Max<int32>(0, 14 - Result.Len()))) + Result;
 
-	FString Desc;
-	FName Group;
-	FName Category;
-	FName ShortName;
-
-	ShortName = Item.NameAndInfo.GetShortName();
-	Group = Item.NameAndInfo.GetGroupName();
-	Category = Item.NameAndInfo.GetGroupCategory();
-	Desc = Item.NameAndInfo.GetDescription();
+	const FString ShortName = Item.NameAndInfo.GetShortName().ToString();
+	const FName Group = Item.NameAndInfo.GetGroupName();
+	const FName Category = Item.NameAndInfo.GetGroupCategory();
+	FString Desc = Item.NameAndInfo.GetDescription();
 	Desc.Trim();
 
-	if (Desc.Len())
+	if( Desc.Len() && Desc != ShortName )
 	{
 		Desc += TEXT(" - ");
+		Desc += ShortName;
 	}
-	Desc += *ShortName.ToString();
 
-	FString GroupStr;
+	FString GroupAndCategoryStr;
 	if (Group != NAME_None)
 	{
-		GroupStr = TEXT(" - ");
-		GroupStr += *Group.ToString();
+		GroupAndCategoryStr = TEXT(" - ");
+		GroupAndCategoryStr += *Group.ToString();
 	}
 	if (Category != NAME_None)
 	{
-		GroupStr = TEXT(" - ");
-		GroupStr += *Category.ToString();
+		GroupAndCategoryStr += TEXT(" - ");
+		GroupAndCategoryStr += *Category.ToString();
 	}
 
-	return FString::Printf(TEXT("  %s  -  %s%s"), *Result, *Desc, *GroupStr);
+	return FString::Printf(TEXT("  %s  -  %s%s"), *Result, *Desc, *GroupAndCategoryStr);
 }
 
 void FStatsUtils::AddMergeStatArray(TArray<FStatMessage>& Dest, TArray<FStatMessage> const& Src)
