Stats - Fixed not visible FCustomNameGraphTaskBase tasks in the 'stat taskgraphtasks'

diff --git a/Engine/Source/Runtime/Core/Public/Async/TaskGraphInterfaces.h b/Engine/Source/Runtime/Core/Public/Async/TaskGraphInterfaces.h
--- a/Engine/Source/Runtime/Core/Public/Async/TaskGraphInterfaces.h
+++ b/Engine/Source/Runtime/Core/Public/Async/TaskGraphInterfaces.h
@@ -820,7 +820,7 @@ public:
 		const FString LongName = FString(TaskClassName) + TEXT(".") + FString(TaskName);
 		const FName StatName = FName( *LongName );
 
-		FStartupMessages::Get().AddMetadata(StatName, TEXT(""),
+		FStartupMessages::Get().AddMetadata(StatName, *LongName,
 			STAT_GROUP_TO_FStatGroup(STATGROUP_TaskGraphTasks)::GetGroupName(),
 			STAT_GROUP_TO_FStatGroup(STATGROUP_TaskGraphTasks)::GetGroupCategory(),
 			STAT_GROUP_TO_FStatGroup(STATGROUP_TaskGraphTasks)::GetDescription(),
