Stats - Task that enters the stall state is exluded from the thread idle stats, because the event wait is not really idle time in this case, asked in https://udn.unrealengine.com/questions/209458/stat-unit-isnt-accurate.html

diff --git a/Engine/Source/Runtime/Core/Private/Async/TaskGraph.cpp b/Engine/Source/Runtime/Core/Private/Async/TaskGraph.cpp
--- a/Engine/Source/Runtime/Core/Private/Async/TaskGraph.cpp
+++ b/Engine/Source/Runtime/Core/Private/Async/TaskGraph.cpp
@@ -193,20 +193,29 @@ public:
 	 **/
 	void ProcessTasks(int32 QueueIndex, bool bAllowStall)
 	{
+		TStatId StallStatId;
+		bool bCountAsStall = false;
 #if STATS
 		TStatId StatName;
 		FCycleCounter ProcessingTasks;
 		if (ThreadId == ENamedThreads::GameThread)
 		{
 			StatName = GET_STATID(STAT_TaskGraph_GameTasks);
+			StallStatId = GET_STATID(STAT_TaskGraph_GameStalls);
+			bCountAsStall = true;
 		}
 		else if (ThreadId == ENamedThreads::RenderThread)
 		{
 			//StatName = none, we need to let the scope empty so that the render thread submits tasks in a timely manner
+			// The render thread has a custom method to calculate the idle time, so leave it as it is for now.
+			//StallStatId = GET_STATID(STAT_TaskGraph_RenderStalls);
+			//bCountAsStall = true;
 		}
 		else if (ThreadId != ENamedThreads::StatsThread)
 		{
 			StatName = GET_STATID(STAT_TaskGraph_OtherTasks);
+			StallStatId = GET_STATID(STAT_TaskGraph_OtherStalls);
+			bCountAsStall = true;
 		}
 		bool bTasksOpen = false;
 #endif
@@ -244,7 +253,7 @@ public:
 								bTasksOpen = false;
 							}
 #endif
-							if (Stall(QueueIndex))
+							if (Stall(QueueIndex, StallStatId, bCountAsStall))
 							{
 								Queue(QueueIndex).IncomingQueue.PopAll(NewTasks);
 							}
@@ -294,7 +303,7 @@ public:
 								bTasksOpen = false;
 							}
 #endif
-							if (Stall(QueueIndex))
+							if (Stall(QueueIndex, StallStatId, bCountAsStall))
 							{
 								Task = Queue(QueueIndex).IncomingQueue.PopIfNotClosed();
 							}
@@ -512,10 +521,12 @@ private:
 
 	/**
 	 *	Internal function to block on the stall wait event.
-	 *  @param QueueIndex, Queue to stall
+	 *  @param QueueIndex		- Queue to stall
+	 *  @param StallStatId		- Stall stat id
+	 *  @param bCountAsStall	- true if StallStatId is a valid stat id
 	 *	@return true if the thread actually stalled; false in the case of a stop request or a task arrived while we were trying to stall.
 	 */
-	bool Stall(int32 QueueIndex)
+	bool Stall(int32 QueueIndex, TStatId StallStatId, bool bCountAsStall)
 	{
 		checkThreadGraph(Queue(QueueIndex).StallRestartEvent); // make sure we are started up
 		if (Queue(QueueIndex).QuitWhenIdle.GetValue() == 0)
@@ -527,10 +538,12 @@ private:
 				FPlatformMisc::MemoryBarrier();
 				if (Queue(QueueIndex).IncomingQueue.CloseIfEmpty())
 				{
+					FScopeCycleCounter Scope( StallStatId );
+
 					int32 NewValue = IsStalled.Increment();
 					NotifyStalling();
 					checkThreadGraph(NewValue == 1); // there should be no concurrent calls to Stall!
-					Queue(QueueIndex).StallRestartEvent->Wait();
+					Queue(QueueIndex).StallRestartEvent->Wait(MAX_uint32, bCountAsStall);
 					NewValue = IsStalled.Decrement();
 					checkThreadGraph(NewValue == 0); // there should be no concurrent calls to Stall!
 					return true;
diff --git a/Engine/Source/Runtime/Core/Private/GenericPlatform/GenericPlatformProcess.cpp b/Engine/Source/Runtime/Core/Private/GenericPlatform/GenericPlatformProcess.cpp
--- a/Engine/Source/Runtime/Core/Private/GenericPlatform/GenericPlatformProcess.cpp
+++ b/Engine/Source/Runtime/Core/Private/GenericPlatform/GenericPlatformProcess.cpp
@@ -260,10 +260,10 @@ void FGenericPlatformProcess::SleepInfinite()
 
 DECLARE_CYCLE_STAT(TEXT("CPU Stall - Wait For Event"),STAT_EventWait,STATGROUP_CPUStalls);
 
-bool FPThreadEvent::Wait(uint32 WaitTime)
+bool FPThreadEvent::Wait(uint32 WaitTime, const bool bIgnoreThreadIdleStats /*= false*/)
 {
 	SCOPE_CYCLE_COUNTER(STAT_EventWait);
-	FThreadIdleStats::FScopeIdle Scope;
+	FThreadIdleStats::FScopeIdle Scope(bIgnoreThreadIdleStats);
 
 	check(bInitialized);
 
diff --git a/Engine/Source/Runtime/Core/Private/Misc/Core.cpp b/Engine/Source/Runtime/Core/Private/Misc/Core.cpp
--- a/Engine/Source/Runtime/Core/Private/Misc/Core.cpp
+++ b/Engine/Source/Runtime/Core/Private/Misc/Core.cpp
@@ -272,10 +272,13 @@ DEFINE_STAT(STAT_GameTickWantedWaitTime);
 DEFINE_STAT(STAT_GameTickAdditionalWaitTime);
 
 DEFINE_STAT(STAT_TaskGraph_OtherTasks);
-DEFINE_STAT(STAT_TaskGraph_RenderIdles);
+DEFINE_STAT(STAT_TaskGraph_OtherStalls);
+
+DEFINE_STAT(STAT_TaskGraph_RenderStalls);
 
 DEFINE_STAT(STAT_TaskGraph_GameTasks);
-DEFINE_STAT(STAT_TaskGraph_GameIdles);
+DEFINE_STAT(STAT_TaskGraph_GameStalls);
+
 DEFINE_STAT(STAT_FlushThreadedLogs);
 DEFINE_STAT(STAT_PumpMessages);
 
diff --git a/Engine/Source/Runtime/Core/Private/WinRT/WinRTEvent.h b/Engine/Source/Runtime/Core/Private/WinRT/WinRTEvent.h
--- a/Engine/Source/Runtime/Core/Private/WinRT/WinRTEvent.h
+++ b/Engine/Source/Runtime/Core/Private/WinRT/WinRTEvent.h
@@ -60,7 +60,7 @@ public:
 		ResetEvent(Event);
 	}
 
-	virtual bool Wait (uint32 WaitTime) override;
+	virtual bool Wait(uint32 WaitTime, const bool bIgnoreThreadIdleStats = false) override;
 
 
 private:
diff --git a/Engine/Source/Runtime/Core/Private/WinRT/WinRTProcess.cpp b/Engine/Source/Runtime/Core/Private/WinRT/WinRTProcess.cpp
--- a/Engine/Source/Runtime/Core/Private/WinRT/WinRTProcess.cpp
+++ b/Engine/Source/Runtime/Core/Private/WinRT/WinRTProcess.cpp
@@ -182,10 +182,10 @@ DECLARE_CYCLE_STAT(TEXT("CPU Stall - Wait For Event"),STAT_EventWait,STATGROUP_C
 
 #include "AllowWinRTPlatformTypes.h"
 
-bool FEventWinRT::Wait (uint32 WaitTime)
+bool FEventWinRT::Wait(uint32 WaitTime, const bool bIgnoreThreadIdleStats /*= false*/)
 {
 	SCOPE_CYCLE_COUNTER(STAT_EventWait);
-	FThreadIdleStats::FScopeIdle Scope;
+	FThreadIdleStats::FScopeIdle Scope(bIgnoreThreadIdleStats);
 	check(Event);
 
 	//		return (WaitForSingleObject(Event, WaitTime) == WAIT_OBJECT_0);
diff --git a/Engine/Source/Runtime/Core/Private/Windows/WindowsEvent.h b/Engine/Source/Runtime/Core/Private/Windows/WindowsEvent.h
--- a/Engine/Source/Runtime/Core/Private/Windows/WindowsEvent.h
+++ b/Engine/Source/Runtime/Core/Private/Windows/WindowsEvent.h
@@ -60,7 +60,7 @@ public:
 		ResetEvent(Event);
 	}
 
-	virtual bool Wait (uint32 WaitTime) override;
+	virtual bool Wait (uint32 WaitTime, const bool bIgnoreThreadIdleStats = false) override;
 
 private:
 
diff --git a/Engine/Source/Runtime/Core/Private/Windows/WindowsPlatformProcess.cpp b/Engine/Source/Runtime/Core/Private/Windows/WindowsPlatformProcess.cpp
--- a/Engine/Source/Runtime/Core/Private/Windows/WindowsPlatformProcess.cpp
+++ b/Engine/Source/Runtime/Core/Private/Windows/WindowsPlatformProcess.cpp
@@ -908,12 +908,12 @@ FEvent* FWindowsPlatformProcess::CreateSynchEvent(bool bIsManualReset)
 
 DECLARE_CYCLE_STAT(TEXT("CPU Stall - Wait For Event"),STAT_EventWait,STATGROUP_CPUStalls);
 
-bool FEventWin::Wait (uint32 WaitTime)
+bool FEventWin::Wait(uint32 WaitTime, const bool bIgnoreThreadIdleStats /*= false*/)
 {
 	SCOPE_CYCLE_COUNTER(STAT_EventWait);
-	FThreadIdleStats::FScopeIdle Scope;
 	check(Event);
 
+	FThreadIdleStats::FScopeIdle Scope(bIgnoreThreadIdleStats);
 	return (WaitForSingleObject(Event, WaitTime) == WAIT_OBJECT_0);
 }
 
diff --git a/Engine/Source/Runtime/Core/Public/HAL/PThreadEvent.h b/Engine/Source/Runtime/Core/Public/HAL/PThreadEvent.h
--- a/Engine/Source/Runtime/Core/Public/HAL/PThreadEvent.h
+++ b/Engine/Source/Runtime/Core/Public/HAL/PThreadEvent.h
@@ -146,7 +146,7 @@ public:
 		UnlockEventMutex();
 	}
 
-	virtual bool Wait(uint32 WaitTime = (uint32)-1) override;
+	virtual bool Wait(uint32 WaitTime = (uint32)-1, const bool bIgnoreThreadIdleStats = false) override;
 
 };
 
diff --git a/Engine/Source/Runtime/Core/Public/HAL/ThreadingBase.h b/Engine/Source/Runtime/Core/Public/HAL/ThreadingBase.h
--- a/Engine/Source/Runtime/Core/Public/HAL/ThreadingBase.h
+++ b/Engine/Source/Runtime/Core/Public/HAL/ThreadingBase.h
@@ -55,14 +55,12 @@ public:
 	 *
 	 * A wait time of MAX_uint32 is treated as infinite wait.
 	 *
-	 * @param WaitTime - The time to wait (in milliseconds).
+	 * @param WaitTime					- The time to wait (in milliseconds).
+	 * @param bIgnoreThreadIdleStats	- If true, ignores ThreadIdleStats
 	 *
 	 * @return true if the event was triggered, false if the wait timed out.
 	 */
-	virtual bool Wait (uint32 WaitTime) = 0;
-
-
-public:
+	virtual bool Wait(uint32 WaitTime, const bool bIgnoreThreadIdleStats = false) = 0;
 
 	/**
 	 * Waits an infinite amount of time for the event to be triggered.
@@ -77,13 +75,14 @@ public:
 	/**
 	 * Waits the specified amount of time for the event to be triggered.
 	 *
-	 * @param WaitTime - The time to wait.
+	 * @param WaitTime					- The time to wait.
+	 * @param bIgnoreThreadIdleStats	- If true, ignores ThreadIdleStats
 	 *
 	 * @return true if the event was triggered, false if the wait timed out.
 	 */
-	bool Wait (const FTimespan& WaitTime)
+	bool Wait (const FTimespan& WaitTime, const bool bIgnoreThreadIdleStats = false)
 	{
-		return Wait(WaitTime.GetTotalMilliseconds());
+		return Wait(WaitTime.GetTotalMilliseconds(), bIgnoreThreadIdleStats);
 	}
 
 
@@ -461,7 +460,7 @@ public:
 	}
 	virtual void Trigger() { bTriggered = true; }
 	virtual void Reset() { bTriggered = false; }
-	virtual bool Wait(uint32 WaitTime) 
+	virtual bool Wait(uint32 WaitTime, const bool bIgnoreThreadIdleStats = false) 
 	{ 
 		// With only one thread it's assumed the event has been triggered
 		// before Wait is called, otherwise it would end up waiting forever or always fail.
diff --git a/Engine/Source/Runtime/Core/Public/Misc/CoreStats.h b/Engine/Source/Runtime/Core/Public/Misc/CoreStats.h
--- a/Engine/Source/Runtime/Core/Public/Misc/CoreStats.h
+++ b/Engine/Source/Runtime/Core/Public/Misc/CoreStats.h
@@ -58,9 +58,13 @@ DECLARE_CYCLE_STAT_EXTERN(TEXT("Game thread tick wait time"),STAT_GameTickWaitTi
 DECLARE_FLOAT_COUNTER_STAT_EXTERN(TEXT("Game thread requested wait time"),STAT_GameTickWantedWaitTime,STATGROUP_Threading, CORE_API);
 DECLARE_FLOAT_COUNTER_STAT_EXTERN(TEXT("Game thread additional wait time"),STAT_GameTickAdditionalWaitTime,STATGROUP_Threading, CORE_API);
 DECLARE_CYCLE_STAT_EXTERN(TEXT("Game TaskGraph Tasks"),STAT_TaskGraph_GameTasks,STATGROUP_Threading, CORE_API);
-DECLARE_CYCLE_STAT_EXTERN(TEXT("Game TaskGraph Idles"),STAT_TaskGraph_GameIdles,STATGROUP_Threading, CORE_API);
-DECLARE_CYCLE_STAT_EXTERN(TEXT("Render TaskGraph Idles"),STAT_TaskGraph_RenderIdles,STATGROUP_Threading, CORE_API);
+DECLARE_CYCLE_STAT_EXTERN(TEXT("Game TaskGraph Stalls"),STAT_TaskGraph_GameStalls,STATGROUP_Threading, CORE_API);
+
+DECLARE_CYCLE_STAT_EXTERN(TEXT("Render TaskGraph Stalls"),STAT_TaskGraph_RenderStalls,STATGROUP_Threading, CORE_API);
+
 DECLARE_CYCLE_STAT_EXTERN(TEXT("Other TaskGraph Tasks"),STAT_TaskGraph_OtherTasks,STATGROUP_Threading, CORE_API);
+DECLARE_CYCLE_STAT_EXTERN(TEXT("Other TaskGraph Stalls"),STAT_TaskGraph_OtherStalls,STATGROUP_Threading, CORE_API);
+
 DECLARE_CYCLE_STAT_EXTERN(TEXT("Flush Threaded Logs"),STAT_FlushThreadedLogs,STATGROUP_Threading, CORE_API);
 DECLARE_CYCLE_STAT_EXTERN(TEXT("Pump Messages"),STAT_PumpMessages,STATGROUP_Threading, CORE_API);
 DECLARE_FLOAT_COUNTER_STAT_EXTERN(TEXT("Percentage CPU utilization"),STAT_CPUTimePct,STATGROUP_Threading, CORE_API);
diff --git a/Engine/Source/Runtime/Core/Public/Stats/Stats2.h b/Engine/Source/Runtime/Core/Public/Stats/Stats2.h
--- a/Engine/Source/Runtime/Core/Public/Stats/Stats2.h
+++ b/Engine/Source/Runtime/Core/Public/Stats/Stats2.h
@@ -23,14 +23,24 @@ public:
 
 	struct FScopeIdle
 	{
-		uint32 Start;
-		FScopeIdle()
+		/** Starting cycle counter. */
+		const uint32 Start;
+
+		/** If true, we ignore this thread idle stats. */
+		const bool bIgnore;
+
+		FScopeIdle( bool bInIgnore = false )
 			: Start(FPlatformTime::Cycles())
+			, bIgnore( bInIgnore )
 		{
 		}
+
 		~FScopeIdle()
 		{
-			FThreadIdleStats::Get().Waits += FPlatformTime::Cycles() - Start;
+			if( !bIgnore )
+			{
+				FThreadIdleStats::Get().Waits += FPlatformTime::Cycles() - Start;
+			}
 		}
 	};
 };
