Math - Adding some additional utility methods spelled out in the Vector2D class.  (Distance, Cross Product)

diff --git a/Engine/Source/Runtime/Core/Public/Math/Vector2D.h b/Engine/Source/Runtime/Core/Public/Math/Vector2D.h
--- a/Engine/Source/Runtime/Core/Public/Math/Vector2D.h
+++ b/Engine/Source/Runtime/Core/Public/Math/Vector2D.h
@@ -255,7 +255,7 @@ public:
 	 * @param Index the index of vector component
 	 * @return reference to component.
 	 */
-    float& operator[]( int32 Index );
+	float& operator[]( int32 Index );
 
 	/**
 	 * Gets specific component of the vector.
@@ -285,6 +285,24 @@ public:
 	 */
 	FORCEINLINE static float DistSquared( const FVector2D& V1, const FVector2D& V2 );
 
+	/**
+	 * Distance between two 2D points.
+	 *
+	 * @param V1 The first point.
+	 * @param V2 The second point.
+	 * @return The squared distance between two 2D points.
+	 */
+	FORCEINLINE static float Distance( const FVector2D& V1, const FVector2D& V2 );
+
+	/**
+	 * Calculate the Cross product of two vectors.
+	 *
+	 * @param A The first vector.
+	 * @param B The second vector.
+	 * @return The Cross product.
+	 */
+	FORCEINLINE static float CrossProduct(const FVector2D& A, const FVector2D& B);
+
 	/**
 	 * Checks for equality with error-tolerant comparison.
 	 *
@@ -555,6 +573,18 @@ FORCEINLINE float FVector2D::DistSquared( const FVector2D &V1, const FVector2D &
 }
 
 
+FORCEINLINE float FVector2D::Distance(const FVector2D& V1, const FVector2D& V2)
+{
+	return FMath::Sqrt(FVector2D::DistSquared(V1, V2));
+}
+
+
+FORCEINLINE float FVector2D::CrossProduct(const FVector2D& A, const FVector2D& B)
+{
+	return A ^ B;
+}
+
+
 FORCEINLINE bool FVector2D::operator==( const FVector2D& V ) const
 {
 	return X==V.X && Y==V.Y;
@@ -772,7 +802,7 @@ FORCEINLINE bool FVector2D::InitFromString( const FString & InSourceString )
 	X = Y = 0;
 
 	// The initialization is only successful if the X and Y values can all be parsed from the string
- 	const bool bSuccessful = FParse::Value( *InSourceString, TEXT("X=") , X ) && FParse::Value( *InSourceString, TEXT("Y="), Y ) ;
+	const bool bSuccessful = FParse::Value( *InSourceString, TEXT("X=") , X ) && FParse::Value( *InSourceString, TEXT("Y="), Y ) ;
 
 	return bSuccessful;
 }
