Upgrade Morpheus to support latest prototype and SDK.  Older prototypes are deprecated.
Incorporate fixes from CCP.
Add full contrller support through HMDServer for PC emulation.

diff --git a/Engine/Source/Runtime/Core/Private/Windows/WindowsApplication.cpp b/Engine/Source/Runtime/Core/Private/Windows/WindowsApplication.cpp
--- a/Engine/Source/Runtime/Core/Private/Windows/WindowsApplication.cpp
+++ b/Engine/Source/Runtime/Core/Private/Windows/WindowsApplication.cpp
@@ -1561,10 +1561,7 @@ void FWindowsApplication::PollGameDeviceState( const float TimeDelta )
 		for( auto InputPluginIt = PluginImplementations.CreateIterator(); InputPluginIt; ++InputPluginIt )
 		{
 			TSharedPtr<IInputDevice> Device = (*InputPluginIt)->CreateInputDevice(MessageHandler);
-			if ( Device.IsValid() )
-			{
-				ExternalInputDevices.Add(Device);
-			}
+			AddExternalInputDevice(Device);			
 		}
 
 		bHasLoadedInputPlugins = true;
@@ -1599,6 +1596,14 @@ void FWindowsApplication::SetChannelValues (int32 ControllerId, const FForceFeed
 	}
 }
 
+void FWindowsApplication::AddExternalInputDevice(TSharedPtr<IInputDevice> InputDevice)
+{
+	if (InputDevice.IsValid())
+	{
+		ExternalInputDevices.Add(InputDevice);
+	}
+}
+
 void FWindowsApplication::DeferDragDropOperation(const FDeferredWindowsDragDropOperation& DeferredDragDropOperation)
 {
 	DeferredDragDropOperations.Add(DeferredDragDropOperation);
diff --git a/Engine/Source/Runtime/Core/Public/GenericPlatform/GenericApplication.h b/Engine/Source/Runtime/Core/Public/GenericPlatform/GenericApplication.h
--- a/Engine/Source/Runtime/Core/Public/GenericPlatform/GenericApplication.h
+++ b/Engine/Source/Runtime/Core/Public/GenericPlatform/GenericApplication.h
@@ -249,6 +249,8 @@ public:
 
 	virtual void SetMessageHandler( const TSharedRef< FGenericApplicationMessageHandler >& InMessageHandler ) { MessageHandler = InMessageHandler; }
 
+	TSharedRef< FGenericApplicationMessageHandler > GetMessageHandler() { return MessageHandler; }
+
 	virtual void PollGameDeviceState( const float TimeDelta ) { }
 
 	virtual void PumpMessages( const float TimeDelta ) { }
diff --git a/Engine/Source/Runtime/Core/Public/Windows/WindowsApplication.h b/Engine/Source/Runtime/Core/Public/Windows/WindowsApplication.h
--- a/Engine/Source/Runtime/Core/Public/Windows/WindowsApplication.h
+++ b/Engine/Source/Runtime/Core/Public/Windows/WindowsApplication.h
@@ -207,6 +207,8 @@ public:
 		return TextInputMethodSystem.Get();
 	}
 
+	virtual void AddExternalInputDevice(TSharedPtr<class IInputDevice> InputDevice);
+
 protected:
 
 	static LRESULT CALLBACK AppWndProc(HWND hwnd, uint32 msg, WPARAM wParam, LPARAM lParam);
