Converted ConsoleManager code over to use the new typed TAutoConsoleVariable. Implemented support for auto typed string variables.
Moved console variables to their corresponding cpp files where possible.

#ttp 319909 - UE4: RENDERING: Change the ConsoleManager code over to use the new typed TAutoConsoleVariable

diff --git a/Engine/Source/Developer/MeshUtilities/Private/MeshUtilities.cpp b/Engine/Source/Developer/MeshUtilities/Private/MeshUtilities.cpp
--- a/Engine/Source/Developer/MeshUtilities/Private/MeshUtilities.cpp
+++ b/Engine/Source/Developer/MeshUtilities/Private/MeshUtilities.cpp
@@ -32,6 +32,22 @@
 
 DEFINE_LOG_CATEGORY_STATIC(LogMeshUtilities,Verbose,All);
 
+// CVars
+static TAutoConsoleVariable<int32> CVarTriangleOrderOptimization(
+	TEXT("r.TriangleOrderOptimization"),
+	1,
+	TEXT("Controls the algorithm to use when optimizing the triangle order for the post-transform cache.\n")
+	TEXT("0: Use NVTriStrip (slower)\n")
+	TEXT("1: Use Forsyth algorithm (fastest)(default)"),
+	ECVF_ReadOnly);
+
+static TAutoConsoleVariable<int32> CVarAllowMeshDistanceFieldRepresentations (
+	TEXT("r.AllowMeshDistanceFieldRepresentations"),
+	0,
+	TEXT("Whether to build distance fields of static meshes, needed for distance field AO, which is used to implement Movable SkyLight shadows.\n")
+	TEXT("Enabling will increase mesh build times and memory usage.  Changing this value will cause a rebuild of all static meshes."),
+	ECVF_ReadOnly);
+
 //@todo - implement required vector intrinsics for other implementations
 #if PLATFORM_ENABLE_VECTORINTRINSICS
 
@@ -250,11 +266,8 @@ static void GenerateSignedDistanceFieldVolumeData(
 	const FBoxSphereBounds& Bounds,
 	float DistanceFieldResolutionScale)
 {
-	static const auto CVar = IConsoleManager::Get().FindTConsoleVariableDataInt(TEXT("r.AllowMeshDistanceFieldRepresentations"));
-
 	if (DistanceFieldResolutionScale > 0 
-		&& CVar 
-		&& CVar->GetValueOnGameThread() != 0)
+		&& CVarAllowMeshDistanceFieldRepresentations.GetValueOnGameThread() != 0)
 	{
 		const FPositionVertexBuffer& PositionVertexBuffer = LODModel.PositionVertexBuffer;
 		FIndexArrayView Indices = LODModel.IndexBuffer.GetArrayView();
@@ -432,11 +445,8 @@ static void GenerateSignedDistanceFieldVolumeData(
 	const FBoxSphereBounds& Bounds,
 	float DistanceFieldResolutionScale)
 {
-	static const auto CVar = IConsoleManager::Get().FindTConsoleVariableDataInt(TEXT("r.AllowMeshDistanceFieldRepresentations"));
-
 	if (DistanceFieldResolutionScale > 0 
-		&& CVar 
-		&& CVar->GetValueOnGameThread() != 0)
+		&& CVarAllowMeshDistanceFieldRepresentations.GetValueOnGameThread() != 0)
 	{
 		UE_LOG(LogMeshUtilities,Error,TEXT("Couldn't generate distance field for mesh, platform is missing required Vector intrinsics."));
 	}
@@ -3751,8 +3761,7 @@ void FMeshUtilities::StartupModule()
 		}
 	}
 
-	TConsoleVariableData<int32>* CVar = IConsoleManager::Get().FindTConsoleVariableDataInt(TEXT("r.TriangleOrderOptimization"));
-	bUsingNvTriStrip = (CVar->GetValueOnGameThread() == 0);
+	bUsingNvTriStrip = (CVarTriangleOrderOptimization.GetValueOnGameThread() == 0);
 
 	// Construct and cache the version string for the mesh utilities module.
 	VersionString = FString::Printf(
diff --git a/Engine/Source/Editor/UnrealEd/Private/MaterialGraphSchema.cpp b/Engine/Source/Editor/UnrealEd/Private/MaterialGraphSchema.cpp
--- a/Engine/Source/Editor/UnrealEd/Private/MaterialGraphSchema.cpp
+++ b/Engine/Source/Editor/UnrealEd/Private/MaterialGraphSchema.cpp
@@ -416,10 +416,18 @@ void UMaterialGraphSchema::GetContextMenuActions(const UEdGraph* CurrentGraph, c
 	Super::GetContextMenuActions(CurrentGraph, InGraphNode, InGraphPin, MenuBuilder, bIsDebugging);
 }
 
+static TAutoConsoleVariable<int32> CVarPreventInvalidMaterialConnections(
+	TEXT("r.PreventInvalidMaterialConnections"),
+	1,
+	TEXT("Controls whether users can make connections in the material editor if the system\n")
+	TEXT("determines that they may cause compile errors\n")
+	TEXT("0: Allow all connections\n")
+	TEXT("1: Prevent invalid connections"),
+	ECVF_Cheat);
+
 const FPinConnectionResponse UMaterialGraphSchema::CanCreateConnection(const UEdGraphPin* A, const UEdGraphPin* B) const
 {
-	static const auto PreventInvalidMaterialConnections = IConsoleManager::Get().FindTConsoleVariableDataInt(TEXT("r.PreventInvalidMaterialConnections"));
-	bool bPreventInvalidConnections = PreventInvalidMaterialConnections->GetValueOnGameThread() != 0;
+	bool bPreventInvalidConnections = CVarPreventInvalidMaterialConnections.GetValueOnGameThread() != 0;
 
 	// Make sure the pins are not on the same node
 	if (A->GetOwningNode() == B->GetOwningNode())
diff --git a/Engine/Source/Runtime/Core/Private/HAL/ConsoleManager.cpp b/Engine/Source/Runtime/Core/Private/HAL/ConsoleManager.cpp
--- a/Engine/Source/Runtime/Core/Private/HAL/ConsoleManager.cpp
+++ b/Engine/Source/Runtime/Core/Private/HAL/ConsoleManager.cpp
@@ -214,6 +214,7 @@ public:
 	virtual FString GetString() const;
 	virtual class TConsoleVariableData<int32>* AsVariableInt() { return 0; }
 	virtual class TConsoleVariableData<float>* AsVariableFloat() { return 0; }
+	virtual class TConsoleVariableData<FString>* AsVariableString() { return 0; }
 
 private: // ----------------------------------------------------
 
@@ -293,6 +294,11 @@ template<> FString FConsoleVariable<FString>::GetString() const
 	return Value();
 }
 
+template<> TConsoleVariableData<FString>* FConsoleVariable<FString>::AsVariableString()
+{
+	return &Data;
+}
+
 // ----
 
 // T: int32, float
@@ -624,7 +630,7 @@ IConsoleVariable* FConsoleManager::RegisterConsoleVariable(const TCHAR* Name, fl
 	return AddConsoleObject(Name, new FConsoleVariable<float>(DefaultValue, Help, (EConsoleVariableFlags)Flags))->AsVariable();
 }
 
-IConsoleVariable* FConsoleManager::RegisterConsoleVariable(const TCHAR* Name, const TCHAR *DefaultValue, const TCHAR* Help, uint32 Flags)
+IConsoleVariable* FConsoleManager::RegisterConsoleVariable(const TCHAR* Name, const FString& DefaultValue, const TCHAR* Help, uint32 Flags)
 {
 	// not supported
 	check((Flags & (uint32)ECVF_RenderThreadSafe) == 0);
@@ -1277,171 +1283,8 @@ void FConsoleManager::Test()
 // These don't belong here, but they belong here more than they belong in launch engine loop.
 void CreateConsoleVariables()
 {
-	// Naming conventions:
-	//
-	// Console variable should start with (suggestion):
-	//
-	// r.      Renderer / 3D Engine / graphical feature
-	// RHI.    Low level RHI (rendering platform) specific
-	// a.	   Animation
-	// s. 	   Sound / Music
-	// n.      Network
-	// ai.     Artificial intelligence
-	// i.      Input e.g. mouse/keyboard
-	// p.      Physics
-	// t.      Timer
-	// g.      Game
-	// Compat.
-	// FX.     Particle effects
-	// sg.     scalability group (used by scalability system, ini load/save or using SCALABILITY console command)
 
 #if !(UE_BUILD_SHIPPING || UE_BUILD_TEST)
-	IConsoleManager::Get().RegisterConsoleVariable(TEXT("r.PreViewTranslation"),
-		1,
-		TEXT("To limit issues with float world space positions we offset the world by the\n")
-		TEXT("PreViewTranslation vector. This command allows to disable updating this vector.\n")
-		TEXT(" 0: disable update\n")
-		TEXT(" 1: update the offset is each frame (default)"),
-		ECVF_Cheat);
-
-	IConsoleManager::Get().RegisterConsoleVariable(TEXT("r.FogDensity"),
-		-1.0f,
-		TEXT("Allows to override the FogDensity setting (needs ExponentialFog in the level).\n")
-		TEXT("Using a strong value allows to quickly see which pixel are affected by fog.\n")
-		TEXT("Using a start distance allows to cull pixels are can speed up rendering.\n")
-		TEXT(" <0: use default settings (default: -1)\n")
-		TEXT(">=0: override settings by the given value (0:off, 1=very dense fog)"),
-		ECVF_Cheat | ECVF_RenderThreadSafe);
-
-	IConsoleManager::Get().RegisterConsoleVariable(TEXT("r.FogStartDistance"),
-		-1.0f,
-		TEXT("Allows to override the FogStartDistance setting (needs ExponentialFog in the level).\n")
-		TEXT(" <0: use default settings (default: -1)\n")
-		TEXT(">=0: override settings by the given value (in world units)"),
-		ECVF_Cheat | ECVF_RenderThreadSafe);
-
-	IConsoleManager::Get().RegisterConsoleVariable(TEXT("r.RenderTimeFrozen"),
-		0,
-		TEXT("Allows to freeze time based effects in order to provide more deterministic render profiling.\n")
-		TEXT(" 0: off\n")
-		TEXT(" 1: on (Note: this also disables occlusion queries)"),
-		ECVF_Cheat);
-
-	IConsoleManager::Get().RegisterConsoleVariable(TEXT("r.VisualizeTexturePool"),
-		0,
-		TEXT("Allows to enable the visualize the texture pool (currently only on console).\n")
-		TEXT(" 0: off\n")
-		TEXT(" 1: on"),
-		ECVF_Cheat | ECVF_RenderThreadSafe);
-
-	IConsoleManager::Get().RegisterConsoleVariable(TEXT("r.D3DCompilerPath"),
-		TEXT(""),	// default
-		TEXT("Allows to specify a HLSL compiler version that is different from the one the code was compiled.\n")
-		TEXT("No path (\"\") means the default one is used.\n")
-		TEXT("If the compiler cannot be found an error is reported and it will compile further with the default one.\n")
-		TEXT("This console variable works with ShaderCompileWorker (with multi threading) and without multi threading.\n")
-		TEXT("This variable can be set in ConsoleVariables.ini to be defined at startup.\n")
-		TEXT("e.g. c:/temp/d3dcompiler_44.dll or \"\""),
-		ECVF_Cheat);
-	
-	IConsoleManager::Get().RegisterConsoleVariable(TEXT("r.RenderTargetPoolTest"),
-		0,
-		TEXT("Clears the texture returned by the rendertarget pool with a special color\n")
-		TEXT("so we can see better which passes would need to clear. Doesn't work on volume textures and non rendertargets yet.\n")
-		TEXT(" 0:off (default), 1:on"),
-		ECVF_Cheat | ECVF_RenderThreadSafe);
-
-	IConsoleManager::Get().RegisterConsoleVariable(TEXT("r.ReflectionEnvironment"),
-		1,	
-		TEXT("0:off, 1:on and blend with scene, 2:on and overwrite scene.\n")
-		TEXT("Whether to render the reflection environment feature, which implements local reflections through Reflection Capture actors."),
-		ECVF_Cheat | ECVF_RenderThreadSafe);
-
-	IConsoleManager::Get().RegisterConsoleVariable(TEXT("r.GPUBusyWait"),
-		0.0f,	
-		TEXT("<=0:off, >0: keep the GPU busy with n units of some fixed amount of work, independent on the resolution\n")
-		TEXT("This can be useful to make GPU timing experiments. The value should roughly represent milliseconds.\n")
-		TEXT("Clamped at 500."),
-		ECVF_Cheat | ECVF_RenderThreadSafe);
-
-	IConsoleManager::Get().RegisterConsoleVariable(TEXT("t.OverrideFPS"),
-		0.0f,
-		TEXT("This allows to override the frame time measurement with a fixed fps number (game can run faster or slower).\n")
-		TEXT("<=0:off, in frames per second, e.g. 60"),
-		ECVF_Cheat);
-
-	IConsoleManager::Get().RegisterConsoleVariable(TEXT("r.LimitRenderingFeatures"),
-		0,
-		TEXT("Allows to quickly reduce render feature to increase render performance.\n")
-		TEXT("This is just a quick way to alter multiple show flags and console variables in the game\n")
-		TEXT("Disabled more feature the higher the number\n")
-		TEXT(" <=0:off, order is defined in code (can be documented here when we settled on an order)"),
-		ECVF_Cheat | ECVF_RenderThreadSafe);
-
-	IConsoleManager::Get().RegisterConsoleVariableRef(TEXT("r.DumpingMovie"),
-		GIsDumpingMovie,
-		TEXT("Allows to dump each rendered frame to disk (slow fps, names MovieFrame..).\n")
-		TEXT("<=0:off (default), <0:remains on, >0:remains on for n frames (n is the number specified)"),
-		ECVF_Cheat);
-
-	// ---------------------------------------
-	// physics related
-	// ---------------------------------------
-
-	IConsoleManager::Get().RegisterConsoleVariable(TEXT("p.ShowInitialOverlaps"),
-		0,
-		TEXT("Show initial overlaps when moving a component, including estimated 'exit' direction.\n")
-		TEXT(" 0:off, otherwise on"),
-		ECVF_Cheat);
-
-	IConsoleManager::Get().RegisterConsoleVariable(TEXT("p.NetShowCorrections"),
-		0,
-		TEXT("Whether to draw client position corrections (red is incorrect, green is corrected).\n")
-		TEXT("0: Disable, 1: Enable"),
-		ECVF_Cheat);
-
-	IConsoleManager::Get().RegisterConsoleVariable(TEXT("p.NetCorrectionLifetime"),
-		4.f,
-		TEXT("How long a visualized network correction persists.\n")
-		TEXT("Time in seconds each visualized network correction persists."),
-		ECVF_Cheat);
-
-	// ---------------------------------------
-
-	IConsoleManager::Get().RegisterConsoleVariable(TEXT("g.DebugCameraTraceComplex"),
-		1,
-		TEXT("Whether DebugCamera should use complex or simple collision for the line trace.\n")
-		TEXT("1: complex collision, 0: simple collision "),
-		ECVF_Cheat);
-
-	// ---------------------------------------
-
-	IConsoleManager::Get().RegisterConsoleVariable(TEXT("r.ExposureOffset"),
-		0.0f,
-		TEXT("For adjusting the exposure on top of post process settings and eye adaptation. For developers only. 0:default"),
-		ECVF_Cheat);
-
-	IConsoleManager::Get().RegisterConsoleVariable(TEXT("r.Shadow.FreezeCamera"),
-		0,
-		TEXT("Debug the shadow methods by allowing to observe the system from outside.\n")
-		TEXT("0: default\n")
-		TEXT("1: freeze camera at current location"),
-		ECVF_Cheat);
-
-	IConsoleManager::Get().RegisterConsoleVariable(TEXT("r.ForceLOD"),
-		-1,
-		TEXT("LOD level to force, -1 is off."),
-		ECVF_Cheat | ECVF_RenderThreadSafe
-		);
-
-	IConsoleManager::Get().RegisterConsoleVariable(TEXT("r.MotionBlurFiltering"),
-		0,
-		TEXT("Useful developer variable\n")
-		TEXT("0: off (default, expected by the shader for better quality)\n")
-		TEXT("1: on"),
-		ECVF_Cheat | ECVF_RenderThreadSafe);
-
-	// ---------------------------------------
 
 	// the following commands are common exec commands that should be added to auto completion (todo: read UnConsole list in ini, discover all exec commands)
 	IConsoleManager::Get().RegisterConsoleCommand(TEXT("VisualizeTexture"),	TEXT("To visualize internal textures"), ECVF_Cheat);
@@ -1453,635 +1296,478 @@ void CreateConsoleVariables()
 
 #endif // !(UE_BUILD_SHIPPING || UE_BUILD_TEST)
 
-	IConsoleManager::Get().RegisterConsoleVariable(TEXT("r.UniformBufferPooling"),
-		1,
-		TEXT("If we pool object in RHICreateUniformBuffer to have less real API calls to creat buffers\n")
-		TEXT(" 0: off (for debugging)\n")
-		TEXT(" 1: on (optimization)"),
-		ECVF_RenderThreadSafe);
-
-	// The following console variable should never be compiled out ------------------
-	IConsoleManager::Get().RegisterConsoleVariable(TEXT("r.ClearWithExcludeRects"),
-		2,
-		TEXT("Control the use of exclude rects when using RHIClear\n")
-		TEXT(" 0: Force off (can be faster on hardware that has fast clears)\n")
-		TEXT(" 1: Use exclude rect if supplied\n")
-		TEXT(" 2: Auto (default is 2, pick what is considered best on this hardware)"),
-		ECVF_RenderThreadSafe);
-
-	IConsoleManager::Get().RegisterConsoleVariable(TEXT("r.SimpleDynamicLighting"),
-		0,	
-		TEXT("Whether to use simple dynamic lighting, which just renders an unshadowed dynamic directional light and a skylight.\n")
-		TEXT("All other lighting features are disabled when true.  This is useful for supporting very low end hardware.\n")
-		TEXT("0:off, 1:on"),
-		ECVF_RenderThreadSafe);
-
-	IConsoleManager::Get().RegisterConsoleVariable(TEXT("r.MobileHDR"),
-		1,	
-		TEXT("0: Mobile renders in LDR gamma space. (suggested for unlit games targeting low-end phones)\n")
-		TEXT("1: Mobile renders in HDR linear space. (default)"),
-		ECVF_RenderThreadSafe);
-
-	IConsoleManager::Get().RegisterConsoleVariable(TEXT("r.MobileHDR32bpp"),
-		0,	
-		TEXT("0: Mobile HDR renders to an FP16 render target.\n")
-		TEXT("1: Mobile HDR renders to an RGBA8 target."),
-		ECVF_RenderThreadSafe);
-
-	IConsoleManager::Get().RegisterConsoleVariable(TEXT("r.ClearSceneMethod"),
-		1,
-		TEXT("Select how scene rendertarget clears are handled\n")
-		TEXT(" 0: No clear\n")
-		TEXT(" 1: RHIClear (default)\n")
-		TEXT(" 2: Quad at max z"),
-		ECVF_RenderThreadSafe);
-
-	IConsoleManager::Get().RegisterConsoleVariable(TEXT("r.SeparateTranslucency"),
-		1,
-		TEXT("Allows to disable the separate translucency feature (all translucency is rendered in separate RT and composited\n")
-		TEXT("after DOF, if not specified otherwise in the material).\n")
-		TEXT(" 0: off (translucency is affected by depth of field)\n")
-		TEXT(" 1: on costs GPU performance and memory but keeps translucency unaffected by Depth of Fieled. (default)"),
-		ECVF_RenderThreadSafe);
-
-	IConsoleManager::Get().RegisterConsoleVariable(TEXT("r.LensFlareQuality"),
-		2,
-		TEXT(" 0: off but best for performance\n")
-		TEXT(" 1: low quality with good performance\n")
-		TEXT(" 2: good quality (default)\n")
-		TEXT(" 3: very good quality but bad performance"),
-		ECVF_Scalability | ECVF_RenderThreadSafe);
-
-	IConsoleManager::Get().RegisterConsoleVariable(TEXT("r.BloomQuality"),
-		5,
-		TEXT(" 0: off, no performance impact.\n")
-		TEXT(" 1: average quality, least performance impact.\n")
-		TEXT(" 2: average quality, least performance impact.\n")
-		TEXT(" 3: good quality.\n")
-		TEXT(" 4: good quality.\n")
-		TEXT(" 5: Best quality, most significant performance impact. (default)\n")
-		TEXT(">5: force experimental higher quality on mobile (can be quite slow on some hardware)"),
-		ECVF_Scalability | ECVF_RenderThreadSafe);
-
-	IConsoleManager::Get().RegisterConsoleVariable(TEXT("r.SceneColorFringeQuality"),
-		1,
-		TEXT(" 0: off but best for performance\n")
-		TEXT(" 1: 3 texture samples (default)n"),
-		ECVF_Scalability | ECVF_RenderThreadSafe);
-
-	IConsoleManager::Get().RegisterConsoleVariable(TEXT("r.UpsampleQuality"),
-		3,
-		TEXT(" 0: Nearest filtering\n")
-		TEXT(" 1: Simple Bilinear\n")
-		TEXT(" 2: 4 tap bilinear\n")
-		TEXT(" 3: Directional blur with unsharp mask upsample. (default)"),
-		ECVF_Scalability | ECVF_RenderThreadSafe);
-
-	// ---------------------------------------
-	// physics related
-	// ---------------------------------------
-
-	IConsoleManager::Get().RegisterConsoleVariable(TEXT("p.InitialOverlapTolerance"),
-		0.0f,
-		TEXT("Tolerance for initial overlapping test in PrimitiveComponent movement.\n")
-		TEXT("Normals within this tolerance are ignored if moving out of the object.\n")
-		TEXT("Dot product of movement direction and surface normal."),
-		ECVF_Default);
-
-	IConsoleManager::Get().RegisterConsoleVariable(TEXT("p.PenetrationPullbackDistance"),
-		0.125f,
-		TEXT("Pull out from penetration of an object by this extra distance.\n")
-		TEXT("Distance added to penetration fix-ups."),
-		ECVF_Default);
-
-	IConsoleManager::Get().RegisterConsoleVariable(TEXT("p.EncroachEpsilon"),
-		0.15f,
-		TEXT("Epsilon value encroach checking for capsules\n")
-		TEXT("0: use full sized capsule. > 0: shrink capsule size by this amount (world units)"),
-		ECVF_Default);
-
-	IConsoleManager::Get().RegisterConsoleVariable(TEXT("p.NetEnableMoveCombining"),
-		1,
-		TEXT("Whether to enable move combining on the client to reduce bandwidth by combining similar moves.\n")
-		TEXT("0: Disable, 1: Enable"),
-		ECVF_Default);
-
-	IConsoleManager::Get().RegisterConsoleVariable(TEXT("p.NetProxyShrinkRadius"),
-		0.01f,
-		TEXT("Shrink simulated proxy capsule radius by this amount, to account for network rounding that may cause encroachment.\n")
-		TEXT("Changing this value at runtime may require the proxy to re-join for correct behavior.\n")
-		TEXT("<= 0: disabled, > 0: shrink by this amount."),
-		ECVF_Default);
-
-	IConsoleManager::Get().RegisterConsoleVariable(TEXT("p.NetProxyShrinkHalfHeight"),
-		0.01f,
-		TEXT("Shrink simulated proxy capsule half height by this amount, to account for network rounding that may cause encroachment.\n")
-		TEXT("Changing this value at runtime may require the proxy to re-join for correct behavior.\n")
-		TEXT("<= 0: disabled, > 0: shrink by this amount."),
-		ECVF_Default);
-
-
-
-	IConsoleManager::Get().RegisterConsoleVariable(TEXT("p.ToleranceScale_Length"),
-		100.f,
-		TEXT("The approximate size of objects in the simulation. Default: 100"),
-		ECVF_ReadOnly);
-
-	IConsoleManager::Get().RegisterConsoleVariable(TEXT("p.ToleranceScale_Mass"),
-		100.f,
-		TEXT("The approximate mass of a length * length * length block. Default: 100"),
-		ECVF_ReadOnly);
-
-	IConsoleManager::Get().RegisterConsoleVariable(TEXT("p.ToleranceScale_Speed"),
-		1000.f,
-		TEXT("The typical magnitude of velocities of objects in simulation. Default: 1000"),
-		ECVF_ReadOnly);
-
-	IConsoleManager::Get().RegisterConsoleVariable(TEXT("p.ConstraintDampingScale"),
-		100000.f,
-		TEXT("The multiplier of constraint damping in simulation. Default: 100000"),
-		ECVF_ReadOnly);
-
-	IConsoleManager::Get().RegisterConsoleVariable(TEXT("p.ConstraintStiffnessScale"),
-		100000.f,
-		TEXT("The multiplier of constraint stiffness in simulation. Default: 100000"),
-		ECVF_ReadOnly);
-
-	IConsoleManager::Get().RegisterConsoleVariable(TEXT("p.ContactOffsetFactor"),
-		0.01f,
-		TEXT("Multiplied by min dimension of object to calculate how close objects get before generating contacts. Default: 0.01"),
-		ECVF_Default);
-
-	IConsoleManager::Get().RegisterConsoleVariable(TEXT("p.MaxContactOffset"),
-		1.f,
-		TEXT("Max value of contact offset, which controls how close objects get before generating contacts. Default: 1.0"),
-		ECVF_Default);
-
-	IConsoleManager::Get().RegisterConsoleVariable(TEXT("p.BounceThresholdVelocity"),
-		200.f,
-		TEXT("A contact with a relative velocity below this will not bounce. Default: 200"),
-		ECVF_Default);
-
-	// ---------------------------------------
-	// APEX
-
-	IConsoleManager::Get().RegisterConsoleVariable(TEXT("p.APEXMaxDestructibleDynamicChunkIslandCount"),
-		2000,
-		TEXT("APEX Max Destructilbe Dynamic Chunk Island Count."),
-		ECVF_Default);
-
-	IConsoleManager::Get().RegisterConsoleVariable(TEXT("p.APEXMaxDestructibleDynamicChunkCount"),
-		2000,
-		TEXT("APEX Max Destructible dynamic Chunk Count."),
-		ECVF_Default);
-
-	IConsoleManager::Get().RegisterConsoleVariable(TEXT("p.bAPEXSortDynamicChunksByBenefit"),
-		1,
-		TEXT("True if APEX should sort dynamic chunks by benefit."),
-		ECVF_Default);
-
-	// ---------------------------------------
-
-	IConsoleManager::Get().RegisterConsoleVariable(TEXT("r.AmbientOcclusionLevels"),
-		3,
-		TEXT("Defines how many mip levels are using during the ambient occlusion calculation. This is useful when tweaking the algorithm.\n")
-		TEXT(" 0:none, 1:one, 2:two, 3:three(default), 4:four(larger at little cost but can flicker)"),
-		ECVF_Scalability | ECVF_RenderThreadSafe);
-
-	IConsoleManager::Get().RegisterConsoleVariable(TEXT("r.AmbientOcclusionRadiusScale"),
-		1.0f,
-		TEXT("Allows to scale the ambient occlusion radius (SSAO).\n")
-		TEXT(" 0:off, 1.0:normal, <1:smaller, >1:larger"),
-		ECVF_Scalability | ECVF_RenderThreadSafe);
-
-	IConsoleManager::Get().RegisterConsoleVariable(TEXT("r.AtmosphereRender"),
-		1,
-		TEXT("Defines atmosphere will render or not. Only changed by r.Atmosphere console command\n")
-		TEXT(" 0: off\n")
-		TEXT(" 1: on"),
-		ECVF_ReadOnly);
-
-	IConsoleManager::Get().RegisterConsoleVariable(TEXT("r.ShadowQuality"),
-		5,
-		TEXT("Defines the shadow method which allows to adjust for quality or performance.\n")
-		TEXT(" 0:off, 1:low(unfiltered), 2:low .. 5:max (default)"),
-		ECVF_Scalability | ECVF_RenderThreadSafe);
-
-	IConsoleManager::Get().RegisterConsoleVariable(TEXT("r.MotionBlurQuality"),
-		4,
-		TEXT("Defines the motion blur method which allows to adjust for quality or performance.\n")
-		TEXT(" 0:off, 1:low, 2:medium, 3:high (default), 4: very high"),
-		ECVF_Scalability | ECVF_RenderThreadSafe);
-
-	IConsoleManager::Get().RegisterConsoleVariable(TEXT("r.PostProcessAAQuality"),
-		4,
-		TEXT("Defines the postprocess anti aliasing method which allows to adjust for quality or performance.\n")
-		TEXT(" 0:off, 1:very low (faster FXAA), 2:low (FXAA), 3:medium (faster TemporalAA), 4:high (default TemporalAA), 5:very high, 6:max"),
-		ECVF_Scalability | ECVF_RenderThreadSafe);
-
-	IConsoleManager::Get().RegisterConsoleVariable(TEXT("r.MotionBlurSoftEdgeSize"),
-		1.0f,
-		TEXT("Defines how lange object motion blur is blurred (percent of screen width) to allow soft edge motion blur.\n")
-		TEXT("This scales linearly with the size (up to a maximum of 32 samples, 2.5 is about 18 samples) and with screen resolution\n")
-		TEXT("Smaller values are better for performance and provide more accurate motion vectors but the blurring outside the object is reduced.\n")
-		TEXT("If needed this can be exposed like the other motionblur settings.\n")
-		TEXT(" 0:off (not free and does never completely disable), >0, 1.0 (default)"),
-		ECVF_RenderThreadSafe);
-
-	IConsoleManager::Get().RegisterConsoleVariable(TEXT("r.Color.Min"),
-		0.0f,
-		TEXT("Allows to define where the value 0 in the color channels is mapped to after color grading.\n")
-		TEXT("The value should be around 0, positive: a gray scale is added to the darks, negative: more dark values become black, Default: 0"),
-		ECVF_RenderThreadSafe);
-
-	IConsoleManager::Get().RegisterConsoleVariable(TEXT("r.Color.Mid"),
-		0.5f,
-		TEXT("Allows to define where the value 0.5 in the color channels is mapped to after color grading (This is similar to a gamma correction).\n")
-		TEXT("Value should be around 0.5, smaller values darken the mid tones, larger values brighten the mid tones, Default: 0.5"),
-		ECVF_RenderThreadSafe);
-
-	IConsoleManager::Get().RegisterConsoleVariable(TEXT("r.Color.Max"),
-		1.0f,
-		TEXT("Allows to define where the value 1.0 in the color channels is mapped to after color grading.\n")
-		TEXT("Value should be around 1, smaller values darken the highlights, larger values move more colors towards white, Default: 1"),
-		ECVF_RenderThreadSafe);
-
-	IConsoleManager::Get().RegisterConsoleVariable(TEXT("r.FullScreenMode"),
-		2,
-		TEXT("Defines how we do full screen when requested (e.g. command line option -fullscreen or in ini [SystemSettings] fullscreen=true)\n")
-		TEXT(" 0: normal full screen (renders faster, more control over vsync, less GPU memory, 10bit color if possible)\n")
-		TEXT(" 1: windowed full screen, desktop resolution (quick switch between applications and window mode, full quality)\n")
-		TEXT(" 2: windowed full screen, specified resolution (like 1 but no unintuitive performance cliff, can be blurry, default)\n")
-		TEXT(" any other number behaves like 0"),
-		ECVF_Scalability);
-
-	IConsoleManager::Get().RegisterConsoleVariable(TEXT("r.GBufferFormat"),
-		1,
-		TEXT("Defines the memory layout used for the GBuffer.\n")
-		TEXT("(affects performance, mostly through bandwidth, quality of normals and material attributes).\n")
-		TEXT(" 0: lower precision (8bit per component, for profiling)\n")
-		TEXT(" 1: low precision (default)\n")
-		TEXT(" 5: high precision"),
-		ECVF_RenderThreadSafe);
-
-	IConsoleManager::Get().RegisterConsoleVariable(TEXT("r.SceneColorFormat"),
-		4,
-		TEXT("Defines the memory layout (RGBA) used for the scene color\n")
-		TEXT("(affects performance, mostly through bandwidth, quality especially with translucency).\n")
-		TEXT(" 0: PF_B8G8R8A8 32Bit (mostly for testing, likely to unusable with HDR)\n")
-		TEXT(" 1: PF_A2B10G10R10 32Bit\n")
-		TEXT(" 2: PF_FloatR11G11B10 32Bit\n")
-		TEXT(" 3: PF_FloatRGB 32Bit\n")
-		TEXT(" 4: PF_FloatRGBA 64Bit (default, might be overkill, especially if translucency is mostly using SeparateTranslucency)\n")
-		TEXT(" 5: PF_A32B32G32R32F 128Bit (unreasonable but good for testing)"),
-		ECVF_Scalability | ECVF_RenderThreadSafe);
-
-	IConsoleManager::Get().RegisterConsoleVariable(TEXT("r.AmbientOcclusionStaticFraction"),
-		-1.0f,
-		TEXT("Allows to override the Ambient Occlusion Static Fraction (see post process volume). Fractions are between 0 and 1.\n")
-		TEXT("<0: use default setting (default -1)\n")
-		TEXT(" 0: no effect on static lighting, 0 is free meaning no extra rendering pass\n")
-		TEXT(" 1: AO affects the stat lighting"),
-		ECVF_Default);
-
-	IConsoleManager::Get().RegisterConsoleVariable(TEXT("r.DepthOfFieldQuality"),
-		2,
-		TEXT("Allows to adjust the depth of field quality. Currently only fully affects BokehDOF. GaussianDOF is either 0 for off, otherwise on.\n")
-		TEXT(" 0: Off\n")
-		TEXT(" 1: Low\n")
-		TEXT(" 2: high quality (default, adaptive, can be 4x slower)"),
-		ECVF_Scalability | ECVF_RenderThreadSafe);
-
-	IConsoleManager::Get().RegisterConsoleVariable(TEXT("r.DepthOfFieldNearBlurSizeThreshold"),
-		0.01f,
-		TEXT("Sets the minimum near blur size before the effect is forcably disabled. Currently only affects Gaussian DOF.\n")
-		TEXT(" (default = 0.01f)"),
-		ECVF_RenderThreadSafe);
-
-	IConsoleManager::Get().RegisterConsoleVariable(TEXT("r.ScreenPercentage"),
-		-1.0f,
-		TEXT("To render in lower resolution and upscale for better performance.\n")
-		TEXT("70 is a good value for low aliasing and performance, can be verified with 'show TestImage'\n")
-		TEXT("in percent, >0 and <=100, <0 means the post process volume settings are used"),
-		ECVF_Scalability | ECVF_Default);
-
-	IConsoleManager::Get().RegisterConsoleVariable(TEXT("r.ScreenPercentageSoftness"),
-		0.3f,
-		TEXT("To scale up with higher quality loosing some sharpness\n")
-		TEXT(" 0..1 (0.3 is good for ScreenPercentage 90"),
-		ECVF_Scalability | ECVF_RenderThreadSafe);
-
-	IConsoleManager::Get().RegisterConsoleVariable(TEXT("r.BlackBorders"),
-		1,	
-		TEXT("To draw black borders around the rendered image\n")
-		TEXT("(prevents artifacts from post processing passes that read outside of the image e.g. PostProcessAA)\n")
-		TEXT("in pixels, 0:off"),
-		ECVF_Default);
+	// testing code
+	{
+		FConsoleManager& ConsoleManager = (FConsoleManager&)IConsoleManager::Get();
 
-	// Changing this causes a full shader recompile
-	IConsoleManager::Get().RegisterConsoleVariable(TEXT("r.UseDiffuseSpecularMaterialInputs"),
-		0,	
-		TEXT("Whether to use DiffuseColor, SpecularColor material inputs, otherwise BaseColor, Metallic, Specular will be used. Cannot be changed at runtime."),
-		ECVF_ReadOnly);
-
-	IConsoleManager::Get().RegisterConsoleVariable(TEXT("r.MaterialQualityLevel"),
-		1,	
-		TEXT("0 corresponds to low quality materials, as defined by quality switches in materials, 1 corresponds to high."),
-		ECVF_Scalability | ECVF_RenderThreadSafe);
-
-	IConsoleManager::Get().RegisterConsoleVariable(TEXT("Compat.UseDXT5NormalMaps"),
-		0,	
-		TEXT("Whether to use DXT5 for normal maps, otherwise BC5 will be used, which is not supported on all hardware.\n")
-		TEXT("Both formats require the same amount of memory (if driver doesn't emulate the format).\n")
-		TEXT("Changing this will cause normal maps to be recompressed on next load (or when using recompile shaders)\n")
-		TEXT(" 0: Use BC5 texture format (default)\n")
-		TEXT(" 1: Use DXT5 texture format (lower quality)"),
-		// 
-		// Changing this causes a full shader recompile
-
-		ECVF_ReadOnly);
-
-	// Changing this is currently unsupported after content has been chunked with the previous setting
-	// Changing this causes a full shader recompile
-	IConsoleManager::Get().RegisterConsoleVariable(TEXT("Compat.MAX_GPUSKIN_BONES"),
-		256,	
-		TEXT("Max number of bones that can be skinned on the GPU in a single draw call. Cannot be changed at runtime."),
-		ECVF_ReadOnly);
+		ConsoleManager.Test();
+	}
+}
+
+// Naming conventions:
+//
+// Console variable should start with (suggestion):
+//
+// r.      Renderer / 3D Engine / graphical feature
+// RHI.    Low level RHI (rendering platform) specific
+// a.	   Animation
+// s. 	   Sound / Music
+// n.      Network
+// ai.     Artificial intelligence
+// i.      Input e.g. mouse/keyboard
+// p.      Physics
+// t.      Timer
+// g.      Game
+// Compat.
+// FX.     Particle effects
+// sg.     scalability group (used by scalability system, ini load/save or using SCALABILITY console command)
+
+#if !(UE_BUILD_SHIPPING || UE_BUILD_TEST)
+static TAutoConsoleVariable<int32> CVarPreViewTranslation(
+	TEXT("r.PreViewTranslation"),
+	1,
+	TEXT("To limit issues with float world space positions we offset the world by the\n")
+	TEXT("PreViewTranslation vector. This command allows to disable updating this vector.\n")
+	TEXT(" 0: disable update\n")
+	TEXT(" 1: update the offset is each frame (default)"),
+	ECVF_Cheat);
+
+static TAutoConsoleVariable<FString> CVarFreezeAtPosition(
+	TEXT("FreezeAtPosition"),
+	TEXT(""),	// default value is empty
+	TEXT("This console variable stores the position and rotation for the FreezeAt command which allows\n")
+	TEXT("to lock the camera in order to provide more deterministic render profiling.\n")
+	TEXT("The FreezeAtPosition can be set in the ConsoleVariables.ini (start the map with MAPNAME?bTourist=1).\n")
+	TEXT("Also see the FreezeAt command console command.\n")
+	TEXT("The number syntax if the same as the one used by the BugIt command:\n")
+	TEXT(" The first three values define the position, the next three define the rotation.\n")
+	TEXT("Example:\n")
+	TEXT(" FreezeAtPosition 2819.5520 416.2633 75.1500 65378 -25879 0"),
+	ECVF_Cheat);
+
+static TAutoConsoleVariable<int32> CVarLimitRenderingFeatures(
+	TEXT("r.LimitRenderingFeatures"),
+	0,
+	TEXT("Allows to quickly reduce render feature to increase render performance.\n")
+	TEXT("This is just a quick way to alter multiple show flags and console variables in the game\n")
+	TEXT("Disabled more feature the higher the number\n")
+	TEXT(" <=0:off, order is defined in code (can be documented here when we settled on an order)"),
+	ECVF_Cheat | ECVF_RenderThreadSafe);
+
+static TAutoConsoleVariable<int32> CVarDumpingMovie(
+	TEXT("r.DumpingMovie"),
+	GIsDumpingMovie,
+	TEXT("Allows to dump each rendered frame to disk (slow fps, names MovieFrame..).\n")
+	TEXT("<=0:off (default), <0:remains on, >0:remains on for n frames (n is the number specified)"),
+	ECVF_Cheat);
+
+#endif // !(UE_BUILD_SHIPPING || UE_BUILD_TEST)
+
+static TAutoConsoleVariable<int32> CVarUniformBufferPooling(
+	TEXT("r.UniformBufferPooling"),
+	1,
+	TEXT("If we pool object in RHICreateUniformBuffer to have less real API calls to creat buffers\n")
+	TEXT(" 0: off (for debugging)\n")
+	TEXT(" 1: on (optimization)"),
+	ECVF_RenderThreadSafe);
+
+// The following console variable should never be compiled out ------------------
+static TAutoConsoleVariable<int32> CVarClearWithExcludeRects(
+	TEXT("r.ClearWithExcludeRects"),
+	2,
+	TEXT("Control the use of exclude rects when using RHIClear\n")
+	TEXT(" 0: Force off (can be faster on hardware that has fast clears)\n")
+	TEXT(" 1: Use exclude rect if supplied\n")
+	TEXT(" 2: Auto (default is 2, pick what is considered best on this hardware)"),
+	ECVF_RenderThreadSafe);
+
+static TAutoConsoleVariable<int32> CVarSimpleDynamicLighting(
+	TEXT("r.SimpleDynamicLighting"),
+	0,
+	TEXT("Whether to use simple dynamic lighting, which just renders an unshadowed dynamic directional light and a skylight.\n")
+	TEXT("All other lighting features are disabled when true.  This is useful for supporting very low end hardware.\n")
+	TEXT("0:off, 1:on"),
+	ECVF_RenderThreadSafe);
+
+static TAutoConsoleVariable<int32> CVarMobileHDR(
+	TEXT("r.MobileHDR"),
+	1,
+	TEXT("0: Mobile renders in LDR gamma space. (suggested for unlit games targeting low-end phones)\n")
+	TEXT("1: Mobile renders in HDR linear space. (default)"),
+	ECVF_RenderThreadSafe);
+
+static TAutoConsoleVariable<int32> CVarMobileHDR32bpp(
+	TEXT("r.MobileHDR32bpp"),
+	0,
+	TEXT("0: Mobile HDR renders to an FP16 render target.\n")
+	TEXT("1: Mobile HDR renders to an RGBA8 target."),
+	ECVF_RenderThreadSafe);
+
+static TAutoConsoleVariable<int32> CVarSetClearSceneMethod(
+	TEXT("r.ClearSceneMethod"),
+	1,
+	TEXT("Select how scene rendertarget clears are handled\n")
+	TEXT(" 0: No clear\n")
+	TEXT(" 1: RHIClear (default)\n")
+	TEXT(" 2: Quad at max z"),
+	ECVF_RenderThreadSafe);
+
+static TAutoConsoleVariable<int32> CVarLensFlareQuality(
+	TEXT("r.LensFlareQuality"),
+	2,
+	TEXT(" 0: off but best for performance\n")
+	TEXT(" 1: low quality with good performance\n")
+	TEXT(" 2: good quality (default)\n")
+	TEXT(" 3: very good quality but bad performance"),
+	ECVF_Scalability | ECVF_RenderThreadSafe);
+
+static TAutoConsoleVariable<int32> CVarBloomQuality(
+	TEXT("r.BloomQuality"),
+	5,
+	TEXT(" 0: off, no performance impact.\n")
+	TEXT(" 1: average quality, least performance impact.\n")
+	TEXT(" 2: average quality, least performance impact.\n")
+	TEXT(" 3: good quality.\n")
+	TEXT(" 4: good quality.\n")
+	TEXT(" 5: Best quality, most significant performance impact. (default)\n")
+	TEXT(">5: force experimental higher quality on mobile (can be quite slow on some hardware)"),
+	ECVF_Scalability | ECVF_RenderThreadSafe);
+
+static TAutoConsoleVariable<int32> CVarSceneColorFringeQuality(
+	TEXT("r.SceneColorFringeQuality"),
+	1,
+	TEXT(" 0: off but best for performance\n")
+	TEXT(" 1: 3 texture samples (default)n"),
+	ECVF_Scalability | ECVF_RenderThreadSafe);
+
+
+// ---------------------------------------
+
+static TAutoConsoleVariable<int32> CVarAmbientOcclusionLevels(
+	TEXT("r.AmbientOcclusionLevels"),
+	3,
+	TEXT("Defines how many mip levels are using during the ambient occlusion calculation. This is useful when tweaking the algorithm.\n")
+	TEXT(" 0:none, 1:one, 2:two, 3:three(default), 4:four(larger at little cost but can flicker)"),
+	ECVF_Scalability | ECVF_RenderThreadSafe);
+
+static TAutoConsoleVariable<float> CVarAmbientOcclusionRadiusScale(
+	TEXT("r.AmbientOcclusionRadiusScale"),
+	1.0f,
+	TEXT("Allows to scale the ambient occlusion radius (SSAO).\n")
+	TEXT(" 0:off, 1.0:normal, <1:smaller, >1:larger"),
+	ECVF_Scalability | ECVF_RenderThreadSafe);
+
+static TAutoConsoleVariable<float> CVarAmbientOcclusionStaticFraction(TEXT("r.AmbientOcclusionStaticFraction"),
+	-1.0f,
+	TEXT("Allows to override the Ambient Occlusion Static Fraction (see post process volume). Fractions are between 0 and 1.\n")
+	TEXT("<0: use default setting (default -1)\n")
+	TEXT(" 0: no effect on static lighting, 0 is free meaning no extra rendering pass\n")
+	TEXT(" 1: AO affects the stat lighting"),
+	ECVF_Default);
+
+static TAutoConsoleVariable<int32> CVarAtmosphereRender(
+	TEXT("r.AtmosphereRender"),
+	1,
+	TEXT("Defines atmosphere will render or not. Only changed by r.Atmosphere console command\n")
+	TEXT(" 0: off\n")
+	TEXT(" 1: on"),
+	ECVF_ReadOnly);
+
+static TAutoConsoleVariable<int32> CVarShadowQuality(
+	TEXT("r.ShadowQuality"),
+	5,
+	TEXT("Defines the shadow method which allows to adjust for quality or performance.\n")
+	TEXT(" 0:off, 1:low(unfiltered), 2:low .. 5:max (default)"),
+	ECVF_Scalability | ECVF_RenderThreadSafe);
+
+static TAutoConsoleVariable<int32> CVarMotionBlurQuality(
+	TEXT("r.MotionBlurQuality"),
+	4,
+	TEXT("Defines the motion blur method which allows to adjust for quality or performance.\n")
+	TEXT(" 0:off, 1:low, 2:medium, 3:high (default), 4: very high"),
+	ECVF_Scalability | ECVF_RenderThreadSafe);
+
+static TAutoConsoleVariable<int32> CVarPostProcessAAQuality(
+	TEXT("r.PostProcessAAQuality"),
+	4,
+	TEXT("Defines the postprocess anti aliasing method which allows to adjust for quality or performance.\n")
+	TEXT(" 0:off, 1:very low (faster FXAA), 2:low (FXAA), 3:medium (faster TemporalAA), 4:high (default TemporalAA), 5:very high, 6:max"),
+	ECVF_Scalability | ECVF_RenderThreadSafe);
 
+static TAutoConsoleVariable<int32> CVarFullscreenMode(
+	TEXT("r.FullScreenMode"),
+	2,
+	TEXT("Defines how we do full screen when requested (e.g. command line option -fullscreen or in ini [SystemSettings] fullscreen=true)\n")
+	TEXT(" 0: normal full screen (renders faster, more control over vsync, less GPU memory, 10bit color if possible)\n")
+	TEXT(" 1: windowed full screen, desktop resolution (quick switch between applications and window mode, full quality)\n")
+	TEXT(" 2: windowed full screen, specified resolution (like 1 but no unintuitive performance cliff, can be blurry, default)\n")
+	TEXT(" any other number behaves like 0"),
+	ECVF_Scalability);
+
+static TAutoConsoleVariable<int32> CVarSceneColorFormat(
+	TEXT("r.SceneColorFormat"),
+	4,
+	TEXT("Defines the memory layout (RGBA) used for the scene color\n")
+	TEXT("(affects performance, mostly through bandwidth, quality especially with translucency).\n")
+	TEXT(" 0: PF_B8G8R8A8 32Bit (mostly for testing, likely to unusable with HDR)\n")
+	TEXT(" 1: PF_A2B10G10R10 32Bit\n")
+	TEXT(" 2: PF_FloatR11G11B10 32Bit\n")
+	TEXT(" 3: PF_FloatRGB 32Bit\n")
+	TEXT(" 4: PF_FloatRGBA 64Bit (default, might be overkill, especially if translucency is mostly using SeparateTranslucency)\n")
+	TEXT(" 5: PF_A32B32G32R32F 128Bit (unreasonable but good for testing)"),
+	ECVF_Scalability | ECVF_RenderThreadSafe);
+
+static TAutoConsoleVariable<int32> CVarDepthOfFieldQuality(
+	TEXT("r.DepthOfFieldQuality"),
+	2,
+	TEXT("Allows to adjust the depth of field quality. Currently only fully affects BokehDOF. GaussianDOF is either 0 for off, otherwise on.\n")
+	TEXT(" 0: Off\n")
+	TEXT(" 1: Low\n")
+	TEXT(" 2: high quality (default, adaptive, can be 4x slower)"),
+	ECVF_Scalability | ECVF_RenderThreadSafe);
+
+static TAutoConsoleVariable<float> CVarScreenPercentage(
+	TEXT("r.ScreenPercentage"),
+	-1.0f,
+	TEXT("To render in lower resolution and upscale for better performance.\n")
+	TEXT("70 is a good value for low aliasing and performance, can be verified with 'show TestImage'\n")
+	TEXT("in percent, >0 and <=100, <0 means the post process volume settings are used"),
+	ECVF_Scalability | ECVF_Default);
+
+// Changing this causes a full shader recompile
+static TAutoConsoleVariable<int32> CVarUseDiffuseSpecularMaterialInputs(
+	TEXT("r.UseDiffuseSpecularMaterialInputs"),
+	0,
+	TEXT("Whether to use DiffuseColor, SpecularColor material inputs, otherwise BaseColor, Metallic, Specular will be used. Cannot be changed at runtime."),
+	ECVF_ReadOnly);
+
+static TAutoConsoleVariable<int32> CVarMaterialQualityLevel(
+	TEXT("r.MaterialQualityLevel"),
+	1,
+	TEXT("0 corresponds to low quality materials, as defined by quality switches in materials, 1 corresponds to high."),
+	ECVF_Scalability | ECVF_RenderThreadSafe);
+
+static TAutoConsoleVariable<int32> CVarUseDXT5NormalMaps(
+	TEXT("Compat.UseDXT5NormalMaps"),
+	0,
+	TEXT("Whether to use DXT5 for normal maps, otherwise BC5 will be used, which is not supported on all hardware.\n")
+	TEXT("Both formats require the same amount of memory (if driver doesn't emulate the format).\n")
+	TEXT("Changing this will cause normal maps to be recompressed on next load (or when using recompile shaders)\n")
+	TEXT(" 0: Use BC5 texture format (default)\n")
+	TEXT(" 1: Use DXT5 texture format (lower quality)"),
+	// 
 	// Changing this causes a full shader recompile
-	IConsoleManager::Get().RegisterConsoleVariable(TEXT("r.AllowStaticLighting"),
-		1,	
-		TEXT("Whether to allow any static lighting to be generated and used, like lightmaps and shadowmaps.\n")
-		TEXT("Games that only use dynamic lighting should set this to 0 to save some static lighting overhead."),
-		ECVF_ReadOnly | ECVF_RenderThreadSafe);
-
-	IConsoleManager::Get().RegisterConsoleVariable(TEXT("r.AllowMeshDistanceFieldRepresentations"),
-		0,	
-		TEXT("Whether to build distance fields of static meshes, needed for distance field AO, which is used to implement Movable SkyLight shadows.\n")
-		TEXT("Enabling will increase mesh build times and memory usage.  Changing this value will cause a rebuild of all static meshes."),
-		ECVF_ReadOnly);
-
-	IConsoleManager::Get().RegisterConsoleVariable(TEXT("con.MinLogVerbosity"),
-		0,
-		TEXT("Allows to see the log in the in game console (by default deactivated to avoid spam and minor performance loss).\n")
-		TEXT(" 0: no logging other than console response (default)\n")
-		TEXT(" 1: Only fatal errors (no that useful)\n")
-		TEXT(" 2: additionally errors\n")
-		TEXT(" 3: additionally warnings\n")
-		TEXT(" 4: additionally display\n")
-		TEXT(" 5: additionally log\n")
-		TEXT("..\n")
-		TEXT(">=7: all"),
-		ECVF_Default);
-	
-	IConsoleManager::Get().RegisterConsoleVariable(TEXT("r.MSAA.CompositingSampleCount"),
-		4,
-		TEXT("Affects the render quality of the editor 3d objects.\n")
-		TEXT(" 1: no MSAA, lowest quality\n")
-		TEXT(" 2: 2x MSAA, medium quality (medium GPU memory consumption)\n")
-		TEXT(" 4: 4x MSAA, high quality (high GPU memory consumption)\n")
-		TEXT(" 8: 8x MSAA, very high quality (insane GPU memory consumption)"),
-		ECVF_RenderThreadSafe);
-
-	IConsoleManager::Get().RegisterConsoleVariable(TEXT("net.DormancyEnable"),
-		1,	
-		TEXT("Enables Network Dormancy System for reducing CPU and bandwidth overhead of infrequently updated actors\n")
-		TEXT("1 Enables network dormancy. 0 disables network dormancy."),
-		ECVF_Default);
-
-	IConsoleManager::Get().RegisterConsoleVariable(TEXT("net.DormancyDraw"),
-		0,
-		TEXT("Draws debug information for network dormancy\n")
-		TEXT("1 Enables network dormancy debugging. 0 disables."),
-		ECVF_Default);
-
-	IConsoleManager::Get().RegisterConsoleVariable(TEXT("net.DormancyDrawCullDistance"),
-		5000.f,
-		TEXT("Cull distance for net.DormancyDraw. World Units")
-		TEXT("Max world units an actor can be away from the local view to draw its dormancy status"),
-		ECVF_Default);
-
-	IConsoleManager::Get().RegisterConsoleVariable(TEXT("net.DormancyValidate"),
-		0,
-		TEXT("Validates that dormant actors do not change state while in a dormant state (on server only)")
-		TEXT("0: Dont validate. 1: Validate on wake up. 2: Validate on each net update"),
-		ECVF_Default);
-
-	IConsoleManager::Get().RegisterConsoleVariable(TEXT("net.PackageMap.DebugObject"),
-		TEXT(""),
-		TEXT("Debugs PackageMap serialization of object")
-		TEXT("Partial name of object to debug"),
-		ECVF_Default);
-
-	IConsoleManager::Get().RegisterConsoleVariable(TEXT("net.Replication.DebugProperty"),
-		TEXT(""),
-		TEXT("Debugs Replication of property by name")
-		TEXT("Partial name of property to debug"),
-		ECVF_Default);
-
-	IConsoleManager::Get().RegisterConsoleVariable(TEXT("net.Reliable.Debug"),
-		0,
-		TEXT("Print all reliable bunches sent over the network\n")
-		TEXT(" 0: no print.\n")
-		TEXT(" 1: Print bunches as they are sent.\n")
-		TEXT(" 2: Print reliable bunch buffer each net update"),
-		ECVF_Default);
-
-	IConsoleManager::Get().RegisterConsoleVariable(TEXT("net.RPC.Debug"),
-		0,
-		TEXT("Print all RPC bunches sent over the network\n")
-		TEXT(" 0: no print.\n")
-		TEXT(" 1: Print bunches as they are sent."),
-		ECVF_Default);
-
-	IConsoleManager::Get().RegisterConsoleVariable(TEXT("net.Montage.Debug"),
-		0,
-		TEXT("Prints Replication information about AnimMontages\n")
-		TEXT(" 0: no print.\n")
-		TEXT(" 1: Print AnimMontage info on client side as they are played."),
-		ECVF_Cheat);
-
-	IConsoleManager::Get().RegisterConsoleVariable(TEXT("r.RenderTargetPoolMin"),
-		400,
-		TEXT("If the render target pool size (in MB) is below this number there is no deallocation of rendertargets")
-		TEXT("Default is 200 MB."),
-		ECVF_Scalability | ECVF_RenderThreadSafe);
-
-	IConsoleManager::Get().RegisterConsoleVariable(TEXT("r.Shadow.DistanceFieldPenumbraSize"),
-		0.05f,
-		TEXT("Controls the size of the uniform penumbra produced by static shadowing.\n")
-		TEXT("This is a fraction of MaxTransitionDistanceWorldSpace in BaseLightmass.ini."),
-		ECVF_RenderThreadSafe);
-
-	static TAutoConsoleVariable<int32> CVarIdleWhenNotForeground(
-		TEXT("t.IdleWhenNotForeground"),0,
-		TEXT("Prevents the engine from taking any CPU or GPU time while not the foreground app."),
-		ECVF_Cheat);
-	
-	IConsoleManager::Get().RegisterConsoleVariable(TEXT("r.TemporalAASamples"),
-		8,
-		TEXT("Number of jittered positions for temporal AA (4, 8=default, 16, 32, 64)."),
-		ECVF_RenderThreadSafe);
-
-	IConsoleManager::Get().RegisterConsoleVariable(TEXT("r.VSync"),
-		0,
-		TEXT("0: VSync is disabled.(default)\n")
-		TEXT("1: VSync is enabled."),
-		ECVF_Scalability | ECVF_RenderThreadSafe);
-
-	IConsoleManager::Get().RegisterConsoleVariable(TEXT("r.DontLimitOnBattery"),
-		0,
-		TEXT("0: Limit performance on devices with a battery.(default)\n")
-		TEXT("1: Do not limit performance due to device having a battery."),
-		ECVF_Scalability | ECVF_RenderThreadSafe);
-
-	IConsoleManager::Get().RegisterConsoleVariable(TEXT("r.FinishCurrentFrame"),
-		0,
-		TEXT("If on, the current frame will be forced to finish and render to the screen instead of being buffered.  This will improve latency, but slow down overall performance."),
-		ECVF_Scalability | ECVF_RenderThreadSafe);
-
-	IConsoleManager::Get().RegisterConsoleVariable(TEXT("r.MaxAnisotropy"),
-		4,
-		TEXT("MaxAnisotropy should range from 1 to 16. Higher values mean better texure quality when using anisotropic filtering but at a cost to performance."),
-		ECVF_Scalability | ECVF_RenderThreadSafe);
-
-	IConsoleManager::Get().RegisterConsoleVariable(TEXT("r.Shadow.MaxResolution"),
-		1024,
-		TEXT("Max square dimensions (in texels) allowed for rendering shadow depths. Range 4 to hardware limit. Higher = better quality shadows but at a performance cost."),
-		ECVF_Scalability | ECVF_RenderThreadSafe);
-
-	IConsoleManager::Get().RegisterConsoleVariable(TEXT("r.Shadow.CSM.TransitionScale"),
-		1.0f,
-		TEXT("Allows to scale the cascaded shadow map transition region. Clamped within 0..2.\n")
-		TEXT("0: no transition (fastest)\n")
-		TEXT("1: as specific in the light settings (default)\n")
-		TEXT("2: 2x larger than what was specified in the light"),
-		ECVF_Scalability | ECVF_RenderThreadSafe);
-
-	IConsoleManager::Get().RegisterConsoleVariable(TEXT("r.MobileContentScaleFactor"),
-		1.0f,
-		TEXT("Content scale multiplier (equates to iOS's contentScaleFactor to support Retina displays"),
-		ECVF_Default);
-
-#if PLATFORM_SUPPORTS_TEXTURE_STREAMING
-	IConsoleManager::Get().RegisterConsoleVariable(TEXT("r.TextureStreaming"),
-		1,
-		TEXT("Allows to define if texture streaming is enabled, can be changed at run time.\n")
-		TEXT("0: off\n")
-		TEXT("1: on (default)"),
-		ECVF_Default | ECVF_RenderThreadSafe);
-#endif
+	ECVF_ReadOnly);
+
+// Changing this is currently unsupported after content has been chunked with the previous setting
+// Changing this causes a full shader recompile
+static TAutoConsoleVariable<int32> CVarMaxGPUSkinBones(
+	TEXT("Compat.MAX_GPUSKIN_BONES"),
+	256,
+	TEXT("Max number of bones that can be skinned on the GPU in a single draw call. Cannot be changed at runtime."),
+	ECVF_ReadOnly);
+
+// Changing this causes a full shader recompile
+static TAutoConsoleVariable<int32> CVarAllowStaticLighting(
+	TEXT("r.AllowStaticLighting"),
+	1,
+	TEXT("Whether to allow any static lighting to be generated and used, like lightmaps and shadowmaps.\n")
+	TEXT("Games that only use dynamic lighting should set this to 0 to save some static lighting overhead."),
+	ECVF_ReadOnly | ECVF_RenderThreadSafe);
 
-	IConsoleManager::Get().RegisterConsoleVariable(TEXT("r.Streaming.PoolSize"),
-		-1,
-		TEXT("-1: Default texture pool size, otherwise the size in MB\n"),
-		ECVF_Scalability );
-
-	IConsoleManager::Get().RegisterConsoleVariable(TEXT("r.Streaming.UseFixedPoolSize"),
-		0,
-		TEXT("If non-zero, do not allow the pool size to change at run time."),
-		ECVF_ReadOnly );
-
-	IConsoleManager::Get().RegisterConsoleVariable(TEXT("r.ShaderDevelopmentMode"),
-		0,
-		TEXT("0: Default, 1: Enable various shader development utilities, such as the ability to retry on failed shader compile, and extra logging as shaders are compiled."),
-		ECVF_Default);
-
-	// this cvar can be removed in shipping to not compile shaders for development (faster)
-	IConsoleManager::Get().RegisterConsoleVariable(TEXT("r.CompileShadersForDevelopment"),
-		1,
-		TEXT("Setting this to 0 allows to ship a game with more optimized shaders as some\n")
-		TEXT("editor and development features are not longer compiled into the shaders.\n")
-		TEXT(" Note: This should be done when shipping but it's not done automatically yet (feature need to mature\n")
-		TEXT("       and shaders will compile slower as shader caching from development isn't shared).\n")
-		TEXT("Cannot be changed at runtime - can be put into BaseEngine.ini\n")
-		TEXT(" 0: off, shader can run a bit faster\n")
-		TEXT(" 1: on (Default)"),
-		ECVF_ReadOnly);
-
-	IConsoleManager::Get().RegisterConsoleVariable(TEXT("r.PS4MixedModeShaderDebugInfo"),
-		0,
-		TEXT("Whether to compile shaders to allow mixed mode shader debugging. This will currently generate slower code.\n")
-		TEXT(" 0: Normal mode\n")
-		TEXT(" 1: Mixed mode)"),
-		ECVF_ReadOnly);
-
-	IConsoleManager::Get().RegisterConsoleVariable(TEXT("r.PS4DumpShaderSDB"),
-		0,
-		TEXT("Whether to dump shader sdb files used for shader association.\n")
-		TEXT(" 0: Disabled\n")
-		TEXT(" 1: Enabled)"),
-		ECVF_ReadOnly);
-
-	IConsoleManager::Get().RegisterConsoleVariable(TEXT("r.ViewDistanceScale"),
-		1.0f,
-		TEXT("Controls the view distance scale. A primitive's MaxDrawDistance is scaled by this value.\n")
-		TEXT("Higher values will increase view distance but at a performance cost.\n")
-		TEXT("Default = 1. Value should be in the range [0.0f, 1.0f]."),
-		ECVF_Scalability | ECVF_RenderThreadSafe);
-
-	IConsoleManager::Get().RegisterConsoleVariable(TEXT("r.LightFunctionQuality"),
-		2,
-		TEXT("Defines the light function quality which allows to adjust for quality or performance.\n")
-		TEXT("<=0: off (fastest)\n")
-		TEXT("  1: low quality (e.g. half res with blurring, not yet implemented)\n")
-		TEXT("  2: normal quality (default)\n")
-		TEXT("  3: high quality (e.g. super-sampled or colored, not yet implemented)"),
-		ECVF_Scalability | ECVF_RenderThreadSafe);
-
-	IConsoleManager::Get().RegisterConsoleVariable(TEXT("r.EyeAdaptationQuality"),
-		2,
-		TEXT("Defines the eye adaptation quality which allows to adjust for quality or performance.\n")
-		TEXT("<=0: off (fastest)\n")
-		TEXT("  1: low quality (e.g. non histogram based, not yet implemented)\n")
-		TEXT("  2: normal quality (default)\n")
-		TEXT("  3: high quality (e.g. screen position localized, not yet implemented)"),
-		ECVF_Scalability | ECVF_RenderThreadSafe);
-
-	IConsoleManager::Get().RegisterConsoleVariable(TEXT("r.Shadow.DistanceScale"),
-		1.0f,
-		TEXT("Scalability option to trade shadow distance versus performance for directional lights (clamped within a reasonable range).\n")
-		TEXT("<1: shorter distance\n")
-		TEXT(" 1: normal (default)\n")
-		TEXT(">1: larger distance"),
-		ECVF_Scalability | ECVF_RenderThreadSafe);
-
-	IConsoleManager::Get().RegisterConsoleVariable(TEXT("r.BokehDOFIndexStyle"),
-#if PLATFORM_MAC // Avoid a driver bug on OSX/NV cards that causes driver to generate an unwound index buffer
-		1,
-#else
-		0,
-#endif
-		TEXT("Controls whether to use a packed or unwound index buffer for Bokeh DOF.\n")
-		TEXT("0: Use packed index buffer (faster) (default)\n")
-		TEXT("1: Use unwound index buffer (slower)"),
-		ECVF_ReadOnly | ECVF_RenderThreadSafe);
+static TAutoConsoleVariable<int32> CVarMinLogVerbosity(
+	TEXT("con.MinLogVerbosity"),
+	0,
+	TEXT("Allows to see the log in the in game console (by default deactivated to avoid spam and minor performance loss).\n")
+	TEXT(" 0: no logging other than console response (default)\n")
+	TEXT(" 1: Only fatal errors (no that useful)\n")
+	TEXT(" 2: additionally errors\n")
+	TEXT(" 3: additionally warnings\n")
+	TEXT(" 4: additionally display\n")
+	TEXT(" 5: additionally log\n")
+	TEXT("..\n")
+	TEXT(">=7: all"),
+	ECVF_Default);
+
+static TAutoConsoleVariable<int32> CVarMSAACompositingSampleCount(
+	TEXT("r.MSAA.CompositingSampleCount"),
+	4,
+	TEXT("Affects the render quality of the editor 3d objects.\n")
+	TEXT(" 1: no MSAA, lowest quality\n")
+	TEXT(" 2: 2x MSAA, medium quality (medium GPU memory consumption)\n")
+	TEXT(" 4: 4x MSAA, high quality (high GPU memory consumption)\n")
+	TEXT(" 8: 8x MSAA, very high quality (insane GPU memory consumption)"),
+	ECVF_RenderThreadSafe);
 
+static TAutoConsoleVariable<FString> CVarNetPackageMapDebugObject(
+	TEXT("net.PackageMap.DebugObject"),
+	TEXT(""),
+	TEXT("Debugs PackageMap serialization of object")
+	TEXT("Partial name of object to debug"),
+	ECVF_Default);
+
+static TAutoConsoleVariable<FString> CVarNetReplicationDebugProperty(
+	TEXT("net.Replication.DebugProperty"),
+	TEXT(""),
+	TEXT("Debugs Replication of property by name")
+	TEXT("Partial name of property to debug"),
+	ECVF_Default);
+
+static TAutoConsoleVariable<int32> CVarNetRPCDebug(
+	TEXT("net.RPC.Debug"),
+	0,
+	TEXT("Print all RPC bunches sent over the network\n")
+	TEXT(" 0: no print.\n")
+	TEXT(" 1: Print bunches as they are sent."),
+	ECVF_Default);
 
+static TAutoConsoleVariable<int32> CVarNetMontageDebug(
+	TEXT("net.Montage.Debug"),
+	0,
+	TEXT("Prints Replication information about AnimMontages\n")
+	TEXT(" 0: no print.\n")
+	TEXT(" 1: Print AnimMontage info on client side as they are played."),
+	ECVF_Cheat);
+
+static TAutoConsoleVariable<int32> CVarRenderTargetPoolMin(
+	TEXT("r.RenderTargetPoolMin"),
+	400,
+	TEXT("If the render target pool size (in MB) is below this number there is no deallocation of rendertargets")
+	TEXT("Default is 200 MB."),
+	ECVF_Scalability | ECVF_RenderThreadSafe);
 
-	IConsoleManager::Get().RegisterConsoleVariable(TEXT("r.FreeSkeletalMeshBuffers"),
-		0,
-		TEXT("Controls whether skeletal mesh buffers are kept in CPU memory to support merging of skeletal meshes.\n")
-		TEXT("0: Keep buffers(default)\n")
-		TEXT("1: Free buffers"),
-		ECVF_RenderThreadSafe );
+static TAutoConsoleVariable<float> CVarShadowDistanceFieldPenumbraSize(
+	TEXT("r.Shadow.DistanceFieldPenumbraSize"),
+	0.05f,
+	TEXT("Controls the size of the uniform penumbra produced by static shadowing.\n")
+	TEXT("This is a fraction of MaxTransitionDistanceWorldSpace in BaseLightmass.ini."),
+	ECVF_RenderThreadSafe);
 
-	IConsoleManager::Get().RegisterConsoleVariable(TEXT("r.PreventInvalidMaterialConnections"),
-		1,
-		TEXT("Controls whether users can make connections in the material editor if the system\n")
-		TEXT("determines that they may cause compile errors\n")
-		TEXT("0: Allow all connections\n")
-		TEXT("1: Prevent invalid connections"),
-		ECVF_Cheat );
+static TAutoConsoleVariable<int32> CVarIdleWhenNotForeground(
+	TEXT("t.IdleWhenNotForeground"), 0,
+	TEXT("Prevents the engine from taking any CPU or GPU time while not the foreground app."),
+	ECVF_Cheat);
 
-	// testing code
-	{
-		FConsoleManager& ConsoleManager = (FConsoleManager&)IConsoleManager::Get();
+static TAutoConsoleVariable<int32> CVarSetVSyncEnabled(
+	TEXT("r.VSync"),
+	0,
+	TEXT("0: VSync is disabled.(default)\n")
+	TEXT("1: VSync is enabled."),
+	ECVF_Scalability | ECVF_RenderThreadSafe);
 
-		ConsoleManager.Test();
-	}
-}
+static TAutoConsoleVariable<int32> CVarFinishCurrentFrame(
+	TEXT("r.FinishCurrentFrame"),
+	0,
+	TEXT("If on, the current frame will be forced to finish and render to the screen instead of being buffered.  This will improve latency, but slow down overall performance."),
+	ECVF_Scalability | ECVF_RenderThreadSafe);
+
+static TAutoConsoleVariable<int32> CVarMaxAnistropy(
+	TEXT("r.MaxAnisotropy"),
+	4,
+	TEXT("MaxAnisotropy should range from 1 to 16. Higher values mean better texure quality when using anisotropic filtering but at a cost to performance."),
+	ECVF_Scalability | ECVF_RenderThreadSafe);
+
+static TAutoConsoleVariable<int32> CVarShadowMaxResolution(
+	TEXT("r.Shadow.MaxResolution"),
+	1024,
+	TEXT("Max square dimensions (in texels) allowed for rendering shadow depths. Range 4 to hardware limit. Higher = better quality shadows but at a performance cost."),
+	ECVF_Scalability | ECVF_RenderThreadSafe);
 
+static TAutoConsoleVariable<float> CVarShadowCSMTransitionScale(
+	TEXT("r.Shadow.CSM.TransitionScale"),
+	1.0f,
+	TEXT("Allows to scale the cascaded shadow map transition region. Clamped within 0..2.\n")
+	TEXT("0: no transition (fastest)\n")
+	TEXT("1: as specific in the light settings (default)\n")
+	TEXT("2: 2x larger than what was specified in the light"),
+	ECVF_Scalability | ECVF_RenderThreadSafe);
+
+static TAutoConsoleVariable<float> CVarMobileContentScaleFactor(
+	TEXT("r.MobileContentScaleFactor"),
+	1.0f,
+	TEXT("Content scale multiplier (equates to iOS's contentScaleFactor to support Retina displays"),
+	ECVF_Default);
+
+// this cvar can be removed in shipping to not compile shaders for development (faster)
+static TAutoConsoleVariable<int32> CVarCompileShadersForDevelopment(
+	TEXT("r.CompileShadersForDevelopment"),
+	1,
+	TEXT("Setting this to 0 allows to ship a game with more optimized shaders as some\n")
+	TEXT("editor and development features are not longer compiled into the shaders.\n")
+	TEXT(" Note: This should be done when shipping but it's not done automatically yet (feature need to mature\n")
+	TEXT("       and shaders will compile slower as shader caching from development isn't shared).\n")
+	TEXT("Cannot be changed at runtime - can be put into BaseEngine.ini\n")
+	TEXT(" 0: off, shader can run a bit faster\n")
+	TEXT(" 1: on (Default)"),
+	ECVF_ReadOnly);
+
+static TAutoConsoleVariable<int32> CVarPS4MixeedModeShaderDebugInfo(
+	TEXT("r.PS4MixedModeShaderDebugInfo"),
+	0,
+	TEXT("Whether to compile shaders to allow mixed mode shader debugging. This will currently generate slower code.\n")
+	TEXT(" 0: Normal mode\n")
+	TEXT(" 1: Mixed mode)"),
+	ECVF_ReadOnly);
+
+static TAutoConsoleVariable<int32> CVarPS4DumpShaderSDB(
+	TEXT("r.PS4DumpShaderSDB"),
+	0,
+	TEXT("Whether to dump shader sdb files used for shader association.\n")
+	TEXT(" 0: Disabled\n")
+	TEXT(" 1: Enabled)"),
+	ECVF_ReadOnly);
+
+static TAutoConsoleVariable<int32> CVarDontLimitOnBattery(
+	TEXT("r.DontLimitOnBattery"),
+	0,
+	TEXT("0: Limit performance on devices with a battery.(default)\n")
+	TEXT("1: Do not limit performance due to device having a battery."),
+	ECVF_Scalability | ECVF_RenderThreadSafe);
+
+static TAutoConsoleVariable<float> CVarViewDistanceScale(
+	TEXT("r.ViewDistanceScale"),
+	1.0f,
+	TEXT("Controls the view distance scale. A primitive's MaxDrawDistance is scaled by this value.\n")
+	TEXT("Higher values will increase view distance but at a performance cost.\n")
+	TEXT("Default = 1. Value should be in the range [0.0f, 1.0f]."),
+	ECVF_Scalability | ECVF_RenderThreadSafe);
+
+static TAutoConsoleVariable<int32> CVarLightFunctionQuality(
+	TEXT("r.LightFunctionQuality"),
+	2,
+	TEXT("Defines the light function quality which allows to adjust for quality or performance.\n")
+	TEXT("<=0: off (fastest)\n")
+	TEXT("  1: low quality (e.g. half res with blurring, not yet implemented)\n")
+	TEXT("  2: normal quality (default)\n")
+	TEXT("  3: high quality (e.g. super-sampled or colored, not yet implemented)"),
+	ECVF_Scalability | ECVF_RenderThreadSafe);
+
+static TAutoConsoleVariable<int32> CVarEyeAdaptationQuality(
+	TEXT("r.EyeAdaptationQuality"),
+	2,
+	TEXT("Defines the eye adaptation quality which allows to adjust for quality or performance.\n")
+	TEXT("<=0: off (fastest)\n")
+	TEXT("  1: low quality (e.g. non histogram based, not yet implemented)\n")
+	TEXT("  2: normal quality (default)\n")
+	TEXT("  3: high quality (e.g. screen position localized, not yet implemented)"),
+	ECVF_Scalability | ECVF_RenderThreadSafe);
+
+static TAutoConsoleVariable<float> CVarShadowDistanceScale(
+	TEXT("r.Shadow.DistanceScale"),
+	1.0f,
+	TEXT("Scalability option to trade shadow distance versus performance for directional lights (clamped within a reasonable range).\n")
+	TEXT("<1: shorter distance\n")
+	TEXT(" 1: normal (default)\n")
+	TEXT(">1: larger distance"),
+	ECVF_Scalability | ECVF_RenderThreadSafe);
+
+static TAutoConsoleVariable<int32> CVarFreeSkeletalMeshBuffers(
+	TEXT("r.FreeSkeletalMeshBuffers"),
+	0,
+	TEXT("Controls whether skeletal mesh buffers are kept in CPU memory to support merging of skeletal meshes.\n")
+	TEXT("0: Keep buffers(default)\n")
+	TEXT("1: Free buffers"),
+	ECVF_RenderThreadSafe);
 
 static TAutoConsoleVariable<int32> CVarTonemapperQuality(
 	TEXT("r.TonemapperQuality"),
@@ -2109,32 +1795,8 @@ static TAutoConsoleVariable<int32> CVarRefractionQuality(
 	TEXT("  3: high quality (e.g. color fringe, not yet implemented)"),
 	ECVF_Scalability | ECVF_RenderThreadSafe);
 
-static TAutoConsoleVariable<float> CVarSSSSS(
-	TEXT("r.SSSSS"),
-	0.0f,
-	TEXT("Very experimental screen space subsurface scattering on non unlit/lit materials\n")
-	TEXT("0: off\n")
-	TEXT("x: SSS radius in world space e.g. 10"),
-	ECVF_RenderThreadSafe);
 
-static TAutoConsoleVariable<int32> CVarTriangleOrderOptimization(
-	TEXT("r.TriangleOrderOptimization"),
-	1,
-	TEXT("Controls the algorithm to use when optimizing the triangle order for the post-transform cache.\n")
-	TEXT("0: Use NVTriStrip (slower)\n")
-	TEXT("1: Use Forsyth algorithm (fastest)(default)"),
-	ECVF_ReadOnly);	
-
-static TAutoConsoleVariable<float> CVarFastBlurThreshold(
-	TEXT("r.FastBlurThreshold"),
-	7.0f,
-	TEXT("Defines at what radius the Gaussian blur optimization kicks in (estimated 25% - 40% faster).\n")
-	TEXT("The optimzation uses slightly less memory and has a quality loss on smallblur radius.\n")
-	TEXT("  0: use the optimization always (fastest, lowest quality)\n")
-	TEXT("  3: use the optimization starting at a 3 pixel radius (quite fast)\n")
-	TEXT("  7: use the optimization starting at a 7 pixel radius (default)\n")
-	TEXT(">15: barely ever use the optimization (high quality)"),
-	ECVF_Scalability | ECVF_RenderThreadSafe);
+
 
 static TAutoConsoleVariable<float> CVarUIBlurRadius(
 	TEXT("UI.BlurRadius"),
@@ -2144,15 +1806,6 @@ static TAutoConsoleVariable<float> CVarUIBlurRadius(
 	TEXT("0 to turn off gaussian blur. (Default 1.2)"),
 	ECVF_Scalability | ECVF_RenderThreadSafe);
 
-static TAutoConsoleVariable<int32> CVarDrawRectangleOptimization(
-	TEXT("r.DrawRectangleOptimization"),
-	1,
-	TEXT("Controls an optimization for DrawRectangle(). When enabled a triangle can be used to draw a quad in certain situations (viewport sized quad).\n")
-	TEXT("Using a triangle allows for slightly faster post processing in lower resolutions but can not always be used.\n")
-	TEXT(" 0: Optimization is disabled, DrawDenormalizedQuad always render with quad\n")
-	TEXT(" 1: Optimization is enabled, a triangle can be rendered where specified (default)"),
-	ECVF_RenderThreadSafe);
-
 static TAutoConsoleVariable<int32> CVarDBuffer(
 	TEXT("r.DBuffer"),
 	0,
@@ -2168,11 +1821,7 @@ static TAutoConsoleVariable<float> CVarSkeletalMeshLODRadiusScale(
 	TEXT("Scale factor for the screen radius used in computing discrete LOD for skeletal meshes. (0.25-1)"),
 	ECVF_Scalability | ECVF_RenderThreadSafe);
 
-static TAutoConsoleVariable<float> CVarStaticMeshLODDistanceScale(
-	TEXT("r.StaticMeshLODDistanceScale"),
-	1.0f,
-	TEXT("Scale factor for the distance used in computing discrete LOD for static meshes. (0.25-1)"),
-	ECVF_Scalability | ECVF_RenderThreadSafe);
+
 
 static TAutoConsoleVariable<int32> CVarOnlyStreamInTextures(
 	TEXT("r.OnlyStreamInTextures"),
@@ -2192,18 +1841,6 @@ static TAutoConsoleVariable<int32> CVarAllowHighQualityLightMaps(
 	TEXT("If set to 1, allow high quality lightmaps which don't bake in direct lighting of stationary lights"),
 	ECVF_RenderThreadSafe);
 
-static TAutoConsoleVariable<float> CVarSetMipMapLODBias(
-	TEXT("r.MipMapLODBias"),
-	0.0f,
-	TEXT("Apply additional mip map bias for all 2D textures, range of -15.0 to 15.0"),
-	ECVF_RenderThreadSafe);
-
-static TAutoConsoleVariable<int32> CVarVirtualTextureEnabled(
-	TEXT("r.VirtualTexture"),
-	1,
-	TEXT("If set to 1, textures will use virtual memory so they can be partially resident."),
-	ECVF_RenderThreadSafe);
-
 static TAutoConsoleVariable<int32> CVarVirtualTextureReducedMemoryEnabled(
 	TEXT("r.VirtualTextureReducedMemory"),
 	0,
@@ -2215,15 +1852,3 @@ static TAutoConsoleVariable<int32> CVarPrecomputedVisibilityWarning(
 	0,
 	TEXT("If set to 1, a warning will be displayed when rendering a scene from a view point without precomputed visibility."),
 	ECVF_RenderThreadSafe);
-
-static TAutoConsoleVariable<float> CVarDemosaicVposOffset(
-	TEXT("r.DemosaicVposOffset"),
-	0.0f,
-	TEXT("This offset is added to the rasterized position used for demosaic in the ES2 tonemapping shader. It exists to workaround driver bugs on some Android devices that have a half-pixel offset."),
-	ECVF_RenderThreadSafe);
-
-static TAutoConsoleVariable<int32> CVarRenderLastFrameInStreamingPause(
-	TEXT("r.RenderLastFrameInStreamingPause"),
-	1,
-	TEXT("If 1 the previous frame is displayed during streaming pause. If zero the screen is left black."),
-	ECVF_RenderThreadSafe);
\ No newline at end of file
diff --git a/Engine/Source/Runtime/Core/Public/HAL/ConsoleManager.h b/Engine/Source/Runtime/Core/Public/HAL/ConsoleManager.h
--- a/Engine/Source/Runtime/Core/Public/HAL/ConsoleManager.h
+++ b/Engine/Source/Runtime/Core/Public/HAL/ConsoleManager.h
@@ -49,7 +49,7 @@ public:
 
 	virtual IConsoleVariable* RegisterConsoleVariable(const TCHAR* Name, int32 DefaultValue, const TCHAR* Help, uint32 Flags) override;
 	virtual IConsoleVariable* RegisterConsoleVariable(const TCHAR* Name, float DefaultValue, const TCHAR* Help, uint32 Flags) override;
-	virtual IConsoleVariable* RegisterConsoleVariable(const TCHAR* Name, const TCHAR *DefaultValue, const TCHAR* Help, uint32 Flags) override;
+	virtual IConsoleVariable* RegisterConsoleVariable(const TCHAR* Name, const FString& DefaultValue, const TCHAR* Help, uint32 Flags) override;
 	virtual IConsoleVariable* RegisterConsoleVariableRef(const TCHAR* Name, int32& RefValue, const TCHAR* Help, uint32 Flags) override;
 	virtual IConsoleVariable* RegisterConsoleVariableRef(const TCHAR* Name, float& RefValue, const TCHAR* Help, uint32 Flags) override;
 	virtual IConsoleVariable* RegisterConsoleVariableBitRef(const TCHAR* CVarName, const TCHAR* FlagName, uint32 BitNumber, uint8* Force0MaskPtr, uint8* Force1MaskPtr, const TCHAR* Help, uint32 Flags) override;
diff --git a/Engine/Source/Runtime/Core/Public/HAL/IConsoleManager.h b/Engine/Source/Runtime/Core/Public/HAL/IConsoleManager.h
--- a/Engine/Source/Runtime/Core/Public/HAL/IConsoleManager.h
+++ b/Engine/Source/Runtime/Core/Public/HAL/IConsoleManager.h
@@ -172,6 +172,11 @@ public:
 		return 0; 
 	}
 
+	virtual class TConsoleVariableData<FString>* AsVariableString()
+	{
+		return 0;
+	}
+
 	/**
 	 * Casts this object to an IConsoleCommand, verifying first that it is safe to do so
 	 */
@@ -298,7 +303,7 @@ struct CORE_API IConsoleManager
 	 * @param Help must not be 0
 	 * @param Flags bitmask combined from EConsoleVariableFlags
 	 */
-	virtual IConsoleVariable* RegisterConsoleVariable(const TCHAR* Name, const TCHAR* DefaultValue, const TCHAR* Help, uint32 Flags = ECVF_Default) = 0;
+	virtual IConsoleVariable* RegisterConsoleVariable(const TCHAR* Name, const FString& DefaultValue, const TCHAR* Help, uint32 Flags = ECVF_Default) = 0;
 	/**
 	 * Create a reference to a int console variable
 	 * @param Name must not be 0
@@ -728,12 +733,12 @@ class TAutoConsoleVariable : public FAutoConsoleObject
 {
 public:
 	/**
-	 * Create a float or int console variable
+	 * Create a float, int or string console variable
 	 * @param Name must not be 0
 	 * @param Help must not be 0
 	 * @param Flags bitmask combined from EConsoleVariableFlags
 	 */
-	TAutoConsoleVariable(const TCHAR* Name, T DefaultValue, const TCHAR* Help, uint32 Flags = ECVF_Default);
+	TAutoConsoleVariable(const TCHAR* Name, const T& DefaultValue, const TCHAR* Help, uint32 Flags = ECVF_Default);
 
 	T GetValueOnGameThread() const
 	{
@@ -755,19 +760,26 @@ private:
 };
 
 template <>
-inline TAutoConsoleVariable<int32>::TAutoConsoleVariable(const TCHAR* Name, int32 DefaultValue, const TCHAR* Help, uint32 Flags)
+inline TAutoConsoleVariable<int32>::TAutoConsoleVariable(const TCHAR* Name, const int32& DefaultValue, const TCHAR* Help, uint32 Flags)
 	: FAutoConsoleObject(IConsoleManager::Get().RegisterConsoleVariable(Name, DefaultValue, Help, Flags))
 {
 	Ref = AsVariable()->AsVariableInt();
 }
 
 template <>
-inline TAutoConsoleVariable<float>::TAutoConsoleVariable(const TCHAR* Name, float DefaultValue, const TCHAR* Help, uint32 Flags)
+inline TAutoConsoleVariable<float>::TAutoConsoleVariable(const TCHAR* Name, const float& DefaultValue, const TCHAR* Help, uint32 Flags)
 	: FAutoConsoleObject(IConsoleManager::Get().RegisterConsoleVariable(Name, DefaultValue, Help, Flags))
 {
 	Ref = AsVariable()->AsVariableFloat();
 }
 
+template <>
+inline TAutoConsoleVariable<FString>::TAutoConsoleVariable(const TCHAR* Name, const FString& DefaultValue, const TCHAR* Help, uint32 Flags)
+	: FAutoConsoleObject(IConsoleManager::Get().RegisterConsoleVariable(Name, DefaultValue, Help, Flags))
+{
+	Ref = AsVariable()->AsVariableString();
+}
+
 /**
  * Autoregistering console command
  */
diff --git a/Engine/Source/Runtime/Engine/Private/CharacterMovementComponent.cpp b/Engine/Source/Runtime/Engine/Private/CharacterMovementComponent.cpp
--- a/Engine/Source/Runtime/Engine/Private/CharacterMovementComponent.cpp
+++ b/Engine/Source/Runtime/Engine/Private/CharacterMovementComponent.cpp
@@ -26,6 +26,45 @@ const float UCharacterMovementComponent::MAX_FLOOR_DIST = 2.4f;
 const float UCharacterMovementComponent::BRAKE_TO_STOP_VELOCITY = 10.f;
 const float UCharacterMovementComponent::SWEEP_EDGE_REJECT_DISTANCE = 0.15f;
 
+// CVars
+static TAutoConsoleVariable<int32> CVarNetEnableMoveCombining(
+	TEXT("p.NetEnableMoveCombining"),
+	1,
+	TEXT("Whether to enable move combining on the client to reduce bandwidth by combining similar moves.\n")
+	TEXT("0: Disable, 1: Enable"),
+	ECVF_Default);
+
+static TAutoConsoleVariable<float> CVarNetProxyShrinkRadius(
+	TEXT("p.NetProxyShrinkRadius"),
+	0.01f,
+	TEXT("Shrink simulated proxy capsule radius by this amount, to account for network rounding that may cause encroachment.\n")
+	TEXT("Changing this value at runtime may require the proxy to re-join for correct behavior.\n")
+	TEXT("<= 0: disabled, > 0: shrink by this amount."),
+	ECVF_Default);
+
+static TAutoConsoleVariable<float> CVarNetProxyShrinkHalfHeight(
+	TEXT("p.NetProxyShrinkHalfHeight"),
+	0.01f,
+	TEXT("Shrink simulated proxy capsule half height by this amount, to account for network rounding that may cause encroachment.\n")
+	TEXT("Changing this value at runtime may require the proxy to re-join for correct behavior.\n")
+	TEXT("<= 0: disabled, > 0: shrink by this amount."),
+	ECVF_Default);
+
+#if !(UE_BUILD_SHIPPING || UE_BUILD_TEST)
+static TAutoConsoleVariable<int32> CvarNetShowCorrections(
+	TEXT("p.NetShowCorrections"),
+	0,
+	TEXT("Whether to draw client position corrections (red is incorrect, green is corrected).\n")
+	TEXT("0: Disable, 1: Enable"),
+	ECVF_Cheat);
+
+static TAutoConsoleVariable<float> CVarNetCorrectionLifetime(
+	TEXT("p.NetCorrectionLifetime"),
+	4.f,
+	TEXT("How long a visualized network correction persists.\n")
+	TEXT("Time in seconds each visualized network correction persists."),
+	ECVF_Cheat);
+#endif // !(UE_BUILD_SHIPPING || UE_BUILD_TEST)
 
 // Version that does not use inverse sqrt estimate, for higher precision.
 FORCEINLINE FVector SafeNormalPrecise(const FVector& V)
@@ -806,11 +845,9 @@ void UCharacterMovementComponent::AdjustProxyCapsuleSize()
 	if (bShrinkProxyCapsule && CharacterOwner && CharacterOwner->Role == ROLE_SimulatedProxy)
 	{
 		bShrinkProxyCapsule = false;
-		static const auto CVarShrinkRadius = IConsoleManager::Get().FindTConsoleVariableDataFloat(TEXT("p.NetProxyShrinkRadius"));
-		static const auto CVarShrinkHalfHeight = IConsoleManager::Get().FindTConsoleVariableDataFloat(TEXT("p.NetProxyShrinkHalfHeight"));
 
-		float ShrinkRadius = CVarShrinkRadius ? FMath::Max(0.f, CVarShrinkRadius->GetValueOnGameThread()) : 0.f;
-		float ShrinkHalfHeight = CVarShrinkHalfHeight ? FMath::Max(0.f, CVarShrinkHalfHeight->GetValueOnGameThread()) : 0.f;
+		float ShrinkRadius = FMath::Max(0.f, CVarNetProxyShrinkRadius.GetValueOnGameThread());
+		float ShrinkHalfHeight = FMath::Max(0.f, CVarNetProxyShrinkHalfHeight.GetValueOnGameThread());
 
 		if (ShrinkRadius == 0.f && ShrinkHalfHeight == 0.f)
 		{
@@ -5024,8 +5061,7 @@ void UCharacterMovementComponent::ReplicateMoveToServer(float DeltaTime, const F
 	// Add NewMove to the list
 	ClientData->SavedMoves.Push(NewMove);
 
-	static const auto CVarCombine = IConsoleManager::Get().FindTConsoleVariableDataInt(TEXT("p.NetEnableMoveCombining"));
-	const bool bCanCombine = (!CVarCombine || CVarCombine->GetValueOnGameThread() != 0);
+	const bool bCanCombine = (CVarNetEnableMoveCombining.GetValueOnGameThread() != 0);
 
 	if (bCanCombine && ClientData->PendingMove.IsValid() == false)
 	{
@@ -5330,12 +5366,10 @@ void UCharacterMovementComponent::ServerMoveHandleClientError(float TimeStamp, f
 		ServerData->PendingAdjustment.NewRot = CharacterOwner->GetActorRotation();
 
 #if !(UE_BUILD_SHIPPING || UE_BUILD_TEST)
-		static const auto CVarShowCorrections = IConsoleManager::Get().FindTConsoleVariableDataInt(TEXT("p.NetShowCorrections"));
-		static const auto CVarCorrectionLifetime = IConsoleManager::Get().FindTConsoleVariableDataFloat(TEXT("p.NetCorrectionLifetime"));
-		if (CVarShowCorrections && CVarShowCorrections->GetValueOnGameThread() != 0)
+		if (CvarNetShowCorrections.GetValueOnGameThread() != 0)
 		{
 			UE_LOG(LogNetPlayerMovement, Warning, TEXT("******** Client Error at %f is %f Accel %s LocDiff %s ClientLoc %s, ServerLoc: %s, Base: %s"), TimeStamp, LocDiff.Size(), *Accel.ToString(), *LocDiff.ToString(), *ClientLoc.ToString(), *CharacterOwner->GetActorLocation().ToString(), *GetNameSafe(MovementBase));
-			const float DebugLifetime = CVarCorrectionLifetime ? CVarCorrectionLifetime->GetValueOnGameThread() : 1.f;
+			const float DebugLifetime = CVarNetCorrectionLifetime.GetValueOnGameThread();
 			DrawDebugCapsule(GetWorld(), CharacterOwner->GetActorLocation(), CharacterOwner->GetSimpleCollisionHalfHeight(), CharacterOwner->GetSimpleCollisionRadius(), FQuat::Identity, FColor(100, 255, 100), true, DebugLifetime);
 			DrawDebugCapsule(GetWorld(), ClientLoc                    , CharacterOwner->GetSimpleCollisionHalfHeight(), CharacterOwner->GetSimpleCollisionRadius(), FQuat::Identity, FColor(255, 100, 100), true, DebugLifetime);
 		}
diff --git a/Engine/Source/Runtime/Engine/Private/ContentStreaming.cpp b/Engine/Source/Runtime/Engine/Private/ContentStreaming.cpp
--- a/Engine/Source/Runtime/Engine/Private/ContentStreaming.cpp
+++ b/Engine/Source/Runtime/Engine/Private/ContentStreaming.cpp
@@ -128,6 +128,28 @@ static TAutoConsoleVariable<float> CVarStreamingBoost(
 	ECVF_Default
 	);
 
+#if PLATFORM_SUPPORTS_TEXTURE_STREAMING
+static TAutoConsoleVariable<int32> CVarSetTextureStreamingEnabled(
+	TEXT("r.TextureStreaming"),
+	1,
+	TEXT("Allows to define if texture streaming is enabled, can be changed at run time.\n")
+	TEXT("0: off\n")
+	TEXT("1: on (default)"),
+	ECVF_Default | ECVF_RenderThreadSafe);
+#endif
+
+static TAutoConsoleVariable<int32> CvarStreamingUseFixedPoolSize(
+	TEXT("r.Streaming.UseFixedPoolSize"),
+	0,
+	TEXT("If non-zero, do not allow the pool size to change at run time."),
+	ECVF_ReadOnly);
+
+static TAutoConsoleVariable<int32> CVarStreamingPoolSize(
+	TEXT("r.Streaming.PoolSize"),
+	-1,
+	TEXT("-1: Default texture pool size, otherwise the size in MB\n"),
+	ECVF_Scalability);
+
 /** Streaming priority: Linear distance factor from 0 to MAX_STREAMINGDISTANCE. */
 #define MAX_STREAMINGDISTANCE	10000.0f
 #define MAX_MIPDELTA			5.0f
@@ -1556,8 +1578,7 @@ FStreamingManagerCollection::FStreamingManagerCollection()
 	// Disable texture streaming if that was requested (needs to happen before the call to ProcessNewlyLoadedUObjects, as that can load textures)
 	if( FParse::Param( FCommandLine::Get(), TEXT( "NoTextureStreaming" ) ) )
 	{
-		static IConsoleVariable* CVar = IConsoleManager::Get().FindConsoleVariable(TEXT("r.TextureStreaming"));
-		CVar->Set(0);
+		CVarSetTextureStreamingEnabled.AsVariable()->Set(0);
 	}
 #endif
 
@@ -1970,10 +1991,7 @@ void FStreamingManagerCollection::AddOrRemoveTextureStreamingManagerIfNeeded(boo
 
 #if PLATFORM_SUPPORTS_TEXTURE_STREAMING
 	{
-		static const auto CVar = IConsoleManager::Get().FindTConsoleVariableDataInt(TEXT("r.TextureStreaming"));
-		check(CVar);
-
-		bUseTextureStreaming = CVar->GetValueOnGameThread() != 0;
+		bUseTextureStreaming = CVarSetTextureStreamingEnabled.GetValueOnGameThread() != 0;
 	}
 
 	if( !GRHISupportsTextureStreaming || IsRunningDedicatedServer() )
@@ -3638,11 +3656,8 @@ void FStreamingManagerTexture::StreamTexturesUnlimited( FStreamingContext& Conte
 
 void FStreamingManagerTexture::CheckUserSettings()
 {	
-	static const auto CVarStreamingTexturePoolSize = IConsoleManager::Get().FindTConsoleVariableDataInt(TEXT("r.Streaming.PoolSize"));
-	int32 PoolSizeSetting = CVarStreamingTexturePoolSize ? CVarStreamingTexturePoolSize->GetValueOnGameThread() : -1;
-
-	static const auto CVarUseFixedStreamingTexturePoolSize = IConsoleManager::Get().FindTConsoleVariableDataInt(TEXT("r.Streaming.UseFixedPoolSize"));
-	int32 FixedPoolSizeSetting = CVarUseFixedStreamingTexturePoolSize ? CVarUseFixedStreamingTexturePoolSize->GetValueOnGameThread() : -1;
+	int32 PoolSizeSetting = CVarStreamingPoolSize.GetValueOnGameThread();
+	int32 FixedPoolSizeSetting = CvarStreamingUseFixedPoolSize.GetValueOnGameThread();
 
 	if (FixedPoolSizeSetting == 0)
 	{
diff --git a/Engine/Source/Runtime/Engine/Private/DataChannel.cpp b/Engine/Source/Runtime/Engine/Private/DataChannel.cpp
--- a/Engine/Source/Runtime/Engine/Private/DataChannel.cpp
+++ b/Engine/Source/Runtime/Engine/Private/DataChannel.cpp
@@ -20,6 +20,15 @@ DEFINE_LOG_CATEGORY_STATIC(LogNetPartialBunch, Warning, All);
 
 extern FAutoConsoleVariable CVarDoReplicationContextString;
 
+static TAutoConsoleVariable<int32> CVarNetReliableDebug(
+	TEXT("net.Reliable.Debug"),
+	0,
+	TEXT("Print all reliable bunches sent over the network\n")
+	TEXT(" 0: no print.\n")
+	TEXT(" 1: Print bunches as they are sent.\n")
+	TEXT(" 2: Print reliable bunch buffer each net update"),
+	ECVF_Default);
+
 /*-----------------------------------------------------------------------------
 	UChannel implementation.
 -----------------------------------------------------------------------------*/
@@ -747,19 +756,15 @@ FOutBunch* UChannel::PrepBunch(FOutBunch* Bunch, FOutBunch* OutBunch, bool Merge
 		Connection->LastOutBunch = OutBunch;
 
 #if !(UE_BUILD_SHIPPING || UE_BUILD_TEST)
-		static const auto CVar = IConsoleManager::Get().FindTConsoleVariableDataInt(TEXT("net.Reliable.Debug"));
-		if (CVar)
+		if (CVarNetReliableDebug.GetValueOnGameThread() == 1)
 		{
-			if (CVar->GetValueOnGameThread() == 1)
-			{
-				UE_LOG(LogNetTraffic, Warning, TEXT("%s. Reliable: %s"), *Describe(), *Bunch->DebugString);
-			}
-			if (CVar->GetValueOnGameThread() == 2)
-			{
-				UE_LOG(LogNetTraffic, Warning, TEXT("%s. Reliable: %s"), *Describe(), *Bunch->DebugString);
-				PrintReliableBunchBuffer();
-				UE_LOG(LogNetTraffic, Warning, TEXT(""));
-			}
+			UE_LOG(LogNetTraffic, Warning, TEXT("%s. Reliable: %s"), *Describe(), *Bunch->DebugString);
+		}
+		if (CVarNetReliableDebug.GetValueOnGameThread() == 2)
+		{
+			UE_LOG(LogNetTraffic, Warning, TEXT("%s. Reliable: %s"), *Describe(), *Bunch->DebugString);
+			PrintReliableBunchBuffer();
+			UE_LOG(LogNetTraffic, Warning, TEXT(""));
 		}
 #endif
 
@@ -1563,8 +1568,7 @@ bool UActorChannel::ReplicateActor()
 	}
 
 #if !(UE_BUILD_SHIPPING || UE_BUILD_TEST)
-	static const auto CVar = IConsoleManager::Get().FindTConsoleVariableDataInt(TEXT("net.Reliable.Debug"));
-	if ( CVar && CVar->GetValueOnGameThread() > 0 )
+	if (CVarNetReliableDebug.GetValueOnGameThread() > 0)
 	{
 		Bunch.DebugString = FString::Printf(TEXT("%.2f ActorBunch: %s"), Connection->Driver->Time, *Actor->GetName() );
 	}
diff --git a/Engine/Source/Runtime/Engine/Private/DebugCameraHUD.cpp b/Engine/Source/Runtime/Engine/Private/DebugCameraHUD.cpp
--- a/Engine/Source/Runtime/Engine/Private/DebugCameraHUD.cpp
+++ b/Engine/Source/Runtime/Engine/Private/DebugCameraHUD.cpp
@@ -8,6 +8,14 @@
 #include "GameFramework/SpectatorPawn.h"
 #include "Engine/DebugCameraHUD.h"
 
+#if !(UE_BUILD_SHIPPING || UE_BUILD_TEST)
+static TAutoConsoleVariable<int32> CVarDebugCameraTraceComplex(
+	TEXT("g.DebugCameraTraceComplex"),
+	1,
+	TEXT("Whether DebugCamera should use complex or simple collision for the line trace.\n")
+	TEXT("1: complex collision, 0: simple collision "),
+	ECVF_Cheat);
+#endif
 // ------------------
 // Externals
 // ------------------
@@ -88,8 +96,11 @@ void ADebugCameraHUD::PostRender()
 			//Canvas->DrawText(FString::Printf(TEXT("CamLoc:%s CamRot:%s"), *CamLoc.ToString(), *CamRot.ToString() ));
 
 			const TCHAR* CVarComplexName = TEXT("g.DebugCameraTraceComplex");
-			static const auto CVarComplex = IConsoleManager::Get().FindTConsoleVariableDataInt(CVarComplexName);
-			const bool bTraceComplex = (CVarComplex ? (CVarComplex->GetValueOnGameThread() != 0) : true);
+			bool bTraceComplex = true;
+
+#if !(UE_BUILD_SHIPPING || UE_BUILD_TEST)
+			bTraceComplex = CVarDebugCameraTraceComplex.GetValueOnGameThread() != 0;
+#endif
 
 			FCollisionQueryParams TraceParams(NAME_None, bTraceComplex, this);
 			FHitResult Hit;
diff --git a/Engine/Source/Runtime/Engine/Private/GameViewportClient.cpp b/Engine/Source/Runtime/Engine/Private/GameViewportClient.cpp
--- a/Engine/Source/Runtime/Engine/Private/GameViewportClient.cpp
+++ b/Engine/Source/Runtime/Engine/Private/GameViewportClient.cpp
@@ -42,6 +42,15 @@ FViewportClient::ESoundShowFlags::Type UGameViewportClient::SoundShowFlags = FVi
 
 DEFINE_STAT(STAT_UIDrawingTime);
 
+static TAutoConsoleVariable<int32> CVarSetBlackBordersEnabled(
+	TEXT("r.BlackBorders"),
+	1,
+	TEXT("To draw black borders around the rendered image\n")
+	TEXT("(prevents artifacts from post processing passes that read outside of the image e.g. PostProcessAA)\n")
+	TEXT("in pixels, 0:off"),
+	ECVF_Default);
+
+
 /**
  * Draw debug info on a game scene view.
  */
@@ -826,8 +835,7 @@ void UGameViewportClient::Draw(FViewport* InViewport, FCanvas* SceneCanvas)
 
 		// To draw black borders around the rendered image (prevents artifacts from post processing passes that read outside of the image e.g. PostProcessAA)
 		{
-			static const auto CVar = IConsoleManager::Get().FindTConsoleVariableDataInt(TEXT("r.BlackBorders"));
-			int32 BlackBorders = FMath::Clamp(CVar->GetValueOnGameThread(), 0, 10);
+			int32 BlackBorders = FMath::Clamp(CVarSetBlackBordersEnabled.GetValueOnGameThread(), 0, 10);
 
 			if(ViewFamily.Views.Num() == 1 && BlackBorders)
 			{
diff --git a/Engine/Source/Runtime/Engine/Private/LevelActor.cpp b/Engine/Source/Runtime/Engine/Private/LevelActor.cpp
--- a/Engine/Source/Runtime/Engine/Private/LevelActor.cpp
+++ b/Engine/Source/Runtime/Engine/Private/LevelActor.cpp
@@ -10,6 +10,14 @@
 	#include "PhysicsEngine/PhysXSupport.h"
 #endif // WITH_PHYSX
 
+// CVars
+static TAutoConsoleVariable<float> CVarEncroachEpsilon(
+	TEXT("p.EncroachEpsilon"),
+	0.15f,
+	TEXT("Epsilon value encroach checking for capsules\n")
+	TEXT("0: use full sized capsule. > 0: shrink capsule size by this amount (world units)"),
+	ECVF_Default);
+
 
 #if LINE_CHECK_TRACING
 
@@ -670,8 +678,7 @@ bool UWorld::FindTeleportSpot(AActor* TestActor, FVector& TestLocation, FRotator
 
 bool UWorld::EncroachingBlockingGeometry(AActor* TestActor, FVector TestLocation, FRotator TestRotation, FVector* ProposedAdjustment )
 {
-	static const auto CVar = IConsoleManager::Get().FindTConsoleVariableDataFloat(TEXT("p.EncroachEpsilon"));
-	float Epsilon = CVar ? CVar->GetValueOnGameThread() : 0.f;
+	float Epsilon = CVarEncroachEpsilon.GetValueOnGameThread();
 
 	// currently just tests RootComponent.  @TODO FIXME should test all colliding components?  Cost/benefit?
 	UPrimitiveComponent* PrimComp = TestActor ? Cast<UPrimitiveComponent>(TestActor->GetRootComponent()) : NULL;
diff --git a/Engine/Source/Runtime/Engine/Private/MovementComponent.cpp b/Engine/Source/Runtime/Engine/Private/MovementComponent.cpp
--- a/Engine/Source/Runtime/Engine/Private/MovementComponent.cpp
+++ b/Engine/Source/Runtime/Engine/Private/MovementComponent.cpp
@@ -293,6 +293,11 @@ bool UMovementComponent::SafeMoveUpdatedComponent(const FVector& Delta, const FR
 	return bMoveResult;
 }
 
+static TAutoConsoleVariable<float> CVarPenetrationPullbackDistance(TEXT("p.PenetrationPullbackDistance"),
+	0.125f,
+	TEXT("Pull out from penetration of an object by this extra distance.\n")
+	TEXT("Distance added to penetration fix-ups."),
+	ECVF_Default);
 
 FVector UMovementComponent::GetPenetrationAdjustment(const FHitResult& Hit) const
 {
@@ -302,8 +307,7 @@ FVector UMovementComponent::GetPenetrationAdjustment(const FHitResult& Hit) cons
 	}
 
 	FVector Result;
-	static const auto CVar = IConsoleManager::Get().FindTConsoleVariableDataFloat(TEXT("p.PenetrationPullbackDistance"));
-	const float PullBackDistance = (CVar ? FMath::Abs(CVar->GetValueOnGameThread()) : 0.125f);
+	const float PullBackDistance = FMath::Abs(CVarPenetrationPullbackDistance.GetValueOnGameThread());
 	const float PenetrationDepth = (Hit.PenetrationDepth > 0.f ? Hit.PenetrationDepth : 0.125f);
 
 	Result = Hit.Normal * (PenetrationDepth + PullBackDistance);
@@ -311,7 +315,6 @@ FVector UMovementComponent::GetPenetrationAdjustment(const FHitResult& Hit) cons
 	return ConstrainDirectionToPlane(Result);
 }
 
-
 bool UMovementComponent::ResolvePenetration(const FVector& ProposedAdjustment, const FHitResult& Hit, const FRotator& NewRotation)
 {
 	const FVector Adjustment = ConstrainDirectionToPlane(ProposedAdjustment);
diff --git a/Engine/Source/Runtime/Engine/Private/NetworkDriver.cpp b/Engine/Source/Runtime/Engine/Private/NetworkDriver.cpp
--- a/Engine/Source/Runtime/Engine/Private/NetworkDriver.cpp
+++ b/Engine/Source/Runtime/Engine/Private/NetworkDriver.cpp
@@ -58,6 +58,35 @@ DEFINE_STAT(STAT_PercentOutVoice);
 #define DEBUG_REMOTEFUNCTION(Format, ...) UE_LOG(LogNet, VeryVerbose, Format, __VA_ARGS__);
 #endif
 
+// CVars
+static TAutoConsoleVariable<int32> CVarSetNetDormancyEnabled(
+	TEXT("net.DormancyEnable"),
+	1,
+	TEXT("Enables Network Dormancy System for reducing CPU and bandwidth overhead of infrequently updated actors\n")
+	TEXT("1 Enables network dormancy. 0 disables network dormancy."),
+	ECVF_Default);
+
+static TAutoConsoleVariable<int32> CVarNetDormancyDraw(
+	TEXT("net.DormancyDraw"),
+	0,
+	TEXT("Draws debug information for network dormancy\n")
+	TEXT("1 Enables network dormancy debugging. 0 disables."),
+	ECVF_Default);
+
+static TAutoConsoleVariable<float> CVarNetDormancyDrawCullDistance(
+	TEXT("net.DormancyDrawCullDistance"),
+	5000.f,
+	TEXT("Cull distance for net.DormancyDraw. World Units")
+	TEXT("Max world units an actor can be away from the local view to draw its dormancy status"),
+	ECVF_Default);
+
+static TAutoConsoleVariable<int32> CVarNetDormancyValidate(
+	TEXT("net.DormancyValidate"),
+	0,
+	TEXT("Validates that dormant actors do not change state while in a dormant state (on server only)")
+	TEXT("0: Dont validate. 1: Validate on wake up. 2: Validate on each net update"),
+	ECVF_Default);
+
 /*-----------------------------------------------------------------------------
 	UNetDriver implementation.
 -----------------------------------------------------------------------------*/
@@ -296,9 +325,7 @@ void UNetDriver::TickFlush(float DeltaSeconds)
 		ClientConnections[i]->Tick();
 	}
 
-	
-	static const auto CVar = IConsoleManager::Get().FindTConsoleVariableDataInt(TEXT("net.DormancyDraw"));
-	if (CVar && CVar->GetValueOnGameThread() > 0)
+	if (CVarNetDormancyDraw.GetValueOnGameThread() > 0)
 	{
 		DrawNetDriverDebug();
 	}
@@ -1449,8 +1476,7 @@ void UNetDriver::FlushActorDormancy(AActor * Actor)
 	// way too, since we dont have to check every dormant actor in ::ServerReplicateActor to see if it needs to go out of dormancy
 
 #if WITH_SERVER_CODE
-	static const auto CVar = IConsoleManager::Get().FindTConsoleVariableDataInt(TEXT("net.DormancyEnable"));
-	if (CVar && CVar->GetValueOnGameThread() == 0)
+	if (CVarSetNetDormancyEnabled.GetValueOnGameThread() == 0)
 		return;
 
 	check(Actor);
@@ -2126,16 +2152,14 @@ int32 UNetDriver::ServerReplicateActors(float DeltaSeconds)
 					UActorChannel* Channel = Connection->ActorChannels.FindRef(Actor);
 
 					// Skip Actor if dormant
-					static const auto CVar = IConsoleManager::Get().FindTConsoleVariableDataInt(TEXT("net.DormancyEnable"));
-					if ( !CVar || CVar->GetValueOnGameThread() == 1 )
+					if ( CVarSetNetDormancyEnabled.GetValueOnGameThread() == 1 )
 					{
 						// If actor is already dormant on this channel, then skip replication entirely
 						if ( Connection->DormantActors.Contains( Actor ) )
 						{
 							// net.DormancyValidate can be set to 2 to validate dormant actor properties on every replicate
 							// (this could be moved to be done every tick instead of every net update if necessary, but seems excessive)
-							static const auto ValidateCVar = IConsoleManager::Get().FindTConsoleVariableDataInt(TEXT("net.DormancyValidate"));
-							if ( ValidateCVar && ValidateCVar->GetValueOnGameThread() == 2 )
+							if ( CVarNetDormancyValidate.GetValueOnGameThread() == 2 )
 							{
 								TSharedRef< FObjectReplicator > * Replicator = Connection->DormantReplicatorMap.Find( Actor );
 
@@ -2607,8 +2631,7 @@ void UNetDriver::DrawNetDriverDebug()
 		return;
 	}
 
-	static const auto CVar = IConsoleManager::Get().FindTConsoleVariableDataFloat(TEXT("net.DormancyDrawCullDistance"));
-	const float CullDist = CVar ? CVar->GetValueOnGameThread() : 5000.f;
+	const float CullDist = CVarNetDormancyDrawCullDistance.GetValueOnGameThread();
 
 	for (FActorIterator It(GetWorld()); It; ++It)
 	{
diff --git a/Engine/Source/Runtime/Engine/Private/PhysicsEngine/BodySetup.cpp b/Engine/Source/Runtime/Engine/Private/PhysicsEngine/BodySetup.cpp
--- a/Engine/Source/Runtime/Engine/Private/PhysicsEngine/BodySetup.cpp
+++ b/Engine/Source/Runtime/Engine/Private/PhysicsEngine/BodySetup.cpp
@@ -25,6 +25,19 @@
 	}
 #endif // WITH_PHYSX
 
+// CVars
+static TAutoConsoleVariable<float> CVarContactOffsetFactor(
+	TEXT("p.ContactOffsetFactor"),
+	0.01f,
+	TEXT("Multiplied by min dimension of object to calculate how close objects get before generating contacts. Default: 0.01"),
+	ECVF_Default);
+
+static TAutoConsoleVariable<float> CVarMaxContactOffset(
+	TEXT("p.MaxContactOffset"),
+	1.f,
+	TEXT("Max value of contact offset, which controls how close objects get before generating contacts. Default: 1.0"),
+	ECVF_Default);
+
 DEFINE_LOG_CATEGORY(LogPhysics);
 UBodySetup::UBodySetup(const class FPostConstructInitializeProperties& PCIP)
 	: Super(PCIP)
@@ -301,10 +314,8 @@ void UBodySetup::AddShapesToRigidActor(PxRigidActor* PDestActor, FVector& Scale3
 	bool bDestStatic = (PDestActor->isRigidStatic() != NULL);
 
 	// Get contact offset params
-	static const auto CVarContactOffsetFactor = IConsoleManager::Get().FindTConsoleVariableDataFloat(TEXT("p.ContactOffsetFactor"));
-	const float ContactOffsetFactor = CVarContactOffsetFactor->GetValueOnGameThread();
-	static const auto CVarMaxContactOffset = IConsoleManager::Get().FindTConsoleVariableDataFloat(TEXT("p.MaxContactOffset"));
-	const float MaxContactOffset = CVarMaxContactOffset->GetValueOnGameThread();
+	const float ContactOffsetFactor = CVarContactOffsetFactor.GetValueOnGameThread();
+	const float MaxContactOffset = CVarMaxContactOffset.GetValueOnGameThread();
 
 	check(GEngine->DefaultPhysMaterial != NULL);
 	PxMaterial* PDefaultMat = GEngine->DefaultPhysMaterial->GetPhysXMaterial();
diff --git a/Engine/Source/Runtime/Engine/Private/PhysicsEngine/PhysLevel.cpp b/Engine/Source/Runtime/Engine/Private/PhysicsEngine/PhysLevel.cpp
--- a/Engine/Source/Runtime/Engine/Private/PhysicsEngine/PhysLevel.cpp
+++ b/Engine/Source/Runtime/Engine/Private/PhysicsEngine/PhysLevel.cpp
@@ -28,6 +28,45 @@ DEFINE_STAT(STAT_SetBodyTransform);
 
 FPhysCommandHandler * GPhysCommandHandler = NULL;
 
+// CVars
+static TAutoConsoleVariable<float> CVarToleranceScaleLength(
+	TEXT("p.ToleranceScale_Length"),
+	100.f,
+	TEXT("The approximate size of objects in the simulation. Default: 100"),
+	ECVF_ReadOnly);
+
+static TAutoConsoleVariable<float> CVarTolerenceScaleMass(
+	TEXT("p.ToleranceScale_Mass"),
+	100.f,
+	TEXT("The approximate mass of a length * length * length block. Default: 100"),
+	ECVF_ReadOnly);
+
+static TAutoConsoleVariable<float> CVarToleranceScaleSpeed(
+	TEXT("p.ToleranceScale_Speed"),
+	1000.f,
+	TEXT("The typical magnitude of velocities of objects in simulation. Default: 1000"),
+	ECVF_ReadOnly);
+
+static TAutoConsoleVariable<int32> CVarAPEXMaxDestructibleDynamicChunkIslandCount(
+	TEXT("p.APEXMaxDestructibleDynamicChunkIslandCount"),
+	2000,
+	TEXT("APEX Max Destructilbe Dynamic Chunk Island Count."),
+	ECVF_Default);
+
+
+static TAutoConsoleVariable<int32> CVarAPEXMaxDestructibleDynamicChunkCount(
+	TEXT("p.APEXMaxDestructibleDynamicChunkCount"),
+	2000,
+	TEXT("APEX Max Destructible dynamic Chunk Count."),
+	ECVF_Default);
+
+static TAutoConsoleVariable<int32> CVarAPEXSortDynamicChunksByBenefit(
+	TEXT("p.bAPEXSortDynamicChunksByBenefit"),
+	1,
+	TEXT("True if APEX should sort dynamic chunks by benefit."),
+	ECVF_Default);
+
+
 
 //////////////////////////////////////////////////////////////////////////
 // UWORLD
@@ -247,15 +286,11 @@ void InitGamePhys()
 	GPhysXProfileZoneManager = &PxProfileZoneManager::createProfileZoneManager(GPhysXFoundation);
 	check(GPhysXProfileZoneManager);
 
-	static const auto CVarToleranceScale_Length = IConsoleManager::Get().FindTConsoleVariableDataFloat(TEXT("p.ToleranceScale_Length"));
-	static const auto CVarToleranceScale_Mass = IConsoleManager::Get().FindTConsoleVariableDataFloat(TEXT("p.ToleranceScale_Mass"));
-	static const auto CVarToleranceScale_Speed = IConsoleManager::Get().FindTConsoleVariableDataFloat(TEXT("p.ToleranceScale_Speed"));
-
 	// Create Physics
 	PxTolerancesScale PScale;
-	PScale.length = CVarToleranceScale_Length->GetValueOnGameThread();
-	PScale.mass = CVarToleranceScale_Mass->GetValueOnGameThread();
-	PScale.speed = CVarToleranceScale_Speed->GetValueOnGameThread();
+	PScale.length = CVarToleranceScaleLength.GetValueOnGameThread();
+	PScale.mass = CVarTolerenceScaleMass.GetValueOnGameThread();
+	PScale.speed = CVarToleranceScaleSpeed.GetValueOnGameThread();
 
 	GPhysXSDK = PxCreatePhysics(PX_PHYSICS_VERSION, *GPhysXFoundation, PScale, false, GPhysXProfileZoneManager);
 	check(GPhysXSDK);
@@ -336,14 +371,10 @@ void InitGamePhys()
 	// Set chunk report for fracture effect callbacks
 	GApexModuleDestructible->setChunkReport(&GApexChunkReport);
 
-	static const auto CVarMaxDestructibleDynamicChunkIslandCount = IConsoleManager::Get().FindTConsoleVariableDataInt(TEXT("p.APEXMaxDestructibleDynamicChunkIslandCount"));
-	GApexModuleDestructible->setMaxDynamicChunkIslandCount((physx::PxU32)FMath::Max(CVarMaxDestructibleDynamicChunkIslandCount->GetValueOnGameThread(), 0));
-
-	static const auto CVarMaxDestructibleDynamicChunkCount = IConsoleManager::Get().FindTConsoleVariableDataInt(TEXT("p.APEXMaxDestructibleDynamicChunkCount"));
-	GApexModuleDestructible->setMaxChunkCount((physx::PxU32)FMath::Max(CVarMaxDestructibleDynamicChunkCount->GetValueOnGameThread(), 0));
-
-	static const auto CVarSortDynamicChunksByBenefit = IConsoleManager::Get().FindTConsoleVariableDataInt(TEXT("p.bAPEXSortDynamicChunksByBenefit"));
-	GApexModuleDestructible->setSortByBenefit(CVarSortDynamicChunksByBenefit->GetValueOnGameThread()!=0);
+	
+	GApexModuleDestructible->setMaxDynamicChunkIslandCount((physx::PxU32)FMath::Max(CVarAPEXMaxDestructibleDynamicChunkIslandCount.GetValueOnGameThread(), 0));
+	GApexModuleDestructible->setMaxChunkCount((physx::PxU32)FMath::Max(CVarAPEXMaxDestructibleDynamicChunkCount.GetValueOnGameThread(), 0));
+	GApexModuleDestructible->setSortByBenefit(CVarAPEXSortDynamicChunksByBenefit.GetValueOnGameThread() != 0);
 
 	GApexModuleDestructible->setChunkReportSendChunkStateEvents(true);
 
diff --git a/Engine/Source/Runtime/Engine/Private/PhysicsEngine/PhysScene.cpp b/Engine/Source/Runtime/Engine/Private/PhysicsEngine/PhysScene.cpp
--- a/Engine/Source/Runtime/Engine/Private/PhysicsEngine/PhysScene.cpp
+++ b/Engine/Source/Runtime/Engine/Private/PhysicsEngine/PhysScene.cpp
@@ -24,6 +24,13 @@
 static int32 PhysXSceneCount = 1;
 static const int PhysXSlowRebuildRate = 10;
 
+// Cvars
+static TAutoConsoleVariable<float> CVarBounceThresholdVelocity(
+	TEXT("p.BounceThresholdVelocity"),
+	200.f,
+	TEXT("A contact with a relative velocity below this will not bounce. Default: 200"),
+	ECVF_Default);
+
 FORCEINLINE EPhysicsSceneType SceneType(const FBodyInstance * BodyInstance)
 {
 #if WITH_PHYSX
@@ -1083,8 +1090,7 @@ void FPhysScene::InitPhysScene(uint32 SceneType)
 	PSceneDesc.flags |= PxSceneFlag::eENABLE_KINEMATIC_PAIRS;
 
 	// Set bounce threshold
-	static const auto CVarBounceThresholdVelocity = IConsoleManager::Get().FindTConsoleVariableDataFloat(TEXT("p.BounceThresholdVelocity"));
-	PSceneDesc.bounceThresholdVelocity = CVarBounceThresholdVelocity->GetValueOnGameThread();
+	PSceneDesc.bounceThresholdVelocity = CVarBounceThresholdVelocity.GetValueOnGameThread();
 
 	// Possibly set flags in async scene for better behavior with piles
 #if USE_ADAPTIVE_FORCES_FOR_ASYNC_SCENE
diff --git a/Engine/Source/Runtime/Engine/Private/PrimitiveComponent.cpp b/Engine/Source/Runtime/Engine/Private/PrimitiveComponent.cpp
--- a/Engine/Source/Runtime/Engine/Private/PrimitiveComponent.cpp
+++ b/Engine/Source/Runtime/Engine/Private/PrimitiveComponent.cpp
@@ -32,6 +32,22 @@ static FAutoConsoleVariable CVarAllowCachedOverlaps(
 	TEXT("AllowCachedOverlaps"), 1,
 	TEXT("0: disable cached overlaps, 1: enable\n"));
 
+static TAutoConsoleVariable<float> CVarInitialOverlapTolerance(
+	TEXT("p.InitialOverlapTolerance"),
+	0.0f,
+	TEXT("Tolerance for initial overlapping test in PrimitiveComponent movement.\n")
+	TEXT("Normals within this tolerance are ignored if moving out of the object.\n")
+	TEXT("Dot product of movement direction and surface normal."),
+	ECVF_Default);
+
+#if !(UE_BUILD_SHIPPING || UE_BUILD_TEST)
+static TAutoConsoleVariable<int32> CVarShowInitialOverlaps(
+	TEXT("p.ShowInitialOverlaps"),
+	0,
+	TEXT("Show initial overlaps when moving a component, including estimated 'exit' direction.\n")
+	TEXT(" 0:off, otherwise on"),
+	ECVF_Cheat);
+#endif // !(UE_BUILD_SHIPPING || UE_BUILD_TEST)
 
 ///////////////////////////////////////////////////////////////////////////////
 // PRIMITIVE COMPONENT
@@ -1139,8 +1155,7 @@ static bool ShouldIgnoreHitResult(const UWorld* InWorld, FHitResult const& TestH
 	// This helps prevent getting stuck in walls.
 	if ( TestHit.bStartPenetrating && TestHit.bBlockingHit )
 	{
-		static const auto OverlapToleranceICVar = IConsoleManager::Get().FindTConsoleVariableDataFloat(TEXT("p.InitialOverlapTolerance"));
-		const float DotTolerance = OverlapToleranceICVar ? OverlapToleranceICVar->GetValueOnGameThread() : 0.f;
+		const float DotTolerance = CVarInitialOverlapTolerance.GetValueOnGameThread();
 
 		// Dot product of movement direction against 'exit' direction
 		const FVector MovementDir = MovementDirDenormalized.SafeNormal();
@@ -1150,8 +1165,7 @@ static bool ShouldIgnoreHitResult(const UWorld* InWorld, FHitResult const& TestH
 
 #if !(UE_BUILD_SHIPPING || UE_BUILD_TEST)
 		{
-			static const auto ICVar = IConsoleManager::Get().FindTConsoleVariableDataInt(TEXT("p.ShowInitialOverlaps"));
-			if(ICVar->GetValueOnGameThread() != 0)
+			if(CVarShowInitialOverlaps.GetValueOnGameThread() != 0)
 			{
 				UE_LOG(LogTemp, Log, TEXT("Overlapping %s Dir %s Dot %f"), *GetNameSafe(TestHit.GetActor()), *MovementDir.ToString(), MoveDot);
 				DrawDebugDirectionalArrow(InWorld, TestHit.TraceStart, TestHit.TraceStart + 30.f * TestHit.ImpactNormal, 5.f, bMovingOut ? FColor(64,128,255) : FColor(255,64,64), true, 1.f);
diff --git a/Engine/Source/Runtime/Engine/Private/SceneView.cpp b/Engine/Source/Runtime/Engine/Private/SceneView.cpp
--- a/Engine/Source/Runtime/Engine/Private/SceneView.cpp
+++ b/Engine/Source/Runtime/Engine/Private/SceneView.cpp
@@ -27,6 +27,28 @@ static TAutoConsoleVariable<float> CVarSSRMaxRoughness(
 	TEXT(" 0..1: use specified max roughness (overrride PostprocessVolume setting)\n")
 	TEXT(" -1: no override (default)"),
 	ECVF_RenderThreadSafe);
+
+static TAutoConsoleVariable<int32> CVarShadowFreezeCamera(
+	TEXT("r.Shadow.FreezeCamera"),
+	0,
+	TEXT("Debug the shadow methods by allowing to observe the system from outside.\n")
+	TEXT("0: default\n")
+	TEXT("1: freeze camera at current location"),
+	ECVF_Cheat);
+
+static TAutoConsoleVariable<float> CVarExposureOffset(
+	TEXT("r.ExposureOffset"),
+	0.0f,
+	TEXT("For adjusting the exposure on top of post process settings and eye adaptation. For developers only. 0:default"),
+	ECVF_Cheat);
+
+static TAutoConsoleVariable<int32> CVarRenderTimeFrozen(
+	TEXT("r.RenderTimeFrozen"),
+	0,
+	TEXT("Allows to freeze time based effects in order to provide more deterministic render profiling.\n")
+	TEXT(" 0: off\n")
+	TEXT(" 1: on (Note: this also disables occlusion queries)"),
+	ECVF_Cheat);
 #endif
 
 static TAutoConsoleVariable<float> CVarSSAOFadeRadiusScale(
@@ -256,8 +278,7 @@ FSceneView::FSceneView(const FSceneViewInitOptions& InitOptions)
 #if !(UE_BUILD_SHIPPING || UE_BUILD_TEST)
 	{
 		// console variable override
-		static const auto CVar = IConsoleManager::Get().FindTConsoleVariableDataInt(TEXT("r.Shadow.FreezeCamera")); 
-		int32 Value = CVar->GetValueOnAnyThread();
+		int32 Value = CVarShadowFreezeCamera.GetValueOnAnyThread();
 
 		static FViewMatrices Backup = ShadowViewMatrices;
 
@@ -870,9 +891,7 @@ void FSceneView::EndFinalPostprocessSettings()
 
 #if !(UE_BUILD_SHIPPING || UE_BUILD_TEST)
 	{
-		static const auto CVar = IConsoleManager::Get().FindTConsoleVariableDataFloat(TEXT("r.ExposureOffset"));
-
-		float Value = CVar->GetValueOnGameThread();
+		float Value = CVarExposureOffset.GetValueOnGameThread();
 		FinalPostProcessSettings.AutoExposureBias += Value;
 	}
 #endif
@@ -1143,9 +1162,7 @@ FSceneViewFamily::FSceneViewFamily( const ConstructionValues& CVS )
 	ensure(CVS.bTimesSet);
 
 #if !(UE_BUILD_SHIPPING || UE_BUILD_TEST)
-	static const auto CVar = IConsoleManager::Get().FindTConsoleVariableDataInt(TEXT("r.RenderTimeFrozen")); 
-	int32 Value = CVar->GetValueOnAnyThread();
-
+	int32 Value = CVarRenderTimeFrozen.GetValueOnAnyThread();
 	if(Value)
 	{
 		CurrentWorldTime = 0;
diff --git a/Engine/Source/Runtime/Engine/Private/ShaderCompiler/ShaderCompiler.cpp b/Engine/Source/Runtime/Engine/Private/ShaderCompiler/ShaderCompiler.cpp
--- a/Engine/Source/Runtime/Engine/Private/ShaderCompiler/ShaderCompiler.cpp
+++ b/Engine/Source/Runtime/Engine/Private/ShaderCompiler/ShaderCompiler.cpp
@@ -30,6 +30,18 @@ static FAutoConsoleVariableRef CVarDumpShaderDebugInfo(
 	TEXT("On iOS, if the PowerVR graphics SDK is installed to the default path, the PowerVR shader compiler will be called and errors will be reported during the cook.")
 	);
 
+#if !(UE_BUILD_SHIPPING || UE_BUILD_TEST)
+static TAutoConsoleVariable<FString> CVarD3DCompilerPath(TEXT("r.D3DCompilerPath"),
+	TEXT(""),	// default
+	TEXT("Allows to specify a HLSL compiler version that is different from the one the code was compiled.\n")
+	TEXT("No path (\"\") means the default one is used.\n")
+	TEXT("If the compiler cannot be found an error is reported and it will compile further with the default one.\n")
+	TEXT("This console variable works with ShaderCompileWorker (with multi threading) and without multi threading.\n")
+	TEXT("This variable can be set in ConsoleVariables.ini to be defined at startup.\n")
+	TEXT("e.g. c:/temp/d3dcompiler_44.dll or \"\""),
+	ECVF_Cheat);
+#endif
+
 // Serialize Queued Job information
 static void DoWriteTasks(TArray<FShaderCompileJob*>& QueuedJobs, FArchive& TransferFile)
 {
@@ -1979,9 +1991,7 @@ void GlobalBeginCompileShader(
 
 #if !(UE_BUILD_SHIPPING || UE_BUILD_TEST)
 	{
-		static IConsoleVariable* CVar = IConsoleManager::Get().FindConsoleVariable(TEXT("r.D3DCompilerPath"));
-
-		FString Path = CVar->GetString();
+		FString Path = CVarD3DCompilerPath.GetValueOnAnyThread();
 		if(!Path.IsEmpty())
 		{
 			Input.Environment.SetDefine(TEXT("D3DCOMPILER_PATH"), *Path);
diff --git a/Engine/Source/Runtime/Engine/Private/Texture2D.cpp b/Engine/Source/Runtime/Engine/Private/Texture2D.cpp
--- a/Engine/Source/Runtime/Engine/Private/Texture2D.cpp
+++ b/Engine/Source/Runtime/Engine/Private/Texture2D.cpp
@@ -43,9 +43,7 @@ static bool CanCreateAsVirtualTexture(const UTexture2D* Texture, uint32 TexCreat
 	const uint32 iRequiredFlags =
 		TexCreate_OfflineProcessed;
 
-	static auto CVarVirtualTextureEnabled = IConsoleManager::Get().FindTConsoleVariableDataInt(TEXT("r.VirtualTexture"));
-	check(CVarVirtualTextureEnabled);
-	bool bCreateVirtual = ((TexCreateFlags & (iDisableFlags | iRequiredFlags)) == iRequiredFlags) && Texture->bIsStreamable && CVarVirtualTextureEnabled->GetValueOnRenderThread();
+	bool bCreateVirtual = ((TexCreateFlags & (iDisableFlags | iRequiredFlags)) == iRequiredFlags) && Texture->bIsStreamable && CVarVirtualTextureEnabled.GetValueOnRenderThread();
 	
 	static auto CVarVirtualTextureReducedMemoryEnabled = IConsoleManager::Get().FindTConsoleVariableDataInt(TEXT("r.VirtualTextureReducedMemory"));
 	check(CVarVirtualTextureReducedMemoryEnabled);
@@ -64,6 +62,19 @@ int32 GDefragmentationRetryCounter = 10;
 /** Number of times to retry to reallocate a texture before trying a panic defragmentation, subsequent times. */
 int32 GDefragmentationRetryCounterLong = 100;
 
+/** CVars */
+static TAutoConsoleVariable<float> CVarSetMipMapLODBias(
+	TEXT("r.MipMapLODBias"),
+	0.0f,
+	TEXT("Apply additional mip map bias for all 2D textures, range of -15.0 to 15.0"),
+	ECVF_RenderThreadSafe);
+
+static TAutoConsoleVariable<int32> CVarVirtualTextureEnabled(
+	TEXT("r.VirtualTexture"),
+	1,
+	TEXT("If set to 1, textures will use virtual memory so they can be partially resident."),
+	ECVF_RenderThreadSafe);
+
 /** Turn on ENABLE_TEXTURE_TRACKING in ContentStreaming.cpp and setup GTrackedTextures to track specific textures through the streaming system. */
 extern bool TrackTextureEvent( FStreamingTexture* StreamingTexture, UTexture2D* Texture, bool bForceMipLevelsToBeResident, const FStreamingManagerTexture* Manager );
 
@@ -1146,10 +1157,11 @@ void UTexture2D::UpdateTextureRegions( int32 MipIndex, uint32 NumRegions, FUpdat
 
 #endif
 
+
+
 float UTexture2D::GetGlobalMipMapLODBias()
 {
-	static const auto CVarMipMapLODBias = IConsoleManager::Get().FindTConsoleVariableDataFloat(TEXT("r.MipMapLODBias"));
-	float BiasOffset = CVarMipMapLODBias->GetValueOnAnyThread(); // called from multiple threads.
+	float BiasOffset = CVarSetMipMapLODBias.GetValueOnAnyThread(); // called from multiple threads.
 	return FMath::Clamp(BiasOffset, -15.0f, 15.0f);
 }
 
diff --git a/Engine/Source/Runtime/Engine/Private/UnrealEngine.cpp b/Engine/Source/Runtime/Engine/Private/UnrealEngine.cpp
--- a/Engine/Source/Runtime/Engine/Private/UnrealEngine.cpp
+++ b/Engine/Source/Runtime/Engine/Private/UnrealEngine.cpp
@@ -138,6 +138,22 @@ static FAutoConsoleVariable CVarSystemResolution(
 	TEXT("  	   e.g. 1920x1080f")
 	);
 
+static TAutoConsoleVariable<float> CVarDepthOfFieldNearBlurSizeThreshold(
+	TEXT("r.DepthOfFieldNearBlurSizeThreshold"),
+	0.01f,
+	TEXT("Sets the minimum near blur size before the effect is forcably disabled. Currently only affects Gaussian DOF.\n")
+	TEXT(" (default = 0.01f)"),
+	ECVF_RenderThreadSafe);
+
+#if !(UE_BUILD_SHIPPING || UE_BUILD_TEST)
+static TAutoConsoleVariable<float> CVarSetOverrideFPS(
+	TEXT("t.OverrideFPS"),
+	0.0f,
+	TEXT("This allows to override the frame time measurement with a fixed fps number (game can run faster or slower).\n")
+	TEXT("<=0:off, in frames per second, e.g. 60"),
+	ECVF_Cheat);
+#endif // !(UE_BUILD_SHIPPING || UE_BUILD_TEST)
+
 /** Enum entries represent index to global object referencer stored in UGameEngine */
 enum EGametypeContentReferencerTypes
 {
@@ -208,12 +224,11 @@ void ScalabilityCVarsSinkCallback()
 	static const auto* MaxAnisotropy = ConsoleMan.FindTConsoleVariableDataInt(TEXT("r.MaxAnisotropy"));
 	static const auto* MaxShadowResolution = ConsoleMan.FindTConsoleVariableDataInt(TEXT("r.Shadow.MaxResolution"));
 	static const auto ViewDistanceScale = ConsoleMan.FindTConsoleVariableDataFloat(TEXT("r.ViewDistanceScale"));
-	static const auto GaussianDOFNearThreshold = IConsoleManager::Get().FindTConsoleVariableDataFloat(TEXT("r.DepthOfFieldNearBlurSizeThreshold"));
 	GCachedScalabilityCVars.MaxAnisotropy = MaxAnisotropy->GetValueOnGameThread();
 	GCachedScalabilityCVars.MaxShadowResolution = MaxShadowResolution->GetValueOnGameThread();
 	GCachedScalabilityCVars.ViewDistanceScale = FMath::Clamp(ViewDistanceScale->GetValueOnGameThread(), 0.0f, 1.0f);
 	GCachedScalabilityCVars.ViewDistanceScaleSquared = FMath::Square(GCachedScalabilityCVars.ViewDistanceScale);
-	GCachedScalabilityCVars.GaussianDOFNearThreshold = GaussianDOFNearThreshold->GetValueOnGameThread();
+	GCachedScalabilityCVars.GaussianDOFNearThreshold = CVarDepthOfFieldNearBlurSizeThreshold.GetValueOnGameThread();
 
 	// action needed if we change r.MaterialQualityLevel at runtime
 	{
@@ -1006,10 +1021,7 @@ void UEngine::UpdateTimeAndHandleMaxTickRate()
 
 #if !(UE_BUILD_SHIPPING || UE_BUILD_TEST)
 	{
-		static const auto ICVarMin = IConsoleManager::Get().FindTConsoleVariableDataFloat(TEXT("t.OverrideFPS"));
-
-		float OverrideFPS = ICVarMin->GetValueOnGameThread();
-
+		float OverrideFPS = CVarSetOverrideFPS.GetValueOnGameThread();
 		if(OverrideFPS >= 0.001f)
 		{
 			// in seconds
diff --git a/Engine/Source/Runtime/RenderCore/Private/RenderCore.cpp b/Engine/Source/Runtime/RenderCore/Private/RenderCore.cpp
--- a/Engine/Source/Runtime/RenderCore/Private/RenderCore.cpp
+++ b/Engine/Source/Runtime/RenderCore/Private/RenderCore.cpp
@@ -151,7 +151,16 @@ DEFINE_STAT(STAT_ReflectiveShadowMapDrawTime);
 DEFINE_STAT(STAT_NumReflectiveShadowMapLights);
 DEFINE_STAT(STAT_RenderWholeSceneReflectiveShadowMapsTime);
 
+#if !(UE_BUILD_SHIPPING || UE_BUILD_TEST)
+
+static TAutoConsoleVariable<int32> CVarForceLOD(
+	TEXT("r.ForceLOD"),
+	-1,
+	TEXT("LOD level to force, -1 is off."),
+	ECVF_Cheat | ECVF_RenderThreadSafe
+	);
 
+#endif // !(UE_BUILD_SHIPPING || UE_BUILD_TEST)
 
 /** Whether to pause the global realtime clock for the rendering thread (read and write only on main thread). */
 bool GPauseRenderingRealtimeClock;
@@ -212,8 +221,7 @@ RENDERCORE_API int32 GetCVarForceLOD()
 
 #if !(UE_BUILD_SHIPPING || UE_BUILD_TEST)
 	{
-		static const auto CVar = IConsoleManager::Get().FindTConsoleVariableDataInt(TEXT("r.ForceLOD"));
-		Ret = CVar->GetValueOnRenderThread();
+		Ret = CVarForceLOD.GetValueOnRenderThread();
 	}
 #endif // !(UE_BUILD_SHIPPING || UE_BUILD_TEST)
 
diff --git a/Engine/Source/Runtime/Renderer/Private/DeferredShadingRenderer.cpp b/Engine/Source/Runtime/Renderer/Private/DeferredShadingRenderer.cpp
--- a/Engine/Source/Runtime/Renderer/Private/DeferredShadingRenderer.cpp
+++ b/Engine/Source/Runtime/Renderer/Private/DeferredShadingRenderer.cpp
@@ -27,6 +27,7 @@ TAutoConsoleVariable<int32> CVarEarlyZPass(
 	TEXT("  x: use built in heuristic (default is 3)\n"),
 	ECVF_Default);
 
+
 int32 GEarlyZPassMovable = 0;
 
 /** Affects static draw lists so must reload level to propagate. */
@@ -46,6 +47,14 @@ static TAutoConsoleVariable<int32> CVarTestUIBlur(
 	TEXT("0: off (default)\n")
 	TEXT("1: on\n"),
 	ECVF_RenderThreadSafe);
+
+static TAutoConsoleVariable<int32> CVarVisualizeTexturePool(
+	TEXT("r.VisualizeTexturePool"),
+	0,
+	TEXT("Allows to enable the visualize the texture pool (currently only on console).\n")
+	TEXT(" 0: off\n")
+	TEXT(" 1: on"),
+	ECVF_Cheat | ECVF_RenderThreadSafe);
 #endif
 
 /*-----------------------------------------------------------------------------
@@ -420,9 +429,7 @@ void FDeferredShadingSceneRenderer::RenderFinish()
 {
 #if !(UE_BUILD_SHIPPING || UE_BUILD_TEST)
 	{
-		static const auto ICVar = IConsoleManager::Get().FindTConsoleVariableDataInt(TEXT("r.VisualizeTexturePool"));
-
-		if(ICVar->GetValueOnRenderThread())
+		if(CVarVisualizeTexturePool.GetValueOnRenderThread())
 		{
 			RenderVisualizeTexturePool();
 		}
diff --git a/Engine/Source/Runtime/Renderer/Private/FogRendering.cpp b/Engine/Source/Runtime/Renderer/Private/FogRendering.cpp
--- a/Engine/Source/Runtime/Renderer/Private/FogRendering.cpp
+++ b/Engine/Source/Runtime/Renderer/Private/FogRendering.cpp
@@ -7,6 +7,26 @@
 #include "RendererPrivate.h"
 #include "ScenePrivate.h"
 
+#if !(UE_BUILD_SHIPPING || UE_BUILD_TEST)
+static TAutoConsoleVariable<float> CVarFogStartDistance(
+	TEXT("r.FogStartDistance"),
+	-1.0f,
+	TEXT("Allows to override the FogStartDistance setting (needs ExponentialFog in the level).\n")
+	TEXT(" <0: use default settings (default: -1)\n")
+	TEXT(">=0: override settings by the given value (in world units)"),
+	ECVF_Cheat | ECVF_RenderThreadSafe);
+
+static TAutoConsoleVariable<float> CVarFogDensity(
+	TEXT("r.FogDensity"),
+	-1.0f,
+	TEXT("Allows to override the FogDensity setting (needs ExponentialFog in the level).\n")
+	TEXT("Using a strong value allows to quickly see which pixel are affected by fog.\n")
+	TEXT("Using a start distance allows to cull pixels are can speed up rendering.\n")
+	TEXT(" <0: use default settings (default: -1)\n")
+	TEXT(">=0: override settings by the given value (0:off, 1=very dense fog)"),
+	ECVF_Cheat | ECVF_RenderThreadSafe);
+#endif
+
 /** Binds the parameters. */
 void FExponentialHeightFogShaderParameters::Bind(const FShaderParameterMap& ParameterMap)
 {
@@ -204,17 +224,9 @@ void FSceneRenderer::InitFogConstants()
 
 #if !(UE_BUILD_SHIPPING || UE_BUILD_TEST)
 	{
-		// console variable override
-		static const auto CVar = IConsoleManager::Get().FindTConsoleVariableDataFloat(TEXT("r.FogDensity")); 
-
-		FogDensityOverride = CVar->GetValueOnAnyThread();
-	}
-
-	{
-		// console variable override
-		static const auto CVar = IConsoleManager::Get().FindTConsoleVariableDataFloat(TEXT("r.FogStartDistance")); 
-
-		FogStartDistanceOverride = CVar->GetValueOnAnyThread();
+		// console variable overrides
+		FogDensityOverride = CVarFogDensity.GetValueOnAnyThread();
+		FogStartDistanceOverride = CVarFogStartDistance.GetValueOnAnyThread();
 	}
 #endif // !(UE_BUILD_SHIPPING || UE_BUILD_TEST)
 
diff --git a/Engine/Source/Runtime/Renderer/Private/PostProcess/PostProcessBokehDOF.cpp b/Engine/Source/Runtime/Renderer/Private/PostProcess/PostProcessBokehDOF.cpp
--- a/Engine/Source/Runtime/Renderer/Private/PostProcess/PostProcessBokehDOF.cpp
+++ b/Engine/Source/Runtime/Renderer/Private/PostProcess/PostProcessBokehDOF.cpp
@@ -11,6 +11,8 @@
 #include "PostProcessing.h"
 #include "PostProcessBokehDOF.h"
 
+
+
 /**
  * Indexing style for DOF
  */
@@ -595,6 +597,20 @@ void FRCPassPostProcessBokehDOF::ComputeDepthOfFieldParams(const FRenderingCompo
 	Out[1] = FVector4(MaxBokehSizeInPixel, YOffsetInUV, UsedYDivTextureY, YOffsetInPixel);
 }
 
+
+static TAutoConsoleVariable<int32> CVarBokehDOFIndexStyle(
+	TEXT("r.BokehDOFIndexStyle"),
+#if PLATFORM_MAC // Avoid a driver bug on OSX/NV cards that causes driver to generate an unwound index buffer
+	1,
+#else
+	0,
+#endif
+	TEXT("Controls whether to use a packed or unwound index buffer for Bokeh DOF.\n")
+	TEXT("0: Use packed index buffer (faster) (default)\n")
+	TEXT("1: Use unwound index buffer (slower)"),
+	ECVF_ReadOnly | ECVF_RenderThreadSafe);
+
+
 void FRCPassPostProcessBokehDOF::Process(FRenderingCompositePassContext& Context)
 {
 	SCOPED_DRAW_EVENT(PassPostProcessBokehDOF, DEC_SCENE_ITEMS);
@@ -660,8 +676,7 @@ void FRCPassPostProcessBokehDOF::Process(FRenderingCompositePassContext& Context
 
 	FIntPoint LeftTop = LocalViewRect.Min;
 
-	static IConsoleVariable* IndexStyleVar = IConsoleManager::Get().FindConsoleVariable(TEXT("r.BokehDOFIndexStyle"));
-	static EBokehIndexStyle IndexStyle = (EBokehIndexStyle)IndexStyleVar->GetInt();
+	static EBokehIndexStyle IndexStyle = (EBokehIndexStyle)CVarBokehDOFIndexStyle.GetValueOnRenderThread();
 	
 	if(bHighQuality)
 	{
diff --git a/Engine/Source/Runtime/Renderer/Private/PostProcess/PostProcessBusyWait.cpp b/Engine/Source/Runtime/Renderer/Private/PostProcess/PostProcessBusyWait.cpp
--- a/Engine/Source/Runtime/Renderer/Private/PostProcess/PostProcessBusyWait.cpp
+++ b/Engine/Source/Runtime/Renderer/Private/PostProcess/PostProcessBusyWait.cpp
@@ -10,6 +10,16 @@
 #include "PostProcessBusyWait.h"
 #include "PostProcessing.h"
 
+#if !(UE_BUILD_SHIPPING || UE_BUILD_TEST)
+static TAutoConsoleVariable<float> CVarSetGPUBusyWait(
+	TEXT("r.GPUBusyWait"),
+	0.0f,
+	TEXT("<=0:off, >0: keep the GPU busy with n units of some fixed amount of work, independent on the resolution\n")
+	TEXT("This can be useful to make GPU timing experiments. The value should roughly represent milliseconds.\n")
+	TEXT("Clamped at 500."),
+	ECVF_Cheat | ECVF_RenderThreadSafe);
+#endif
+
 /** Encapsulates the post processing busy wait pixel shader. */
 class FPostProcessBusyWaitPS : public FGlobalShader
 {
@@ -45,11 +55,9 @@ public:
 
 #if !(UE_BUILD_SHIPPING || UE_BUILD_TEST)
 		{
-			static const auto CVar = IConsoleManager::Get().FindTConsoleVariableDataFloat(TEXT("r.GPUBusyWait"));
-
 			uint32 PixelCount = Context.View.ViewRect.Size().X * Context.View.ViewRect.Size().Y;
 
-			float CVarValue = FMath::Clamp(CVar->GetValueOnRenderThread(), 0.0f, 500.0f);
+			float CVarValue = FMath::Clamp(CVarSetGPUBusyWait.GetValueOnRenderThread(), 0.0f, 500.0f);
 
 			// multiply with large number to get more human friendly number range
 			// calibrated on a NV580 to be roughly a millisecond
diff --git a/Engine/Source/Runtime/Renderer/Private/PostProcess/PostProcessCombineLUTs.cpp b/Engine/Source/Runtime/Renderer/Private/PostProcess/PostProcessCombineLUTs.cpp
--- a/Engine/Source/Runtime/Renderer/Private/PostProcess/PostProcessCombineLUTs.cpp
+++ b/Engine/Source/Runtime/Renderer/Private/PostProcess/PostProcessCombineLUTs.cpp
@@ -11,6 +11,29 @@
 #include "PostProcessing.h"
 #include "ScreenRendering.h"
 
+
+// CVars
+static TAutoConsoleVariable<float> CVarColorMin(
+	TEXT("r.Color.Min"),
+	0.0f,
+	TEXT("Allows to define where the value 0 in the color channels is mapped to after color grading.\n")
+	TEXT("The value should be around 0, positive: a gray scale is added to the darks, negative: more dark values become black, Default: 0"),
+	ECVF_RenderThreadSafe);
+
+static TAutoConsoleVariable<float> CVarColorMid(
+	TEXT("r.Color.Mid"),
+	0.5f,
+	TEXT("Allows to define where the value 0.5 in the color channels is mapped to after color grading (This is similar to a gamma correction).\n")
+	TEXT("Value should be around 0.5, smaller values darken the mid tones, larger values brighten the mid tones, Default: 0.5"),
+	ECVF_RenderThreadSafe);
+
+static TAutoConsoleVariable<float> CVarColorMax(
+	TEXT("r.Color.Max"),
+	1.0f,
+	TEXT("Allows to define where the value 1.0 in the color channels is mapped to after color grading.\n")
+	TEXT("Value should be around 1, smaller values darken the highlights, larger values move more colors towards white, Default: 1"),
+	ECVF_RenderThreadSafe);
+
 // false:use 256x16 texture / true:use volume texture (faster, requires geometry shader)
 // USE_VOLUME_LUT: needs to be the same for C++ and HLSL
 static bool UseVolumeTextureLUT(EShaderPlatform Platform) 
@@ -54,19 +77,10 @@ FColorRemapShaderParameters::FColorRemapShaderParameters(const FShaderParameterM
 void FColorRemapShaderParameters::Set(FRHICommandList& RHICmdList, const FPixelShaderRHIParamRef ShaderRHI)
 {
 	FColorTransform ColorTransform;
-
-	{
-		static const auto ICVarMin = IConsoleManager::Get().FindTConsoleVariableDataFloat(TEXT("r.Color.Min"));
-		static const auto ICVarMid = IConsoleManager::Get().FindTConsoleVariableDataFloat(TEXT("r.Color.Mid"));
-		static const auto ICVarMax = IConsoleManager::Get().FindTConsoleVariableDataFloat(TEXT("r.Color.Max"));
-		
-		check(ICVarMin && ICVarMid && ICVarMax);
-		
-		ColorTransform.MinValue = FMath::Clamp(ICVarMin->GetValueOnRenderThread(), -10.0f, 10.0f);
-		ColorTransform.MidValue = FMath::Clamp(ICVarMid->GetValueOnRenderThread(), -10.0f, 10.0f);
-		ColorTransform.MaxValue = FMath::Clamp(ICVarMax->GetValueOnRenderThread(), -10.0f, 10.0f);
-	}
-
+	ColorTransform.MinValue = FMath::Clamp(CVarColorMin.GetValueOnRenderThread(), -10.0f, 10.0f);
+	ColorTransform.MidValue = FMath::Clamp(CVarColorMid.GetValueOnRenderThread(), -10.0f, 10.0f);
+	ColorTransform.MaxValue = FMath::Clamp(CVarColorMax.GetValueOnRenderThread(), -10.0f, 10.0f);
+	
 	{
 		// x is the input value, y the output value
 		// RGB = a, b, c where y = a * x*x + b * x + c
diff --git a/Engine/Source/Runtime/Renderer/Private/PostProcess/PostProcessMotionBlur.cpp b/Engine/Source/Runtime/Renderer/Private/PostProcess/PostProcessMotionBlur.cpp
--- a/Engine/Source/Runtime/Renderer/Private/PostProcess/PostProcessMotionBlur.cpp
+++ b/Engine/Source/Runtime/Renderer/Private/PostProcess/PostProcessMotionBlur.cpp
@@ -10,8 +10,15 @@
 #include "PostProcessMotionBlur.h"
 #include "PostProcessing.h"
 
-
-
+#if !(UE_BUILD_SHIPPING || UE_BUILD_TEST)
+static TAutoConsoleVariable<int32> CVarMotionBlurFiltering(
+	TEXT("r.MotionBlurFiltering"),
+	0,
+	TEXT("Useful developer variable\n")
+	TEXT("0: off (default, expected by the shader for better quality)\n")
+	TEXT("1: on"),
+	ECVF_Cheat | ECVF_RenderThreadSafe);
+#endif
 
 
 /** Encapsulates the post processing motion blur pixel shader. */
@@ -264,8 +271,7 @@ public:
 			bool bFiltered = false;
 
 #if !(UE_BUILD_SHIPPING || UE_BUILD_TEST)
-			static const auto CVar = IConsoleManager::Get().FindTConsoleVariableDataInt(TEXT("r.MotionBlurFiltering"));
-			bFiltered = CVar->GetValueOnRenderThread() != 0;
+			bFiltered = CVarMotionBlurFiltering.GetValueOnRenderThread() != 0;
 #endif // !(UE_BUILD_SHIPPING || UE_BUILD_TEST)
 
 			if(bFiltered)
diff --git a/Engine/Source/Runtime/Renderer/Private/PostProcess/PostProcessUpscale.cpp b/Engine/Source/Runtime/Renderer/Private/PostProcess/PostProcessUpscale.cpp
--- a/Engine/Source/Runtime/Renderer/Private/PostProcess/PostProcessUpscale.cpp
+++ b/Engine/Source/Runtime/Renderer/Private/PostProcess/PostProcessUpscale.cpp
@@ -12,6 +12,13 @@
 #include "PostProcessHistogram.h"
 #include "PostProcessEyeAdaptation.h"
 
+static TAutoConsoleVariable<float> CVarScreenPercentageSoftness(
+	TEXT("r.ScreenPercentageSoftness"),
+	0.3f,
+	TEXT("To scale up with higher quality loosing some sharpness\n")
+	TEXT(" 0..1 (0.3 is good for ScreenPercentage 90"),
+	ECVF_Scalability | ECVF_RenderThreadSafe);
+
 /** Encapsulates the post processing eye adaptation pixel shader. */
 template <uint32 Method>
 class FPostProcessUpscalePS : public FGlobalShader
@@ -62,9 +69,7 @@ public:
 		// If the method needs softness value
 		if(Method == 2)
 		{
-			static const auto* CVar = IConsoleManager::Get().FindTConsoleVariableDataFloat(TEXT("r.ScreenPercentageSoftness"));
-
-			float UpscaleSoftnessValue = FMath::Clamp(CVar->GetValueOnRenderThread(), 0.0f, 1.0f);
+			float UpscaleSoftnessValue = FMath::Clamp(CVarScreenPercentageSoftness.GetValueOnRenderThread(), 0.0f, 1.0f);
 
 			SetShaderValue(Context.RHICmdList, ShaderRHI, UpscaleSoftness, UpscaleSoftnessValue);
 		}
diff --git a/Engine/Source/Runtime/Renderer/Private/PostProcess/PostProcessWeightedSampleSum.cpp b/Engine/Source/Runtime/Renderer/Private/PostProcess/PostProcessWeightedSampleSum.cpp
--- a/Engine/Source/Runtime/Renderer/Private/PostProcess/PostProcessWeightedSampleSum.cpp
+++ b/Engine/Source/Runtime/Renderer/Private/PostProcess/PostProcessWeightedSampleSum.cpp
@@ -514,6 +514,18 @@ FPooledRenderTargetDesc FRCPassPostProcessWeightedSampleSum::ComputeOutputDesc(E
 	return Ret;
 }
 
+
+static TAutoConsoleVariable<float> CVarFastBlurThreshold(
+	TEXT("r.FastBlurThreshold"),
+	7.0f,
+	TEXT("Defines at what radius the Gaussian blur optimization kicks in (estimated 25% - 40% faster).\n")
+	TEXT("The optimzation uses slightly less memory and has a quality loss on smallblur radius.\n")
+	TEXT("  0: use the optimization always (fastest, lowest quality)\n")
+	TEXT("  3: use the optimization starting at a 3 pixel radius (quite fast)\n")
+	TEXT("  7: use the optimization starting at a 7 pixel radius (default)\n")
+	TEXT(">15: barely ever use the optimization (high quality)"),
+	ECVF_Scalability | ECVF_RenderThreadSafe);
+
 bool FRCPassPostProcessWeightedSampleSum::DoFastBlur() const
 {
 	bool bRet = false;
@@ -536,13 +548,10 @@ bool FRCPassPostProcessWeightedSampleSum::DoFastBlur() const
 			// we scale by width because FOV is defined horizontally
 			float EffectiveBlurRadius = SizeScale * SrcSizeForThisAxis  * 2 / 100.0f;
 
-			// can be hard coded once the feature works reliable
-			static const auto CVar = IConsoleManager::Get().FindTConsoleVariableDataFloat(TEXT("r.FastBlurThreshold"));
-
 #if PLATFORM_HTML5
-			float FastBlurThreshold = CVar->GetValueOnGameThread();
+			float FastBlurThreshold = CVarFastBlurThreshold.GetValueOnGameThread();
 #else
-			float FastBlurThreshold = CVar->GetValueOnRenderThread();
+			float FastBlurThreshold = CVarFastBlurThreshold.GetValueOnRenderThread();
 #endif
 
 			// small radius look too different with this optimization so we only to it for larger radius
diff --git a/Engine/Source/Runtime/Renderer/Private/PostProcess/PostProcessing.cpp b/Engine/Source/Runtime/Renderer/Private/PostProcess/PostProcessing.cpp
--- a/Engine/Source/Runtime/Renderer/Private/PostProcess/PostProcessing.cpp
+++ b/Engine/Source/Runtime/Renderer/Private/PostProcess/PostProcessing.cpp
@@ -60,6 +60,33 @@ static TAutoConsoleVariable<int32> CVarRenderTargetSwitchWorkaround(
 	TEXT("We want this enabled (1) on all 32 bit iOS devices (implemented through DeviceProfiles)."),
 	ECVF_RenderThreadSafe);
 
+static TAutoConsoleVariable<int32> CVarUpsampleQuality(
+	TEXT("r.UpsampleQuality"),
+	3,
+	TEXT(" 0: Nearest filtering\n")
+	TEXT(" 1: Simple Bilinear\n")
+	TEXT(" 2: 4 tap bilinear\n")
+	TEXT(" 3: Directional blur with unsharp mask upsample. (default)"),
+	ECVF_Scalability | ECVF_RenderThreadSafe);
+
+static TAutoConsoleVariable<float> CVarMotionBlurSoftEdgeSize(
+	TEXT("r.MotionBlurSoftEdgeSize"),
+	1.0f,
+	TEXT("Defines how lange object motion blur is blurred (percent of screen width) to allow soft edge motion blur.\n")
+	TEXT("This scales linearly with the size (up to a maximum of 32 samples, 2.5 is about 18 samples) and with screen resolution\n")
+	TEXT("Smaller values are better for performance and provide more accurate motion vectors but the blurring outside the object is reduced.\n")
+	TEXT("If needed this can be exposed like the other motionblur settings.\n")
+	TEXT(" 0:off (not free and does never completely disable), >0, 1.0 (default)"),
+	ECVF_RenderThreadSafe);
+
+static TAutoConsoleVariable<float> CVarSSSSS(
+	TEXT("r.SSSSS"),
+	0.0f,
+	TEXT("Very experimental screen space subsurface scattering on non unlit/lit materials\n")
+	TEXT("0: off\n")
+	TEXT("x: SSS radius in world space e.g. 10"),
+	ECVF_RenderThreadSafe);
+
 IMPLEMENT_SHADER_TYPE(,FPostProcessVS,TEXT("PostProcessBloom"),TEXT("MainPostprocessCommonVS"),SF_Vertex);
 
 // -------------------------------------------------------
@@ -781,8 +808,7 @@ void FPostProcessing::Process(const FViewInfo& View, TRefCountPtr<IPooledRenderT
 		{
 			// Screen Space Subsurface Scattering
 			{
-				static const auto ICVar = IConsoleManager::Get().FindTConsoleVariableDataFloat(TEXT("r.SSSSS"));
-				float Radius = ICVar->GetValueOnRenderThread();
+				float Radius = CVarSSSSS.GetValueOnRenderThread();
 
 				if(Radius > 0 && !bSimpleDynamicLighting)
 				{
@@ -896,12 +922,7 @@ void FPostProcessing::Process(const FViewInfo& View, TRefCountPtr<IPooledRenderT
 
 					SoftEdgeVelocity = FRenderingCompositeOutputRef(QuarterResVelocity);
 
-					float MotionBlurSoftEdgeSize;
-					{
-						static const auto ICVar = IConsoleManager::Get().FindTConsoleVariableDataFloat(TEXT("r.MotionBlurSoftEdgeSize"));
-
-						MotionBlurSoftEdgeSize = ICVar->GetValueOnRenderThread();
-					}
+					float MotionBlurSoftEdgeSize = CVarMotionBlurSoftEdgeSize.GetValueOnRenderThread();
 
 					if(MotionBlurSoftEdgeSize > 0.01f)
 					{
diff --git a/Engine/Source/Runtime/Renderer/Private/PostProcess/RenderTargetPool.cpp b/Engine/Source/Runtime/Renderer/Private/PostProcess/RenderTargetPool.cpp
--- a/Engine/Source/Runtime/Renderer/Private/PostProcess/RenderTargetPool.cpp
+++ b/Engine/Source/Runtime/Renderer/Private/PostProcess/RenderTargetPool.cpp
@@ -45,6 +45,8 @@ void RenderTargetPoolEvents(const TArray<FString>& Args)
 		UE_LOG(LogRenderTargetPool, Display, TEXT("r.DumpRenderTargetPoolEvents is now disabled, use r.DumpRenderTargetPoolEvents <SizeInKB> to enable or r.DumpRenderTargetPoolEvents ? for help"));
 	}
 }
+
+// CVars and commands
 static FAutoConsoleCommand GRenderTargetPoolEventsCmd(
 	TEXT("r.RenderTargetPool.Events"),
 	TEXT("Visualize the render target pool events over time in one frame. Optional parameter defines threshold in KB.\n")
@@ -52,6 +54,16 @@ static FAutoConsoleCommand GRenderTargetPoolEventsCmd(
 	FConsoleCommandWithArgsDelegate::CreateStatic(RenderTargetPoolEvents)
 	);
 
+#if !(UE_BUILD_SHIPPING || UE_BUILD_TEST)
+static TAutoConsoleVariable<int32> CVarRenderTargetPoolTest(
+	TEXT("r.RenderTargetPoolTest"),
+	0,
+	TEXT("Clears the texture returned by the rendertarget pool with a special color\n")
+	TEXT("so we can see better which passes would need to clear. Doesn't work on volume textures and non rendertargets yet.\n")
+	TEXT(" 0:off (default), 1:on"),
+	ECVF_Cheat | ECVF_RenderThreadSafe);
+#endif
+
 bool FRenderTargetPool::IsEventRecordingEnabled() const
 {
 #if !(UE_BUILD_SHIPPING || UE_BUILD_TEST)
@@ -305,12 +317,10 @@ bool FRenderTargetPool::FindFreeElement(const FPooledRenderTargetDesc& Desc, TRe
 
 #if !(UE_BUILD_SHIPPING || UE_BUILD_TEST)
 	{
-		static const auto ICVar = IConsoleManager::Get().FindTConsoleVariableDataInt(TEXT("r.RenderTargetPoolTest"));
-		
 		//@todo-rco: RHIPacketList
 		FRHICommandList& RHICmdList = FRHICommandList::GetNullRef();
 
-		if(ICVar->GetValueOnRenderThread())
+		if(CVarRenderTargetPoolTest.GetValueOnRenderThread())
 		{
 			if(Found->GetDesc().TargetableFlags & TexCreate_RenderTargetable)
 			{
diff --git a/Engine/Source/Runtime/Renderer/Private/PostProcess/SceneFilterRendering.cpp b/Engine/Source/Runtime/Renderer/Private/PostProcess/SceneFilterRendering.cpp
--- a/Engine/Source/Runtime/Renderer/Private/PostProcess/SceneFilterRendering.cpp
+++ b/Engine/Source/Runtime/Renderer/Private/PostProcess/SceneFilterRendering.cpp
@@ -84,12 +84,21 @@ IMPLEMENT_UNIFORM_BUFFER_STRUCT(FDrawRectangleParameters,TEXT("DrawRectanglePara
 
 typedef TUniformBufferRef<FDrawRectangleParameters> FDrawRectangleBufferRef;
 
+
+static TAutoConsoleVariable<int32> CVarDrawRectangleOptimization(
+	TEXT("r.DrawRectangleOptimization"),
+	1,
+	TEXT("Controls an optimization for DrawRectangle(). When enabled a triangle can be used to draw a quad in certain situations (viewport sized quad).\n")
+	TEXT("Using a triangle allows for slightly faster post processing in lower resolutions but can not always be used.\n")
+	TEXT(" 0: Optimization is disabled, DrawDenormalizedQuad always render with quad\n")
+	TEXT(" 1: Optimization is enabled, a triangle can be rendered where specified (default)"),
+	ECVF_RenderThreadSafe);
+
 static void DoDrawRectangleFlagOverride(EDrawRectangleFlags& Flags)
 {
 #if !(UE_BUILD_SHIPPING || UE_BUILD_TEST)
 	// Determine triangle draw mode
-	static auto* TriangleModeCVar = IConsoleManager::Get().FindTConsoleVariableDataInt(TEXT("r.DrawRectangleOptimization"));
-	int Value = TriangleModeCVar->GetValueOnRenderThread();
+	int Value = CVarDrawRectangleOptimization.GetValueOnRenderThread();
 
 	if(!Value)
 	{
diff --git a/Engine/Source/Runtime/Renderer/Private/PostProcess/SceneRenderTargets.cpp b/Engine/Source/Runtime/Renderer/Private/PostProcess/SceneRenderTargets.cpp
--- a/Engine/Source/Runtime/Renderer/Private/PostProcess/SceneRenderTargets.cpp
+++ b/Engine/Source/Runtime/Renderer/Private/PostProcess/SceneRenderTargets.cpp
@@ -57,7 +57,7 @@ static TAutoConsoleVariable<int32> CVarCustomDepth(
 	ECVF_RenderThreadSafe
 	);
 
-TAutoConsoleVariable<int32> CVarMobileMSAA(
+static TAutoConsoleVariable<int32> CVarMobileMSAA(
 	TEXT("r.MobileMSAA"),
 	0,
 	TEXT("Use MSAA instead of Temporal AA on mobile:\n")
@@ -67,6 +67,16 @@ TAutoConsoleVariable<int32> CVarMobileMSAA(
 	ECVF_RenderThreadSafe
 	);
 
+static TAutoConsoleVariable<int32> CVarGBufferFormat(
+	TEXT("r.GBufferFormat"),
+	1,
+	TEXT("Defines the memory layout used for the GBuffer.\n")
+	TEXT("(affects performance, mostly through bandwidth, quality of normals and material attributes).\n")
+	TEXT(" 0: lower precision (8bit per component, for profiling)\n")
+	TEXT(" 1: low precision (default)\n")
+	TEXT(" 5: high precision"),
+	ECVF_RenderThreadSafe);
+
 /** The global render targets used for scene rendering. */
 TGlobalResource<FSceneRenderTargets> GSceneRenderTargets;
 
@@ -143,17 +153,12 @@ void FSceneRenderTargets::Allocate(const FSceneViewFamily& ViewFamily)
 	check(DesiredBufferSize.X > 0 && DesiredBufferSize.Y > 0);
 	QuantizeBufferSize(DesiredBufferSize.X, DesiredBufferSize.Y);
 
-	int GBufferFormat;
-	{
-		static const auto CVar = IConsoleManager::Get().FindTConsoleVariableDataInt(TEXT("r.GBufferFormat"));
-	
-		GBufferFormat = CVar->GetValueOnRenderThread();
-	}
+	int GBufferFormat = CVarGBufferFormat.GetValueOnRenderThread();
 
 	int SceneColorFormat;
 	{
 		static const auto CVar = IConsoleManager::Get().FindTConsoleVariableDataInt(TEXT("r.SceneColorFormat"));
-	
+
 		SceneColorFormat = CVar->GetValueOnRenderThread();
 	}
 		
@@ -916,6 +921,14 @@ const FTexture2DRHIRef& FSceneRenderTargets::GetEditorPrimitivesDepth()
 	return (const FTexture2DRHIRef&)EditorPrimitivesDepth->GetRenderTargetItem().TargetableTexture;
 }
 
+static TAutoConsoleVariable<int32> CVarSetSeperateTranslucencyEnabled(
+	TEXT("r.SeparateTranslucency"),
+	1,
+	TEXT("Allows to disable the separate translucency feature (all translucency is rendered in separate RT and composited\n")
+	TEXT("after DOF, if not specified otherwise in the material).\n")
+	TEXT(" 0: off (translucency is affected by depth of field)\n")
+	TEXT(" 1: on costs GPU performance and memory but keeps translucency unaffected by Depth of Fieled. (default)"),
+	ECVF_RenderThreadSafe);
 
 bool FSceneRenderTargets::IsSeparateTranslucencyActive(const FViewInfo& View) const
 {
diff --git a/Engine/Source/Runtime/Renderer/Private/ReflectionEnvironment.cpp b/Engine/Source/Runtime/Renderer/Private/ReflectionEnvironment.cpp
--- a/Engine/Source/Runtime/Renderer/Private/ReflectionEnvironment.cpp
+++ b/Engine/Source/Runtime/Renderer/Private/ReflectionEnvironment.cpp
@@ -30,13 +30,20 @@ static TAutoConsoleVariable<int32> CVarDiffuseFromCaptures(
 	TEXT(" 0 is off (default), 1 is on"),
 	ECVF_RenderThreadSafe);
 
+#if !(UE_BUILD_SHIPPING || UE_BUILD_TEST)
+static TAutoConsoleVariable<int32> CVarReflectionEnvironment(
+	TEXT("r.ReflectionEnvironment"),
+	1,
+	TEXT("0:off, 1:on and blend with scene, 2:on and overwrite scene.\n")
+	TEXT("Whether to render the reflection environment feature, which implements local reflections through Reflection Capture actors."),
+	ECVF_Cheat | ECVF_RenderThreadSafe);
+#endif 
+
 // to avoid having direct access from many places
 static int GetReflectionEnvironmentCVar()
 {
 #if !(UE_BUILD_SHIPPING || UE_BUILD_TEST)
-	static const auto CVar = IConsoleManager::Get().FindTConsoleVariableDataInt(TEXT("r.ReflectionEnvironment"));
-	check(CVar);
-	return CVar->GetValueOnAnyThread();
+	return CVarReflectionEnvironment.GetValueOnAnyThread();
 #endif
 
 	// on, default mode
diff --git a/Engine/Source/Runtime/Renderer/Private/SceneRendering.cpp b/Engine/Source/Runtime/Renderer/Private/SceneRendering.cpp
--- a/Engine/Source/Runtime/Renderer/Private/SceneRendering.cpp
+++ b/Engine/Source/Runtime/Renderer/Private/SceneRendering.cpp
@@ -33,6 +33,13 @@ static TAutoConsoleVariable<int32> CVarAllowOcclusionQueries(
 	ECVF_RenderThreadSafe
 	);
 
+static TAutoConsoleVariable<float> CVarDemosaicVposOffset(
+	TEXT("r.DemosaicVposOffset"),
+	0.0f,
+	TEXT("This offset is added to the rasterized position used for demosaic in the ES2 tonemapping shader. It exists to workaround driver bugs on some Android devices that have a half-pixel offset."),
+	ECVF_RenderThreadSafe);
+
+
 #if !(UE_BUILD_SHIPPING || UE_BUILD_TEST)
 static TAutoConsoleVariable<float> CVarGeneralPurposeTweak(
 	TEXT("r.GeneralPurposeTweak"),
@@ -438,11 +445,7 @@ TUniformBufferRef<FViewUniformShaderParameters> FViewInfo::CreateUniformBuffer(
 
 	ViewUniformShaderParameters.DemosaicVposOffset = 0.0f;
 	{
-		static auto* DemosaicVposOffsetCvar = IConsoleManager::Get().FindTConsoleVariableDataFloat(TEXT("r.DemosaicVposOffset"));
-		if (DemosaicVposOffsetCvar)
-		{
-			ViewUniformShaderParameters.DemosaicVposOffset = DemosaicVposOffsetCvar->GetValueOnRenderThread();
-		}
+		ViewUniformShaderParameters.DemosaicVposOffset = CVarDemosaicVposOffset.GetValueOnRenderThread();
 	}
 
 	ViewUniformShaderParameters.IndirectLightingColorScale = FVector(FinalPostProcessSettings.IndirectLightingColor.R * FinalPostProcessSettings.IndirectLightingIntensity,
diff --git a/Engine/Source/Runtime/Renderer/Private/SceneVisibility.cpp b/Engine/Source/Runtime/Renderer/Private/SceneVisibility.cpp
--- a/Engine/Source/Runtime/Renderer/Private/SceneVisibility.cpp
+++ b/Engine/Source/Runtime/Renderer/Private/SceneVisibility.cpp
@@ -46,6 +46,12 @@ static FAutoConsoleVariableRef CVarMinScreenRadiusForCSMDepth(
 	ECVF_RenderThreadSafe
 	);
 
+static TAutoConsoleVariable<int32> CVarTemporalAASamples(
+	TEXT("r.TemporalAASamples"),
+	8,
+	TEXT("Number of jittered positions for temporal AA (4, 8=default, 16, 32, 64)."),
+	ECVF_RenderThreadSafe);
+
 #if PLATFORM_MAC // @todo: disabled until rendering problems with HZB occlusion in OpenGL are solved
 static int32 GHZBOcclusion = 0;
 #else
@@ -76,6 +82,12 @@ static TAutoConsoleVariable<int32> CVarLightShaftQuality(
 	TEXT("  1: on (default)\n"),
 	ECVF_Scalability | ECVF_RenderThreadSafe);
 
+static TAutoConsoleVariable<float> CVarStaticMeshLODDistanceScale(
+	TEXT("r.StaticMeshLODDistanceScale"),
+	1.0f,
+	TEXT("Scale factor for the distance used in computing discrete LOD for static meshes. (0.25-1)"),
+	ECVF_Scalability | ECVF_RenderThreadSafe);
+
 /** Distance fade cvars */
 static int32 GDisableLODFade = false;
 static FAutoConsoleVariableRef CVarDisableLODFade( TEXT("r.DisableLODFade"), GDisableLODFade, TEXT("Disable fading for distance culling"), ECVF_RenderThreadSafe );
@@ -775,9 +787,7 @@ static void MarkRelevantStaticMeshesForView(
 	// outside of the loop to be more efficient
 	int32 ForcedLODLevel = (View.Family->EngineShowFlags.LOD) ? GetCVarForceLOD() : 0;
 	
-	static const auto* StaticMeshLODDistanceScale = IConsoleManager::Get().FindTConsoleVariableDataFloat(TEXT("r.StaticMeshLODDistanceScale"));
-	check(StaticMeshLODDistanceScale);
-	const float LODScale = StaticMeshLODDistanceScale->GetValueOnRenderThread();
+	const float LODScale = CVarStaticMeshLODDistanceScale.GetValueOnRenderThread();
 	const float InvLODScale = 1.0f / LODScale;
 
 	const float MinScreenRadiusForCSMDepthSquared = GMinScreenRadiusForCSMDepth * GMinScreenRadiusForCSMDepth;
@@ -978,17 +988,14 @@ void FSceneRenderer::PreVisibilityFrameSetup()
 			View.ViewRect.Width() * View.ViewRect.Height();
 		View.OneOverNumPossiblePixels = NumPossiblePixels > 0.0 ? 1.0f / NumPossiblePixels : 0.0f;
 
-
-		// Subpixel jitter for temporal AA
-		static const auto CVarTemporalAASamples = IConsoleManager::Get().FindTConsoleVariableDataInt( TEXT("r.TemporalAASamples") );
-
 		// Still need no jitter to be set for temporal feedback on SSR (it is enabled even when temporal AA is off).
 		View.TemporalJitterPixelsX = 0.0f;
 		View.TemporalJitterPixelsY = 0.0f;
 
 		if( View.FinalPostProcessSettings.AntiAliasingMethod == AAM_TemporalAA && ViewState )
 		{
-			int32 TemporalAASamples = CVarTemporalAASamples->GetValueOnRenderThread();
+			// Subpixel jitter for temporal AA
+			int32 TemporalAASamples = CVarTemporalAASamples.GetValueOnRenderThread();
 		
 			if( TemporalAASamples > 1 )
 			{
diff --git a/Engine/Source/Runtime/ShaderCore/Private/ShaderCore.cpp b/Engine/Source/Runtime/ShaderCore/Private/ShaderCore.cpp
--- a/Engine/Source/Runtime/ShaderCore/Private/ShaderCore.cpp
+++ b/Engine/Source/Runtime/ShaderCore/Private/ShaderCore.cpp
@@ -12,13 +12,18 @@
 
 FSHAHash GGlobalShaderMapHash;
 
+static TAutoConsoleVariable<int32> CVarShaderDevelopmentMode(
+	TEXT("r.ShaderDevelopmentMode"),
+	0,
+	TEXT("0: Default, 1: Enable various shader development utilities, such as the ability to retry on failed shader compile, and extra logging as shaders are compiled."),
+	ECVF_Default);
+
 void UpdateShaderDevelopmentMode()
 {
 	// Keep LogShaders verbosity in sync with r.ShaderDevelopmentMode
 	// r.ShaderDevelopmentMode==1 results in all LogShaders log messages being displayed
-	static const auto CVar0 = IConsoleManager::Get().FindTConsoleVariableDataInt(TEXT("r.ShaderDevelopmentMode"));
 	bool bLogShadersUnsuppressed = UE_LOG_ACTIVE(LogShaders, Log);
-	bool bDesiredLogShadersUnsuppressed = CVar0->GetValueOnGameThread() == 1;
+	bool bDesiredLogShadersUnsuppressed = CVarShaderDevelopmentMode.GetValueOnGameThread() == 1;
 
 	if (bLogShadersUnsuppressed != bDesiredLogShadersUnsuppressed)
 	{
diff --git a/Engine/Source/Runtime/StreamingPauseRendering/Private/StreamingPauseRendering.cpp b/Engine/Source/Runtime/StreamingPauseRendering/Private/StreamingPauseRendering.cpp
--- a/Engine/Source/Runtime/StreamingPauseRendering/Private/StreamingPauseRendering.cpp
+++ b/Engine/Source/Runtime/StreamingPauseRendering/Private/StreamingPauseRendering.cpp
@@ -11,6 +11,12 @@ StreamingPauseRendering.cpp: Streaming pause implementation.
 
 IMPLEMENT_MODULE(FStreamingPauseRenderingModule, StreamingPauseRendering);
 
+static TAutoConsoleVariable<int32> CVarRenderLastFrameInStreamingPause(
+	TEXT("r.RenderLastFrameInStreamingPause"),
+	1,
+	TEXT("If 1 the previous frame is displayed during streaming pause. If zero the screen is left black."),
+	ECVF_RenderThreadSafe);
+
 FStreamingPauseRenderingModule::FStreamingPauseRenderingModule()
 : Viewport(NULL)
 , ViewportWidget(NULL)
@@ -58,8 +64,7 @@ void FStreamingPauseRenderingModule::BeginStreamingPause( FViewport* GameViewpor
 	);
 
 	//Render the current scene to a new viewport.
-	static const auto* CVarRenderLastFrameInStreamingPause = IConsoleManager::Get().FindTConsoleVariableDataInt(TEXT("r.RenderLastFrameInStreamingPause"));
-	if( CVarRenderLastFrameInStreamingPause && CVarRenderLastFrameInStreamingPause->GetValueOnGameThread() != 0 ) 
+	if (CVarRenderLastFrameInStreamingPause.GetValueOnGameThread() != 0)
 	{
 		ViewportClient = GameViewport->GetClient();
 
