Added a custom icu::CharacterIterator implementation, which will attempt to interface with FText/FString as efficiently as it can based on the native platform string format.

TTP# 336464 - Editor: Finish the Multiline Editable Text Block

There are two versions of this, which are typedef'd to FICUTextCharacterIterator based on the current platform string traits.

FICUTextCharacterIterator_NativeUTF16 is used on platforms where TCHAR is already in a native UTF-16 format (as used by ICU), so performs no conversion and just passes the internal string data through to an icu::UCharCharacterIterator. Windows, Xbox One, and PS4 use this iterator.

FICUTextCharacterIterator_ConvertToUnicodeString is used on platforms where TCHAR is not in a native UTF-16 format, and performs a conversion to an icu::UnicodeString which is then passed into an icu::StringCharacterIterator - this has the same behaviour and cost as every platforms had prior to this change. OS X, iOS, Android, and Linux use this iterator.

Other changes:
- Renamed the Convert functions in ICUUtilities to either ConvertString or ConvertChar32 to make it clearer what they were doing.
- Removed the ConvertChar32 functions which were attempting to convert a UChar32 back to a TCHAR (as this is neither safe or sane).
- Added IsUnicodeEncoded to FGenericPlatformString. This is used along with the sizeof(TCHAR) at compile time to work out if the platform is using UTF-16.

diff --git a/Engine/Source/Runtime/Core/Private/Internationalization/ICUCharacterBoundaryIterator.cpp b/Engine/Source/Runtime/Core/Private/Internationalization/ICUCharacterBoundaryIterator.cpp
--- a/Engine/Source/Runtime/Core/Private/Internationalization/ICUCharacterBoundaryIterator.cpp
+++ b/Engine/Source/Runtime/Core/Private/Internationalization/ICUCharacterBoundaryIterator.cpp
@@ -4,7 +4,7 @@
 #include "CharacterBoundaryIterator.h"
 
 #if UE_ENABLE_ICU
-#include "ICUUtilities.h"
+#include "ICUTextCharacterIterator.h"
 #include <unicode/brkiter.h>
 
 namespace
@@ -22,8 +22,7 @@ public:
 	FImplementation(const TCHAR* const String, const int32 StringLength)
 		: ICUBreakIterator( CreateGraphemeBreakIterator() )
 	{
-		ICUUtilities::Convert(FString(String, StringLength), ICUString);
-		ICUBreakIterator->setText(ICUString);
+		ICUBreakIterator->adoptText(new FICUTextCharacterIterator(String, StringLength)); // ICUBreakIterator takes ownership of this instance
 	}
 
 	int32 GetCurrentPosition() const
@@ -63,7 +62,6 @@ public:
 
 private:
 	TSharedRef<icu::BreakIterator> ICUBreakIterator;
-	icu::UnicodeString ICUString;
 };
 
 FCharacterBoundaryIterator::FCharacterBoundaryIterator(const FText& Text)
diff --git a/Engine/Source/Runtime/Core/Private/Internationalization/ICUCulture.cpp b/Engine/Source/Runtime/Core/Private/Internationalization/ICUCulture.cpp
--- a/Engine/Source/Runtime/Core/Private/Internationalization/ICUCulture.cpp
+++ b/Engine/Source/Runtime/Core/Private/Internationalization/ICUCulture.cpp
@@ -82,18 +82,14 @@ FString FCulture::FICUCultureImplementation::GetDisplayName() const
 {
 	icu::UnicodeString ICUResult;
 	ICULocale.getDisplayName(ICUResult);
-	FString Result;
-	ICUUtilities::Convert(ICUResult, Result);
-	return Result;
+	return ICUUtilities::ConvertString(ICUResult);
 }
 
 FString FCulture::FICUCultureImplementation::GetEnglishName() const
 {
 	icu::UnicodeString ICUResult;
 	ICULocale.getDisplayName(icu::Locale("en"), ICUResult);
-	FString Result;
-	ICUUtilities::Convert(ICUResult, Result);
-	return Result;
+	return ICUUtilities::ConvertString(ICUResult);
 }
 
 int FCulture::FICUCultureImplementation::GetKeyboardLayoutId() const
@@ -115,9 +111,7 @@ FString FCulture::FICUCultureImplementation::GetNativeName() const
 {
 	icu::UnicodeString ICUResult;
 	ICULocale.getDisplayName(ICULocale, ICUResult);
-	FString Result;
-	ICUUtilities::Convert(ICUResult, Result);
-	return Result;
+	return ICUUtilities::ConvertString(ICUResult);
 }
 
 FString FCulture::FICUCultureImplementation::GetNativeLanguage() const
@@ -125,12 +119,12 @@ FString FCulture::FICUCultureImplementation::GetNativeLanguage() const
 	icu::UnicodeString ICUNativeLanguage;
 	ICULocale.getDisplayLanguage(ICULocale, ICUNativeLanguage);
 	FString NativeLanguage;
-	ICUUtilities::Convert(ICUNativeLanguage, NativeLanguage);
+	ICUUtilities::ConvertString(ICUNativeLanguage, NativeLanguage);
 
 	icu::UnicodeString ICUNativeScript;
 	ICULocale.getDisplayScript(ICULocale, ICUNativeScript);
 	FString NativeScript;
-	ICUUtilities::Convert(ICUNativeScript, NativeScript);
+	ICUUtilities::ConvertString(ICUNativeScript, NativeScript);
 
 	if ( !NativeScript.IsEmpty() )
 	{
@@ -144,12 +138,12 @@ FString FCulture::FICUCultureImplementation::GetNativeRegion() const
 	icu::UnicodeString ICUNativeCountry;
 	ICULocale.getDisplayCountry(ICULocale, ICUNativeCountry);
 	FString NativeCountry;
-	ICUUtilities::Convert(ICUNativeCountry, NativeCountry);
+	ICUUtilities::ConvertString(ICUNativeCountry, NativeCountry);
 
 	icu::UnicodeString ICUNativeVariant;
 	ICULocale.getDisplayVariant(ICULocale, ICUNativeVariant);
 	FString NativeVariant;
-	ICUUtilities::Convert(ICUNativeVariant, NativeVariant);
+	ICUUtilities::ConvertString(ICUNativeVariant, NativeVariant);
 
 	if ( !NativeVariant.IsEmpty() )
 	{
@@ -306,7 +300,7 @@ TSharedRef<const icu::DateFormat> FCulture::FICUCultureImplementation::GetDateFo
 TSharedRef<const icu::DateFormat> FCulture::FICUCultureImplementation::GetTimeFormatter(const EDateTimeStyle::Type TimeStyle, const FString& TimeZone) const
 {
 	icu::UnicodeString InputTimeZoneID;
-	ICUUtilities::Convert(TimeZone, InputTimeZoneID, false);
+	ICUUtilities::ConvertString(TimeZone, InputTimeZoneID, false);
 
 	const TSharedRef<const icu::DateFormat> DefaultFormatter( ICUTimeFormat );
 
@@ -346,7 +340,7 @@ TSharedRef<const icu::DateFormat> FCulture::FICUCultureImplementation::GetTimeFo
 TSharedRef<const icu::DateFormat> FCulture::FICUCultureImplementation::GetDateTimeFormatter(const EDateTimeStyle::Type DateStyle, const EDateTimeStyle::Type TimeStyle, const FString& TimeZone) const
 {
 	icu::UnicodeString InputTimeZoneID;
-	ICUUtilities::Convert(TimeZone, InputTimeZoneID, false);
+	ICUUtilities::ConvertString(TimeZone, InputTimeZoneID, false);
 
 	const TSharedRef<const icu::DateFormat> DefaultFormatter( ICUDateTimeFormat );
 
diff --git a/Engine/Source/Runtime/Core/Private/Internationalization/ICULineBreakIterator.cpp b/Engine/Source/Runtime/Core/Private/Internationalization/ICULineBreakIterator.cpp
--- a/Engine/Source/Runtime/Core/Private/Internationalization/ICULineBreakIterator.cpp
+++ b/Engine/Source/Runtime/Core/Private/Internationalization/ICULineBreakIterator.cpp
@@ -4,7 +4,7 @@
 #include "LineBreakIterator.h"
 
 #if UE_ENABLE_ICU
-#include "ICUUtilities.h"
+#include "ICUTextCharacterIterator.h"
 #include <unicode/brkiter.h>
 
 namespace
@@ -22,8 +22,7 @@ public:
 	FImplementation(const TCHAR* const String, const int32 StringLength)
 		: ICUBreakIterator( CreateLineBreakIterator() )
 	{
-		ICUUtilities::Convert(FString(String, StringLength), ICUString);
-		ICUBreakIterator->setText(ICUString);
+		ICUBreakIterator->adoptText(new FICUTextCharacterIterator(String, StringLength)); // ICUBreakIterator takes ownership of this instance
 	}
 
 	int32 GetCurrentPosition() const
@@ -63,7 +62,6 @@ public:
 
 private:
 	TSharedRef<icu::BreakIterator> ICUBreakIterator;
-	icu::UnicodeString ICUString;
 };
 
 FLineBreakIterator::FLineBreakIterator(const FText& Text)
diff --git a/Engine/Source/Runtime/Core/Private/Internationalization/ICURegex.cpp b/Engine/Source/Runtime/Core/Private/Internationalization/ICURegex.cpp
--- a/Engine/Source/Runtime/Core/Private/Internationalization/ICURegex.cpp
+++ b/Engine/Source/Runtime/Core/Private/Internationalization/ICURegex.cpp
@@ -12,7 +12,7 @@ namespace
 	TSharedRef<icu::RegexPattern> CreateRegexPattern(const FString& SourceString)
 	{
 		icu::UnicodeString ICUSourceString;
-		ICUUtilities::Convert(SourceString, ICUSourceString);
+		ICUUtilities::ConvertString(SourceString, ICUSourceString);
 
 		UErrorCode ICUStatus = U_ZERO_ERROR;
 		return MakeShareable( icu::RegexPattern::compile(ICUSourceString, 0, ICUStatus) );
@@ -46,7 +46,7 @@ namespace
 class FRegexMatcherImplementation
 {
 public:
-	FRegexMatcherImplementation(const FRegexPatternImplementation& Pattern, const FString& InputString) : ICUInputString( ICUUtilities::Convert(InputString) ), ICURegexMatcher( CreateRegexMatcher(Pattern, ICUInputString) )
+	FRegexMatcherImplementation(const FRegexPatternImplementation& Pattern, const FString& InputString) : ICUInputString( ICUUtilities::ConvertString(InputString) ), ICURegexMatcher( CreateRegexMatcher(Pattern, ICUInputString) )
 	{
 	}
 
diff --git a/Engine/Source/Runtime/Core/Private/Internationalization/ICUText.cpp b/Engine/Source/Runtime/Core/Private/Internationalization/ICUText.cpp
--- a/Engine/Source/Runtime/Core/Private/Internationalization/ICUText.cpp
+++ b/Engine/Source/Runtime/Core/Private/Internationalization/ICUText.cpp
@@ -11,6 +11,13 @@
 #include <unicode/numfmt.h>
 #include <unicode/msgfmt.h>
 #include "ICUUtilities.h"
+#include "ICUTextCharacterIterator.h"
+
+bool FText::IsWhitespace( const TCHAR Char )
+{
+	const UChar32 ICUChar = ICUUtilities::ConvertChar32(Char);
+	return u_isWhitespace(ICUChar) != 0;
+}
 
 FText FText::AsDate(const FDateTime& DateTime, const EDateTimeStyle::Type DateStyle, const TSharedPtr<FCulture, ESPMode::ThreadSafe>& TargetCulture)
 {
@@ -27,7 +34,7 @@ FText FText::AsDate(const FDateTime& DateTime, const EDateTimeStyle::Type DateSt
 	ICUDateFormat->format(ICUDate, FormattedString);
 
 	FString NativeString;
-	ICUUtilities::Convert(FormattedString, NativeString);
+	ICUUtilities::ConvertString(FormattedString, NativeString);
 
 	FText ResultText = FText::CreateChronologicalText(NativeString);
 	ResultText.History = MakeShareable(new FTextHistory_AsDate(DateTime, DateStyle, TargetCulture));
@@ -50,7 +57,7 @@ FText FText::AsTime(const FDateTime& DateTime, const EDateTimeStyle::Type TimeSt
 	ICUDateFormat->format(ICUDate, FormattedString);
 
 	FString NativeString;
-	ICUUtilities::Convert(FormattedString, NativeString);
+	ICUUtilities::ConvertString(FormattedString, NativeString);
 
 	FText ResultText = FText::CreateChronologicalText(NativeString);
 	ResultText.History = MakeShareable(new FTextHistory_AsTime(DateTime, TimeStyle, TimeZone, TargetCulture));
@@ -97,7 +104,7 @@ FText FText::AsDateTime(const FDateTime& DateTime, const EDateTimeStyle::Type Da
 	ICUDateFormat->format(ICUDate, FormattedString);
 
 	FString NativeString;
-	ICUUtilities::Convert(FormattedString, NativeString);
+	ICUUtilities::ConvertString(FormattedString, NativeString);
 
 	FText ResultText = FText::CreateChronologicalText(NativeString);
 	ResultText.History = MakeShareable(new FTextHistory_AsDateTime(DateTime, DateStyle, TimeStyle, TimeZone, TargetCulture));
@@ -133,15 +140,20 @@ FText FText::AsMemory(SIZE_T NumBytes, const FNumberFormattingOptions* const Opt
 
 int32 FText::CompareTo( const FText& Other, const ETextComparisonLevel::Type ComparisonLevel ) const
 {
-	UErrorCode ICUStatus = U_ZERO_ERROR;
-
 	const TSharedRef<const icu::Collator, ESPMode::ThreadSafe> Collator( FInternationalization::Get().GetCurrentCulture()->Implementation->GetCollator(ComparisonLevel) );
 
-	icu::UnicodeString A;
-	ICUUtilities::Convert(DisplayString.Get(), A);
-	icu::UnicodeString B;
-	ICUUtilities::Convert(Other.DisplayString.Get(), B);
-	UCollationResult Result = Collator->compare(A, B, ICUStatus);
+	// Create an iterator for 'this' so that we can interface with ICU
+	UCharIterator DisplayStringICUIterator;
+	FICUTextCharacterIterator DisplayStringIterator(&DisplayString.Get());
+	uiter_setCharacterIterator(&DisplayStringICUIterator, &DisplayStringIterator);
+
+	// Create an iterator for 'Other' so that we can interface with ICU
+	UCharIterator OtherDisplayStringICUIterator;
+	FICUTextCharacterIterator OtherDisplayStringIterator(&Other.DisplayString.Get());
+	uiter_setCharacterIterator(&OtherDisplayStringICUIterator, &OtherDisplayStringIterator);
+
+	UErrorCode ICUStatus = U_ZERO_ERROR;
+	const UCollationResult Result = Collator->compare(DisplayStringICUIterator, OtherDisplayStringICUIterator, ICUStatus);
 
 	return Result;
 }
@@ -153,17 +165,7 @@ int32 FText::CompareToCaseIgnored( const FText& Other ) const
 
 bool FText::EqualTo( const FText& Other, const ETextComparisonLevel::Type ComparisonLevel ) const
 {
-	UErrorCode ICUStatus = U_ZERO_ERROR;
-
-	const TSharedRef<const icu::Collator, ESPMode::ThreadSafe> Collator( FInternationalization::Get().GetCurrentCulture()->Implementation->GetCollator(ComparisonLevel) );
-
-	icu::UnicodeString A;
-	ICUUtilities::Convert(DisplayString.Get(), A);
-	icu::UnicodeString B;
-	ICUUtilities::Convert(Other.DisplayString.Get(), B);
-	UCollationResult Result = Collator->compare(A, B, ICUStatus);
-
-	return Result == 0;
+	return CompareTo(Other, ComparisonLevel) == 0;
 }
 
 bool FText::EqualToCaseIgnored( const FText& Other ) const
@@ -173,17 +175,6 @@ bool FText::EqualToCaseIgnored( const FText& Other ) const
 
 class FText::FSortPredicate::FSortPredicateImplementation
 {
-private:
-	icu::CollationKey GetCollationKeyFromText(const FText& Text)
-	{
-		UErrorCode ICUStatus = U_ZERO_ERROR;
-		icu::CollationKey CollationKey;
-		const FString& String = Text.ToString();
-		const UChar* const Characters = reinterpret_cast<const UChar* const>(*String);
-		int32_t Length = String.Len();
-		return ICUCollator->getCollationKey(Characters, Length, CollationKey, ICUStatus);
-	}
-
 public:
 	FSortPredicateImplementation(const ETextComparisonLevel::Type InComparisonLevel)
 		: ComparisonLevel(InComparisonLevel)
@@ -193,9 +184,20 @@ public:
 
 	bool Compare(const FText& A, const FText& B)
 	{
-		const icu::CollationKey ACollationKey = GetCollationKeyFromText(A);
-		const icu::CollationKey BCollationKey = GetCollationKeyFromText(B);
-		return (ACollationKey.compareTo(BCollationKey) == icu::Collator::EComparisonResult::GREATER ? false : true);
+		// Create an iterator for 'A' so that we can interface with ICU
+		UCharIterator ADisplayStringICUIterator;
+		FICUTextCharacterIterator ADisplayStringIterator(&A.DisplayString.Get());
+		uiter_setCharacterIterator(&ADisplayStringICUIterator, &ADisplayStringIterator);
+
+		// Create an iterator for 'B' so that we can interface with ICU
+		UCharIterator BDisplayStringICUIterator;
+		FICUTextCharacterIterator BDisplayStringIterator(&B.DisplayString.Get());
+		uiter_setCharacterIterator(&BDisplayStringICUIterator, &BDisplayStringIterator);
+
+		UErrorCode ICUStatus = U_ZERO_ERROR;
+		const UCollationResult Result = ICUCollator->compare(ADisplayStringICUIterator, BDisplayStringICUIterator, ICUStatus);
+
+		return Result != UCOL_GREATER;
 	}
 
 private:
@@ -221,7 +223,7 @@ void FText::GetFormatPatternParameters(const FText& Pattern, TArray<FString>& Pa
 	const FString& NativePatternString = Pattern.ToString();
 
 	icu::UnicodeString ICUPatternString;
-	ICUUtilities::Convert(NativePatternString, ICUPatternString);
+	ICUUtilities::ConvertString(NativePatternString, ICUPatternString);
 
 	UParseError ICUParseError;
 	icu::MessagePattern ICUMessagePattern(ICUPatternString, &(ICUParseError), ICUStatus);
@@ -240,7 +242,7 @@ void FText::GetFormatPatternParameters(const FText& Pattern, TArray<FString>& Pa
 
 			bool bIsCaseSensitiveUnique = true;
 			FString CurrentArgument;
-			ICUUtilities::Convert(ICUMessagePattern.getSubstring(Part), CurrentArgument);
+			ICUUtilities::ConvertString(ICUMessagePattern.getSubstring(Part), CurrentArgument);
 
 			for(auto It = ParameterNames.CreateConstIterator(); It; ++It)
 			{
@@ -276,7 +278,7 @@ FText FText::FormatInternal(const FText& Pattern, const FFormatNamedArguments& A
 	{
 		const FString& NativeArgumentName = It.Key();
 		icu::UnicodeString ICUArgumentName;
-		ICUUtilities::Convert(NativeArgumentName, ICUArgumentName, false);
+		ICUUtilities::ConvertString(NativeArgumentName, ICUArgumentName, false);
 		ArgumentNames.Add( ICUArgumentName );
 
 		const FFormatArgumentValue& ArgumentValue = It.Value();
@@ -312,7 +314,7 @@ FText FText::FormatInternal(const FText& Pattern, const FFormatNamedArguments& A
 
 				FString NativeStringValue = bInRebuildAsSource? ArgumentValue.TextValue->BuildSourceString() : ArgumentValue.TextValue->ToString();
 				icu::UnicodeString ICUStringValue;
-				ICUUtilities::Convert(NativeStringValue, ICUStringValue, false);
+				ICUUtilities::ConvertString(NativeStringValue, ICUStringValue, false);
 				ArgumentValues.Add( icu::Formattable( ICUStringValue ) );
 			}
 			break;
@@ -321,7 +323,7 @@ FText FText::FormatInternal(const FText& Pattern, const FFormatNamedArguments& A
 
 	const FString& NativePatternString = bInRebuildAsSource? Pattern.BuildSourceString() : Pattern.ToString();
 	icu::UnicodeString ICUPatternString;
-	ICUUtilities::Convert(NativePatternString, ICUPatternString);
+	ICUUtilities::ConvertString(NativePatternString, ICUPatternString);
 
 	UErrorCode ICUStatus = U_ZERO_ERROR;
 	UParseError ICUParseError;
@@ -339,7 +341,7 @@ FText FText::FormatInternal(const FText& Pattern, const FFormatNamedArguments& A
 	}
 
 	FString NativeResultString;
-	ICUUtilities::Convert(ICUResultString, NativeResultString);
+	ICUUtilities::ConvertString(ICUResultString, NativeResultString);
 
 	FText ResultText(NativeResultString);
 	ResultText.History = MakeShareable(new FTextHistory_NamedFormat(Pattern, Arguments));
@@ -393,7 +395,7 @@ FText FText::FormatInternal(const FText& Pattern, const FFormatOrderedArguments&
 
 				FString NativeStringValue = bInRebuildAsSource? ArgumentValue.TextValue->BuildSourceString() : ArgumentValue.TextValue->ToString();
 				icu::UnicodeString ICUStringValue;
-				ICUUtilities::Convert(NativeStringValue, ICUStringValue, false);
+				ICUUtilities::ConvertString(NativeStringValue, ICUStringValue, false);
 				ArgumentValues.Add( icu::Formattable( ICUStringValue ) );
 			}
 			break;
@@ -402,7 +404,7 @@ FText FText::FormatInternal(const FText& Pattern, const FFormatOrderedArguments&
 
 	const FString& NativePatternString = bInRebuildAsSource? Pattern.BuildSourceString() : Pattern.ToString();
 	icu::UnicodeString ICUPatternString;
-	ICUUtilities::Convert(NativePatternString, ICUPatternString);
+	ICUUtilities::ConvertString(NativePatternString, ICUPatternString);
 
 	UErrorCode ICUStatus = U_ZERO_ERROR;
 	UParseError ICUParseError;
@@ -420,7 +422,7 @@ FText FText::FormatInternal(const FText& Pattern, const FFormatOrderedArguments&
 	}
 
 	FString NativeResultString;
-	ICUUtilities::Convert(ICUResultString, NativeResultString);
+	ICUUtilities::ConvertString(ICUResultString, NativeResultString);
 
 	FText ResultText = FText(NativeResultString);
 
@@ -451,7 +453,7 @@ FText FText::FormatInternal(const FText& Pattern, const TArray< FFormatArgumentD
 	{
 		const FString& ArgumentName = InArguments[x].ArgumentName.ToString();
 		icu::UnicodeString ICUArgumentName;
-		ICUUtilities::Convert(ArgumentName, ICUArgumentName, false);
+		ICUUtilities::ConvertString(ArgumentName, ICUArgumentName, false);
 		ArgumentNames.Add( ICUArgumentName );
 
 		const FFormatArgumentValue& ArgumentValue = InArguments[x].ArgumentValue;
@@ -487,7 +489,7 @@ FText FText::FormatInternal(const FText& Pattern, const TArray< FFormatArgumentD
 
 				FString StringValue = bInRebuildAsSource? ArgumentValue.TextValue->BuildSourceString() : ArgumentValue.TextValue->ToString();
 				icu::UnicodeString ICUStringValue;
-				ICUUtilities::Convert(StringValue, ICUStringValue, false);
+				ICUUtilities::ConvertString(StringValue, ICUStringValue, false);
 				ArgumentValues.Add( ICUStringValue );
 			}
 			break;
@@ -496,7 +498,7 @@ FText FText::FormatInternal(const FText& Pattern, const TArray< FFormatArgumentD
 
 	const FString& PatternString = bInRebuildAsSource? Pattern.BuildSourceString() : Pattern.ToString();
 	icu::UnicodeString ICUPatternString;
-	ICUUtilities::Convert(PatternString, ICUPatternString, false);
+	ICUUtilities::ConvertString(PatternString, ICUPatternString, false);
 
 	UErrorCode ICUStatus = U_ZERO_ERROR;
 	UParseError ICUParseError;
@@ -514,7 +516,7 @@ FText FText::FormatInternal(const FText& Pattern, const TArray< FFormatArgumentD
 	}
 
 	FString NativeResultString;
-	ICUUtilities::Convert(Result, NativeResultString);
+	ICUUtilities::ConvertString(Result, NativeResultString);
 
 	FText ResultText = FText(NativeResultString);
 
diff --git a/Engine/Source/Runtime/Core/Private/Internationalization/ICUText.h b/Engine/Source/Runtime/Core/Private/Internationalization/ICUText.h
--- a/Engine/Source/Runtime/Core/Private/Internationalization/ICUText.h
+++ b/Engine/Source/Runtime/Core/Private/Internationalization/ICUText.h
@@ -23,7 +23,7 @@ FText FText::AsNumberTemplate(T1 Val, const FNumberFormattingOptions* const Opti
 	ICUDecimalFormat->format(FormattableVal, FormattedString, ICUStatus);
 
 	FString NativeString;
-	ICUUtilities::Convert(FormattedString, NativeString);
+	ICUUtilities::ConvertString(FormattedString, NativeString);
 
 	FText ReturnText = FText::CreateNumericalText(NativeString);
 	ReturnText.History = MakeShareable(new FTextHistory_AsNumber(Val, Options, TargetCulture));
@@ -43,7 +43,7 @@ FText FText::AsCurrencyTemplate(T1 Val, const FNumberFormattingOptions* const Op
 	ICUDecimalFormat->format(FormattableVal, FormattedString, ICUStatus);
 
 	FString NativeString;
-	ICUUtilities::Convert(FormattedString, NativeString);
+	ICUUtilities::ConvertString(FormattedString, NativeString);
 
 	FText ReturnText = FText::CreateNumericalText(NativeString);
 	ReturnText.History = MakeShareable(new FTextHistory_AsCurrency(Val, Options, TargetCulture));
@@ -63,7 +63,7 @@ FText FText::AsPercentTemplate(T1 Val, const FNumberFormattingOptions* const Opt
 	ICUDecimalFormat->format(FormattableVal, FormattedString, ICUStatus);
 
 	FString NativeString;
-	ICUUtilities::Convert(FormattedString, NativeString);
+	ICUUtilities::ConvertString(FormattedString, NativeString);
 
 	FText ReturnText = FText::CreateNumericalText(NativeString);
 	ReturnText.History = MakeShareable(new FTextHistory_AsPercent(Val, Options, TargetCulture));
diff --git a/Engine/Source/Runtime/Core/Private/Internationalization/ICUTextCharacterIterator.cpp b/Engine/Source/Runtime/Core/Private/Internationalization/ICUTextCharacterIterator.cpp
new file mode 100644
--- /dev/null
+++ b/Engine/Source/Runtime/Core/Private/Internationalization/ICUTextCharacterIterator.cpp
@@ -0,0 +1,127 @@
+// Copyright 1998-2014 Epic Games, Inc. All Rights Reserved.
+
+#include "CorePrivate.h"
+
+#if UE_ENABLE_ICU
+
+#include "ICUTextCharacterIterator.h"
+#include "ICUUtilities.h"
+
+UOBJECT_DEFINE_RTTI_IMPLEMENTATION(FICUTextCharacterIterator_NativeUTF16)
+
+FICUTextCharacterIterator_NativeUTF16::FICUTextCharacterIterator_NativeUTF16(const FText& InText)
+	: InternalString(InText.ToString())
+	, StringPtr(&InternalString)
+{
+	setText(reinterpret_cast<const UChar*>(**StringPtr), StringPtr->Len()); // scary cast from TCHAR* to UChar* so that this builds on platforms where TCHAR isn't UTF-16 (but we won't use it!)
+}
+
+FICUTextCharacterIterator_NativeUTF16::FICUTextCharacterIterator_NativeUTF16(const FString& InString)
+	: InternalString(InString)
+	, StringPtr(&InternalString)
+{
+	setText(reinterpret_cast<const UChar*>(**StringPtr), StringPtr->Len()); // scary cast from TCHAR* to UChar* so that this builds on platforms where TCHAR isn't UTF-16 (but we won't use it!)
+}
+
+FICUTextCharacterIterator_NativeUTF16::FICUTextCharacterIterator_NativeUTF16(const TCHAR* const InString, const int32 InStringLength)
+	: InternalString(InString, InStringLength)
+	, StringPtr(&InternalString)
+{
+	setText(reinterpret_cast<const UChar*>(**StringPtr), StringPtr->Len()); // scary cast from TCHAR* to UChar* so that this builds on platforms where TCHAR isn't UTF-16 (but we won't use it!)
+}
+
+FICUTextCharacterIterator_NativeUTF16::FICUTextCharacterIterator_NativeUTF16(FString&& InString)
+	: InternalString(InString)
+	, StringPtr(&InternalString)
+{
+	setText(reinterpret_cast<const UChar*>(**StringPtr), StringPtr->Len()); // scary cast from TCHAR* to UChar* so that this builds on platforms where TCHAR isn't UTF-16 (but we won't use it!)
+}
+
+FICUTextCharacterIterator_NativeUTF16::FICUTextCharacterIterator_NativeUTF16(const FString* const InString)
+	: StringPtr(InString)
+{
+	check(StringPtr);
+	setText(reinterpret_cast<const UChar*>(**StringPtr), StringPtr->Len()); // scary cast from TCHAR* to UChar* so that this builds on platforms where TCHAR isn't UTF-16 (but we won't use it!)
+}
+
+FICUTextCharacterIterator_NativeUTF16::FICUTextCharacterIterator_NativeUTF16(const FICUTextCharacterIterator_NativeUTF16& Other)
+	: icu::UCharCharacterIterator(Other)
+	, InternalString(*Other.StringPtr)
+	, StringPtr(&InternalString)
+{
+	setText(reinterpret_cast<const UChar*>(**StringPtr), StringPtr->Len()); // scary cast from TCHAR* to UChar* so that this builds on platforms where TCHAR isn't UTF-16 (but we won't use it!)
+}
+
+FICUTextCharacterIterator_NativeUTF16::~FICUTextCharacterIterator_NativeUTF16()
+{
+}
+
+FICUTextCharacterIterator_NativeUTF16& FICUTextCharacterIterator_NativeUTF16::operator=(const FICUTextCharacterIterator_NativeUTF16& Other)
+{
+	if(this != &Other)
+	{
+		icu::UCharCharacterIterator::operator=(Other);
+		InternalString = *Other.StringPtr;
+		StringPtr = &InternalString;
+		setText(reinterpret_cast<const UChar*>(**StringPtr), StringPtr->Len()); // scary cast from TCHAR* to UChar* so that this builds on platforms where TCHAR isn't UTF-16 (but we won't use it!)
+	}
+	return *this;
+}
+
+icu::CharacterIterator* FICUTextCharacterIterator_NativeUTF16::clone() const
+{
+	return new FICUTextCharacterIterator_NativeUTF16(*this);
+}
+
+
+UOBJECT_DEFINE_RTTI_IMPLEMENTATION(FICUTextCharacterIterator_ConvertToUnicodeString)
+
+FICUTextCharacterIterator_ConvertToUnicodeString::FICUTextCharacterIterator_ConvertToUnicodeString(const FText& InText)
+	: icu::StringCharacterIterator(ICUUtilities::ConvertString(InText.ToString()))
+{
+}
+
+FICUTextCharacterIterator_ConvertToUnicodeString::FICUTextCharacterIterator_ConvertToUnicodeString(const FString& InString)
+	: icu::StringCharacterIterator(ICUUtilities::ConvertString(InString))
+{
+}
+
+FICUTextCharacterIterator_ConvertToUnicodeString::FICUTextCharacterIterator_ConvertToUnicodeString(const TCHAR* const InString, const int32 InStringLength)
+	: icu::StringCharacterIterator(ICUUtilities::ConvertString(FString(InString, InStringLength)))
+{
+}
+
+FICUTextCharacterIterator_ConvertToUnicodeString::FICUTextCharacterIterator_ConvertToUnicodeString(FString&& InString)
+	: icu::StringCharacterIterator(ICUUtilities::ConvertString(InString))
+{
+}
+
+FICUTextCharacterIterator_ConvertToUnicodeString::FICUTextCharacterIterator_ConvertToUnicodeString(const FString* const InString)
+	: icu::StringCharacterIterator(ICUUtilities::ConvertString(*InString))
+{
+}
+
+FICUTextCharacterIterator_ConvertToUnicodeString::FICUTextCharacterIterator_ConvertToUnicodeString(const FICUTextCharacterIterator_ConvertToUnicodeString& Other)
+	: icu::StringCharacterIterator(Other)
+{
+}
+
+FICUTextCharacterIterator_ConvertToUnicodeString::~FICUTextCharacterIterator_ConvertToUnicodeString()
+{
+}
+
+FICUTextCharacterIterator_ConvertToUnicodeString& FICUTextCharacterIterator_ConvertToUnicodeString::operator=(const FICUTextCharacterIterator_ConvertToUnicodeString& Other)
+{
+	if(this != &Other)
+	{
+		icu::StringCharacterIterator::operator=(Other);
+	}
+	return *this;
+}
+
+icu::CharacterIterator* FICUTextCharacterIterator_ConvertToUnicodeString::clone() const
+{
+	return new FICUTextCharacterIterator_ConvertToUnicodeString(*this);
+}
+
+#endif
diff --git a/Engine/Source/Runtime/Core/Private/Internationalization/ICUTextCharacterIterator.h b/Engine/Source/Runtime/Core/Private/Internationalization/ICUTextCharacterIterator.h
new file mode 100644
--- /dev/null
+++ b/Engine/Source/Runtime/Core/Private/Internationalization/ICUTextCharacterIterator.h
@@ -0,0 +1,87 @@
+// Copyright 1998-2014 Epic Games, Inc. All Rights Reserved.
+
+#pragma once
+
+#if UE_ENABLE_ICU
+
+#include <unicode/uchriter.h>
+#include <unicode/schriter.h>
+
+/**
+ * Implementation of an ICU CharacterIterator that is able to iterate over an FText/FString directly since the native string format for this platform is already UTF-16 (as used by ICU)
+ * This can be used with the ICU break iterator types by passing it in via adoptText(...)
+ * Note: Do not use this type directly! Use the FICUTextCharacterIterator typedef, which will be set correctly for your platform
+ */
+class FICUTextCharacterIterator_NativeUTF16 : public icu::UCharCharacterIterator
+{
+public:
+	/** Construct from a string by value */
+	FICUTextCharacterIterator_NativeUTF16(const FText& InText);
+	FICUTextCharacterIterator_NativeUTF16(const FString& InString);
+	FICUTextCharacterIterator_NativeUTF16(const TCHAR* const InString, const int32 InStringLength);
+
+	/** Construct from a string by stealing (moving) the data */
+	FICUTextCharacterIterator_NativeUTF16(FString&& InString);
+
+	/** Construct from a string by referencing the original data (the string *must* exist for the duration of FICUTextCharacterIterator's lifetime */
+	FICUTextCharacterIterator_NativeUTF16(const FString* const InString);
+
+	FICUTextCharacterIterator_NativeUTF16(const FICUTextCharacterIterator_NativeUTF16& Other);
+	virtual ~FICUTextCharacterIterator_NativeUTF16();
+
+	FICUTextCharacterIterator_NativeUTF16& operator=(const FICUTextCharacterIterator_NativeUTF16& Other);
+
+	// ICU RTTI
+	static UClassID getStaticClassID();
+	virtual UClassID getDynamicClassID() const override;
+
+	// Implement icu::CharacterIterator
+	virtual icu::CharacterIterator* clone() const override;
+
+private:
+	/** Internal copy of the string (if constructed by value/move). Do not access directly, use StringPtr instead */
+	FString InternalString;
+
+	/** Pointer to either InternalString, or the external string we were told to reference */
+	const FString* StringPtr;
+};
+
+
+/**
+ * Implementation of an ICU CharacterIterator that converts an FText/FString to an icu::UnicodeString, since the native string format for this platform is not UTF-16 (as used by ICU)
+ * This can be used with the ICU break iterator types by passing it in via adoptText(...)
+ * Note: Do not use this type directly! Use the FICUTextCharacterIterator typedef, which will be set correctly for your platform
+ */
+class FICUTextCharacterIterator_ConvertToUnicodeString : public icu::StringCharacterIterator
+{
+public:
+	/** Construct from a string by value */
+	FICUTextCharacterIterator_ConvertToUnicodeString(const FText& InText);
+	FICUTextCharacterIterator_ConvertToUnicodeString(const FString& InString);
+	FICUTextCharacterIterator_ConvertToUnicodeString(const TCHAR* const InString, const int32 InStringLength);
+
+	/** Construct from a string by stealing (moving) the data */
+	FICUTextCharacterIterator_ConvertToUnicodeString(FString&& InString);
+
+	/** Construct from a string by referencing the original data (the string *must* exist for the duration of FICUTextCharacterIterator's lifetime */
+	FICUTextCharacterIterator_ConvertToUnicodeString(const FString* const InString);
+
+	FICUTextCharacterIterator_ConvertToUnicodeString(const FICUTextCharacterIterator_ConvertToUnicodeString& Other);
+	virtual ~FICUTextCharacterIterator_ConvertToUnicodeString();
+
+	FICUTextCharacterIterator_ConvertToUnicodeString& operator=(const FICUTextCharacterIterator_ConvertToUnicodeString& Other);
+
+	// ICU RTTI
+	static UClassID getStaticClassID();
+	virtual UClassID getDynamicClassID() const override;
+
+	// Implement icu::CharacterIterator
+	virtual icu::CharacterIterator* clone() const override;
+};
+
+/** Work out the best character iterator to use based upon our native platform string traits */
+template <bool IsUnicode, size_t TCHARSize> struct FICUTextCharacterIterator_PlatformSpecific { typedef FICUTextCharacterIterator_ConvertToUnicodeString Type; };
+template <> struct FICUTextCharacterIterator_PlatformSpecific<true, 2> { typedef FICUTextCharacterIterator_NativeUTF16 Type; }; // A unicode encoding with a wchar_t size of 2 bytes is assumed to be UTF-16
+typedef FICUTextCharacterIterator_PlatformSpecific<FPlatformString::IsUnicodeEncoded, sizeof(TCHAR)>::Type FICUTextCharacterIterator;
+
+#endif
diff --git a/Engine/Source/Runtime/Core/Private/Internationalization/ICUUtilities.cpp b/Engine/Source/Runtime/Core/Private/Internationalization/ICUUtilities.cpp
--- a/Engine/Source/Runtime/Core/Private/Internationalization/ICUUtilities.cpp
+++ b/Engine/Source/Runtime/Core/Private/Internationalization/ICUUtilities.cpp
@@ -8,7 +8,7 @@
 
 namespace ICUUtilities
 {
-	void Convert(const FString& Source, icu::UnicodeString& Destination, const bool ShouldNullTerminate)
+	void ConvertString(const FString& Source, icu::UnicodeString& Destination, const bool ShouldNullTerminate)
 	{
 		if( Source.Len() )
 		{
@@ -47,14 +47,14 @@ namespace ICUUtilities
 		}
 	}
 
-	icu::UnicodeString Convert(const FString& Source, const bool ShouldNullTerminate)
+	icu::UnicodeString ConvertString(const FString& Source, const bool ShouldNullTerminate)
 	{
 		icu::UnicodeString Destination;
-		Convert(Source, Destination, ShouldNullTerminate);
+		ConvertString(Source, Destination, ShouldNullTerminate);
 		return Destination;
 	}
 
-	void Convert(const icu::UnicodeString& Source, FString& Destination)
+	void ConvertString(const icu::UnicodeString& Source, FString& Destination)
 	{
 		if( Source.length() )
 		{
@@ -90,14 +90,14 @@ namespace ICUUtilities
 		}
 	}
 
-	FString Convert(const icu::UnicodeString& Source)
+	FString ConvertString(const icu::UnicodeString& Source)
 	{
 		FString Destination;
-		Convert(Source, Destination);
+		ConvertString(Source, Destination);
 		return Destination;
 	}
 
-	void Convert(const TCHAR Source, UChar32& Destination)
+	void ConvertChar32(const TCHAR Source, UChar32& Destination)
 	{
 		const char* const SourceEncoding = FPlatformString::GetEncodingName();
 	
@@ -117,28 +117,10 @@ namespace ICUUtilities
 		ucnv_close(ICUConverter);
 	}
 
-	UChar32 Convert(const TCHAR Source)
+	UChar32 ConvertChar32(const TCHAR Source)
 	{
 		UChar32 Destination = 0;
-		Convert(Source, Destination);
-		return Destination;
-	}
-
-	void Convert(const UChar32 Source, TCHAR& Destination)
-	{
-		// ICU has ucnv_getNextUChar as an optimization for converting a single character, but it doesn't seem to have an inverse of this
-		// Because of this we're using the string variants instead (which will be slower)
-		icu::UnicodeString SourceString;
-		SourceString.setTo(Source);
-		FString DestinationString;
-		Convert(SourceString, DestinationString);
-		Destination = DestinationString.Len() ? DestinationString[0] : 0;
-	}
-
-	TCHAR Convert(const UChar32 Source)
-	{
-		TCHAR Destination = 0;
-		Convert(Source, Destination);
+		ConvertChar32(Source, Destination);
 		return Destination;
 	}
 }
diff --git a/Engine/Source/Runtime/Core/Private/Internationalization/ICUUtilities.h b/Engine/Source/Runtime/Core/Private/Internationalization/ICUUtilities.h
--- a/Engine/Source/Runtime/Core/Private/Internationalization/ICUUtilities.h
+++ b/Engine/Source/Runtime/Core/Private/Internationalization/ICUUtilities.h
@@ -7,14 +7,14 @@
 
 namespace ICUUtilities
 {
-	void Convert(const FString& Source, icu::UnicodeString& Destination, const bool ShouldNullTerminate = true);
-	icu::UnicodeString Convert(const FString& Source, const bool ShouldNullTerminate = true);
-	void Convert(const icu::UnicodeString& Source, FString& Destination);
-	FString Convert(const icu::UnicodeString& Source);
+	/** Convert FString <-> icu::UnicodeString */
+	void ConvertString(const FString& Source, icu::UnicodeString& Destination, const bool ShouldNullTerminate = true);
+	icu::UnicodeString ConvertString(const FString& Source, const bool ShouldNullTerminate = true);
+	void ConvertString(const icu::UnicodeString& Source, FString& Destination);
+	FString ConvertString(const icu::UnicodeString& Source);
 
-	void Convert(const TCHAR Source, UChar32& Destination);
-	UChar32 Convert(const TCHAR Source);
-	void Convert(const UChar32 Source, TCHAR& Destination);
-	TCHAR Convert(const UChar32 Source);
+	/** Convert TCHAR -> UChar32 */
+	void ConvertChar32(const TCHAR Source, UChar32& Destination);
+	UChar32 ConvertChar32(const TCHAR Source);
 }
 #endif
\ No newline at end of file
diff --git a/Engine/Source/Runtime/Core/Private/Internationalization/ICUWordBreakIterator.cpp b/Engine/Source/Runtime/Core/Private/Internationalization/ICUWordBreakIterator.cpp
--- a/Engine/Source/Runtime/Core/Private/Internationalization/ICUWordBreakIterator.cpp
+++ b/Engine/Source/Runtime/Core/Private/Internationalization/ICUWordBreakIterator.cpp
@@ -4,7 +4,7 @@
 #include "WordBreakIterator.h"
 
 #if UE_ENABLE_ICU
-#include "ICUUtilities.h"
+#include "ICUTextCharacterIterator.h"
 #include <unicode/brkiter.h>
 
 namespace
@@ -22,8 +22,7 @@ public:
 	FImplementation(const TCHAR* const String, const int32 StringLength)
 		: ICUBreakIterator(CreateWordBreakIterator())
 	{
-		ICUUtilities::Convert(FString(String, StringLength), ICUString);
-		ICUBreakIterator->setText(ICUString);
+		ICUBreakIterator->adoptText(new FICUTextCharacterIterator(String, StringLength)); // ICUBreakIterator takes ownership of this instance
 	}
 
 	int32 GetCurrentPosition() const
@@ -63,7 +62,6 @@ public:
 
 private:
 	TSharedRef<icu::BreakIterator> ICUBreakIterator;
-	icu::UnicodeString ICUString;
 };
 
 FWordBreakIterator::FWordBreakIterator(const FText& Text)
diff --git a/Engine/Source/Runtime/Core/Private/Internationalization/Internationalization.cpp b/Engine/Source/Runtime/Core/Private/Internationalization/Internationalization.cpp
--- a/Engine/Source/Runtime/Core/Private/Internationalization/Internationalization.cpp
+++ b/Engine/Source/Runtime/Core/Private/Internationalization/Internationalization.cpp
@@ -83,7 +83,7 @@ void FInternationalization::GetTimeZonesIDs(TArray<FString>& TimeZonesIDs) const
 		if(ICUString)
 		{
 			FString NativeString;
-			ICUUtilities::Convert(*ICUString, NativeString);
+			ICUUtilities::ConvertString(*ICUString, NativeString);
 			TimeZonesIDs.Add( NativeString );
 		}
 	} while( ICUString );
diff --git a/Engine/Source/Runtime/Core/Private/Internationalization/Text.cpp b/Engine/Source/Runtime/Core/Private/Internationalization/Text.cpp
--- a/Engine/Source/Runtime/Core/Private/Internationalization/Text.cpp
+++ b/Engine/Source/Runtime/Core/Private/Internationalization/Text.cpp
@@ -410,16 +410,6 @@ FText FText::TrimPrecedingAndTrailing( const FText& InText )
 	return NewText;
 }
 
-bool FText::IsWhitespace( const TCHAR Char )
-{
-#if UE_ENABLE_ICU
-	const UChar32 ICUChar = ICUUtilities::Convert(Char);
-	return u_isWhitespace(ICUChar) != 0;
-#else
-	return FChar::IsWhitespace(Char);
-#endif
-}
-
 FText FText::Format(const FText& Fmt,const FText& v1)
 {
 	FFormatOrderedArguments Arguments;
diff --git a/Engine/Source/Runtime/Core/Private/Tests/Internationalization/TextTest.cpp b/Engine/Source/Runtime/Core/Private/Tests/Internationalization/TextTest.cpp
--- a/Engine/Source/Runtime/Core/Private/Tests/Internationalization/TextTest.cpp
+++ b/Engine/Source/Runtime/Core/Private/Tests/Internationalization/TextTest.cpp
@@ -632,13 +632,13 @@ bool FICUTextTest::RunTest (const FString& Parameters)
 
 	// ---------------------------------------------------------------------
 
-	ICUUtilities::Convert(SourceString, ICUString);
+	ICUUtilities::ConvertString(SourceString, ICUString);
 	if (SourceString.Len() != ICUString.length())
 	{
 		AddError(FString::Printf(TEXT("icu::UnicodeString is the incorrect length (%d; expected %d)."), ICUString.length(), SourceString.Len()));
 	}
 
-	ICUUtilities::Convert(ICUString, ConversionBackStr);
+	ICUUtilities::ConvertString(ICUString, ConversionBackStr);
 	if (ICUString.length() != ConversionBackStr.Len())
 	{
 		AddError(FString::Printf(TEXT("FString is the incorrect length (%d; expected %d)."), ConversionBackStr.Len(), ICUString.length()));
@@ -650,13 +650,13 @@ bool FICUTextTest::RunTest (const FString& Parameters)
 
 	// ---------------------------------------------------------------------
 
-	ICUUtilities::Convert(SourceString2, ICUString);
+	ICUUtilities::ConvertString(SourceString2, ICUString);
 	if (SourceString2.Len() != ICUString.length())
 	{
 		AddError(FString::Printf(TEXT("icu::UnicodeString is the incorrect length (%d; expected %d)."), ICUString.length(), SourceString2.Len()));
 	}
 
-	ICUUtilities::Convert(ICUString, ConversionBackStr);
+	ICUUtilities::ConvertString(ICUString, ConversionBackStr);
 	if (ICUString.length() != ConversionBackStr.Len())
 	{
 		AddError(FString::Printf(TEXT("FString is the incorrect length (%d; expected %d)."), ConversionBackStr.Len(), ICUString.length()));
@@ -668,13 +668,13 @@ bool FICUTextTest::RunTest (const FString& Parameters)
 
 	// ---------------------------------------------------------------------
 
-	ICUUtilities::Convert(SourceString, ICUString);
+	ICUUtilities::ConvertString(SourceString, ICUString);
 	if (SourceString.Len() != ICUString.length())
 	{
 		AddError(FString::Printf(TEXT("icu::UnicodeString is the incorrect length (%d; expected %d)."), ICUString.length(), SourceString.Len()));
 	}
 
-	ICUUtilities::Convert(ICUString, ConversionBackStr);
+	ICUUtilities::ConvertString(ICUString, ConversionBackStr);
 	if (ICUString.length() != ConversionBackStr.Len())
 	{
 		AddError(FString::Printf(TEXT("FString is the incorrect length (%d; expected %d)."), ConversionBackStr.Len(), ICUString.length()));
diff --git a/Engine/Source/Runtime/Core/Public/Apple/ApplePlatformString.h b/Engine/Source/Runtime/Core/Public/Apple/ApplePlatformString.h
--- a/Engine/Source/Runtime/Core/Public/Apple/ApplePlatformString.h
+++ b/Engine/Source/Runtime/Core/Public/Apple/ApplePlatformString.h
@@ -64,6 +64,8 @@ struct FApplePlatformString : public FStandardPlatformString
 	{
 		return "UTF-32LE";
 	}
+
+	static const bool IsUnicodeEncoded = true;
 };
 
 typedef FApplePlatformString FPlatformString;
diff --git a/Engine/Source/Runtime/Core/Public/GenericPlatform/GenericPlatformString.h b/Engine/Source/Runtime/Core/Public/GenericPlatform/GenericPlatformString.h
--- a/Engine/Source/Runtime/Core/Public/GenericPlatform/GenericPlatformString.h
+++ b/Engine/Source/Runtime/Core/Public/GenericPlatform/GenericPlatformString.h
@@ -49,6 +49,11 @@ struct FGenericPlatformString
 	template <typename Encoding>
 	static const TCHAR* GetEncodingTypeName();
 
+	/**
+	 * True if the encoding type of the string is some form of unicode
+	 */
+	static const bool IsUnicodeEncoded = false;
+
 
 	/**
 	 * Metafunction which tests whether a given character type represents a fixed-width encoding.
diff --git a/Engine/Source/Runtime/Core/Public/Linux/LinuxPlatformString.h b/Engine/Source/Runtime/Core/Public/Linux/LinuxPlatformString.h
--- a/Engine/Source/Runtime/Core/Public/Linux/LinuxPlatformString.h
+++ b/Engine/Source/Runtime/Core/Public/Linux/LinuxPlatformString.h
@@ -28,6 +28,8 @@ struct FLinuxPlatformString : public FStandardPlatformString
 	{
 		return "UTF-32LE";
 	}
+
+	static const bool IsUnicodeEncoded = true;
 };
 
 typedef FLinuxPlatformString FPlatformString;
diff --git a/Engine/Source/Runtime/Core/Public/Windows/WindowsPlatformString.h b/Engine/Source/Runtime/Core/Public/Windows/WindowsPlatformString.h
--- a/Engine/Source/Runtime/Core/Public/Windows/WindowsPlatformString.h
+++ b/Engine/Source/Runtime/Core/Public/Windows/WindowsPlatformString.h
@@ -23,11 +23,6 @@ struct FWindowsPlatformString : public FMicrosoftPlatformString
 	{
 		::MultiByteToWideChar(CP_ACP,0,Source,LengthM1+1,Dest,LengthM1+1);
 	}*/
-
-	static const ANSICHAR* GetEncodingName()
-	{
-		return "UTF-16LE";
-	}
 };
 
 typedef FWindowsPlatformString FPlatformString;
