Move definition of Swap() to fix compilation on actually conformant compilers

diff --git a/Engine/Source/Runtime/Core/Public/Templates/UnrealTemplate.h b/Engine/Source/Runtime/Core/Public/Templates/UnrealTemplate.h
--- a/Engine/Source/Runtime/Core/Public/Templates/UnrealTemplate.h
+++ b/Engine/Source/Runtime/Core/Public/Templates/UnrealTemplate.h
@@ -85,16 +85,6 @@ template< class T > inline T AlignArbitrary( const T Ptr, uint32 Alignment )
 {
 	return (T) ( ( ((UPTRINT)Ptr + Alignment - 1) / Alignment ) * Alignment );
 }
-template< class T > inline void Swap( T& A, T& B )
-{
-	T Temp = MoveTemp(A);
-	A = MoveTemp(B);
-	B = MoveTemp(Temp);
-}
-template< class T > inline void Exchange( T& A, T& B )
-{
-	Swap(A, B);
-}
 
 /**
  * Chooses between the two parameters based on whether the first is NULL or not.
@@ -417,6 +407,20 @@ FORCEINLINE T&& Forward(typename TRemoveReference<T>::Type&& Obj)
 	return (T&&)Obj;
 }
 
+/**
+ * Swap two values, using moves if possible
+ */
+template< class T > inline void Swap(T& A, T& B)
+{
+	T Temp = MoveTemp(A);
+	A = MoveTemp(B);
+	B = MoveTemp(Temp);
+}
+template< class T > inline void Exchange(T& A, T& B)
+{
+	Swap(A, B);
+}
+
 /**
  * This exists to avoid a Visual Studio bug where using a cast to forward an rvalue reference array argument
  * to a pointer parameter will cause bad code generation.  Wrapping the cast in a function causes the correct
