Fixed missing const on JSON function.

diff --git a/Engine/Source/Runtime/Core/Public/Serialization/Json/JsonDocumentObjectModel.h b/Engine/Source/Runtime/Core/Public/Serialization/Json/JsonDocumentObjectModel.h
--- a/Engine/Source/Runtime/Core/Public/Serialization/Json/JsonDocumentObjectModel.h
+++ b/Engine/Source/Runtime/Core/Public/Serialization/Json/JsonDocumentObjectModel.h
@@ -168,7 +168,7 @@ public:
 		return TSharedPtr<FJsonValue>(new FJsonValueNull());
 	}
 
-	TSharedPtr<FJsonValue> TryGetField( const FString& FieldName )
+	TSharedPtr<FJsonValue> TryGetField( const FString& FieldName ) const
 	{
 		const TSharedPtr<FJsonValue>* Field = Values.Find(FieldName);
 		return (Field != NULL && Field->IsValid()) ? *Field : TSharedPtr<FJsonValue>();
