Delegates: Added RemoveAll() to dynamic multicast delegates

diff --git a/Engine/Source/Runtime/Core/Public/UObject/ScriptDelegates.h b/Engine/Source/Runtime/Core/Public/UObject/ScriptDelegates.h
--- a/Engine/Source/Runtime/Core/Public/UObject/ScriptDelegates.h
+++ b/Engine/Source/Runtime/Core/Public/UObject/ScriptDelegates.h
@@ -296,6 +296,27 @@ public:
 		CompactInvocationList();
 	}
 
+	/**
+	 * Removes all delegate bindings from this multicast delegate's
+	 * invocation list that are bound to the specified object.
+	 *
+	 * This method also compacts the invocation list.
+	 *
+	 * @param InObject The object to remove bindings for.
+	 */
+	void RemoveAll( UObject* Object )
+	{
+		for (int32 BindingIndex = InvocationList.Num() - 1; BindingIndex >= 0; --BindingIndex)
+		{
+			const TScriptDelegate<TWeakPtr>& Binding = InvocationList[BindingIndex];
+
+			if (Binding.IsBoundToObject(Object) || Binding.IsCompactable())
+			{
+				InvocationList.RemoveAtSwap(BindingIndex);
+			}
+		}
+	}
+
 	/**
 	 * Removes all functions from this delegate's invocation list
 	 */
