- Removed sRGB back buffer for Metal (no longer too bright)

diff --git a/Engine/Source/Runtime/Core/Private/IOS/IOSView.cpp b/Engine/Source/Runtime/Core/Private/IOS/IOSView.cpp
--- a/Engine/Source/Runtime/Core/Private/IOS/IOSView.cpp
+++ b/Engine/Source/Runtime/Core/Private/IOS/IOSView.cpp
@@ -70,7 +70,7 @@ id<MTLDevice> GMetalDevice = nil;
 
 			// set the device on the rendering layer and provide a pixel format
 			MetalLayer.device = MetalDevice;
-			MetalLayer.pixelFormat = MTLPixelFormatBGRA8Unorm_sRGB;
+			MetalLayer.pixelFormat = MTLPixelFormatBGRA8Unorm;
 			MetalLayer.framebufferOnly = NO;
 		
 		}
