FString::IsNumeric now accepts a leading '+' in numeric strings, in line with the string to number conversion functions in TCString<>.

diff --git a/Engine/Source/Runtime/Core/Private/Containers/String.cpp b/Engine/Source/Runtime/Core/Private/Containers/String.cpp
--- a/Engine/Source/Runtime/Core/Private/Containers/String.cpp
+++ b/Engine/Source/Runtime/Core/Private/Containers/String.cpp
@@ -491,7 +491,7 @@ bool FString::IsNumeric() const
 
 	TCHAR C = (*this)[0];
 	
-	if( C == '-' || C =='.' || FChar::IsDigit( C ) )
+	if( C == '-' || C == '+' || C =='.' || FChar::IsDigit( C ) )
 	{
 		bool HasDot = (C == '.');
 
