[GitHub] 231 : Optimized FAsyncIOSystemBase so that on every tick it doesn't do expensive string comparisons per IO request
#ttp 338723

diff --git a/Engine/Source/Runtime/Core/Private/Serialization/AsyncIOSystemBase.cpp b/Engine/Source/Runtime/Core/Private/Serialization/AsyncIOSystemBase.cpp
--- a/Engine/Source/Runtime/Core/Private/Serialization/AsyncIOSystemBase.cpp
+++ b/Engine/Source/Runtime/Core/Private/Serialization/AsyncIOSystemBase.cpp
@@ -36,6 +36,7 @@ uint64 FAsyncIOSystemBase::QueueIORequest(
 	IORequest.RequestIndex				= RequestIndex++;
 	IORequest.FileSortKey				= INDEX_NONE;
 	IORequest.FileName					= FileName;
+	IORequest.FileNameHash				= FCrc::StrCrc32<TCHAR>(FileName.ToLower().GetCharArray().GetData());
 	IORequest.Offset					= Offset;
 	IORequest.Size						= Size;
 	IORequest.UncompressedSize			= UncompressedSize;
@@ -78,6 +79,7 @@ uint64 FAsyncIOSystemBase::QueueDestroyHandleRequest(const FString& FileName)
 	FAsyncIORequest IORequest;
 	IORequest.RequestIndex				= RequestIndex++;
 	IORequest.FileName					= FileName;
+	IORequest.FileNameHash				= FCrc::StrCrc32<TCHAR>(FileName.ToLower().GetCharArray().GetData());
 	IORequest.Priority					= AIOP_MIN;
 	IORequest.bIsDestroyHandleRequest	= true;
 
@@ -390,7 +392,7 @@ IFileHandle* FAsyncIOSystemBase::GetCachedFileHandle( const FString& FileName )
 		// Make sure it's valid before caching and using it.
 		if( FileHandle )
 		{
-			NameToHandleMap.Add( *FileName, FileHandle );
+			NameHashToHandleMap.Add(FCrc::StrCrc32<TCHAR>(FileName.ToLower().GetCharArray().GetData()), FileHandle);
 		}
 	}
 
@@ -399,7 +401,12 @@ IFileHandle* FAsyncIOSystemBase::GetCachedFileHandle( const FString& FileName )
 
 IFileHandle* FAsyncIOSystemBase::FindCachedFileHandle( const FString& FileName )
 {
-	return NameToHandleMap.FindRef( FileName );
+	return FindCachedFileHandle(FCrc::StrCrc32<TCHAR>(FileName.ToLower().GetCharArray().GetData()));
+}
+
+IFileHandle* FAsyncIOSystemBase::FindCachedFileHandle(const uint32 FileNameHash)
+{
+	return NameHashToHandleMap.FindRef(FileNameHash);
 }
 
 uint64 FAsyncIOSystemBase::LoadData( 
@@ -595,7 +602,7 @@ void FAsyncIOSystemBase::Tick()
 				FAsyncIORequest& OutstandingRequest = OutstandingRequests[RequestIdx];
 				if( OutstandingRequest.bHasAlreadyRequestedHandleToBeCached == false
 				&&	OutstandingRequest.bIsDestroyHandleRequest == false 
-				&&	FindCachedFileHandle( OutstandingRequest.FileName ) == NULL )
+				&&	FindCachedFileHandle( OutstandingRequest.FileNameHash ) == NULL )
 				{
 					new(FileNamesToCacheHandles)FString(*OutstandingRequest.FileName);
 					OutstandingRequest.bHasAlreadyRequestedHandleToBeCached = true;
@@ -640,12 +647,12 @@ void FAsyncIOSystemBase::Tick()
 		// handle a destroy handle request from the queue
 		if( IORequest.bIsDestroyHandleRequest )
 		{
-			IFileHandle*	FileHandle = FindCachedFileHandle( IORequest.FileName );
+			IFileHandle*	FileHandle = FindCachedFileHandle( IORequest.FileNameHash );
 			if( FileHandle )
 			{
 				// destroy and remove the handle
 				delete FileHandle;
-				NameToHandleMap.Remove(IORequest.FileName);
+				NameHashToHandleMap.Remove(IORequest.FileNameHash);
 			}
 		}
 		else
@@ -739,11 +746,11 @@ void FAsyncIOSystemBase::FlushHandles()
 {
 	FScopeLock ScopeLock( CriticalSection );
 	// Iterate over all file handles, destroy them and empty name to handle map.
-	for( TMap<FString,IFileHandle*>::TIterator It(NameToHandleMap); It; ++It )
+	for (TMap<uint32, IFileHandle*>::TIterator It(NameHashToHandleMap); It; ++It)
 	{
 		delete It.Value();
 	}
-	NameToHandleMap.Empty();
+	NameHashToHandleMap.Empty();
 }
 
 /** Thread used for async IO manager */
diff --git a/Engine/Source/Runtime/Core/Private/Serialization/AsyncIOSystemBase.h b/Engine/Source/Runtime/Core/Private/Serialization/AsyncIOSystemBase.h
--- a/Engine/Source/Runtime/Core/Private/Serialization/AsyncIOSystemBase.h
+++ b/Engine/Source/Runtime/Core/Private/Serialization/AsyncIOSystemBase.h
@@ -177,6 +177,8 @@ protected:
 
 		/** Name of file.																			*/
 		FString				FileName;
+		/** Hash of the name of the file. This avoids the need to do a string comparison.			*/
+		uint32				FileNameHash;
 		/** Offset into file.																		*/
 		int64				Offset;
 		/** Size in bytes of data to read.															*/
@@ -311,6 +313,15 @@ protected:
 	 */
 	IFileHandle* FindCachedFileHandle( const FString& FileName );
 
+	/**
+	 * Returns cached file handle if found, or NULL if not. This function does
+	 * NOT create any file handles and therefore is not blocking.
+	 *
+	 * @param	FileNameHash	hash of the file name to retrieve cached handle for
+	 * @return	cached file handle, NULL if not cached
+	 */
+	IFileHandle* FindCachedFileHandle(const uint32 FileNameHash);
+
 	/**
 	 * Flushes all file handles.
 	 */
@@ -376,7 +387,7 @@ protected:
 	/** Critical section used to syncronize access to outstanding requests map						*/
 	FCriticalSection*				CriticalSection;
 	/** TMap of file name string hash to file handles												*/
-	TMap<FString,IFileHandle*>		NameToHandleMap;
+	TMap<uint32, IFileHandle*>		NameHashToHandleMap;
 	/** Array of outstanding requests, processed in FIFO order										*/
 	TArray<FAsyncIORequest>			OutstandingRequests;
 	/** Event that is signaled if there are outstanding requests									*/
