Add the ability to try and cast JSON values to different types and be able to test whether they succeed. Also add functions for reading integers.

diff --git a/Engine/Source/Runtime/Core/Private/Serialization/Json/JsonDocumentObjectModel.cpp b/Engine/Source/Runtime/Core/Private/Serialization/Json/JsonDocumentObjectModel.cpp
--- a/Engine/Source/Runtime/Core/Private/Serialization/Json/JsonDocumentObjectModel.cpp
+++ b/Engine/Source/Runtime/Core/Private/Serialization/Json/JsonDocumentObjectModel.cpp
@@ -8,6 +8,78 @@
 const TArray< TSharedPtr<FJsonValue> > FJsonValue::EMPTY_ARRAY;
 const TSharedPtr<FJsonObject> FJsonValue::EMPTY_OBJECT(new FJsonObject());
 
+double FJsonValue::AsNumber() const
+{
+	double Number = 0.0;
+	if(!TryGetNumber(Number))
+	{
+		ErrorMessage(TEXT("Number"));
+	}
+	return Number;
+}
+
+FString FJsonValue::AsString() const 
+{
+	FString String;
+	if(!TryGetString(String))
+	{
+		ErrorMessage(TEXT("String"));
+	}
+	return String;
+}
+
+bool FJsonValue::AsBool() const 
+{
+	bool Bool = false;
+	if(!TryGetBool(Bool))
+	{
+		ErrorMessage(TEXT("Boolean")); 
+	}
+	return Bool;
+}
+
+const TArray< TSharedPtr<FJsonValue> >& FJsonValue::AsArray() const 
+{
+	const TArray< TSharedPtr<FJsonValue> > *Array = &EMPTY_ARRAY;
+	if(!TryGetArray(Array))
+	{
+		ErrorMessage(TEXT("Array")); 
+	}
+	return *Array;
+}
+
+const TSharedPtr<FJsonObject>& FJsonValue::AsObject() const 
+{
+	const TSharedPtr<FJsonObject> *Object = &EMPTY_OBJECT;
+	if(!TryGetObject(Object))
+	{
+		ErrorMessage(TEXT("Object"));
+	}
+	return *Object;
+}
+
+bool FJsonValue::TryGetNumber(int32& OutNumber) const
+{
+	double Double;
+	if(TryGetNumber(Double) && Double >= INT_MIN && Double <= INT_MAX)
+	{
+		OutNumber = (int32)(Double + 0.5);
+		return true;
+	}
+	return false;
+}
+
+bool FJsonValue::TryGetNumber(uint32& OutNumber) const
+{
+	double Double;
+	if(TryGetNumber(Double) && Double >= 0.0 && Double <= UINT_MAX)
+	{
+		OutNumber = (uint32)(Double + 0.5);
+		return true;
+	}
+	return false;
+}
+
 //static 
 bool FJsonValue::CompareEqual(const FJsonValue& Lhs, const FJsonValue& Rhs)
 {
@@ -107,6 +179,24 @@ double FJsonObject::GetNumberField(const FString& FieldName) const
 	return GetField<EJson::None>(FieldName)->AsNumber();
 }
 
+bool FJsonObject::TryGetNumberField(const FString& FieldName, double& OutNumber) const
+{
+	TSharedPtr<FJsonValue> Field = TryGetField(FieldName);
+	return Field.IsValid() && Field->TryGetNumber(OutNumber);
+}
+
+bool FJsonObject::TryGetNumberField(const FString& FieldName, int32& OutNumber) const
+{
+	TSharedPtr<FJsonValue> Field = TryGetField(FieldName);
+	return Field.IsValid() && Field->TryGetNumber(OutNumber);
+}
+
+bool FJsonObject::TryGetNumberField(const FString& FieldName, uint32& OutNumber) const
+{
+	TSharedPtr<FJsonValue> Field = TryGetField(FieldName);
+	return Field.IsValid() && Field->TryGetNumber(OutNumber);
+}
+
 void FJsonObject::SetNumberField( const FString& FieldName, double Number )
 {
 	this->Values.Add( FieldName, MakeShareable(new FJsonValueNumber(Number)) );
@@ -117,6 +207,12 @@ FString FJsonObject::GetStringField(const FString& FieldName) const
 	return GetField<EJson::None>(FieldName)->AsString();
 }
 
+bool FJsonObject::TryGetStringField(const FString& FieldName, FString& OutString) const
+{
+	TSharedPtr<FJsonValue> Field = TryGetField(FieldName);
+	return Field.IsValid() && Field->TryGetString(OutString);
+}
+
 void FJsonObject::SetStringField( const FString& FieldName, const FString& StringValue )
 {
 	this->Values.Add( FieldName, MakeShareable(new FJsonValueString(StringValue)) );
@@ -127,6 +223,12 @@ bool FJsonObject::GetBoolField(const FString& FieldName) const
 	return GetField<EJson::None>(FieldName)->AsBool();
 }
 
+bool FJsonObject::TryGetBoolField(const FString& FieldName, bool& OutBool) const
+{
+	TSharedPtr<FJsonValue> Field = TryGetField(FieldName);
+	return Field.IsValid() && Field->TryGetBool(OutBool);
+}
+
 void FJsonObject::SetBoolField( const FString& FieldName, bool InValue )
 {
 	this->Values.Add( FieldName, MakeShareable( new FJsonValueBoolean(InValue) ) );
@@ -137,6 +239,12 @@ const TArray< TSharedPtr<FJsonValue> >& FJsonObject::GetArrayField(const FString
 	return GetField<EJson::Array>(FieldName)->AsArray();
 }
 
+bool FJsonObject::TryGetArrayField(const FString& FieldName, const TArray< TSharedPtr<FJsonValue> >*& OutArray) const
+{
+	TSharedPtr<FJsonValue> Field = TryGetField(FieldName);
+	return Field.IsValid() && Field->TryGetArray(OutArray);
+}
+
 void FJsonObject::SetArrayField( const FString& FieldName, const TArray< TSharedPtr<FJsonValue> >& Array )
 {
 	this->Values.Add( FieldName, MakeShareable( new FJsonValueArray(Array) ) );
@@ -147,6 +255,12 @@ const TSharedPtr<FJsonObject>& FJsonObject::GetObjectField(const FString& FieldN
 	return GetField<EJson::Object>(FieldName)->AsObject();
 }
 
+bool FJsonObject::TryGetObjectField(const FString& FieldName, const TSharedPtr<FJsonObject>*& OutObject) const
+{
+	TSharedPtr<FJsonValue> Field = TryGetField(FieldName);
+	return Field.IsValid() && Field->TryGetObject(OutObject);
+}
+
 void FJsonObject::SetObjectField( const FString& FieldName, const TSharedPtr<FJsonObject>& JsonObject )
 {
 	if (JsonObject.IsValid())
diff --git a/Engine/Source/Runtime/Core/Public/Serialization/Json/JsonDocumentObjectModel.h b/Engine/Source/Runtime/Core/Public/Serialization/Json/JsonDocumentObjectModel.h
--- a/Engine/Source/Runtime/Core/Public/Serialization/Json/JsonDocumentObjectModel.h
+++ b/Engine/Source/Runtime/Core/Public/Serialization/Json/JsonDocumentObjectModel.h
@@ -10,20 +10,41 @@ class FJsonObject;
 class CORE_API FJsonValue
 {
 public:
-	/** Returns this value as a double, throwing an error if this is not an Json Number */
-	virtual double AsNumber() const {ErrorMessage(TEXT("Number")); return 0.f;}
-
-	/** Returns this value as a number, throwing an error if this is not an Json String */
-	virtual FString AsString() const {ErrorMessage(TEXT("String")); return FString();}
+	/** Returns this value as a double, logging an error and returning zero if this is not an Json Number */
+	double AsNumber() const;
+	
+	/** Returns this value as a number, logging an error and returning an empty string if not possible */
+	FString AsString() const;
 
-	/** Returns this value as a boolean, throwing an error if this is not an Json Bool */
-	virtual bool AsBool() const {ErrorMessage(TEXT("Boolean")); return false;}
+	/** Returns this value as a boolean, logging an error and returning false if not possible */
+	bool AsBool() const;
 
-	/** Returns this value as an array, throwing an error if this is not an Json Array */
-	virtual const TArray< TSharedPtr<FJsonValue> >& AsArray() const {ErrorMessage(TEXT("Array")); return EMPTY_ARRAY;}
+	/** Returns this value as an array, logging an error and returning an empty array reference if not possible */
+	const TArray< TSharedPtr<FJsonValue> >& AsArray() const;
 
 	/** Returns this value as an object, throwing an error if this is not an Json Object */
-	virtual const TSharedPtr<FJsonObject>& AsObject() const {ErrorMessage(TEXT("Object")); return EMPTY_OBJECT;}
+	virtual const TSharedPtr<FJsonObject>& AsObject() const;
+
+	/** Tries to convert this value to a number, returning false if not possible */
+	virtual bool TryGetNumber(double& OutDouble) const { return false; }
+
+	/** Tries to convert this value to an integer, returning false if not possible */
+	bool TryGetNumber(int32& OutNumber) const;
+
+	/** Tries to convert this value to an integer, returning false if not possible */
+	bool TryGetNumber(uint32& OutNumber) const;
+
+	/** Tries to convert this value to a string, returning false if not possible */
+	virtual bool TryGetString(FString& OutString) const { return false; }
+
+	/** Tries to convert this value to a bool, returning false if not possible */
+	virtual bool TryGetBool(bool& OutBool) const { return false; }
+
+	/** Tries to convert this value to an array, returning false if not possible */
+	virtual bool TryGetArray(const TArray< TSharedPtr<FJsonValue> >*& OutArray) const { return false; }
+
+	/** Tries to convert this value to an object, returning false if not possible */
+	virtual bool TryGetObject(const TSharedPtr<FJsonObject>*& Object) const { return false; }
 
 	/** Returns true if this value is a 'null' */
 	bool IsNull() const { return Type == EJson::Null || Type == EJson::None; }
@@ -59,9 +80,9 @@ class CORE_API FJsonValueString : public FJsonValue
 public:
 	FJsonValueString(const FString& InString) : Value(InString) {Type = EJson::String;}
 
-	virtual FString AsString() const {return Value;}
-	virtual double AsNumber() const {return Value.IsNumeric() ? FCString::Atod(*Value) : FJsonValue::AsNumber();}
-	virtual bool AsBool() const {return Value.ToBool();}
+	virtual bool TryGetString(FString& OutString) const override	{ OutString = Value; return true; }
+	virtual bool TryGetNumber(double& OutDouble) const override		{ if(Value.IsNumeric()){ OutDouble = FCString::Atod(*Value); return true; } else { return false; } }
+	virtual bool TryGetBool(bool& OutBool) const override			{ OutBool = Value.ToBool(); return true; }
 
 protected:
 	FString Value;
@@ -74,9 +95,9 @@ class CORE_API FJsonValueNumber : public FJsonValue
 {
 public:
 	FJsonValueNumber(double InNumber) : Value(InNumber) {Type = EJson::Number;}
-	virtual double AsNumber() const {return Value;}
-	virtual bool AsBool() const {return Value != 0.0;}
-	virtual FString AsString() const {return FString::SanitizeFloat(Value);}
+	virtual bool TryGetNumber(double& OutNumber) const override		{ OutNumber = Value; return true; }
+	virtual bool TryGetBool(bool& OutBool) const override			{ OutBool = (Value != 0.0); return true; }
+	virtual bool TryGetString(FString& OutString) const override	{ OutString = FString::SanitizeFloat(Value); return true; }
 	
 protected:
 	double Value;
@@ -89,9 +110,9 @@ class CORE_API FJsonValueBoolean : public FJsonValue
 {
 public:
 	FJsonValueBoolean(bool InBool) : Value(InBool) {Type = EJson::Boolean;}
-	virtual double AsNumber() const {return Value ? 1 : 0;}
-	virtual bool AsBool() const {return Value;}
-	virtual FString AsString() const {return Value ? TEXT("true") : TEXT("false");}
+	virtual bool TryGetNumber(double& OutNumber) const override		{ OutNumber = Value ? 1 : 0; return true; }
+	virtual bool TryGetBool(bool& OutBool) const override			{ OutBool = Value; return true; }
+	virtual bool TryGetString(FString& OutString) const override	{ OutString = Value ? TEXT("true") : TEXT("false"); return true; }
 	
 protected:
 	bool Value;
@@ -104,7 +125,7 @@ class CORE_API FJsonValueArray : public FJsonValue
 {
 public:
 	FJsonValueArray(const TArray< TSharedPtr<FJsonValue> >& InArray) : Value(InArray) {Type = EJson::Array;}
-	virtual const TArray< TSharedPtr<FJsonValue> >& AsArray() const {return Value;}
+	virtual bool TryGetArray(const TArray< TSharedPtr<FJsonValue> >*& OutArray) const override	{ OutArray = &Value; return true; }
 	
 protected:
 	TArray< TSharedPtr<FJsonValue> > Value;
@@ -117,7 +138,7 @@ class CORE_API FJsonValueObject : public FJsonValue
 {
 public:
 	FJsonValueObject(TSharedPtr<FJsonObject> InObject) : Value(InObject) {Type = EJson::Object;}
-	virtual const TSharedPtr<FJsonObject>& AsObject() const {return Value;}
+	virtual bool TryGetObject(const TSharedPtr<FJsonObject>*& OutObject) const override			{ OutObject = &Value; return true; }
 	
 protected:
 	TSharedPtr<FJsonObject> Value;
@@ -206,30 +227,51 @@ public:
 	/** Get the field named FieldName as a number. Ensures that the field is present and is of type Json number. */
 	double GetNumberField(const FString& FieldName) const;
 
+	/** Get the field named FieldName as a number. Returns false if it cannot be converted. */
+	bool TryGetNumberField(const FString& FieldName, double& OutNumber) const;
+
+	/** Get the field named FieldName as a number, and makes sure it's within int32 range. Returns false if it cannot be converted. */
+	bool TryGetNumberField(const FString& FieldName, int32& OutNumber) const;
+
+	/** Get the field named FieldName as a number, and makes sure it's within uint32 range. Returns false if it cannot be converted.  */
+	bool TryGetNumberField(const FString& FieldName, uint32& OutNumber) const;
+
 	/** Add a field named FieldName with Number as value */
 	void SetNumberField( const FString& FieldName, double Number );
 
 	/** Get the field named FieldName as a string. */
 	FString GetStringField(const FString& FieldName) const;
 
+	/** Get the field named FieldName as a string. Returns false if it cannot be converted. */
+	bool TryGetStringField(const FString& FieldName, FString& OutString) const;
+
 	/** Add a field named FieldName with value of StringValue */
 	void SetStringField( const FString& FieldName, const FString& StringValue );
 
 	/** Get the field named FieldName as a boolean. */
 	bool GetBoolField(const FString& FieldName) const;
 
+	/** Get the field named FieldName as a string. Returns false if it cannot be converted. */
+	bool TryGetBoolField(const FString& FieldName, bool& OutBool) const;
+
 	/** Set a boolean field named FieldName and value of InValue */
 	void SetBoolField( const FString& FieldName, bool InValue );
 
 	/** Get the field named FieldName as an array. */
 	const TArray< TSharedPtr<FJsonValue> >& GetArrayField(const FString& FieldName) const;
 
+	/** Try to get the field named FieldName as an array, or return false if it's another type */
+	bool TryGetArrayField(const FString& FieldName, const TArray< TSharedPtr<FJsonValue> >*& OutArray) const;
+
 	/** Set an array field named FieldName and value of Array */
 	void SetArrayField( const FString& FieldName, const TArray< TSharedPtr<FJsonValue> >& Array );
 
 	/** Get the field named FieldName as a Json object. */
 	const TSharedPtr<FJsonObject>& GetObjectField(const FString& FieldName) const;
 
+	/** Try to get the field named FieldName as an object, or return false if it's another type */
+	bool TryGetObjectField(const FString& FieldName, const TSharedPtr<FJsonObject>*& OutObject) const;
+
 	/** Set an ObjectField named FieldName and value of JsonObject */
 	void SetObjectField( const FString& FieldName, const TSharedPtr<FJsonObject>& JsonObject );
 };
