Removed ICUUtilities::ConvertChar32

It wasn't correct since it was potentially only taking half of a UTF-16 pair, and was only used for the FText::IsWhitespace check where a TCHAR will suffice due to whitespace never spanning multiple character pairs.

diff --git a/Engine/Source/Runtime/Core/Private/Internationalization/ICUText.cpp b/Engine/Source/Runtime/Core/Private/Internationalization/ICUText.cpp
--- a/Engine/Source/Runtime/Core/Private/Internationalization/ICUText.cpp
+++ b/Engine/Source/Runtime/Core/Private/Internationalization/ICUText.cpp
@@ -15,7 +15,9 @@
 
 bool FText::IsWhitespace( const TCHAR Char )
 {
-	const UChar32 ICUChar = ICUUtilities::ConvertChar32(Char);
+	// TCHAR should either be UTF-16 or UTF-32, so we should be fine to cast it to a UChar32 for the whitespace
+	// check, since whitespace is never a pair of UTF-16 characters
+	const UChar32 ICUChar = static_cast<UChar32>(Char);
 	return u_isWhitespace(ICUChar) != 0;
 }
 
diff --git a/Engine/Source/Runtime/Core/Private/Internationalization/ICUUtilities.cpp b/Engine/Source/Runtime/Core/Private/Internationalization/ICUUtilities.cpp
--- a/Engine/Source/Runtime/Core/Private/Internationalization/ICUUtilities.cpp
+++ b/Engine/Source/Runtime/Core/Private/Internationalization/ICUUtilities.cpp
@@ -96,32 +96,5 @@ namespace ICUUtilities
 		ConvertString(Source, Destination);
 		return Destination;
 	}
-
-	void ConvertChar32(const TCHAR Source, UChar32& Destination)
-	{
-		const char* const SourceEncoding = FPlatformString::GetEncodingName();
-	
-		UErrorCode ICUStatus = U_ZERO_ERROR;
-		
-		UConverter* ICUConverter = ucnv_open(SourceEncoding, &ICUStatus);
-		if( U_FAILURE(ICUStatus) ) 
-		{
-			return;
-		}
-
-
-		const char* NativeChars = reinterpret_cast<const char*>(&Source);
-		const int32_t NativeCharsLen = sizeof(TCHAR);
-		Destination = ucnv_getNextUChar(ICUConverter, &NativeChars, NativeChars + NativeCharsLen, &ICUStatus);
-
-		ucnv_close(ICUConverter);
-	}
-
-	UChar32 ConvertChar32(const TCHAR Source)
-	{
-		UChar32 Destination = 0;
-		ConvertChar32(Source, Destination);
-		return Destination;
-	}
 }
 #endif
\ No newline at end of file
diff --git a/Engine/Source/Runtime/Core/Private/Internationalization/ICUUtilities.h b/Engine/Source/Runtime/Core/Private/Internationalization/ICUUtilities.h
--- a/Engine/Source/Runtime/Core/Private/Internationalization/ICUUtilities.h
+++ b/Engine/Source/Runtime/Core/Private/Internationalization/ICUUtilities.h
@@ -12,9 +12,5 @@ namespace ICUUtilities
 	icu::UnicodeString ConvertString(const FString& Source, const bool ShouldNullTerminate = true);
 	void ConvertString(const icu::UnicodeString& Source, FString& Destination);
 	FString ConvertString(const icu::UnicodeString& Source);
-
-	/** Convert TCHAR -> UChar32 */
-	void ConvertChar32(const TCHAR Source, UChar32& Destination);
-	UChar32 ConvertChar32(const TCHAR Source);
 }
 #endif
\ No newline at end of file
