Add a JSON utility function to read an array of strings (seems to be a very common thing to do, and it's long winded)

diff --git a/Engine/Source/Runtime/Core/Private/Serialization/Json/JsonDocumentObjectModel.cpp b/Engine/Source/Runtime/Core/Private/Serialization/Json/JsonDocumentObjectModel.cpp
--- a/Engine/Source/Runtime/Core/Private/Serialization/Json/JsonDocumentObjectModel.cpp
+++ b/Engine/Source/Runtime/Core/Private/Serialization/Json/JsonDocumentObjectModel.cpp
@@ -213,6 +213,32 @@ bool FJsonObject::TryGetStringField(const FString& FieldName, FString& OutString
 	return Field.IsValid() && Field->TryGetString(OutString);
 }
 
+bool FJsonObject::TryGetStringArrayField(const FString& FieldName, TArray<FString>& OutArray) const
+{
+	TSharedPtr<FJsonValue> Field = TryGetField(FieldName);
+	if(!Field.IsValid())
+	{
+		return false;
+	}
+
+	const TArray< TSharedPtr<FJsonValue> > *Array;
+	if(!Field->TryGetArray(Array))
+	{
+		return false;
+	}
+
+	for(int Idx = 0; Idx < Array->Num(); Idx++)
+	{
+		FString Element;
+		if(!(*Array)[Idx]->TryGetString(Element))
+		{
+			return false;
+		}
+		OutArray.Add(Element);
+	}
+	return true;
+}
+
 void FJsonObject::SetStringField( const FString& FieldName, const FString& StringValue )
 {
 	this->Values.Add( FieldName, MakeShareable(new FJsonValueString(StringValue)) );
diff --git a/Engine/Source/Runtime/Core/Public/Serialization/Json/JsonDocumentObjectModel.h b/Engine/Source/Runtime/Core/Public/Serialization/Json/JsonDocumentObjectModel.h
--- a/Engine/Source/Runtime/Core/Public/Serialization/Json/JsonDocumentObjectModel.h
+++ b/Engine/Source/Runtime/Core/Public/Serialization/Json/JsonDocumentObjectModel.h
@@ -227,13 +227,13 @@ public:
 	/** Get the field named FieldName as a number. Ensures that the field is present and is of type Json number. */
 	double GetNumberField(const FString& FieldName) const;
 
-	/** Get the field named FieldName as a number. Returns false if it cannot be converted. */
+	/** Get the field named FieldName as a number. Returns false if it doesn't exist or cannot be converted. */
 	bool TryGetNumberField(const FString& FieldName, double& OutNumber) const;
 
-	/** Get the field named FieldName as a number, and makes sure it's within int32 range. Returns false if it cannot be converted. */
+	/** Get the field named FieldName as a number, and makes sure it's within int32 range. Returns false if it doesn't exist or cannot be converted. */
 	bool TryGetNumberField(const FString& FieldName, int32& OutNumber) const;
 
-	/** Get the field named FieldName as a number, and makes sure it's within uint32 range. Returns false if it cannot be converted.  */
+	/** Get the field named FieldName as a number, and makes sure it's within uint32 range. Returns false if it doesn't exist or cannot be converted.  */
 	bool TryGetNumberField(const FString& FieldName, uint32& OutNumber) const;
 
 	/** Add a field named FieldName with Number as value */
@@ -242,16 +242,19 @@ public:
 	/** Get the field named FieldName as a string. */
 	FString GetStringField(const FString& FieldName) const;
 
-	/** Get the field named FieldName as a string. Returns false if it cannot be converted. */
+	/** Get the field named FieldName as a string. Returns false if it doesn't exist or cannot be converted. */
 	bool TryGetStringField(const FString& FieldName, FString& OutString) const;
 
+	/** Get the field named FieldName as an array of strings. Returns false if it doesn't exist or any member cannot be converted. */
+	bool TryGetStringArrayField(const FString& FieldName, TArray<FString>& OutArray) const;
+
 	/** Add a field named FieldName with value of StringValue */
 	void SetStringField( const FString& FieldName, const FString& StringValue );
 
 	/** Get the field named FieldName as a boolean. */
 	bool GetBoolField(const FString& FieldName) const;
 
-	/** Get the field named FieldName as a string. Returns false if it cannot be converted. */
+	/** Get the field named FieldName as a string. Returns false if it doesn't exist or cannot be converted. */
 	bool TryGetBoolField(const FString& FieldName, bool& OutBool) const;
 
 	/** Set a boolean field named FieldName and value of InValue */
