#ue4 delete the archive when we're done with it, fixes some of the memory leaks that occur when a blueprint is opened and closed

diff --git a/Engine/Source/Runtime/Core/Private/Internationalization/Internationalization.cpp b/Engine/Source/Runtime/Core/Private/Internationalization/Internationalization.cpp
--- a/Engine/Source/Runtime/Core/Private/Internationalization/Internationalization.cpp
+++ b/Engine/Source/Runtime/Core/Private/Internationalization/Internationalization.cpp
@@ -212,7 +212,7 @@ namespace
 
 		FileAr->Serialize(*contents, FileSize); 
 
-		FileAr->Close(); // Close the file archive to free up the handle - it is unneeded since the data is in memory.
+		delete FileAr;
 
 		return TRUE;
 	}
