FText::BuildSourceString will check if the History is valid before attempting to use it.

#ttp 339868 - CRASH: BLUEPRINT: Unpinning Get DataTableRow node

diff --git a/Engine/Source/Runtime/Core/Private/Internationalization/Text.cpp b/Engine/Source/Runtime/Core/Private/Internationalization/Text.cpp
--- a/Engine/Source/Runtime/Core/Private/Internationalization/Text.cpp
+++ b/Engine/Source/Runtime/Core/Private/Internationalization/Text.cpp
@@ -707,7 +707,11 @@ const FString& FText::ToString() const
 
 FString FText::BuildSourceString() const
 {
-	return History->ToText(true).ToString();
+	if(History.IsValid())
+	{
+		return History->ToText(true).ToString();
+	}
+	return DisplayString.Get();
 }
 
 void FText::Rebuild() const
