Use 32-bit thread ids on Linux.

This is less portable approach than pthread_self(), but the latter requires more intrusive changes throughout the engine (upping thread ids to 64 bit or even to an opaque type), which will break NetworkProfiler serialization and don't feel warranted in general.

Thanks to [REDACTED_PERSON] for bringing this up on #ue4linux.

diff --git a/Engine/Source/Runtime/Core/Public/Linux/LinuxPlatformTLS.h b/Engine/Source/Runtime/Core/Public/Linux/LinuxPlatformTLS.h
--- a/Engine/Source/Runtime/Core/Public/Linux/LinuxPlatformTLS.h
+++ b/Engine/Source/Runtime/Core/Public/Linux/LinuxPlatformTLS.h
@@ -7,8 +7,12 @@
 
 #pragma once
 
+#if defined(_GNU_SOURCE)
+	#include <sys/syscall.h>	// SYS_gettid
+#endif // _GNU_SOURCE
+
 /**
- * Android implementation of the TLS OS functions
+ * Linux implementation of the TLS OS functions
  */
 struct CORE_API FLinuxTLS : public FGenericPlatformTLS
 {
@@ -17,7 +21,15 @@ struct CORE_API FLinuxTLS : public FGenericPlatformTLS
 	 */
 	static FORCEINLINE uint32 GetCurrentThreadId(void)
 	{
-		return pthread_self();
+		// note: cannot use pthread_self() without updating the rest of API to opaque (or at least 64-bit) thread handles
+#if defined(_GNU_SOURCE)
+		pid_t ThreadId = static_cast<pid_t>(syscall(SYS_gettid));
+		static_assert(sizeof(pid_t) <= sizeof(uint32), "pid_t is larger than uint32, reconsider implementation of GetCurrentThreadId()");
+		return static_cast<pid_t>(ThreadId);
+#else
+		// better than nothing...
+		return static_cast< uint32 >(pthread_self());
+#endif
 	}
 
 	/**
