Fixed incorrect locks being used with ValidateHeap() and DumpAllocatorStats(). Fixes TTP #340022

diff --git a/Engine/Source/Runtime/Core/Public/HAL/MallocBinned.h b/Engine/Source/Runtime/Core/Public/HAL/MallocBinned.h
--- a/Engine/Source/Runtime/Core/Public/HAL/MallocBinned.h
+++ b/Engine/Source/Runtime/Core/Public/HAL/MallocBinned.h
@@ -1090,17 +1090,20 @@ public:
 	 */
 	virtual bool ValidateHeap() override
 	{
-#ifdef USE_INTERNAL_LOCKS
+#ifdef USE_COARSE_GRAIN_LOCKS
 		FScopeLock ScopedLock(&AccessGuard);
 #endif
 		for( int32 i = 0; i < POOL_COUNT; i++ )
 		{
 			FPoolTable* Table = &PoolTable[i];
+#ifdef USE_FINE_GRAIN_LOCKS
+			FScopeLock TableLock(&Table->CriticalSection);
+#endif
 			for( FPoolInfo** PoolPtr = &Table->FirstPool; *PoolPtr; PoolPtr = &(*PoolPtr)->Next )
 			{
 				FPoolInfo* Pool = *PoolPtr;
 				check(Pool->PrevLink == PoolPtr);
-				check(Pool->FirstMem /*|| Pool->Taken == (Pool->AllocSize/Pool->TableIndex)*/);
+				check(Pool->FirstMem);
 				for( FFreeMem* Free = Pool->FirstMem; Free; Free = Free->Next )
 				{
 					check(Free->NumFreeBlocks > 0);
@@ -1174,7 +1177,7 @@ public:
 	{
 		FBufferedOutputDevice BufferedOutput;
 		{
-#ifdef USE_INTERNAL_LOCKS
+#ifdef USE_COARSE_GRAIN_LOCKS
 			FScopeLock ScopedLock( &AccessGuard );
 #endif
 			ValidateHeap();
@@ -1216,6 +1219,10 @@ public:
 
 				Table = MemSizeToPoolTable[i];
 
+#ifdef USE_FINE_GRAIN_LOCKS
+				FScopeLock TableLock(&Table->CriticalSection);
+#endif
+
 				uint32 TableAllocSize = (Table->BlockSize > BinnedSizeLimit ? (((3 * (i - BinnedSizeLimit)) + 3)*BINNED_ALLOC_POOL_SIZE) : BINNED_ALLOC_POOL_SIZE);
 				// The amount of memory allocated from the OS
 				uint32 MemAllocated = (Table->NumActivePools * TableAllocSize) / 1024;
