Support from reading binary streams from pipes.

External contribution by SRombauts (Sebastien.Rombauts), PR 203 (with changes): https://github.com/EpicGames/UnrealEngine/pull/203

diff --git a/Engine/Source/Runtime/Core/Private/GenericPlatform/GenericPlatformProcess.cpp b/Engine/Source/Runtime/Core/Private/GenericPlatform/GenericPlatformProcess.cpp
--- a/Engine/Source/Runtime/Core/Private/GenericPlatform/GenericPlatformProcess.cpp
+++ b/Engine/Source/Runtime/Core/Private/GenericPlatform/GenericPlatformProcess.cpp
@@ -396,6 +396,12 @@ FString FGenericPlatformProcess::ReadPipe( void* ReadPipe )
 	return FString();
 }
 
+bool FGenericPlatformProcess::ReadPipeToArray(void* ReadPipe, TArray<uint8> & Output)
+{
+	UE_LOG(LogHAL, Fatal, TEXT("FGenericPlatformProcess::ReadPipeToArray not implemented on this platform"));
+	return false;
+}
+
 bool FGenericPlatformProcess::SupportsMultithreading()
 {
 	static bool bSupportsMultithreading = !FParse::Param(FCommandLine::Get(), TEXT("nothreading"));
diff --git a/Engine/Source/Runtime/Core/Private/Linux/LinuxPlatformProcess.cpp b/Engine/Source/Runtime/Core/Private/Linux/LinuxPlatformProcess.cpp
--- a/Engine/Source/Runtime/Core/Private/Linux/LinuxPlatformProcess.cpp
+++ b/Engine/Source/Runtime/Core/Private/Linux/LinuxPlatformProcess.cpp
@@ -200,6 +200,35 @@ FString FPipeHandle::Read()
 	return Output;
 }
 
+bool FPipeHandle::ReadToArray(TArray<uint8> & Output)
+{
+	int BytesAvailable = 0;
+	if (ioctl(PipeDesc, FIONREAD, &BytesAvailable) == 0)
+	{
+		if (BytesAvailable > 0)
+		{
+			Output.Init(BytesAvailable);
+			int BytesRead = read(PipeDesc, Output.GetData(), BytesAvailable);
+			if (BytesRead > 0)
+			{
+				if (BytesRead < BytesAvailable)
+				{
+					Output.SetNum(BytesRead);
+				}
+
+				return true;
+			}
+			else
+			{
+				Output.Empty();
+			}
+		}
+	}
+
+	return false;
+}
+
+
 void FLinuxPlatformProcess::ClosePipe( void* ReadPipe, void* WritePipe )
 {
 	if (ReadPipe)
@@ -243,6 +272,17 @@ FString FLinuxPlatformProcess::ReadPipe( void* ReadPipe )
 	return FString();
 }
 
+bool FLinuxPlatformProcess::ReadPipeToArray(void* ReadPipe, TArray<uint8> & Output)
+{
+	if (ReadPipe)
+	{
+		FPipeHandle * PipeHandle = reinterpret_cast<FPipeHandle*>(ReadPipe);
+		return PipeHandle->ReadToArray(Output);
+	}
+
+	return false;
+}
+
 FRunnableThread* FLinuxPlatformProcess::CreateRunnableThread()
 {
 	return new FRunnableThreadLinux();
diff --git a/Engine/Source/Runtime/Core/Private/Mac/MacPlatformProcess.cpp b/Engine/Source/Runtime/Core/Private/Mac/MacPlatformProcess.cpp
--- a/Engine/Source/Runtime/Core/Private/Mac/MacPlatformProcess.cpp
+++ b/Engine/Source/Runtime/Core/Private/Mac/MacPlatformProcess.cpp
@@ -871,6 +871,35 @@ FString FMacPlatformProcess::ReadPipe( void* ReadPipe )
 	return Output;
 }
 
+bool FMacPlatformProcess::ReadPipeToArray(void* ReadPipe, TArray<uint8> & Output)
+{
+	SCOPED_AUTORELEASE_POOL;
+
+	const int32 READ_SIZE = 32768;
+
+	if (ReadPipe)
+	{
+		Output.Init(READ_SIZE);
+		int32 BytesRead = 0;
+		BytesRead = read([(NSFileHandle*)ReadPipe fileDescriptor], Output.GetData(), READ_SIZE);
+		if (BytesRead > 0)
+		{
+			if (BytesRead < READ_SIZE)
+			{
+				Output.SetNum(BytesRead);
+			}
+
+			return true;
+		}
+		else
+		{
+			Output.Empty();
+		}
+	}
+
+	return false;
+}
+
 #include "MacPlatformRunnableThread.h"
 
 FRunnableThread* FMacPlatformProcess::CreateRunnableThread()
diff --git a/Engine/Source/Runtime/Core/Private/Windows/WindowsPlatformProcess.cpp b/Engine/Source/Runtime/Core/Private/Windows/WindowsPlatformProcess.cpp
--- a/Engine/Source/Runtime/Core/Private/Windows/WindowsPlatformProcess.cpp
+++ b/Engine/Source/Runtime/Core/Private/Windows/WindowsPlatformProcess.cpp
@@ -978,6 +978,31 @@ FString FWindowsPlatformProcess::ReadPipe( void* ReadPipe )
 	return Output;
 }
 
+bool FWindowsPlatformProcess::ReadPipeToArray(void* ReadPipe, TArray<uint8> & Output)
+{
+	uint32 BytesAvailable = 0;
+	if (::PeekNamedPipe(ReadPipe, NULL, 0, NULL, (::DWORD*)&BytesAvailable, NULL) && (BytesAvailable > 0))
+	{
+		Output.Init(BytesAvailable);
+		uint32 BytesRead = 0;
+		if (::ReadFile(ReadPipe, Output.GetData(), BytesAvailable, (::DWORD*)&BytesRead, NULL))
+		{
+			if (BytesRead < BytesAvailable)
+			{
+				Output.SetNum(BytesRead);
+			}
+
+			return true;
+		}
+		else
+		{
+			Output.Empty();
+		}
+	}
+
+	return false;
+}
+
 #include "AllowWindowsPlatformTypes.h"
 
 FWindowsPlatformProcess::FWindowsSemaphore::FWindowsSemaphore(const FString & InName, HANDLE InSemaphore)
diff --git a/Engine/Source/Runtime/Core/Public/GenericPlatform/GenericPlatformProcess.h b/Engine/Source/Runtime/Core/Public/GenericPlatform/GenericPlatformProcess.h
--- a/Engine/Source/Runtime/Core/Public/GenericPlatform/GenericPlatformProcess.h
+++ b/Engine/Source/Runtime/Core/Public/GenericPlatform/GenericPlatformProcess.h
@@ -420,6 +420,19 @@ struct CORE_API FGenericPlatformProcess
 	 */
 	static FString ReadPipe( void* ReadPipe );
 
+	/**
+	 * Reads all pending data from an anonymous pipe, such as STDOUT or STDERROR of a process.
+	 *
+	 * @param Pipe - The handle to the pipe to read from.
+	 * @param Output - The data read.
+	 *
+	 * @return true if successful (i.e. any data was read)
+	 *
+	 * @see ClosePipe
+	 * @see CreatePipe
+	 */
+	static bool ReadPipeToArray(void* ReadPipe, TArray<uint8> & Output);
+
 	/**
 	 * Gets whether this platform can use multiple threads.
 	 *
diff --git a/Engine/Source/Runtime/Core/Public/Linux/LinuxPlatformProcess.h b/Engine/Source/Runtime/Core/Public/Linux/LinuxPlatformProcess.h
--- a/Engine/Source/Runtime/Core/Public/Linux/LinuxPlatformProcess.h
+++ b/Engine/Source/Runtime/Core/Public/Linux/LinuxPlatformProcess.h
@@ -108,6 +108,11 @@ struct FPipeHandle
 	 */
 	FString Read();
 
+	/**
+	 * Reads until EOF.
+	 */
+	bool ReadToArray(TArray<uint8> & Output);
+
 	/**
 	 * Returns raw file handle.
 	 */
@@ -141,6 +146,7 @@ struct CORE_API FLinuxPlatformProcess : public FGenericPlatformProcess
 	static void ClosePipe( void* ReadPipe, void* WritePipe );
 	static bool CreatePipe( void*& ReadPipe, void*& WritePipe );
 	static FString ReadPipe( void* ReadPipe );
+	static bool ReadPipeToArray(void* ReadPipe, TArray<uint8> & Output);
 	static class FRunnableThread* CreateRunnableThread();
 	static void LaunchURL(const TCHAR* URL, const TCHAR* Parms, FString* Error);
 	static FProcHandle CreateProc(const TCHAR* URL, const TCHAR* Parms, bool bLaunchDetached, bool bLaunchHidden, bool bLaunchReallyHidden, uint32* OutProcessID, int32 PriorityModifier, const TCHAR* OptionalWorkingDirectory, void* PipeWrite);
diff --git a/Engine/Source/Runtime/Core/Public/Mac/MacPlatformProcess.h b/Engine/Source/Runtime/Core/Public/Mac/MacPlatformProcess.h
--- a/Engine/Source/Runtime/Core/Public/Mac/MacPlatformProcess.h
+++ b/Engine/Source/Runtime/Core/Public/Mac/MacPlatformProcess.h
@@ -61,7 +61,8 @@ struct CORE_API FMacPlatformProcess : public FGenericPlatformProcess
 	static void ClosePipe( void* ReadPipe, void* WritePipe );
 	static bool CreatePipe( void*& ReadPipe, void*& WritePipe );
 	static FString ReadPipe( void* ReadPipe );
-	static bool ExecProcess( const TCHAR* URL, const TCHAR* Params, int32* OutReturnCode, FString* OutStdOut, FString* OutStdErr );
+	static bool ReadPipeToArray(void* ReadPipe, TArray<uint8> & Output);
+	static bool ExecProcess(const TCHAR* URL, const TCHAR* Params, int32* OutReturnCode, FString* OutStdOut, FString* OutStdErr);
 
 	// Mac specific
 	static const TCHAR* UserPreferencesDir();
diff --git a/Engine/Source/Runtime/Core/Public/Windows/WindowsPlatformProcess.h b/Engine/Source/Runtime/Core/Public/Windows/WindowsPlatformProcess.h
--- a/Engine/Source/Runtime/Core/Public/Windows/WindowsPlatformProcess.h
+++ b/Engine/Source/Runtime/Core/Public/Windows/WindowsPlatformProcess.h
@@ -109,6 +109,7 @@ struct CORE_API FWindowsPlatformProcess : public FGenericPlatformProcess
 	static void ClosePipe( void* ReadPipe, void* WritePipe );
 	static bool CreatePipe( void*& ReadPipe, void*& WritePipe );
 	static FString ReadPipe( void* ReadPipe );
+	static bool ReadPipeToArray(void* ReadPipe, TArray<uint8> & Output);
 	static FSemaphore * NewInterprocessSynchObject(const FString & Name, bool bCreate, uint32 MaxLocks = 1);
 	static bool DeleteInterprocessSynchObject(FSemaphore * Object);
 	static bool Daemonize();
