Use uuid_generate rather than implementing our own GUID generation for much improved GUID generation.

diff --git a/Engine/Source/Runtime/Core/Private/Mac/MacPlatformMisc.cpp b/Engine/Source/Runtime/Core/Private/Mac/MacPlatformMisc.cpp
--- a/Engine/Source/Runtime/Core/Private/Mac/MacPlatformMisc.cpp
+++ b/Engine/Source/Runtime/Core/Private/Mac/MacPlatformMisc.cpp
@@ -22,6 +22,7 @@
 #include <mach-o/dyld.h>
 #include <libproc.h>
 #include <notify.h>
+#include <uuid/uuid.h>
 
 
 /**
@@ -793,16 +794,14 @@ void FMacPlatformMisc::ClipboardPaste(class FString& Result)
 
 void FMacPlatformMisc::CreateGuid(FGuid& Result)
 {
-	int32 Year = 0, Month = 0, DayOfWeek = 0, Day = 0, Hour = 0, Min = 0, Sec = 0, MSec = 0;
-	FPlatformTime::SystemTime(Year, Month, DayOfWeek, Day, Hour, Min, Sec, MSec);
-
-	FGuid   GUID;
-	GUID.A = Day | (Hour << 16);
-	GUID.B = Month | (Sec << 16);
-	GUID.C = MSec | (Min << 16);
-	GUID.D = Year ^ FPlatformTime::Cycles();
-
-	Result = GUID;
+    uuid_t UUID;
+	uuid_generate(UUID);
+    
+    uint32* Values = (uint32*)(&UUID[0]);
+    Result[0] = Values[0];
+    Result[1] = Values[1];
+    Result[2] = Values[2];
+    Result[3] = Values[3];
 }
 
 EAppReturnType::Type FMacPlatformMisc::MessageBoxExt(EAppMsgType::Type MsgType, const TCHAR* Text, const TCHAR* Caption)
