Use the proper I/O Kit power management assertions to prevent screensave and display sleep.

diff --git a/Engine/Source/Runtime/Core/Private/Mac/MacPlatformMisc.cpp b/Engine/Source/Runtime/Core/Private/Mac/MacPlatformMisc.cpp
--- a/Engine/Source/Runtime/Core/Private/Mac/MacPlatformMisc.cpp
+++ b/Engine/Source/Runtime/Core/Private/Mac/MacPlatformMisc.cpp
@@ -19,6 +19,7 @@
 #include <IOKit/network/IOEthernetController.h>
 #include <IOKit/ps/IOPowerSources.h>
 #include <IOKit/ps/IOPSKeys.h>
+#include <IOKit/pwr_mgt/IOPMLib.h>
 #include <mach-o/dyld.h>
 #include <libproc.h>
 #include <notify.h>
@@ -324,14 +325,45 @@ void FMacPlatformMisc::UpdateWindowMenu()
 	[WindowMenu addItem:BringAllToFrontItem];
 }
 
-void FMacPlatformMisc::PreventScreenSaver()
+bool FMacPlatformMisc::ControlScreensaver(EScreenSaverAction Action)
 {
-	CGEventRef LocationEvent = CGEventCreate(NULL);
-	CGPoint MouseLocation = CGEventGetLocation(LocationEvent);
-	CFRelease(LocationEvent);
-	CGEventRef MouseMoveCommand = CGEventCreateMouseEvent(NULL, kCGEventMouseMoved, MouseLocation, 0);
-	CGEventPost(kCGHIDEventTap, MouseMoveCommand);
-	CFRelease(MouseMoveCommand);
+	static uint32 IOPMNoSleepAssertion = 0;
+	static bool bDisplaySleepEnabled = true;
+	
+	switch(Action)
+	{
+		case EScreenSaverAction::Disable:
+		{
+			// Prevent display sleep.
+			if(bDisplaySleepEnabled)
+			{
+				SCOPED_AUTORELEASE_POOL;
+				
+				//  NOTE: IOPMAssertionCreateWithName limits the string to 128 characters.
+				FString ReasonForActivity = FString::Printf(TEXT("Running %s"), FApp::GetGameName());
+				
+				CFStringRef ReasonForActivityCF = (CFStringRef)ReasonForActivity.GetNSString();
+				
+				IOReturn Success = IOPMAssertionCreateWithName(kIOPMAssertionTypeNoDisplaySleep, kIOPMAssertionLevelOn, ReasonForActivityCF, &IOPMNoSleepAssertion);
+				bDisplaySleepEnabled = !(Success == kIOReturnSuccess);
+				ensure(!bDisplaySleepEnabled);
+			}
+			break;
+		}
+		case EScreenSaverAction::Enable:
+		{
+			// Stop preventing display sleep now that we are done.
+			if(!bDisplaySleepEnabled)
+			{
+				IOReturn Success = IOPMAssertionRelease(IOPMNoSleepAssertion);
+				bDisplaySleepEnabled = (Success == kIOReturnSuccess);
+				ensure(bDisplaySleepEnabled);
+			}
+			break;
+		}
+    }
+	
+	return true;
 }
 
 GenericApplication* FMacPlatformMisc::CreateApplication()
diff --git a/Engine/Source/Runtime/Core/Public/Mac/MacPlatformMisc.h b/Engine/Source/Runtime/Core/Public/Mac/MacPlatformMisc.h
--- a/Engine/Source/Runtime/Core/Public/Mac/MacPlatformMisc.h
+++ b/Engine/Source/Runtime/Core/Public/Mac/MacPlatformMisc.h
@@ -54,7 +54,7 @@ struct CORE_API FMacPlatformMisc : public FGenericPlatformMisc
 	static void ClipboardPaste(class FString& Dest);
 	static void CreateGuid(class FGuid& Result);
 	static EAppReturnType::Type MessageBoxExt( EAppMsgType::Type MsgType, const TCHAR* Text, const TCHAR* Caption );
-	static void PreventScreenSaver();
+	static bool ControlScreensaver(EScreenSaverAction Action);
 	static bool CommandLineCommands();
 	static int32 NumberOfCores();
 	static int32 NumberOfCoresIncludingHyperthreads();
