Use POSIX APIs for FMacPlatformProcess::IsApplicationRunning as it is much more reliable.

diff --git a/Engine/Source/Runtime/Core/Private/Mac/MacPlatformProcess.cpp b/Engine/Source/Runtime/Core/Private/Mac/MacPlatformProcess.cpp
--- a/Engine/Source/Runtime/Core/Private/Mac/MacPlatformProcess.cpp
+++ b/Engine/Source/Runtime/Core/Private/Mac/MacPlatformProcess.cpp
@@ -519,25 +519,35 @@ bool FMacPlatformProcess::GetProcReturnCode( FProcHandle & ProcessHandle, int32*
 
 bool FMacPlatformProcess::IsApplicationRunning( const TCHAR* ProcName )
 {
-	SCOPED_AUTORELEASE_POOL;
-	
-	NSArray* ActiveApps = [[NSWorkspace sharedWorkspace] runningApplications ];
+    CFStringRef ProcessName = FPlatformString::TCHARToCFString( ProcName );
+    check(ProcessName);
+	uint32 ThisProcessID = getpid();
+	bool bFound = false;
 	
-	if( ActiveApps )
+	int32 Mib[4] = { CTL_KERN, KERN_PROC, KERN_PROC_ALL, 0 };
+	size_t BufferSize = 0;
+	if (sysctl(Mib, 4, NULL, &BufferSize, NULL, 0) != -1 && BufferSize > 0)
 	{
-		for( NSRunningApplication* App in ActiveApps )
+		char Buffer[MAX_PATH];
+		struct kinfo_proc* Processes = (struct kinfo_proc*)FMemory::Malloc(BufferSize);
+		if (sysctl(Mib, 4, Processes, &BufferSize, NULL, 0) != -1)
 		{
-			if( App )
+			uint32 ProcCount = (uint32)(BufferSize / sizeof(struct kinfo_proc));
+			for (uint32 Index = 0; Index < ProcCount; Index++)
 			{
-				NSString* AppName = [App localizedName];
-				if( FString(AppName) == FString(ProcName) )
+				proc_pidpath(Processes[Index].kp_proc.p_pid, Buffer, sizeof(Buffer));
+				if (Processes[Index].kp_proc.p_pid != ThisProcessID && [(NSString*)ProcessName isEqualToString: [[NSString stringWithUTF8String: Buffer] lastPathComponent]])
 				{
-					return true;
+					bFound = true;
 				}
 			}
 		}
+		
+		FMemory::Free(Processes);
 	}
-	return false;
+    CFRelease(ProcessName);
+    
+	return bFound;
 }
 
 bool FMacPlatformProcess::IsApplicationRunning( uint32 ProcessId )
