Use the system notifications for application activation and deactivation to update the application volume modifier and reduce the CPU usage when deactivated.

diff --git a/Engine/Source/Runtime/Core/Private/Mac/MacApplication.cpp b/Engine/Source/Runtime/Core/Private/Mac/MacApplication.cpp
--- a/Engine/Source/Runtime/Core/Private/Mac/MacApplication.cpp
+++ b/Engine/Source/Runtime/Core/Private/Mac/MacApplication.cpp
@@ -77,6 +77,39 @@ FMacApplication::FMacApplication()
 		TextInputMethodSystem.Reset();
 	}
 
+	AppActivationObserver = [[NSNotificationCenter defaultCenter] addObserverForName:NSApplicationDidBecomeActiveNotification object:[NSApplication sharedApplication] queue:[NSOperationQueue mainQueue] usingBlock:^(NSNotification* Notification){
+								// If editor thread doesn't have the focus, don't suck up too much CPU time.
+								if( GIsEditor )
+								{
+									// Boost our priority back to normal.
+									struct sched_param Sched;
+									FMemory::Memzero(&Sched, sizeof(struct sched_param));
+									Sched.sched_priority = 15;
+									pthread_setschedparam(pthread_self(), SCHED_RR, &Sched);
+								}
+
+								// app is active, allow sound
+								GVolumeMultiplier = 1.0f;
+							}];
+
+	AppDeactivationObserver = [[NSNotificationCenter defaultCenter] addObserverForName:NSApplicationDidResignActiveNotification object:[NSApplication sharedApplication] queue:[NSOperationQueue mainQueue] usingBlock:^(NSNotification* Notification){
+									// If editor thread doesn't have the focus, don't suck up too much CPU time.
+									if( GIsEditor )
+									{
+										// Drop our priority to speed up whatever is in the foreground.
+										struct sched_param Sched;
+										FMemory::Memzero(&Sched, sizeof(struct sched_param));
+										Sched.sched_priority = 5;
+										pthread_setschedparam(pthread_self(), SCHED_RR, &Sched);
+
+										// Sleep for a bit to not eat up all CPU time.
+										FPlatformProcess::Sleep(0.005f);
+									}
+
+									// app is inactive, silence it
+									GVolumeMultiplier = 0.0f;
+								}];
+	
 #if WITH_EDITOR
 	FMemory::MemZero(GestureUsage);
 	LastGestureUsed = EGestureEvent::None;
@@ -85,6 +118,18 @@ FMacApplication::FMacApplication()
 
 FMacApplication::~FMacApplication()
 {
+	if(AppActivationObserver)
+	{
+		[[NSNotificationCenter defaultCenter] removeObserver:AppActivationObserver];
+		AppActivationObserver = nil;
+	}
+	
+	if(AppDeactivationObserver)
+	{
+		[[NSNotificationCenter defaultCenter] removeObserver:AppDeactivationObserver];
+		AppDeactivationObserver = nil;
+	}
+	
 	CGDisplayRemoveReconfigurationCallback(FMacApplication::OnDisplayReconfiguration, this);
 	if (MouseCaptureWindow)
 	{
diff --git a/Engine/Source/Runtime/Core/Private/Mac/MacPlatformMisc.cpp b/Engine/Source/Runtime/Core/Private/Mac/MacPlatformMisc.cpp
--- a/Engine/Source/Runtime/Core/Private/Mac/MacPlatformMisc.cpp
+++ b/Engine/Source/Runtime/Core/Private/Mac/MacPlatformMisc.cpp
@@ -650,39 +650,6 @@ void FMacPlatformMisc::PumpMessages( bool bFromMainLoop )
 
 			[NSApp sendEvent: Event];
 		}
-
-		const bool HasFocus = [NSApp isActive];
-
-		// If editor thread doesn't have the focus, don't suck up too much CPU time.
-		if( GIsEditor )
-		{
-			static bool HadFocus=1;
-			if( HadFocus && !HasFocus )
-			{
-				// Drop our priority to speed up whatever is in the foreground.
-				struct sched_param Sched;
-				FMemory::Memzero(&Sched, sizeof(struct sched_param));
-				Sched.sched_priority = 5;
-				pthread_setschedparam(pthread_self(), SCHED_RR, &Sched);
-			}
-			else if( HasFocus && !HadFocus )
-			{
-				// Boost our priority back to normal.
-				struct sched_param Sched;
-				FMemory::Memzero(&Sched, sizeof(struct sched_param));
-				Sched.sched_priority = 15;
-				pthread_setschedparam(pthread_self(), SCHED_RR, &Sched);
-			}
-			if( !HasFocus )
-			{
-				// Sleep for a bit to not eat up all CPU time.
-				FPlatformProcess::Sleep(0.005f);
-			}
-			HadFocus = HasFocus;
-		}
-
-		// if app is active, allow sound, otherwise silence it
-		GVolumeMultiplier = HasFocus ? 1.0f : 0.0f;
 	}
 }
 
diff --git a/Engine/Source/Runtime/Core/Public/Mac/MacApplication.h b/Engine/Source/Runtime/Core/Public/Mac/MacApplication.h
--- a/Engine/Source/Runtime/Core/Public/Mac/MacApplication.h
+++ b/Engine/Source/Runtime/Core/Public/Mac/MacApplication.h
@@ -159,6 +159,12 @@ private:
 
 	TSharedPtr<FMacTextInputMethodSystem> TextInputMethodSystem;
 
+	/** Notification center observer for application activation events */
+	id AppActivationObserver;
+
+	/** Notification center observer for application deactivation events */
+	id AppDeactivationObserver;
+
 #if WITH_EDITOR
 	/** Holds the last gesture used to try and capture unique uses for gestures. */
 	EGestureEvent::Type LastGestureUsed;
