Move the Mac mouse event handling into a single location in FSlateTextView.

diff --git a/Engine/Source/Developer/StandaloneRenderer/Private/Mac/OpenGL/SlateOpenGLContext.cpp b/Engine/Source/Developer/StandaloneRenderer/Private/Mac/OpenGL/SlateOpenGLContext.cpp
--- a/Engine/Source/Developer/StandaloneRenderer/Private/Mac/OpenGL/SlateOpenGLContext.cpp
+++ b/Engine/Source/Developer/StandaloneRenderer/Private/Mac/OpenGL/SlateOpenGLContext.cpp
@@ -123,90 +123,6 @@ char const* const CompositedBlitFragmentShader = "#version 120\n"
 	return PixelFormat;
 }
 
-- (BOOL)acceptsFirstMouse:(NSEvent *)Event
-{
-	return YES;
-}
-
-/** Forward mouse events up to the window rather than through the responder chain - thus avoiding
- *	the hidden titlebar controls. Normal windows just use the responder chain as usual.
- */
-- (void)mouseDown:(NSEvent*)Event
-{
-	// Swallowed by FSlateTextView
-	[super mouseDown:Event];
-	
-	FSlateCocoaWindow* SlateCocoaWindow = [[self window] isKindOfClass:[FSlateCocoaWindow class]] ? (FSlateCocoaWindow*)[self window] : nil;
-	if (SlateCocoaWindow)
-	{
-		[SlateCocoaWindow mouseDown:Event];
-	}
-}
-
-- (void)rightMouseDown:(NSEvent*)Event
-{
-	FSlateCocoaWindow* SlateCocoaWindow = [[self window] isKindOfClass:[FSlateCocoaWindow class]] ? (FSlateCocoaWindow*)[self window] : nil;
-	if (SlateCocoaWindow)
-	{
-		[SlateCocoaWindow rightMouseDown:Event];
-	}
-	else
-	{
-		[super rightMouseDown:Event];
-	}
-}
-
-- (void)otherMouseDown:(NSEvent*)Event
-{
-	FSlateCocoaWindow* SlateCocoaWindow = [[self window] isKindOfClass:[FSlateCocoaWindow class]] ? (FSlateCocoaWindow*)[self window] : nil;
-	if (SlateCocoaWindow)
-	{
-		[SlateCocoaWindow otherMouseDown:Event];
-	}
-	else
-	{
-		[super otherMouseDown:Event];
-	}
-}
-
-- (void)mouseUp:(NSEvent*)Event
-{
-	// Swallowed by FSlateTextView
-	[super mouseUp:Event];
-	
-	FSlateCocoaWindow* SlateCocoaWindow = [[self window] isKindOfClass:[FSlateCocoaWindow class]] ? (FSlateCocoaWindow*)[self window] : nil;
-	if (SlateCocoaWindow)
-	{
-		[SlateCocoaWindow mouseUp:Event];
-	}
-}
-
-- (void)rightMouseUp:(NSEvent*)Event
-{
-	FSlateCocoaWindow* SlateCocoaWindow = [[self window] isKindOfClass:[FSlateCocoaWindow class]] ? (FSlateCocoaWindow*)[self window] : nil;
-	if (SlateCocoaWindow)
-	{
-		[SlateCocoaWindow rightMouseUp:Event];
-	}
-	else
-	{
-		[super rightMouseUp:Event];
-	}
-}
-
-- (void)otherMouseUp:(NSEvent*)Event
-{
-	FSlateCocoaWindow* SlateCocoaWindow = [[self window] isKindOfClass:[FSlateCocoaWindow class]] ? (FSlateCocoaWindow*)[self window] : nil;
-	if (SlateCocoaWindow)
-	{
-		[SlateCocoaWindow otherMouseUp:Event];
-	}
-	else
-	{
-		[super otherMouseUp:Event];
-	}
-}
-
 - (void)renewGState
 {
 	extern bool GMacEnableCocoaScreenUpdates;
diff --git a/Engine/Source/Runtime/Core/Private/Mac/MacTextInputMethodSystem.cpp b/Engine/Source/Runtime/Core/Private/Mac/MacTextInputMethodSystem.cpp
--- a/Engine/Source/Runtime/Core/Private/Mac/MacTextInputMethodSystem.cpp
+++ b/Engine/Source/Runtime/Core/Private/Mac/MacTextInputMethodSystem.cpp
@@ -42,12 +42,26 @@ DEFINE_LOG_CATEGORY_STATIC(LogMacTextInputMethodSystem, Log, All);
 	}
 }
 
+/** Forward mouse events up to the window rather than through the responder chain - thus avoiding
+ *	the hidden titlebar controls. Normal windows just use the responder chain as usual.
+ */
+- (BOOL)acceptsFirstMouse:(NSEvent *)Event
+{
+	return YES;
+}
+
 - (void)mouseDown:(NSEvent *)theEvent
 {
 	if (IMMContext.IsValid())
 	{
 		[[self inputContext] handleEvent:theEvent];
 	}
+	
+	FSlateCocoaWindow* SlateCocoaWindow = [[self window] isKindOfClass:[FSlateCocoaWindow class]] ? (FSlateCocoaWindow*)[self window] : nil;
+	if (SlateCocoaWindow)
+	{
+		[SlateCocoaWindow mouseDown:theEvent];
+	}
 }
 
 - (void)mouseDragged:(NSEvent *)theEvent
@@ -64,6 +78,64 @@ DEFINE_LOG_CATEGORY_STATIC(LogMacTextInputMethodSystem, Log, All);
 	{
 		[[self inputContext] handleEvent:theEvent];
 	}
+	
+	FSlateCocoaWindow* SlateCocoaWindow = [[self window] isKindOfClass:[FSlateCocoaWindow class]] ? (FSlateCocoaWindow*)[self window] : nil;
+	if (SlateCocoaWindow)
+	{
+		[SlateCocoaWindow mouseUp:theEvent];
+	}
+}
+
+- (void)rightMouseDown:(NSEvent*)Event
+{
+	FSlateCocoaWindow* SlateCocoaWindow = [[self window] isKindOfClass:[FSlateCocoaWindow class]] ? (FSlateCocoaWindow*)[self window] : nil;
+	if (SlateCocoaWindow)
+	{
+		[SlateCocoaWindow rightMouseDown:Event];
+	}
+	else
+	{
+		[super rightMouseDown:Event];
+	}
+}
+
+- (void)otherMouseDown:(NSEvent*)Event
+{
+	FSlateCocoaWindow* SlateCocoaWindow = [[self window] isKindOfClass:[FSlateCocoaWindow class]] ? (FSlateCocoaWindow*)[self window] : nil;
+	if (SlateCocoaWindow)
+	{
+		[SlateCocoaWindow otherMouseDown:Event];
+	}
+	else
+	{
+		[super otherMouseDown:Event];
+	}
+}
+
+- (void)rightMouseUp:(NSEvent*)Event
+{
+	FSlateCocoaWindow* SlateCocoaWindow = [[self window] isKindOfClass:[FSlateCocoaWindow class]] ? (FSlateCocoaWindow*)[self window] : nil;
+	if (SlateCocoaWindow)
+	{
+		[SlateCocoaWindow rightMouseUp:Event];
+	}
+	else
+	{
+		[super rightMouseUp:Event];
+	}
+}
+
+- (void)otherMouseUp:(NSEvent*)Event
+{
+	FSlateCocoaWindow* SlateCocoaWindow = [[self window] isKindOfClass:[FSlateCocoaWindow class]] ? (FSlateCocoaWindow*)[self window] : nil;
+	if (SlateCocoaWindow)
+	{
+		[SlateCocoaWindow otherMouseUp:Event];
+	}
+	else
+	{
+		[super otherMouseUp:Event];
+	}
 }
 
 - (void)activateInputMethod:(const TSharedRef<ITextInputMethodContext>&)InContext
diff --git a/Engine/Source/Runtime/OpenGLDrv/Private/Mac/OpenGLMac.cpp b/Engine/Source/Runtime/OpenGLDrv/Private/Mac/OpenGLMac.cpp
--- a/Engine/Source/Runtime/OpenGLDrv/Private/Mac/OpenGLMac.cpp
+++ b/Engine/Source/Runtime/OpenGLDrv/Private/Mac/OpenGLMac.cpp
@@ -217,90 +217,6 @@ bool GMacEnableCocoaScreenUpdates = true;
 	[super renewGState];
 }
 
-- (BOOL)acceptsFirstMouse:(NSEvent *)Event
-{
-	return YES;
-}
-
-/** Forward mouse events up to the window rather than through the responder chain - thus avoiding 
- *	the hidden titlebar controls. Normal windows just use the responder chain as usual.
- */
-- (void)mouseDown:(NSEvent*)Event
-{
-	// Swallowed by FSlateTextView
-	[super mouseDown:Event];
-	
-	FSlateCocoaWindow* SlateCocoaWindow = [[self window] isKindOfClass:[FSlateCocoaWindow class]] ? (FSlateCocoaWindow*)[self window] : nil;
-	if (SlateCocoaWindow)
-	{
-		[SlateCocoaWindow mouseDown:Event];
-	}
-}
-
-- (void)rightMouseDown:(NSEvent*)Event
-{
-	FSlateCocoaWindow* SlateCocoaWindow = [[self window] isKindOfClass:[FSlateCocoaWindow class]] ? (FSlateCocoaWindow*)[self window] : nil;
-	if (SlateCocoaWindow)
-	{
-		[SlateCocoaWindow rightMouseDown:Event];
-	}
-	else
-	{
-		[super rightMouseDown:Event];
-	}
-}
-
-- (void)otherMouseDown:(NSEvent*)Event
-{
-	FSlateCocoaWindow* SlateCocoaWindow = [[self window] isKindOfClass:[FSlateCocoaWindow class]] ? (FSlateCocoaWindow*)[self window] : nil;
-	if (SlateCocoaWindow)
-	{
-		[SlateCocoaWindow otherMouseDown:Event];
-	}
-	else
-	{
-		[super otherMouseDown:Event];
-	}
-}
-
-- (void)mouseUp:(NSEvent*)Event
-{
-	// Swallowed by FSlateTextView
-	[super mouseUp:Event];
-	
-	FSlateCocoaWindow* SlateCocoaWindow = [[self window] isKindOfClass:[FSlateCocoaWindow class]] ? (FSlateCocoaWindow*)[self window] : nil;
-	if (SlateCocoaWindow)
-	{
-		[SlateCocoaWindow mouseUp:Event];
-	}
-}
-
-- (void)rightMouseUp:(NSEvent*)Event
-{
-	FSlateCocoaWindow* SlateCocoaWindow = [[self window] isKindOfClass:[FSlateCocoaWindow class]] ? (FSlateCocoaWindow*)[self window] : nil;
-	if (SlateCocoaWindow)
-	{
-		[SlateCocoaWindow rightMouseUp:Event];
-	}
-	else
-	{
-		[super rightMouseUp:Event];
-	}
-}
-
-- (void)otherMouseUp:(NSEvent*)Event
-{
-	FSlateCocoaWindow* SlateCocoaWindow = [[self window] isKindOfClass:[FSlateCocoaWindow class]] ? (FSlateCocoaWindow*)[self window] : nil;
-	if (SlateCocoaWindow)
-	{
-		[SlateCocoaWindow otherMouseUp:Event];
-	}
-	else
-	{
-		[super otherMouseUp:Event];
-	}
-}
-
 @end
 
 class FMacOpenGLTimer
