Improve the Mac IME support so that it shows the range being edited by the pop-up glyph selection window for input-methods that provide one.

diff --git a/Engine/Source/Runtime/Core/Private/Mac/MacTextInputMethodSystem.cpp b/Engine/Source/Runtime/Core/Private/Mac/MacTextInputMethodSystem.cpp
--- a/Engine/Source/Runtime/Core/Private/Mac/MacTextInputMethodSystem.cpp
+++ b/Engine/Source/Runtime/Core/Private/Mac/MacTextInputMethodSystem.cpp
@@ -253,10 +253,31 @@ DEFINE_LOG_CATEGORY_STATIC(LogMacTextInputMethodSystem, Log, All);
 			}
 			markedRange = NSMakeRange(SelectionLocation, [aString length]);
 			
+			__block NSRange CompositionRange = markedRange;
+			
 			NSString* TheString;
 			if ([aString isKindOfClass:[NSAttributedString class]])
 			{
-				TheString = [(NSAttributedString*)aString string];
+				// While the whole string is being composed NSUnderlineStyleAttributeName is 1 to show a single line below the whole string.
+				// When using the pop-up glyph selection window in some IME's the NSAttributedString is broken up into separate glyph ranges, each with its own set of attributes.
+				// Each range specifies NSMarkedClauseSegment, incrementing the NSNumber value from 0 as well as NSUnderlineStyleAttributeName, which makes the underlining show the different ranges.
+				// The subrange being edited by the pop-up glyph selection window will set NSUnderlineStyleAttributeName to a value >1, while all other ranges will be set NSUnderlineStyleAttributeName to 1.
+				NSAttributedString* AttributedString = (NSAttributedString*)aString;
+				[AttributedString enumerateAttribute:NSUnderlineStyleAttributeName inRange:NSMakeRange(0, [aString length]) options:0 usingBlock:^(id Value, NSRange Range, BOOL* bStop) {
+					if(Value && [Value isKindOfClass:[NSNumber class]])
+					{
+						NSNumber* NumberValue = (NSNumber*)Value;
+						const int UnderlineValue = [NumberValue intValue];
+						if(UnderlineValue > 1)
+						{
+							// Found the active range, stop enumeration.
+							*bStop = YES;
+							CompositionRange = Range;
+						}
+					}
+				}];
+				
+				TheString = [AttributedString string];
 			}
 			else
 			{
@@ -265,8 +286,8 @@ DEFINE_LOG_CATEGORY_STATIC(LogMacTextInputMethodSystem, Log, All);
 			
 			FString TheFString(TheString);
 			IMMContext->SetTextInRange(SelectionLocation, SelectionLength, TheFString);
-			IMMContext->UpdateCompositionRange(markedRange.location, markedRange.length);
-			IMMContext->SetSelectionRange(SelectionLocation+TheFString.Len(), 0, ITextInputMethodContext::ECaretPosition::Ending);
+			IMMContext->UpdateCompositionRange(CompositionRange.location, CompositionRange.length);
+			IMMContext->SetSelectionRange(CompositionRange.location + CompositionRange.length, 0, ITextInputMethodContext::ECaretPosition::Ending);
 		}
 		[[self inputContext] invalidateCharacterCoordinates]; // recentering
 	}
