When parsing the attributed string to locate where to draw the IME highlight on OS X remember that it is an offset within the marked range, not from the start of the whole string.

diff --git a/Engine/Source/Runtime/Core/Private/Mac/MacTextInputMethodSystem.cpp b/Engine/Source/Runtime/Core/Private/Mac/MacTextInputMethodSystem.cpp
--- a/Engine/Source/Runtime/Core/Private/Mac/MacTextInputMethodSystem.cpp
+++ b/Engine/Source/Runtime/Core/Private/Mac/MacTextInputMethodSystem.cpp
@@ -272,7 +272,8 @@ DEFINE_LOG_CATEGORY_STATIC(LogMacTextInputMethodSystem, Log, All);
 						{
 							// Found the active range, stop enumeration.
 							*bStop = YES;
-							CompositionRange = Range;
+							CompositionRange.location += Range.location;
+							CompositionRange.length = Range.length;
 						}
 					}
 				}];
