Fix reinsertion of composed text after clicking within the existing composition when using the OS X IME and a failure to clear the composed range when ending IME composition.
As we aren't using proper Cocoa views we have to manually clear the OS X IME's marked text when we deactivate our internal IME context to ensure that we can never reinsert previous text. We then need to ensure that we clean-up the composition range when unmarking text as we won't get further insertion/replacement messages.

diff --git a/Engine/Source/Runtime/Core/Private/Mac/MacTextInputMethodSystem.cpp b/Engine/Source/Runtime/Core/Private/Mac/MacTextInputMethodSystem.cpp
--- a/Engine/Source/Runtime/Core/Private/Mac/MacTextInputMethodSystem.cpp
+++ b/Engine/Source/Runtime/Core/Private/Mac/MacTextInputMethodSystem.cpp
@@ -144,6 +144,7 @@ DEFINE_LOG_CATEGORY_STATIC(LogMacTextInputMethodSystem, Log, All);
 	{
 		[self unmarkText];
 		[[self inputContext] deactivate];
+		[[self inputContext] discardMarkedText];
 	}
 	IMMContext = InContext;
 	[[self inputContext] activate];
@@ -154,6 +155,7 @@ DEFINE_LOG_CATEGORY_STATIC(LogMacTextInputMethodSystem, Log, All);
 	[self unmarkText];
 	IMMContext = NULL;
 	[[self inputContext] deactivate];
+	[[self inputContext] discardMarkedText];
 }
 
 //@protocol NSTextInputClient
@@ -302,10 +304,14 @@ DEFINE_LOG_CATEGORY_STATIC(LogMacTextInputMethodSystem, Log, All);
  */
 - (void)unmarkText
 {
-	markedRange = {NSNotFound, 0};
-	if (IMMContext.IsValid())
+	if(markedRange.location != NSNotFound)
 	{
-		IMMContext->EndComposition();
+		markedRange = {NSNotFound, 0};
+		if (IMMContext.IsValid())
+		{
+			IMMContext->UpdateCompositionRange(0, 0);
+			IMMContext->EndComposition();
+		}
 	}
 }
 
