Update the caret position to the specified selection position when setting the marked text in the OS X IME.

diff --git a/Engine/Source/Runtime/Core/Private/Mac/MacTextInputMethodSystem.cpp b/Engine/Source/Runtime/Core/Private/Mac/MacTextInputMethodSystem.cpp
--- a/Engine/Source/Runtime/Core/Private/Mac/MacTextInputMethodSystem.cpp
+++ b/Engine/Source/Runtime/Core/Private/Mac/MacTextInputMethodSystem.cpp
@@ -290,7 +290,7 @@ DEFINE_LOG_CATEGORY_STATIC(LogMacTextInputMethodSystem, Log, All);
 			FString TheFString(TheString);
 			IMMContext->SetTextInRange(SelectionLocation, SelectionLength, TheFString);
 			IMMContext->UpdateCompositionRange(CompositionRange.location, CompositionRange.length);
-			IMMContext->SetSelectionRange(CompositionRange.location + CompositionRange.length, 0, ITextInputMethodContext::ECaretPosition::Ending);
+			IMMContext->SetSelectionRange(markedRange.location + selectedRange.location, 0, ITextInputMethodContext::ECaretPosition::Ending);
 		}
 		[[self inputContext] invalidateCharacterCoordinates]; // recentering
 	}
