Remove usage of Cocoa's child-window mechanism, it just doesn't do what we want. Instead set tear-out or panel-style windows to the floating-window level but be otherwise normal windows. They will appear in Expose/Spaces/Mission Control & the window cycling. Modal windows are treated slightly specially to ensure that the last shown modal window currently still open remains on top and attempts to change focus fail with visual indication. The code to detect the current window under the cursor has also been necessarily changed to handle all this.

diff --git a/Engine/Source/Editor/MainFrame/Private/MainFrameModule.cpp b/Engine/Source/Editor/MainFrame/Private/MainFrameModule.cpp
--- a/Engine/Source/Editor/MainFrame/Private/MainFrameModule.cpp
+++ b/Engine/Source/Editor/MainFrame/Private/MainFrameModule.cpp
@@ -89,8 +89,11 @@ void FMainFrameModule::CreateDefaultMainFrame( const bool bStartImmersivePIE )
 
 		const bool bShowRootWindowImmediately = false;
 		FSlateApplication::Get().AddWindow( RootWindow, bShowRootWindowImmediately );
+		
+#if !PLATFORM_MAC // On OS X we don't want Top-Level windows to have a parent, as we don't really support the notion of child windows on that OS
 		FGlobalTabmanager::Get()->SetRootWindow(RootWindow);
 		FSlateNotificationManager::Get().SetRootWindow(RootWindow);
+#endif
 
 		TSharedPtr<SWidget> MainFrameContent;
 		bool bLevelEditorIsMainTab = false;
diff --git a/Engine/Source/Runtime/Core/Private/Mac/MacApplication.cpp b/Engine/Source/Runtime/Core/Private/Mac/MacApplication.cpp
--- a/Engine/Source/Runtime/Core/Private/Mac/MacApplication.cpp
+++ b/Engine/Source/Runtime/Core/Private/Mac/MacApplication.cpp
@@ -186,7 +186,6 @@ void FMacApplication::OnWindowDraggingFinished()
 	if( DraggedWindow )
 	{
 		SCOPED_AUTORELEASE_POOL;
-		[DraggedWindow reconnectChildWindows];
 		DraggedWindow = NULL;
 	}
 }
@@ -652,69 +651,70 @@ FSlateCocoaWindow* FMacApplication::FindEventWindow( NSEvent* Event )
 		else
 		{
 			const NSPoint CursorPos = [NSEvent mouseLocation];
-
-			// Only the editor needs to handle the space-per-display logic introduced in Mavericks.
-#if WITH_EDITOR
-			NSScreen* MouseScreen = nil;
-			// Only fetch the spans-displays once - it requires a log-out to change.
-			static bool bScreensHaveSeparateSpaces = false;
-			static bool bSettingFetched = false;
-			if(!bSettingFetched)
+			TSharedPtr<FMacWindow> WindowUnderCursor = StaticCastSharedPtr<FMacWindow>(LocateWindowUnderCursor(FVector2D(CursorPos.x, FPlatformMisc::ConvertSlateYPositionToCocoa(CursorPos.y))));
+			if (WindowUnderCursor.IsValid())
 			{
-				bSettingFetched = true;
-				bScreensHaveSeparateSpaces = [NSScreen screensHaveSeparateSpaces];
+				EventWindow = WindowUnderCursor->GetWindowHandle();
 			}
-			if(bScreensHaveSeparateSpaces)
-			{
-				// New default mode which uses a separate Space per display
-				// Find the screen the cursor is currently on so we can ignore invisible window regions.
-				NSEnumerator* ScreenEnumerator = [[NSScreen screens] objectEnumerator];
-				while ((MouseScreen = [ScreenEnumerator nextObject]) && !NSMouseInRect(CursorPos, MouseScreen.frame, NO))
-					;
-			}
-#endif
-			
-			NSArray* AllWindows = [NSApp orderedWindows];
-			for( int32 Index = 0; Index < [AllWindows count]; Index++ )
-			{
-				NSWindow* Window = (NSWindow*)[AllWindows objectAtIndex: Index];
-				if( [Window isMiniaturized] || ![Window isVisible] || ![Window isOnActiveSpace] || ![Window isKindOfClass: [FSlateCocoaWindow class]] )
-				{
-					continue;
-				}
+		}
+	}
 
-				if( [Window canBecomeKeyWindow] == false )
-				{
-					NSWindow* ParentWindow = [Window parentWindow];
-					while( ParentWindow )
-					{
-						if( [ParentWindow canBecomeKeyWindow] )
-						{
-							Window = ParentWindow;
-							break;
-						}
-						ParentWindow = [ParentWindow parentWindow];
-					}
-				}
-				
-				NSRect VisibleFrame = [Window frame];
+	return EventWindow;
+}
+
+TSharedPtr<FGenericWindow> FMacApplication::LocateWindowUnderCursor( const FVector2D& CursorPos )
+{
+	const NSPoint Position = NSMakePoint(CursorPos.X, FPlatformMisc::ConvertSlateYPositionToCocoa(CursorPos.Y));
+
+	NSScreen* MouseScreen = nil;
+	if ([NSScreen screensHaveSeparateSpaces])
+	{
+		// New default mode which uses a separate Space per display
+		// Find the screen the cursor is currently on so we can ignore invisible window regions.
+		NSEnumerator* ScreenEnumerator = [[NSScreen screens] objectEnumerator];
+		while ((MouseScreen = [ScreenEnumerator nextObject]) && !NSMouseInRect(Position, MouseScreen.frame, NO))
+			;
+	}
+
+	NSArray* AllWindows = [NSApp orderedWindows];
+	for (int32 Index = 0; Index < [AllWindows count]; Index++)
+	{
+		NSWindow* NativeWindow = (NSWindow*)[AllWindows objectAtIndex: Index];
+		if ([NativeWindow isMiniaturized] || ![NativeWindow isVisible] || ![NativeWindow isOnActiveSpace] || ![NativeWindow isKindOfClass: [FSlateCocoaWindow class]])
+		{
+			continue;
+		}
+
+        if( [NativeWindow canBecomeKeyWindow] == false )
+        {
+            NSWindow* ParentWindow = [NativeWindow parentWindow];
+            while( ParentWindow )
+            {
+                if( [ParentWindow canBecomeKeyWindow] )
+                {
+                    NativeWindow = ParentWindow;
+                    break;
+                }
+                ParentWindow = [ParentWindow parentWindow];
+            }
+        }
+        
+        NSRect VisibleFrame = [NativeWindow frame];
 #if WITH_EDITOR
-				if(MouseScreen != nil)
-				{
-					VisibleFrame = NSIntersectionRect([MouseScreen frame], VisibleFrame);
-				}
+        if(MouseScreen != nil)
+        {
+            VisibleFrame = NSIntersectionRect([MouseScreen frame], VisibleFrame);
+        }
 #endif
 
-				if( NSPointInRect( CursorPos, VisibleFrame ) )
-				{
-					EventWindow = (FSlateCocoaWindow*)Window;
-					break;
-				}
-			}
+		if (NSPointInRect(Position, VisibleFrame))
+		{
+			TSharedPtr<FMacWindow> MacWindow = FindWindowByNSWindow(Windows, (FSlateCocoaWindow*)NativeWindow);
+			return MacWindow;
 		}
 	}
 
-	return EventWindow;
+	return NULL;
 }
 
 void FMacApplication::PollGameDeviceState( const float TimeDelta )
@@ -993,7 +993,6 @@ void FMacApplication::OnWindowWillMove( FSlateCocoaWindow* Window )
 	SCOPED_AUTORELEASE_POOL;
 
 	DraggedWindow = Window;
-	[DraggedWindow disconnectChildWindows];
 }
 
 void FMacApplication::OnWindowDidMove( FSlateCocoaWindow* Window )
diff --git a/Engine/Source/Runtime/Core/Private/Mac/MacWindow.cpp b/Engine/Source/Runtime/Core/Private/Mac/MacWindow.cpp
--- a/Engine/Source/Runtime/Core/Private/Mac/MacWindow.cpp
+++ b/Engine/Source/Runtime/Core/Private/Mac/MacWindow.cpp
@@ -4,9 +4,7 @@
 #include "MacWindow.h"
 #include "MacApplication.h"
 
-@interface NSWindow (UE4Extensions)
-- (bool)shouldAddChildWindow:(NSWindow*)Window;
-@end
+TArray< FSlateCocoaWindow* > GRunningModalWindows;
 
 /**
  * Custom window class used for input handling
@@ -20,9 +18,7 @@
 	id NewSelf = [super initWithContentRect:ContentRect styleMask:Style backing:BufferingType defer:Flag];
 	if(NewSelf)
 	{
-		Parent = NULL;
 		WindowMode = EWindowMode::Windowed;
-		ChildWindows = nil;;
 		bAcceptsInput = false;
 		bRoundedCorners = false;
 		bDisplayReconfiguring = false;
@@ -122,11 +118,15 @@
 
 - (void)orderWindow:(NSWindowOrderingMode)OrderingMode relativeTo:(NSInteger)OtherWindowNumber
 {
-	if([self alphaValue] > 0.0f)
+	bool bModal = GRunningModalWindows.Num() == 0 || GRunningModalWindows.Last() == self;
+	if(OrderingMode == NSWindowOut || bModal)
 	{
-		[self performDeferredSetFrame];
+		if([self alphaValue] > 0.0f)
+		{
+			[self performDeferredSetFrame];
+		}
+		[super orderWindow:OrderingMode relativeTo:OtherWindowNumber];
 	}
-	[super orderWindow:OrderingMode relativeTo:OtherWindowNumber];
 }
 
 - (bool)roundedCorners
@@ -167,35 +167,13 @@
 	bDisplayReconfiguring = bIsDisplayReconfiguring;
 }
 
-- (void)setParent:(NSWindow*)InParent
-{
-	Parent = InParent;
-}
-
-- (bool)shouldAddChildWindow:(NSWindow*)Window
-{
-	bool bShouldAddChildWindow = false;
-	if(Window && ![self isMiniaturized] && NSContainsRect([[self screen] frame], [Window frame]))
-	{
-		bool bSelfOnScreen = [Window isOnActiveSpace];
-		bool bParentOnScreen = [self isOnActiveSpace];
-		bShouldAddChildWindow = bSelfOnScreen && bParentOnScreen;
-	}
-	return bShouldAddChildWindow;
-}
-
-- (void)orderFrontEvenIfChildAndMakeMain:(bool)bMain andKey:(bool)bKey
+- (void)orderFrontAndMakeMain:(bool)bMain andKey:(bool)bKey
 {
 	if ([NSApp isHidden] == NO)
 	{
-		NSWindow* ParentWindow = [self getParent];
-		if( ParentWindow && [ParentWindow shouldAddChildWindow:self] )
-		{
-			[ParentWindow removeChildWindow: self];
-			[ParentWindow addChildWindow: self ordered: NSWindowAbove];
-		}
 
-		if (![self isVisible])
+		bool bModal = GRunningModalWindows.Num() == 0 || GRunningModalWindows.Last() == self;
+		if (![self isVisible] && bModal)
 		{
 			[self orderFront:nil];
 		}
@@ -211,67 +189,11 @@
 	}
 }
 
-- (void)disconnectChildWindows
-{
-	ChildWindows = [[NSMutableArray alloc] initWithArray: [self childWindows]];
-	for( int32 Index = 0; Index < [ChildWindows count]; Index++ )
-	{
-		NSWindow* Window = (NSWindow*)[ChildWindows objectAtIndex: Index];
-		[self removeChildWindow: Window];
-		[Window setLevel: NSFloatingWindowLevel];
-	}
-	if(Parent && [self parentWindow])
-	{
-		[Parent removeChildWindow:self];
-		[self setLevel: NSFloatingWindowLevel];
-	}
-}
-
-- (void)reconnectChildWindows
-{
-	if( ChildWindows )
-	{
-		for( int32 Index = 0; Index < [ChildWindows count]; Index++ )
-		{
-			NSWindow* Window = (NSWindow*)[ChildWindows objectAtIndex: Index];
-			if([self shouldAddChildWindow:Window])
-			{
-				[self addChildWindow: Window ordered: NSWindowAbove];
-				[Window setLevel: ([Window styleMask] & NSClosableWindowMask) ? NSNormalWindowLevel : NSFloatingWindowLevel];
-			}
-		}
-		[ChildWindows release];
-		ChildWindows = NULL;
-	}
-	if(Parent && [self parentWindow] == nil && [Parent shouldAddChildWindow:self])
-	{
-		[Parent addChildWindow: self ordered: NSWindowAbove];
-		[self setLevel: ([self styleMask] & NSClosableWindowMask) ? NSNormalWindowLevel : NSFloatingWindowLevel];
-	}
-}
-
-- (bool)shouldAddChildWindows
-{
-	return (ChildWindows == NULL) && ![self isMiniaturized];
-}
-
-- (NSWindow*)getParent
-{
-	return Parent;
-}
-
-- (void)removeCachedChild:(NSWindow*)InChild
-{
-	if(ChildWindows && InChild)
-	{
-		[ChildWindows removeObject:InChild];
-	}
-}
-
 // Following few methods overload NSWindow's methods from Cocoa API, so have to use Cocoa's BOOL (signed char), not bool (unsigned int)
 - (BOOL)canBecomeMainWindow
 {
-	return bAcceptsInput && ([self styleMask] != NSBorderlessWindowMask);
+	bool bNoModalOrCurrent = GRunningModalWindows.Num() == 0 || GRunningModalWindows.Last() == self;
+	return bAcceptsInput && ([self styleMask] != NSBorderlessWindowMask) && bNoModalOrCurrent;
 }
 
 - (BOOL)canBecomeKeyWindow
@@ -366,30 +288,6 @@
 {
 }
 
-- (void)miniaturize:(id)sender
-{
-	// If parented detach or the whole app will hide, though the mouse handling will still be directed through our app.
-	// There's a reason why Inspector windows in Cocoa apps don't have a minimise button...
-	if(Parent)
-	{
-		[Parent removeChildWindow:self];
-	}
-	[super miniaturize:sender];
-}
-
-- (void)windowDidDeminiaturize:(NSNotification*)Notification
-{
-	// If parented ensure that the parent is visible and then reattach, to match Windows behaviour.
-	if(Parent)
-	{
-		if([Parent isMiniaturized] == YES)
-		{
-			[Parent deminiaturize:self];
-		}
-		[Parent addChildWindow:self ordered: NSWindowAbove];
-	}
-}
-
 - (void)windowDidEnterFullScreen:(NSNotification *)notification
 {
 	WindowMode = EWindowMode::WindowedFullscreen;
@@ -412,9 +310,19 @@
 {
 	if([NSApp isHidden] == NO)
 	{
-		// First, make sure we're in front - on our child window level.
-		// That's what happens on PC, and engine doesn't give us any BringToFront() calls.
-		[self orderFrontEvenIfChildAndMakeMain:false andKey:false];
+		bool bModal = GRunningModalWindows.Num() != 0 && GRunningModalWindows.Last() == self;
+		if(bModal)
+		{
+			[self setLevel:NSModalPanelWindowLevel];
+		}
+		if(GRunningModalWindows.Num() == 0 || GRunningModalWindows.Last() == self)
+		{
+			[self orderFrontAndMakeMain:false andKey:false];
+		}
+		else
+		{
+			[GRunningModalWindows.Last() orderFrontAndMakeMain:true andKey:true];
+		}
 	}
 
 	if(self.bForwardEvents)
@@ -428,6 +336,11 @@
 	[self setMovable: YES];
 	[self setMovableByWindowBackground: NO];
 	
+	if([self level] == NSModalPanelWindowLevel)
+	{
+		[self setLevel:NSNormalWindowLevel];
+	}
+	
 	if(self.bForwardEvents)
 	{
 		MacApplication->OnWindowDidResignKey( self );
@@ -463,15 +376,6 @@
 	// It does however, work fine for handling display arrangement changes that cause a window to go offscreen.
 	if(bDisplayReconfiguring)
 	{
-		if(Parent && [self parentWindow] == nil && [Parent shouldAddChildWindow:self])
-		{
-			[Parent addChildWindow: self ordered: NSWindowAbove];
-		}
-		else if(Parent && [self parentWindow])
-		{
-			[Parent removeChildWindow:self];
-		}
-		
 		NSScreen* Screen = [self screen];
 		NSRect Frame = [self frame];
 		NSRect VisibleFrame = [Screen visibleFrame];
@@ -709,7 +613,6 @@
 
 @end
 
-
 FMacWindow::~FMacWindow()
 {
 	// While on Windows invalid HWNDs fail silently, accessing an invalid NSWindow is fatal.
@@ -805,7 +708,15 @@ void FMacWindow::Initialize( FMacApplication* const Application, const TSharedRe
 	[WindowHandle setDisplayReconfiguring: false];
 	[WindowHandle setAcceptsMouseMovedEvents: YES];
 	[WindowHandle setDelegate: WindowHandle];
-	[WindowHandle setLevel: Definition->IsRegularWindow ? NSNormalWindowLevel : NSFloatingWindowLevel];
+	if(Definition->IsRegularWindow && !InParent.IsValid())
+	{
+		[WindowHandle setLevel: NSNormalWindowLevel];
+	}
+	else
+	{
+		[WindowHandle setLevel: NSFloatingWindowLevel];
+		[WindowHandle setHidesOnDeactivate: YES];
+	}
 	
 	// Use of rounded corners will always render with system values for rounding
 	[WindowHandle setRoundedCorners: (Definition->CornerRadius != 0)];
@@ -850,20 +761,6 @@ void FMacWindow::Initialize( FMacApplication* const Application, const TSharedRe
 		[WindowHandle setCollectionBehavior:NSWindowCollectionBehaviorCanJoinAllSpaces|NSWindowCollectionBehaviorTransient|NSWindowCollectionBehaviorIgnoresCycle];
 	}
 
-	if( InParent.IsValid() )
-	{
-		FSlateCocoaWindow* ParentWindowHandle = static_cast<FSlateCocoaWindow*>(InParent->WindowHandle);
-		[WindowHandle setParent:ParentWindowHandle];
-		if(![ParentWindowHandle shouldAddChildWindows] || ![ParentWindowHandle shouldAddChildWindow:WindowHandle])
-		{
-			[WindowHandle setLevel: ([WindowHandle styleMask] & NSClosableWindowMask) ? NSNormalWindowLevel : NSFloatingWindowLevel];
-		}
-	}
-	else
-	{
-		[WindowHandle setParent:nil];
-	}
-
 	if( Definition->SupportsTransparency )
 	{
 		SetOpacity( Definition->Opacity );
@@ -944,7 +841,7 @@ void FMacWindow::BringToFront( bool bForce )
 	if (bIsVisible)
 	{
 		SCOPED_AUTORELEASE_POOL;
-		[WindowHandle orderFrontEvenIfChildAndMakeMain:false andKey:false];
+		[WindowHandle orderFrontAndMakeMain:false andKey:false];
 	}
 }
 
@@ -957,19 +854,15 @@ void FMacWindow::Destroy()
 		// This makes it possible to correctly handle any NSEvent's sent to the window during destruction
 		// as the WindowHandle is still a valid NSWindow. Unlike HWNDs accessing a destructed NSWindow* is fatal.
 		FSlateCocoaWindow* Window = [WindowHandle retain];
+		if(!Definition->SupportsMinimize && !Definition->SupportsMaximize)
+		{
+			GRunningModalWindows.Remove(WindowHandle);
+		}
 		Window.bForwardEvents = false;
 		if( MacApplication->OnWindowDestroyed( Window ) )
 		{
 			// This FMacWindow may have been destructed by now & so the WindowHandle will probably be invalid memory.
 			
-			// Remove from parent
-			FSlateCocoaWindow* ParentWindow = (FSlateCocoaWindow*)[Window getParent];
-			if( ParentWindow )
-			{
-				[ParentWindow removeCachedChild: Window];
-				[ParentWindow removeChildWindow: Window];
-			}
-			
 			// Then change the focus to something useful, either the previous in the stack
 			TSharedPtr<FMacWindow> KeyWindow = MacApplication->GetKeyWindow();
 			if( KeyWindow.IsValid() && [Window isKeyWindow] && KeyWindow->GetOSWindowHandle() && ![(FSlateCocoaWindow*)KeyWindow->GetOSWindowHandle() isMiniaturized] )
@@ -977,14 +870,7 @@ void FMacWindow::Destroy()
 				// Activate specified previous window if still present, provided it isn't minimized
 				KeyWindow->SetWindowFocus();
 			}
-			// Or, activate parent to make sure that some other window won't steal focus, provided it isn't minimized. Important for menu stack.
-			else if( ParentWindow && [Window isKeyWindow] && ![ParentWindow isMiniaturized] )
-			{
-				// If miniaturized dispatch the message to prevent PrivateDrawWindows from exploding because the
-				// window we are currently closing hasn't been removed from the window list.
-				[ParentWindow orderFrontEvenIfChildAndMakeMain:true andKey:true];
-			}
-			
+
 			// Close the window, but don't destruct it...
 			[Window performClose:nil];
 			
@@ -1034,8 +920,12 @@ void FMacWindow::Show()
 	SCOPED_AUTORELEASE_POOL;
 	if (!bIsVisible)
 	{
+		if(!Definition->SupportsMinimize && !Definition->SupportsMaximize && !GRunningModalWindows.Contains(WindowHandle))
+		{
+			GRunningModalWindows.Add(WindowHandle);
+		}
 		bool bMakeMainAndKey = ([WindowHandle canBecomeKeyWindow] && Definition->ActivateWhenFirstShown);
-		[WindowHandle orderFrontEvenIfChildAndMakeMain:bMakeMainAndKey andKey:bMakeMainAndKey];
+		[WindowHandle orderFrontAndMakeMain:bMakeMainAndKey andKey:bMakeMainAndKey];
 		
 		bIsVisible = [WindowHandle isVisible];
 		static bool bCannotRecurse = false;
@@ -1062,6 +952,10 @@ void FMacWindow::Hide()
 	{
 		SCOPED_AUTORELEASE_POOL;
 		bIsVisible = false;
+		if(!Definition->SupportsMinimize && !Definition->SupportsMaximize)
+		{
+			GRunningModalWindows.Remove(WindowHandle);
+		}
 		[WindowHandle orderOut:nil];
 	}
 }
@@ -1140,7 +1034,7 @@ bool FMacWindow::GetRestoredDimensions(int32& X, int32& Y, int32& Width, int32&
 void FMacWindow::SetWindowFocus()
 {
 	SCOPED_AUTORELEASE_POOL;
-	[WindowHandle orderFrontEvenIfChildAndMakeMain:true andKey:true];
+	[WindowHandle orderFrontAndMakeMain:true andKey:true];
 }
 
 void FMacWindow::SetOpacity( const float InOpacity )
diff --git a/Engine/Source/Runtime/Core/Public/Mac/MacApplication.h b/Engine/Source/Runtime/Core/Public/Mac/MacApplication.h
--- a/Engine/Source/Runtime/Core/Public/Mac/MacApplication.h
+++ b/Engine/Source/Runtime/Core/Public/Mac/MacApplication.h
@@ -114,6 +114,7 @@ private:
 	TCHAR TranslateCharCode( TCHAR CharCode, uint32 KeyCode );
 
 	FSlateCocoaWindow* FindEventWindow( NSEvent* CocoaEvent );
+	TSharedPtr<FGenericWindow> LocateWindowUnderCursor( const FVector2D& CursorPos );
 
 	NSScreen* FindScreenByPoint( int32 X, int32 Y ) const;
 
diff --git a/Engine/Source/Runtime/Core/Public/Mac/MacWindow.h b/Engine/Source/Runtime/Core/Public/Mac/MacWindow.h
--- a/Engine/Source/Runtime/Core/Public/Mac/MacWindow.h
+++ b/Engine/Source/Runtime/Core/Public/Mac/MacWindow.h
@@ -12,9 +12,7 @@
  */
 @interface FSlateCocoaWindow : NSWindow <NSWindowDelegate>
 {
-	NSWindow* Parent;
 	EWindowMode::Type WindowMode;
-	NSMutableArray* ChildWindows;
 	NSRect DeferFrame;
 	CGFloat DeferOpacity;
 	bool bAcceptsInput;
@@ -65,26 +63,8 @@
 /** Mutator that specifies that the display arrangement is being reconfigured when bIsDisplayReconfiguring is true. */
 - (void)setDisplayReconfiguring:(bool)bIsDisplayReconfiguring;
 
-/** Record the parent window separately from Cocoa so that parented windows can still be minimised without hiding the app. */
-- (void)setParent:(NSWindow*)InParent;
-
-/** Cached current child windows & disconnect them so the window can be moved independently. */
-- (void)disconnectChildWindows;
-
-/** Order window to the front, on top of sibling windows if it is a child. */
-- (void)orderFrontEvenIfChildAndMakeMain:(bool)bMain andKey:(bool)bKey;
-
-/** Reattach cached child windows. */
-- (void)reconnectChildWindows;
-
-/** Returns true if it is permissable to add a child window, false if it will start being dragged around with the parent. */
-- (bool)shouldAddChildWindows;
-
-/** Get the cached parent window, regardless of whether it is attached. */
-- (NSWindow*)getParent;
-
-/** Remove InChild from any cached list of children. */
-- (void)removeCachedChild:(NSWindow*)InChild;
+/** Order window to the front. */
+- (void)orderFrontAndMakeMain:(bool)bMain andKey:(bool)bKey;
 
 @end
 
@@ -121,14 +101,6 @@ class CORE_API FMacWindow : public FGenericWindow, public TSharedFromThis<FMacWi
 public:
 	~FMacWindow();
 
-	/** Create a new FMacWindow.
-	 *
-	 * @param SlateWindows		List of all top level Slate windows.  This function will add the owner window to this list.
-	 * @param OwnerWindow		The SlateWindow for which we are crating a backing Win32Window
-	 * @param InHInstance		Win32 application instance handle
-	 * @param InParent			Parent Win32 window; usually NULL.
-	 * @param bShowImmediately	True to show this window as soon as its initialized
-	 */
 	static TSharedRef< FMacWindow > Make();
 
 	FSlateCocoaWindow* GetWindowHandle() const;
