Stops loading FText's from stomping other duplicated FText's and causing them to display incorrect information.

#ttp 340038 - CRITICAL: BP: All pins turn to True Pose

diff --git a/Engine/Source/Runtime/Core/Private/Internationalization/TextHistory.cpp b/Engine/Source/Runtime/Core/Private/Internationalization/TextHistory.cpp
--- a/Engine/Source/Runtime/Core/Private/Internationalization/TextHistory.cpp
+++ b/Engine/Source/Runtime/Core/Private/Internationalization/TextHistory.cpp
@@ -18,6 +18,17 @@ TSharedPtr< FString, ESPMode::ThreadSafe > FTextHistory::GetSourceString() const
 	return NULL;
 }
 
+void FTextHistory::SerializeForDisplayString(FArchive& Ar, TSharedRef<FString, ESPMode::ThreadSafe>& InOutDisplayString)
+{
+	if(Ar.IsLoading())
+	{
+		//When duplicating, the CDO is used as the template, then values for the instance are assigned.
+		//If we don't duplicate the string, the CDO and the instance are both pointing at the same thing.
+		//This would result in all subsequently duplicated objects stamping over formerly duplicated ones.
+		InOutDisplayString = MakeShareable(new FString());
+	}
+}
+
 ///////////////////////////////////////
 // FTextHistory_NamedFormat
 
diff --git a/Engine/Source/Runtime/Core/Private/Internationalization/TextHistory.h b/Engine/Source/Runtime/Core/Private/Internationalization/TextHistory.h
--- a/Engine/Source/Runtime/Core/Private/Internationalization/TextHistory.h
+++ b/Engine/Source/Runtime/Core/Private/Internationalization/TextHistory.h
@@ -29,8 +29,8 @@ public:
 	/** Serializes the history to/from an FArchive */
 	virtual void Serialize(FArchive& Ar) = 0;
 
-	/** Serializes data needed to get the FText's DisplayString, most history classes do not need to do anything here */
-	virtual void SerializeForDisplayString(FArchive& Ar, TSharedRef<FString, ESPMode::ThreadSafe>& InOutDisplayString) {};
+	/** Serializes data needed to get the FText's DisplayString */
+	virtual void SerializeForDisplayString(FArchive& Ar, TSharedRef<FString, ESPMode::ThreadSafe>& InOutDisplayString);
 
 	/** Returns TRUE if the Revision is out of date */
 	virtual bool IsOutOfDate(int32 InRevision);
