Compiler pre-setup includes clean-up.

diff --git a/Engine/Source/Runtime/Core/Public/GenericPlatform/GenericPlatformCompilerPreSetup.h b/Engine/Source/Runtime/Core/Public/GenericPlatform/GenericPlatformCompilerPreSetup.h
new file mode 100644
--- /dev/null
+++ b/Engine/Source/Runtime/Core/Public/GenericPlatform/GenericPlatformCompilerPreSetup.h
@@ -0,0 +1,15 @@
+// Copyright 1998-2014 Epic Games, Inc. All Rights Reserved.
+
+/*=============================================================================================
+	GenericPlatformCompilerPreSetup.h: pragmas, version checks and other things for generic compiler
+==============================================================================================*/
+
+#pragma once
+
+#ifndef DEPRECATED
+	#define DEPRECATED(VERSION, MESSAGE)
+#endif // DEPRECATED
+
+#ifndef EMIT_DEPRECATED_WARNING_MESSAGE
+	#define EMIT_DEPRECATED_WARNING_MESSAGE(Msg)
+#endif // EMIT_DEPRECATED_WARNING_MESSAGE
\ No newline at end of file
diff --git a/Engine/Source/Runtime/Core/Public/GenericPlatform/GenericPlatformCompilerSetup.h b/Engine/Source/Runtime/Core/Public/GenericPlatform/GenericPlatformCompilerSetup.h
deleted file mode 100644
diff --git a/Engine/Source/Runtime/Core/Public/HAL/Platform.h b/Engine/Source/Runtime/Core/Public/HAL/Platform.h
--- a/Engine/Source/Runtime/Core/Public/HAL/Platform.h
+++ b/Engine/Source/Runtime/Core/Public/HAL/Platform.h
@@ -5,41 +5,6 @@
 ==================================================================================*/
 #pragma once
 
-#if PLATFORM_APPLE
-	#include <stddef.h> // needed for size_t
-#endif
-
-#include "GenericPlatform/GenericPlatform.h"
-
-//---------------------------------------------------------
-// Identify the current platform and include that header
-//---------------------------------------------------------
-
-//@port Identify the platform here and include the platform header to setup the platform types, etc
-#if PLATFORM_WINDOWS
-	#include "Windows/WIndowsPlatform.h"
-#elif PLATFORM_PS4
-	#include "PS4/PS4Platform.h"
-#elif PLATFORM_XBOXONE
-	#include "XboxOne/XboxOnePlatform.h"
-#elif PLATFORM_MAC
-	#include "Mac/MacPlatform.h"
-#elif PLATFORM_IOS
-	#include "IOS/IOSPlatform.h"
-#elif PLATFORM_ANDROID
-	#include "Android/AndroidPlatform.h"
-#elif PLATFORM_WINRT_ARM
-	#include "WinRT/WinRTARMPlatform.h"
-#elif PLATFORM_WINRT
-	#include "WinRT/WinRTPlatform.h"
-#elif PLATFORM_HTML5
-	#include "HTML5/HTML5Platform.h"
-#elif PLATFORM_LINUX
-	#include "Linux/LinuxPlatform.h"
-#else
-	#error Unknown Compiler
-#endif
-
 // define all other platforms to be zero
 //@port Define the platform here to be zero when compiling for other platforms
 #if !defined(PLATFORM_WINDOWS)
@@ -76,6 +41,51 @@
 	#define PLATFORM_LINUX 0
 #endif
 
+// Platform specific compiler pre-setup.
+#if PLATFORM_WINDOWS
+	#include "Windows/WindowsPlatformCompilerPreSetup.h"
+#elif PLATFORM_MAC
+	#include "Mac/MacPlatformCompilerPreSetup.h"
+#endif
+
+// Generic compiler pre-setup.
+#include "GenericPlatform/GenericPlatformCompilerPreSetup.h"
+
+#if PLATFORM_APPLE
+	#include <stddef.h> // needed for size_t
+#endif
+
+#include "GenericPlatform/GenericPlatform.h"
+
+//---------------------------------------------------------
+// Identify the current platform and include that header
+//---------------------------------------------------------
+
+//@port Identify the platform here and include the platform header to setup the platform types, etc
+#if PLATFORM_WINDOWS
+	#include "Windows/WIndowsPlatform.h"
+#elif PLATFORM_PS4
+	#include "PS4/PS4Platform.h"
+#elif PLATFORM_XBOXONE
+	#include "XboxOne/XboxOnePlatform.h"
+#elif PLATFORM_MAC
+	#include "Mac/MacPlatform.h"
+#elif PLATFORM_IOS
+	#include "IOS/IOSPlatform.h"
+#elif PLATFORM_ANDROID
+	#include "Android/AndroidPlatform.h"
+#elif PLATFORM_WINRT_ARM
+	#include "WinRT/WinRTARMPlatform.h"
+#elif PLATFORM_WINRT
+	#include "WinRT/WinRTPlatform.h"
+#elif PLATFORM_HTML5
+	#include "HTML5/HTML5Platform.h"
+#elif PLATFORM_LINUX
+	#include "Linux/LinuxPlatform.h"
+#else
+	#error Unknown Compiler
+#endif
+
 //------------------------------------------------------------------
 // Finalize define setup
 //------------------------------------------------------------------
@@ -292,10 +302,6 @@
 	#define TYPENAME_OUTSIDE_TEMPLATE	typename
 #endif
 
-#ifndef EMIT_DEPRECATED_WARNING_MESSAGE
-	#define EMIT_DEPRECATED_WARNING_MESSAGE(Msg)
-#endif
-
 // Legacy method modifier macros.  You shouldn't use these macros in modern code.  Use the built-in keyword directly.
 #ifndef OVERRIDE
 	#define OVERRIDE \
diff --git a/Engine/Source/Runtime/Core/Public/HAL/PlatformIncludes.h b/Engine/Source/Runtime/Core/Public/HAL/PlatformIncludes.h
--- a/Engine/Source/Runtime/Core/Public/HAL/PlatformIncludes.h
+++ b/Engine/Source/Runtime/Core/Public/HAL/PlatformIncludes.h
@@ -15,15 +15,7 @@
 #include <float.h>
 #include <string.h>
 
-// Platform specific compiler pre-setup, must be included before everything else
-#if PLATFORM_WINDOWS
-#include "Windows/WindowsPlatformCompilerPreSetup.h"
-#elif PLATFORM_MAC
-#include "Mac/MacPlatformCompilerPreSetup.h"
-#endif
-
 // Generic versions of the platform abstractions
-#include "GenericPlatform/GenericPlatformCompilerSetup.h"
 #include "GenericPlatform/GenericPlatformMemory.h"
 #include "GenericPlatform/GenericPlatformString.h"
 #include "GenericPlatform/GenericPlatformMisc.h"
diff --git a/Engine/Source/Runtime/Core/Public/Mac/MacPlatformCompilerPreSetup.h b/Engine/Source/Runtime/Core/Public/Mac/MacPlatformCompilerPreSetup.h
--- a/Engine/Source/Runtime/Core/Public/Mac/MacPlatformCompilerPreSetup.h
+++ b/Engine/Source/Runtime/Core/Public/Mac/MacPlatformCompilerPreSetup.h
@@ -7,7 +7,5 @@
 #pragma once
 
 #ifndef DISABLE_DEPRECATION
-#define DEPRECATED(VERSION, MESSAGE) __attribute__((deprecated(MESSAGE " Please update your code to the new API before upgrading to the next release, otherwise your project will no longer compile.")))
-#else
-	#define DEPRECATED(VERSION, MESSAGE)
+	#define DEPRECATED(VERSION, MESSAGE) __attribute__((deprecated(MESSAGE " Please update your code to the new API before upgrading to the next release, otherwise your project will no longer compile.")))
 #endif // DISABLE_DEPRECATION
diff --git a/Engine/Source/Runtime/Core/Public/Windows/WindowsPlatformCompilerPreSetup.h b/Engine/Source/Runtime/Core/Public/Windows/WindowsPlatformCompilerPreSetup.h
--- a/Engine/Source/Runtime/Core/Public/Windows/WindowsPlatformCompilerPreSetup.h
+++ b/Engine/Source/Runtime/Core/Public/Windows/WindowsPlatformCompilerPreSetup.h
@@ -15,7 +15,4 @@
 	#endif // EMIT_DEPRECATED_WARNING_MESSAGE_STR
 
 	#define EMIT_DEPRECATED_WARNING_MESSAGE(Msg) __pragma(message(__FILE__ "(" EMIT_DEPRECATED_WARNING_MESSAGE_STR(__LINE__) "): warning C4996: " Msg))
-#else
-	#define DEPRECATED(VERSION, MESSAGE)
-	#define EMIT_DEPRECATED_WARNING_MESSAGE(Msg)
 #endif // DISABLE_DEPRECATION
