Merging using UE4-To-UE4-Releases-Latest-Full

diff --git a/Engine/Source/Runtime/Core/Private/Apple/ApplePlatformFile.cpp b/Engine/Source/Runtime/Core/Private/Apple/ApplePlatformFile.cpp
--- a/Engine/Source/Runtime/Core/Private/Apple/ApplePlatformFile.cpp
+++ b/Engine/Source/Runtime/Core/Private/Apple/ApplePlatformFile.cpp
@@ -85,7 +85,7 @@ public:
 		if( IsManaged() )
 		{
 			FileOffset = NewPosition >= FileSize ? FileSize - 1 : NewPosition;
-			return ActiveHandles[ HandleSlot ] == this ? lseek(FileHandle, FileOffset, SEEK_SET) != -1 : true;
+			return IsValid() && ActiveHandles[ HandleSlot ] == this ? lseek(FileHandle, FileOffset, SEEK_SET) != -1 : true;
 		}
 		else
 #endif
@@ -102,7 +102,7 @@ public:
 		if( IsManaged() )
 		{
 			FileOffset = (NewPositionRelativeToEnd >= FileSize) ? 0 : ( FileSize + NewPositionRelativeToEnd - 1 );
-			return ActiveHandles[ HandleSlot ] == this ? lseek(FileHandle, FileOffset, SEEK_SET) != -1 : true;
+			return IsValid() && ActiveHandles[ HandleSlot ] == this ? lseek(FileHandle, FileOffset, SEEK_SET) != -1 : true;
 		}
 		else
 #endif
