Fixed an bug in FMacPlatformProcess::GetDllHandle() that allowed certain dylibs to be loaded twice due to dyld not recognizing dylibs loaded with RPATH as the same as dylibs loaded with full path.

diff --git a/Engine/Source/Runtime/Core/Private/Mac/MacPlatformProcess.cpp b/Engine/Source/Runtime/Core/Private/Mac/MacPlatformProcess.cpp
--- a/Engine/Source/Runtime/Core/Private/Mac/MacPlatformProcess.cpp
+++ b/Engine/Source/Runtime/Core/Private/Mac/MacPlatformProcess.cpp
@@ -13,21 +13,31 @@ void* FMacPlatformProcess::GetDllHandle( const TCHAR* Filename )
 {
 	SCOPED_AUTORELEASE_POOL;
 
-	check( Filename );
-	// First, try inside the bundle's Frameworks folder
-	CFStringRef CFStr = FPlatformString::TCHARToCFString(Filename);
-	NSString *FullPath = [[[[NSBundle mainBundle] executablePath] stringByDeletingLastPathComponent] stringByAppendingPathComponent:(NSString *)CFStr];
-	CFRelease(CFStr);
+	check(Filename);
 
-	void *Handle = dlopen( [FullPath fileSystemRepresentation], RTLD_LAZY | RTLD_LOCAL );
+	NSFileManager* FileManager = [NSFileManager defaultManager];
+	NSString* DylibPath = [NSString stringWithUTF8String:TCHAR_TO_UTF8(Filename)];
+	if (![FileManager fileExistsAtPath:DylibPath])
+	{
+		// If it's not a absolute or relative path, try to find the file in the app bundle
+		DylibPath = [[[[NSBundle mainBundle] executablePath] stringByDeletingLastPathComponent] stringByAppendingPathComponent:FString(Filename).GetNSString()];
+	}
+
+	// Check if dylib is already loaded
+	void* Handle = dlopen([DylibPath fileSystemRepresentation], RTLD_NOLOAD | RTLD_LAZY | RTLD_LOCAL);
+	if (!Handle)
+	{
+		// Maybe it was loaded using RPATH
+		Handle = dlopen([[@"@rpath" stringByAppendingPathComponent:[DylibPath lastPathComponent]] fileSystemRepresentation], RTLD_NOLOAD | RTLD_LAZY | RTLD_LOCAL);
+	}
 	if (!Handle)
 	{
-		// Maybe we're not a bundle. Try opening from current working dir.
-		Handle = dlopen( TCHAR_TO_ANSI(Filename), RTLD_LAZY | RTLD_LOCAL );
+		// Not loaded yet, so try to open it
+		Handle = dlopen([DylibPath fileSystemRepresentation], RTLD_LAZY | RTLD_LOCAL);
 	}
 	if (!Handle)
 	{
-		UE_LOG(LogMac, Warning, TEXT("dlopen failed: %s"), ANSI_TO_TCHAR(dlerror()) );
+		UE_LOG(LogMac, Warning, TEXT("dlopen failed: %s"), ANSI_TO_TCHAR(dlerror()));
 	}
 	return Handle;
 }
