Disable Windows accessibility features while game is running - more or less a port of the UE3 implementation

diff --git a/Engine/Source/Runtime/Core/Private/Windows/WindowsApplication.cpp b/Engine/Source/Runtime/Core/Private/Windows/WindowsApplication.cpp
--- a/Engine/Source/Runtime/Core/Private/Windows/WindowsApplication.cpp
+++ b/Engine/Source/Runtime/Core/Private/Windows/WindowsApplication.cpp
@@ -72,6 +72,116 @@ FWindowsApplication::FWindowsApplication( const HINSTANCE HInstance, const HICON
 
 	// Get initial display metrics. (display information for existing desktop, before we start changing resolutions)
 	GetDisplayMetrics(InitialDisplayMetrics);
+
+	// Save the current sticky/toggle/filter key settings so they can be restored them later
+	// If there are .ini settings, use them instead of the current system settings.
+	// NOTE: Whenever we exit and restore these settings gracefully, the .ini settings are removed.
+	FMemory::MemZero(StartupStickyKeys);
+	FMemory::MemZero(StartupToggleKeys);
+	FMemory::MemZero(StartupFilterKeys);
+	
+	StartupStickyKeys.cbSize = sizeof(StartupStickyKeys);
+	StartupToggleKeys.cbSize = sizeof(StartupToggleKeys);
+	StartupFilterKeys.cbSize = sizeof(StartupFilterKeys);
+
+	SystemParametersInfo(SPI_GETSTICKYKEYS, sizeof(STICKYKEYS), &StartupStickyKeys, 0);
+	SystemParametersInfo(SPI_GETTOGGLEKEYS, sizeof(TOGGLEKEYS), &StartupToggleKeys, 0);
+	SystemParametersInfo(SPI_GETFILTERKEYS, sizeof(FILTERKEYS), &StartupFilterKeys, 0);
+
+	bool bSKHotkey = (StartupStickyKeys.dwFlags & SKF_HOTKEYACTIVE) ? true : false;
+	bool bTKHotkey = (StartupToggleKeys.dwFlags & TKF_HOTKEYACTIVE) ? true : false;
+	bool bFKHotkey = (StartupFilterKeys.dwFlags & FKF_HOTKEYACTIVE) ? true : false;
+	bool bSKConfirmation = (StartupStickyKeys.dwFlags & SKF_CONFIRMHOTKEY) ? true : false;
+	bool bTKConfirmation = (StartupToggleKeys.dwFlags & TKF_CONFIRMHOTKEY) ? true : false;
+	bool bFKConfirmation = (StartupFilterKeys.dwFlags & FKF_CONFIRMHOTKEY) ? true : false;
+
+	GConfig->GetBool(TEXT("WindowsApplication.Accessibility"), TEXT("StickyKeysHotkey"), bSKHotkey, GEngineIni);
+	GConfig->GetBool(TEXT("WindowsApplication.Accessibility"), TEXT("ToggleKeysHotkey"), bTKHotkey, GEngineIni);
+	GConfig->GetBool(TEXT("WindowsApplication.Accessibility"), TEXT("FilterKeysHotkey"), bFKHotkey, GEngineIni);
+	GConfig->GetBool(TEXT("WindowsApplication.Accessibility"), TEXT("StickyKeysConfirmation"), bSKConfirmation, GEngineIni);
+	GConfig->GetBool(TEXT("WindowsApplication.Accessibility"), TEXT("ToggleKeysConfirmation"), bTKConfirmation, GEngineIni);
+	GConfig->GetBool(TEXT("WindowsApplication.Accessibility"), TEXT("FilterKeysConfirmation"), bFKConfirmation, GEngineIni);
+
+	StartupStickyKeys.dwFlags = bSKHotkey ? (StartupStickyKeys.dwFlags | SKF_HOTKEYACTIVE) : (StartupStickyKeys.dwFlags & ~SKF_HOTKEYACTIVE);
+	StartupToggleKeys.dwFlags = bTKHotkey ? (StartupToggleKeys.dwFlags | TKF_HOTKEYACTIVE) : (StartupToggleKeys.dwFlags & ~TKF_HOTKEYACTIVE);
+	StartupFilterKeys.dwFlags = bFKHotkey ? (StartupFilterKeys.dwFlags | FKF_HOTKEYACTIVE) : (StartupFilterKeys.dwFlags & ~FKF_HOTKEYACTIVE);
+	StartupStickyKeys.dwFlags = bSKConfirmation ? (StartupStickyKeys.dwFlags | SKF_CONFIRMHOTKEY) : (StartupStickyKeys.dwFlags & ~SKF_CONFIRMHOTKEY);
+	StartupToggleKeys.dwFlags = bTKConfirmation ? (StartupToggleKeys.dwFlags | TKF_CONFIRMHOTKEY) : (StartupToggleKeys.dwFlags & ~TKF_CONFIRMHOTKEY);
+	StartupFilterKeys.dwFlags = bFKConfirmation ? (StartupFilterKeys.dwFlags | FKF_CONFIRMHOTKEY) : (StartupFilterKeys.dwFlags & ~FKF_CONFIRMHOTKEY);
+
+	GConfig->SetBool(TEXT("WindowsApplication.Accessibility"), TEXT("StickyKeysHotkey"), bSKHotkey, GEngineIni);
+	GConfig->SetBool(TEXT("WindowsApplication.Accessibility"), TEXT("ToggleKeysHotkey"), bTKHotkey, GEngineIni);
+	GConfig->SetBool(TEXT("WindowsApplication.Accessibility"), TEXT("FilterKeysHotkey"), bFKHotkey, GEngineIni);
+	GConfig->SetBool(TEXT("WindowsApplication.Accessibility"), TEXT("StickyKeysConfirmation"), bSKConfirmation, GEngineIni);
+	GConfig->SetBool(TEXT("WindowsApplication.Accessibility"), TEXT("ToggleKeysConfirmation"), bTKConfirmation, GEngineIni);
+	GConfig->SetBool(TEXT("WindowsApplication.Accessibility"), TEXT("FilterKeysConfirmation"), bFKConfirmation, GEngineIni);
+
+	GConfig->Flush(false, GEngineIni);
+
+	FCoreDelegates::OnShutdownAfterError.AddRaw(this, &FWindowsApplication::ShutDownAfterError);
+
+	// Disable accessibility shortcuts
+	AllowAccessibilityShortcutKeys(false);
+}
+
+void FWindowsApplication::AllowAccessibilityShortcutKeys(const bool bAllowKeys)
+{
+	if (bAllowKeys)
+	{
+		// Restore StickyKeys/etc to original state and enable Windows key      
+		SystemParametersInfo(SPI_SETSTICKYKEYS, sizeof(STICKYKEYS), &StartupStickyKeys, 0);
+		SystemParametersInfo(SPI_SETTOGGLEKEYS, sizeof(TOGGLEKEYS), &StartupToggleKeys, 0);
+		SystemParametersInfo(SPI_SETFILTERKEYS, sizeof(FILTERKEYS), &StartupFilterKeys, 0);
+	}
+	else
+	{
+		// Disable StickyKeys/etc shortcuts but if the accessibility feature is on, 
+		// then leave the settings alone as its probably being usefully used
+
+		STICKYKEYS skOff = StartupStickyKeys;
+		if ((skOff.dwFlags & SKF_STICKYKEYSON) == 0)
+		{
+			// Disable the hotkey and the confirmation
+			skOff.dwFlags &= ~SKF_HOTKEYACTIVE;
+			skOff.dwFlags &= ~SKF_CONFIRMHOTKEY;
+
+			SystemParametersInfo(SPI_SETSTICKYKEYS, sizeof(STICKYKEYS), &skOff, 0);
+		}
+
+		TOGGLEKEYS tkOff = StartupToggleKeys;
+		if ((tkOff.dwFlags & TKF_TOGGLEKEYSON) == 0)
+		{
+			// Disable the hotkey and the confirmation
+			tkOff.dwFlags &= ~TKF_HOTKEYACTIVE;
+			tkOff.dwFlags &= ~TKF_CONFIRMHOTKEY;
+
+			SystemParametersInfo(SPI_SETTOGGLEKEYS, sizeof(TOGGLEKEYS), &tkOff, 0);
+		}
+
+		FILTERKEYS fkOff = StartupFilterKeys;
+		if ((fkOff.dwFlags & FKF_FILTERKEYSON) == 0)
+		{
+			// Disable the hotkey and the confirmation
+			fkOff.dwFlags &= ~FKF_HOTKEYACTIVE;
+			fkOff.dwFlags &= ~FKF_CONFIRMHOTKEY;
+
+			SystemParametersInfo(SPI_SETFILTERKEYS, sizeof(FILTERKEYS), &fkOff, 0);
+		}
+	}
+}
+
+void FWindowsApplication::DestroyApplication()
+{
+	// Restore accessibility shortcuts and remove the saved state from the .ini
+	AllowAccessibilityShortcutKeys(true);
+	GConfig->EmptySection(TEXT("WindowsApplication.Accessibility"), GEngineIni);
+}
+
+void FWindowsApplication::ShutDownAfterError()
+{
+	// Restore accessibility shortcuts and remove the saved state from the .ini
+	AllowAccessibilityShortcutKeys(true);
+	GConfig->EmptySection(TEXT("WindowsApplication.Accessibility"), GEngineIni);
 }
 
 bool FWindowsApplication::RegisterClass( const HINSTANCE HInstance, const HICON HIcon )
@@ -460,11 +570,6 @@ EWindowTitleAlignment::Type FWindowsApplication::GetWindowTitleAlignment() const
 	return EWindowTitleAlignment::Left;
 }
 
-void FWindowsApplication::DestroyApplication()
-{
-
-}
-
 static TSharedPtr< FWindowsWindow > FindWindowByHWND( const TArray< TSharedRef< FWindowsWindow > >& WindowsToSearch, HWND HandleToFind )
 {
 	for (int32 WindowIndex=0; WindowIndex < WindowsToSearch.Num(); ++WindowIndex)
diff --git a/Engine/Source/Runtime/Core/Public/Windows/WindowsApplication.h b/Engine/Source/Runtime/Core/Public/Windows/WindowsApplication.h
--- a/Engine/Source/Runtime/Core/Public/Windows/WindowsApplication.h
+++ b/Engine/Source/Runtime/Core/Public/Windows/WindowsApplication.h
@@ -251,6 +251,10 @@ private:
 
 	void CheckForShiftUpEvents(const int32 KeyCode);
 
+	void ShutDownAfterError();
+
+	void AllowAccessibilityShortcutKeys(const bool bAllowKeys);
+
 private:
 
 	static const FIntPoint MinimizedWindowPosition;
@@ -283,4 +287,9 @@ private:
 	FDisplayMetrics InitialDisplayMetrics;
 
 	TSharedPtr<FWindowsTextInputMethodSystem> TextInputMethodSystem;
+
+	// Accessibility shortcut keys
+	STICKYKEYS							StartupStickyKeys;
+	TOGGLEKEYS							StartupToggleKeys;
+	FILTERKEYS							StartupFilterKeys;
 };
\ No newline at end of file
