CrashReport - Prevent CrashReportClient from spawning another CrashReportClient (Windows only)

diff --git a/Engine/Source/Runtime/Core/Private/Windows/WindowsPlatformExceptionHandling.cpp b/Engine/Source/Runtime/Core/Private/Windows/WindowsPlatformExceptionHandling.cpp
--- a/Engine/Source/Runtime/Core/Private/Windows/WindowsPlatformExceptionHandling.cpp
+++ b/Engine/Source/Runtime/Core/Private/Windows/WindowsPlatformExceptionHandling.cpp
@@ -297,78 +297,84 @@ int32 ReportCrashUsingCrashReportClient(EXCEPTION_POINTERS* ExceptionInfo, const
 	// Flush out the log
 	GLog->Flush();
 
-	// Set the report to force queue
-	FScopedWERQueuing ScopedQueueForcer;
+	// Prevent CrashReportClient from spawning another CrashReportClient.
+	const TCHAR* ExecutableName = FPlatformProcess::ExecutableName();
+	const bool bCanRunCrashReportClient = FCString::Stristr( ExecutableName, TEXT( "CrashReportClient" ) ) == nullptr;
+	if( bCanRunCrashReportClient )
+	{
+		// Set the report to force queue
+		FScopedWERQueuing ScopedQueueForcer;
 
-	// Construct the report details
-	WER_REPORT_INFORMATION ReportInformation = { sizeof(WER_REPORT_INFORMATION) };
+		// Construct the report details
+		WER_REPORT_INFORMATION ReportInformation = {sizeof( WER_REPORT_INFORMATION )};
 
-	StringCchCopy( ReportInformation.wzConsentKey, ARRAYSIZE( ReportInformation.wzConsentKey ), TEXT( "" ) );
-	StringCchCopy( ReportInformation.wzApplicationName, ARRAYSIZE( ReportInformation.wzApplicationName ), FApp::GetGameName() );
-	StringCchCopy( ReportInformation.wzApplicationPath, ARRAYSIZE( ReportInformation.wzApplicationPath ), FPlatformProcess::BaseDir() );
-	StringCchCat( ReportInformation.wzApplicationPath, ARRAYSIZE( ReportInformation.wzApplicationPath ), FPlatformProcess::ExecutableName() );
-	StringCchCat( ReportInformation.wzApplicationPath, ARRAYSIZE( ReportInformation.wzApplicationPath ), TEXT( ".exe" ) );
+		StringCchCopy( ReportInformation.wzConsentKey, ARRAYSIZE( ReportInformation.wzConsentKey ), TEXT( "" ) );
+		StringCchCopy( ReportInformation.wzApplicationName, ARRAYSIZE( ReportInformation.wzApplicationName ), FApp::GetGameName() );
+		StringCchCopy( ReportInformation.wzApplicationPath, ARRAYSIZE( ReportInformation.wzApplicationPath ), FPlatformProcess::BaseDir() );
+		StringCchCat( ReportInformation.wzApplicationPath, ARRAYSIZE( ReportInformation.wzApplicationPath ), FPlatformProcess::ExecutableName() );
+		StringCchCat( ReportInformation.wzApplicationPath, ARRAYSIZE( ReportInformation.wzApplicationPath ), TEXT( ".exe" ) );
 
-	GetCrashDescription( ReportInformation, ErrorMessage );
-	
-	// Create a crash event report
-	HREPORT ReportHandle = NULL;
-	if (WerReportCreate(APPCRASH_EVENT, WerReportApplicationCrash, &ReportInformation, &ReportHandle) == S_OK)
-	{
-		// Set the standard set of a crash parameters
-		SetReportParameters(ReportHandle, ExceptionInfo);
+		GetCrashDescription( ReportInformation, ErrorMessage );
 
-		// Add a manually generated minidump
-		AddMiniDump(ReportHandle, ExceptionInfo);
+		// Create a crash event report
+		HREPORT ReportHandle = NULL;
+		if( WerReportCreate( APPCRASH_EVENT, WerReportApplicationCrash, &ReportInformation, &ReportHandle ) == S_OK )
+		{
+			// Set the standard set of a crash parameters
+			SetReportParameters( ReportHandle, ExceptionInfo );
 
-		// Add the log and video
-		AddMiscFiles(ReportHandle);
+			// Add a manually generated minidump
+			AddMiniDump( ReportHandle, ExceptionInfo );
 
-		// Submit
-		WER_SUBMIT_RESULT SubmitResult;
-		WerReportSubmit(ReportHandle, WerConsentAlwaysPrompt, WER_SUBMIT_QUEUE | WER_SUBMIT_BYPASS_DATA_THROTTLING, &SubmitResult);
+			// Add the log and video
+			AddMiscFiles( ReportHandle );
 
-		// Cleanup
-		WerReportCloseHandle(ReportHandle);
-	}
+			// Submit
+			WER_SUBMIT_RESULT SubmitResult;
+			WerReportSubmit( ReportHandle, WerConsentAlwaysPrompt, WER_SUBMIT_QUEUE | WER_SUBMIT_BYPASS_DATA_THROTTLING, &SubmitResult );
 
-	// Build machines do not upload these automatically since it is not okay to have lingering processes after the build completes.
-	if (GIsBuildMachine)
-		return EXCEPTION_CONTINUE_EXECUTION;
+			// Cleanup
+			WerReportCloseHandle( ReportHandle );
+		}
 
-	FString CrashReportClientArguments;
+		// Build machines do not upload these automatically since it is not okay to have lingering processes after the build completes.
+		if( GIsBuildMachine )
+			return EXCEPTION_CONTINUE_EXECUTION;
 
-	// Suppress the user input dialog if we're running in unattended mode
-	bool bNoDialog = FApp::IsUnattended() || ReportUI == EErrorReportUI::ReportInUnattendedMode;
-	if (bNoDialog)
-	{
-		CrashReportClientArguments += TEXT(" -Unattended");
-	}
+		FString CrashReportClientArguments;
 
-	if (FApp::IsInstalled())
-	{
-		// Temporary workaround for CrashReportClient being built in Development, not Shipping (TTP328030). The
-		// following ensures that logs are saved to the user directory when UE4 is installed.
-		CrashReportClientArguments += TEXT(" -Installed");
-	}
+		// Suppress the user input dialog if we're running in unattended mode
+		bool bNoDialog = FApp::IsUnattended() || ReportUI == EErrorReportUI::ReportInUnattendedMode;
+		if( bNoDialog )
+		{
+			CrashReportClientArguments += TEXT( " -Unattended" );
+		}
 
-	CrashReportClientArguments += FString(TEXT(" -AppName=")) + ReportInformation.wzApplicationName;
+		if( FApp::IsInstalled() )
+		{
+			// Temporary workaround for CrashReportClient being built in Development, not Shipping (TTP328030). The
+			// following ensures that logs are saved to the user directory when UE4 is installed.
+			CrashReportClientArguments += TEXT( " -Installed" );
+		}
 
-	static const TCHAR CrashReportClientExeName[] = TEXT("CrashReportClient.exe");
-	FString CrashClientPath = FString(TEXT("..\\..\\..\\Engine\\Binaries")) / FPlatformProcess::GetBinariesSubdirectory() / CrashReportClientExeName;
+		CrashReportClientArguments += FString( TEXT( " -AppName=" ) ) + ReportInformation.wzApplicationName;
 
-	bool bCrashReporterRan = FPlatformProcess::CreateProc(*CrashClientPath, *CrashReportClientArguments, true, false, false, NULL, 0, NULL, NULL).IsValid();
+		static const TCHAR CrashReportClientExeName[] = TEXT( "CrashReportClient.exe" );
+		FString CrashClientPath = FString( TEXT( "..\\..\\..\\Engine\\Binaries" ) ) / FPlatformProcess::GetBinariesSubdirectory() / CrashReportClientExeName;
 
-	if (!bCrashReporterRan && !bNoDialog)
-	{
-		UE_LOG(LogWindows, Log, TEXT("Could not start %s"), CrashReportClientExeName);
-		FPlatformMemory::DumpStats(*GWarn);
-		FText MessageTitle(FText::Format(
-			NSLOCTEXT("MessageDialog", "AppHasCrashed", "The {0} {1} has crashed and will close"),
-			FText::FromString(ReportInformation.wzApplicationName),
-			FText::FromString(GetEngineMode())
-		));
-		FMessageDialog::Open(EAppMsgType::Ok, FText::FromString(GErrorHist), &MessageTitle);
+		bool bCrashReporterRan = FPlatformProcess::CreateProc( *CrashClientPath, *CrashReportClientArguments, true, false, false, NULL, 0, NULL, NULL ).IsValid();
+
+		if( !bCrashReporterRan && !bNoDialog )
+		{
+			UE_LOG( LogWindows, Log, TEXT( "Could not start %s" ), CrashReportClientExeName );
+			FPlatformMemory::DumpStats( *GWarn );
+			FText MessageTitle( FText::Format(
+				NSLOCTEXT( "MessageDialog", "AppHasCrashed", "The {0} {1} has crashed and will close" ),
+				FText::FromString( ReportInformation.wzApplicationName ),
+				FText::FromString( GetEngineMode() )
+				) );
+			FMessageDialog::Open( EAppMsgType::Ok, FText::FromString( GErrorHist ), &MessageTitle );
+		}
 	}
 
 	// Let the system take back over (return value only used by NewReportEnsure)
