CrashReport - Prevent CrashReportClient from spawning another CrashReportClient (Mac)

diff --git a/Engine/Source/Runtime/Core/Private/Mac/MacPlatformMisc.cpp b/Engine/Source/Runtime/Core/Private/Mac/MacPlatformMisc.cpp
--- a/Engine/Source/Runtime/Core/Private/Mac/MacPlatformMisc.cpp
+++ b/Engine/Source/Runtime/Core/Private/Mac/MacPlatformMisc.cpp
@@ -1487,7 +1487,12 @@ void FMacCrashContext::GenerateCrashInfoAndLaunchReporter() const
 	FCStringAnsi::Strcat(CrashInfoFolder, PATH_MAX, GMacAppInfo.AppNameUTF8);
 	FCStringAnsi::Strcat(CrashInfoFolder, PATH_MAX, "-pid-");
 	FCStringAnsi::Strcat(CrashInfoFolder, PATH_MAX, ItoANSI(getpid(), 10));
-	if(!mkdir(CrashInfoFolder, 0766))
+
+	// Prevent CrashReportClient from spawning another CrashReportClient.
+	const TCHAR* ExecutableName = FPlatformProcess::ExecutableName();
+	const bool bCanRunCrashReportClient = FCString::Stristr( ExecutableName, TEXT( "CrashReportClient" ) ) == nullptr;
+
+	if(!mkdir(CrashInfoFolder, 0766) && bCanRunCrashReportClient)
 	{
 		char FilePath[PATH_MAX] = {};
 		FCStringAnsi::Strncpy(FilePath, CrashInfoFolder, PATH_MAX);
