Variadic template support enabled in VS2013.

#change Platform macro changed to enable variadics for 2013.
#fix TUniqueObj 'multiple default constructors' compilation error fixed.
#add TUniqueObj supports the same number of forwarded arguments as TUniquePtr for the non-variadic case.

diff --git a/Engine/Source/Runtime/Core/Public/Templates/UniqueObj.h b/Engine/Source/Runtime/Core/Public/Templates/UniqueObj.h
--- a/Engine/Source/Runtime/Core/Public/Templates/UniqueObj.h
+++ b/Engine/Source/Runtime/Core/Public/Templates/UniqueObj.h
@@ -11,11 +11,6 @@ template <typename T>
 class TUniqueObj
 {
 public:
-	TUniqueObj()
-		: Obj(MakeUnique<T>())
-	{
-	}
-
 	TUniqueObj(const TUniqueObj& other)
 		: Obj(MakeUnique<T>(*other))
 	{
@@ -28,18 +23,22 @@ public:
 	{
 	}
 
-	template <typename Arg>
-	explicit TUniqueObj(Arg&& arg)
-		: Obj(MakeUnique<T>(Forward<Arg>(arg)))
-	{
-	}
-
 #if PLATFORM_COMPILER_HAS_VARIADIC_TEMPLATES
+
 	template <typename... Args>
 	explicit TUniqueObj(Args&&... args)
 		: Obj(MakeUnique<T>(Forward<Args>(args)...))
 	{
 	}
+
+#else
+
+	                                                                      explicit TUniqueObj(                                                  ) : Obj(MakeUnique<T>(                                                                                  )) {}
+	template <typename Arg0                                             > explicit TUniqueObj(Arg0&& arg0                                       ) : Obj(MakeUnique<T>(Forward<Arg0>(arg0)                                                               )) {}
+	template <typename Arg0, typename Arg1                              > explicit TUniqueObj(Arg0&& arg0, Arg1&& arg1                          ) : Obj(MakeUnique<T>(Forward<Arg0>(arg0), Forward<Arg1>(arg1)                                          )) {}
+	template <typename Arg0, typename Arg1, typename Arg2               > explicit TUniqueObj(Arg0&& arg0, Arg1&& arg1, Arg2&& arg2             ) : Obj(MakeUnique<T>(Forward<Arg0>(arg0), Forward<Arg1>(arg1), Forward<Arg2>(arg2)                     )) {}
+	template <typename Arg0, typename Arg1, typename Arg2, typename Arg3> explicit TUniqueObj(Arg0&& arg0, Arg1&& arg1, Arg2&& arg2, Arg3&& arg3) : Obj(MakeUnique<T>(Forward<Arg0>(arg0), Forward<Arg1>(arg1), Forward<Arg2>(arg2), Forward<Arg3>(arg3))) {}
+
 #endif
 
 	// Disallow copy-assignment
diff --git a/Engine/Source/Runtime/Core/Public/Windows/WIndowsPlatform.h b/Engine/Source/Runtime/Core/Public/Windows/WIndowsPlatform.h
--- a/Engine/Source/Runtime/Core/Public/Windows/WIndowsPlatform.h
+++ b/Engine/Source/Runtime/Core/Public/Windows/WIndowsPlatform.h
@@ -67,7 +67,11 @@ typedef FWindowsPlatformTypes FPlatformTypes;
 #define PLATFORM_SUPPORTS_TBB								1
 #define PLATFORM_SUPPORTS_NAMED_PIPES						1
 #define PLATFORM_COMPILER_HAS_DEFAULTED_FUNCTIONS			0
-#define PLATFORM_COMPILER_HAS_VARIADIC_TEMPLATES			0
+#if _MSC_VER >= 1800
+	#define PLATFORM_COMPILER_HAS_VARIADIC_TEMPLATES			1
+#else
+	#define PLATFORM_COMPILER_HAS_VARIADIC_TEMPLATES			0
+#endif
 #define PLATFORM_COMPILER_HAS_EXPLICIT_OPERATORS			0
 #define PLATFORM_COMPILER_HAS_TCHAR_WMAIN					1
 
