Scoped timers moved to Core.

#change FDurationTimer and FScopedDurationTimer moved from BP commandlet to ProfilingDebugging.
#add New FScopedDurationTimeLogger added for auto-logging the lifetime of the object.

diff --git a/Engine/Source/Editor/UnrealEd/Private/Commandlets/DumpBlueprintsInfoCommandlet.cpp b/Engine/Source/Editor/UnrealEd/Private/Commandlets/DumpBlueprintsInfoCommandlet.cpp
--- a/Engine/Source/Editor/UnrealEd/Private/Commandlets/DumpBlueprintsInfoCommandlet.cpp
+++ b/Engine/Source/Editor/UnrealEd/Private/Commandlets/DumpBlueprintsInfoCommandlet.cpp
@@ -15,6 +15,7 @@
 #include "BehaviorTree/BehaviorTree.h"
 #include "AssetSelection.h"
 #include "K2Node_CustomEvent.h"
+#include "ScopedTimers.h"
 
 DEFINE_LOG_CATEGORY_STATIC(LogBlueprintInfoDump, Log, All);
 
@@ -150,60 +151,6 @@ DumpBlueprintsInfo commandlet params: \n\
 	/** Tracks instantiated blueprints (so we don't have to create more the we have to). */
 	static TMap<UClass*, UBlueprint*> ClassBlueprints;
 
-	/**
-	 * Utility stopwatch class for tracking the duration of some action (tracks 
-	 * time in seconds and adds it to the specified variable on destruction).
-	 */
-	class FDurationTimer
-	{
-	public:
-		FDurationTimer(double& AccumulatorIn)
-			: StartTime(FPlatformTime::Seconds())
-			, Accumulator(AccumulatorIn)
-		{}
-		
-		double Start()
-		{
-			StartTime = FPlatformTime::Seconds();
-			return StartTime;
-		}
-		
-		double Stop()
-		{
-			double StopTime = FPlatformTime::Seconds();
-			Accumulator += (StopTime - StartTime);
-			StartTime = StopTime;
-			
-			return StopTime;
-		}
-
-	private:
-		/** Start time, captured in ctor. */
-		double StartTime;
-		/** Time variable to update. */
-		double& Accumulator;
-	};
-	
-	/**
-	 * Utility class for tracking the duration of a scoped action (the user 
-	 * doesn't have to call Start() and Stop() manually).
-	 */
-	class FScopedDurationTimer : private FDurationTimer
-	{
-	public:
-		FScopedDurationTimer(double& AccumulatorIn)
-			: FDurationTimer(AccumulatorIn)
-		{
-			Start();
-		}
-		
-		/** Dtor, updating seconds with time delta. */
-		~FScopedDurationTimer()
-		{
-			Stop();
-		}
-	};
-
 	/**
 	 * Certain blueprints (like level blueprints) require a level outer, and 
 	 * for certain actions we need a level actor selected. This utility function
diff --git a/Engine/Source/Runtime/Core/Public/ProfilingDebugging/ScopedTimers.h b/Engine/Source/Runtime/Core/Public/ProfilingDebugging/ScopedTimers.h
new file mode 100644
--- /dev/null
+++ b/Engine/Source/Runtime/Core/Public/ProfilingDebugging/ScopedTimers.h
@@ -0,0 +1,84 @@
+// Copyright 1998-2014 Epic Games, Inc. All Rights Reserved.
+
+/**
+ * Utility stopwatch class for tracking the duration of some action (tracks 
+ * time in seconds and adds it to the specified variable on destruction).
+ */
+class FDurationTimer
+{
+public:
+	explicit FDurationTimer(double& AccumulatorIn)
+		: StartTime(FPlatformTime::Seconds())
+		, Accumulator(AccumulatorIn)
+	{}
+
+	double Start()
+	{
+		StartTime = FPlatformTime::Seconds();
+		return StartTime;
+	}
+
+	double Stop()
+	{
+		double StopTime = FPlatformTime::Seconds();
+		Accumulator += (StopTime - StartTime);
+		StartTime = StopTime;
+			
+		return StopTime;
+	}
+
+private:
+	/** Start time, captured in ctor. */
+	double StartTime;
+	/** Time variable to update. */
+	double& Accumulator;
+};
+
+/**
+ * Utility class for tracking the duration of a scoped action (the user 
+ * doesn't have to call Start() and Stop() manually).
+ */
+class FScopedDurationTimer : private FDurationTimer
+{
+public:
+	explicit FScopedDurationTimer(double& AccumulatorIn)
+		: FDurationTimer(AccumulatorIn)
+	{
+		Start();
+	}
+
+	/** Dtor, updating seconds with time delta. */
+	~FScopedDurationTimer()
+	{
+		Stop();
+	}
+};
+
+/**
+ * Utility class for logging the duration of a scoped action (the user 
+ * doesn't have to call Start() and Stop() manually).
+ */
+class FScopedDurationTimeLogger
+{
+public:
+	explicit FScopedDurationTimeLogger(FString InMsg = TEXT("Scoped action"), FOutputDevice* InDevice = GLog)
+		: Msg        (MoveTemp(InMsg))
+		, Device     (InDevice)
+		, Accumulator(0.0)
+		, Timer      (Accumulator)
+	{
+		Timer.Start();
+	}
+
+	~FScopedDurationTimeLogger()
+	{
+		Timer.Stop();
+		Device->Logf(TEXT("%s: %f secs"), *Msg, Accumulator);
+	}
+
+private:
+	FString        Msg;
+	FOutputDevice* Device;
+	double         Accumulator;
+	FDurationTimer Timer;
+};
