Fix using menus and tooltips when a modal dialog is open on OS X.

diff --git a/Engine/Source/Runtime/Core/Private/Mac/MacWindow.cpp b/Engine/Source/Runtime/Core/Private/Mac/MacWindow.cpp
--- a/Engine/Source/Runtime/Core/Private/Mac/MacWindow.cpp
+++ b/Engine/Source/Runtime/Core/Private/Mac/MacWindow.cpp
@@ -118,7 +118,7 @@ TArray< FSlateCocoaWindow* > GRunningModalWindows;
 
 - (void)orderWindow:(NSWindowOrderingMode)OrderingMode relativeTo:(NSInteger)OtherWindowNumber
 {
-	bool bModal = GRunningModalWindows.Num() == 0 || GRunningModalWindows.Last() == self;
+	bool bModal = GRunningModalWindows.Num() == 0 || GRunningModalWindows.Last() == self || [self styleMask] == NSBorderlessWindowMask;
 	if(OrderingMode == NSWindowOut || bModal)
 	{
 		if([self alphaValue] > 0.0f)
@@ -172,7 +172,7 @@ TArray< FSlateCocoaWindow* > GRunningModalWindows;
 	if ([NSApp isHidden] == NO)
 	{
 
-		bool bModal = GRunningModalWindows.Num() == 0 || GRunningModalWindows.Last() == self;
+		bool bModal = GRunningModalWindows.Num() == 0 || GRunningModalWindows.Last() == self || [self styleMask] == NSBorderlessWindowMask;
 		if (![self isVisible] && bModal)
 		{
 			[self orderFront:nil];
@@ -315,7 +315,7 @@ TArray< FSlateCocoaWindow* > GRunningModalWindows;
 		{
 			[self setLevel:NSModalPanelWindowLevel];
 		}
-		if(GRunningModalWindows.Num() == 0 || GRunningModalWindows.Last() == self)
+		if(GRunningModalWindows.Num() == 0 || GRunningModalWindows.Last() == self || [self styleMask] == NSBorderlessWindowMask)
 		{
 			[self orderFrontAndMakeMain:false andKey:false];
 		}
