Notification windows on OS X aren't modal windows so don't treat them as such otherwise some of their widgets won't always work.

diff --git a/Engine/Source/Runtime/Core/Private/Mac/MacWindow.cpp b/Engine/Source/Runtime/Core/Private/Mac/MacWindow.cpp
--- a/Engine/Source/Runtime/Core/Private/Mac/MacWindow.cpp
+++ b/Engine/Source/Runtime/Core/Private/Mac/MacWindow.cpp
@@ -854,10 +854,7 @@ void FMacWindow::Destroy()
 		// This makes it possible to correctly handle any NSEvent's sent to the window during destruction
 		// as the WindowHandle is still a valid NSWindow. Unlike HWNDs accessing a destructed NSWindow* is fatal.
 		FSlateCocoaWindow* Window = [WindowHandle retain];
-		if(!Definition->SupportsMinimize && !Definition->SupportsMaximize)
-		{
-			GRunningModalWindows.Remove(WindowHandle);
-		}
+		GRunningModalWindows.Remove(WindowHandle);
 		Window.bForwardEvents = false;
 		if( MacApplication->OnWindowDestroyed( Window ) )
 		{
@@ -920,7 +917,7 @@ void FMacWindow::Show()
 	SCOPED_AUTORELEASE_POOL;
 	if (!bIsVisible)
 	{
-		if(!Definition->SupportsMinimize && !Definition->SupportsMaximize && !GRunningModalWindows.Contains(WindowHandle))
+		if(!Definition->SupportsMinimize && !Definition->SupportsMaximize && Definition->IsRegularWindow && !GRunningModalWindows.Contains(WindowHandle))
 		{
 			GRunningModalWindows.Add(WindowHandle);
 		}
