Allow the DebugGame configuration setting to be compiled into monolithic executables. The -debug command line parameter is still used for editor builds.

diff --git a/Engine/Source/Programs/UnrealBuildTool/Configuration/UEBuildModule.cs b/Engine/Source/Programs/UnrealBuildTool/Configuration/UEBuildModule.cs
--- a/Engine/Source/Programs/UnrealBuildTool/Configuration/UEBuildModule.cs
+++ b/Engine/Source/Programs/UnrealBuildTool/Configuration/UEBuildModule.cs
@@ -1596,10 +1596,11 @@ namespace UnrealBuildTool
 			Result.Config.bEnableExceptions                      = bEnableExceptions;
 			Result.Config.OutputDirectory                        = Path.Combine(Binary.Config.IntermediateDirectory, Name);
 
-			// Switch the optimization flag if we're building a game module
+			// Switch the optimization flag if we're building a game module. Also pass the definition for building in DebugGame along (see ModuleManager.h for notes).
 			if (Target.Configuration == UnrealTargetConfiguration.DebugGame && Type == UEBuildModuleType.Game)
 			{
 				Result.Config.Target.Configuration = CPPTargetConfiguration.Debug;
+				Result.Config.Definitions.Add("UE_BUILD_DEVELOPMENT_WITH_DEBUGGAME=1");
 			}
 
 			// Add the module's private definitions.
diff --git a/Engine/Source/Runtime/Core/Private/Misc/App.cpp b/Engine/Source/Runtime/Core/Private/Misc/App.cpp
--- a/Engine/Source/Runtime/Core/Private/Misc/App.cpp
+++ b/Engine/Source/Runtime/Core/Private/Misc/App.cpp
@@ -36,8 +36,14 @@ EBuildConfigurations::Type FApp::GetBuildConfiguration()
 	return EBuildConfigurations::Debug;
 
 #elif UE_BUILD_DEVELOPMENT
-	static bool bUsingDebugGame = FParse::Param(FCommandLine::Get(), TEXT("debug"));
-	return bUsingDebugGame ? EBuildConfigurations::DebugGame : EBuildConfigurations::Development;
+	// Detect DebugGame using an extern variable in monolithic configurations, or a command line argument in modular configurations.
+	#if IS_MONOLITHIC
+		extern const bool GIsDebugGame;
+		return GIsDebugGame? EBuildConfigurations::DebugGame : EBuildConfigurations::Development;
+	#else
+		static const bool bUsingDebugGame = FParse::Param(FCommandLine::Get(), TEXT("debug"));
+		return bUsingDebugGame? EBuildConfigurations::DebugGame : EBuildConfigurations::Development;
+	#endif
 
 #elif UE_BUILD_SHIPPING || UI_BUILD_SHIPPING_EDITOR
 	return EBuildConfigurations::Shipping;
diff --git a/Engine/Source/Runtime/Core/Private/Modules/ModuleManager.cpp b/Engine/Source/Runtime/Core/Private/Modules/ModuleManager.cpp
--- a/Engine/Source/Runtime/Core/Private/Modules/ModuleManager.cpp
+++ b/Engine/Source/Runtime/Core/Private/Modules/ModuleManager.cpp
@@ -1138,18 +1138,28 @@ void FModuleManager::FindModulePaths(const TCHAR* NamePattern, TMap<FName, FStri
 void FModuleManager::FindModulePathsInDirectory(const FString& InDirectoryName, bool bIsGameDirectory, const TCHAR* NamePattern, TMap<FName, FString> &OutModulePaths) const
 {
 	// Get the module configuration for this directory type
-	#if UE_BUILD_DEBUG
-		const TCHAR *ConfigSuffix = TEXT("-Debug");
-	#elif UE_BUILD_DEVELOPMENT
-		static bool bUsingDebugGame = FParse::Param(FCommandLine::Get(), TEXT("debug"));
-		const TCHAR *ConfigSuffix = (bIsGameDirectory && bUsingDebugGame) ? TEXT("-DebugGame") : NULL;
-	#elif UE_BUILD_TEST
-		const TCHAR *ConfigSuffix = TEXT("-Test");
-	#elif UE_BUILD_SHIPPING
-		const TCHAR *ConfigSuffix = TEXT("-Shipping");
-	#else
-		#error "Unknown configuration type"
-	#endif
+	const TCHAR* ConfigSuffix = NULL;
+	switch(FApp::GetBuildConfiguration())
+	{
+	case EBuildConfigurations::Debug:
+		ConfigSuffix = TEXT("-Debug");
+		break;
+	case EBuildConfigurations::DebugGame:
+		ConfigSuffix = bIsGameDirectory? TEXT("-DebugGame") : NULL;
+		break;
+	case EBuildConfigurations::Development:
+		ConfigSuffix = NULL;
+		break;
+	case EBuildConfigurations::Test:
+		ConfigSuffix = TEXT("-Test");
+		break;
+	case EBuildConfigurations::Shipping:
+		ConfigSuffix = TEXT("-Shipping");
+		break;
+	default:
+		check(false);
+		break;
+	}
 
 	// Get the base name for modules of this application
 	FString ModulePrefix = FPaths::GetBaseFilename(FPlatformProcess::ExecutableName());
@@ -1462,14 +1472,7 @@ void FModuleManager::MakeUniqueModuleFilename( const FName InModuleName, FString
 
 const TCHAR *FModuleManager::GetUBTConfiguration()
 {
-#if UE_BUILD_DEBUG
-	return TEXT("Debug");
-#elif UE_BUILD_SHIPPING
-	return TEXT("Shipping");
-#else
-	static bool bIsDebugGame = FParse::Param(FCommandLine::Get(), TEXT("debug"));
-	return bIsDebugGame? TEXT("DebugGame") : TEXT("Development");
-#endif
+	return EBuildConfigurations::ToString(FApp::GetBuildConfiguration());
 }
 
 
diff --git a/Engine/Source/Runtime/Core/Public/Modules/ModuleManager.h b/Engine/Source/Runtime/Core/Public/Modules/ModuleManager.h
--- a/Engine/Source/Runtime/Core/Public/Modules/ModuleManager.h
+++ b/Engine/Source/Runtime/Core/Public/Modules/ModuleManager.h
@@ -945,6 +945,20 @@ class FDefaultGameModuleImpl
 	#define IMPLEMENT_FOREIGN_ENGINE_DIR() 
 #endif
 
+/**
+ * Macro for declaring the GIsDebugGame variable for monolithic development builds. NB: This define, and the UE_BUILD_DEVELOPMENT_WITH_DEBUGGAME defines like it, should NEVER be 
+ * directly used or defined for engine code, because it prevents sharing the same build products with the development build (important for Rocket build sizes). In modular builds, 
+ * DebugGame modules will be loaded by specifying the -debug parameter on the command-line.
+ */
+#if IS_MONOLITHIC && UE_BUILD_DEVELOPMENT
+	#if UE_BUILD_DEVELOPMENT_WITH_DEBUGGAME
+		#define IMPLEMENT_DEBUGGAME() extern const bool GIsDebugGame = true;
+	#else
+		#define IMPLEMENT_DEBUGGAME() extern const bool GIsDebugGame = false;
+	#endif
+#else
+	#define IMPLEMENT_DEBUGGAME()
+#endif 
 
 #if IS_PROGRAM
 
@@ -952,6 +966,7 @@ class FDefaultGameModuleImpl
 		#define IMPLEMENT_APPLICATION( ModuleName, GameName ) \
 			/* For monolithic builds, we must statically define the game's name string (See Core.h) */ \
 			TCHAR GGameName[64] = TEXT( GameName ); \
+			IMPLEMENT_DEBUGGAME() \
 			IMPLEMENT_FOREIGN_ENGINE_DIR() \
 			IMPLEMENT_GAME_MODULE(FDefaultGameModuleImpl, ModuleName) \
 			PER_MODULE_BOILERPLATE \
@@ -985,6 +1000,7 @@ class FDefaultGameModuleImpl
 			TCHAR GGameName[64] = TEXT( GameName ); \
 			/* Implement the GIsGameAgnosticExe variable (See Core.h). */ \
 			bool GIsGameAgnosticExe = false; \
+			IMPLEMENT_DEBUGGAME() \
 			IMPLEMENT_FOREIGN_ENGINE_DIR() \
 			IMPLEMENT_GAME_MODULE( ModuleImplClass, ModuleName ) \
 			PER_MODULE_BOILERPLATE \
@@ -999,6 +1015,7 @@ class FDefaultGameModuleImpl
 		#define IMPLEMENT_PRIMARY_GAME_MODULE( ModuleImplClass, ModuleName, GameName ) \
 			/* For monolithic builds, we must statically define the game's name string (See Core.h) */ \
 			TCHAR GGameName[64] = TEXT( GameName ); \
+			IMPLEMENT_DEBUGGAME() \
 			PER_MODULE_BOILERPLATE \
 			IMPLEMENT_FOREIGN_ENGINE_DIR() \
 			IMPLEMENT_GAME_MODULE( ModuleImplClass, ModuleName ) \
diff --git a/Engine/Source/Runtime/UE4Game/Private/UE4Game.cpp b/Engine/Source/Runtime/UE4Game/Private/UE4Game.cpp
--- a/Engine/Source/Runtime/UE4Game/Private/UE4Game.cpp
+++ b/Engine/Source/Runtime/UE4Game/Private/UE4Game.cpp
@@ -9,5 +9,6 @@ IMPLEMENT_MODULE(FDefaultModuleImpl, UE4Game);
 PER_MODULE_BOILERPLATE
 bool GIsGameAgnosticExe = true;
 TCHAR GGameName[64] = TEXT("");
+IMPLEMENT_DEBUGGAME()
 IMPLEMENT_FOREIGN_ENGINE_DIR()
 #endif
