Editor Preferences - Fix an issue with user default array serialization which caused duplicated values in the UI.

TTP# 340060

There is still an issue with array serialization where values are repeated, e.g. Level Editor->Viewports->Decimal Grid Intervals.  User defaults are not saved with any index information so when when they are deserialized the last instances of the changed values are swapped, instead of the ones that the user actually changed.  Fixing that seemed like it would take a larger refactor to the user default serialization and seemed outside of the scope of this bug fix.

diff --git a/Engine/Source/Runtime/Core/Private/Misc/ConfigCacheIni.cpp b/Engine/Source/Runtime/Core/Private/Misc/ConfigCacheIni.cpp
--- a/Engine/Source/Runtime/Core/Private/Misc/ConfigCacheIni.cpp
+++ b/Engine/Source/Runtime/Core/Private/Misc/ConfigCacheIni.cpp
@@ -331,7 +331,7 @@ void FConfigFile::CombineFromBuffer(const FString& Filename,const FString& Buffe
 				else if( Cmd=='-' )	
 				{
 					// Remove if present.
-					CurrentSection->Remove( Start, *ProcessedValue );
+					CurrentSection->RemoveSingle( Start, *ProcessedValue );
 					CurrentSection->Compact();
 				}
 				else if ( Cmd=='.' )
@@ -1069,9 +1069,19 @@ void FConfigFile::ProcessPropertyAndWriteForDefaults( const TArray< FString >& I
 		FConfigSection* FoundSection = HierarchyFile.Find( SectionName );
 		if( FoundSection != NULL )
 		{
-			// As we are dealing with default configs, we must process array elements correct syntax syntax. I.e. + or -
+			// The non-default config hierarchy array property names are prefixed with a . and not a +, so change
+			// that before searching.
+			FString HierarchyPropertyName;
+			if (PropertyName.StartsWith(TEXT("+")))
+			{
+				HierarchyPropertyName = TEXT(".") + PropertyName.RightChop(1);
+			}
+			else
+			{
+				HierarchyPropertyName = PropertyName;
+			}
 			TArray< FString > HierearchysArrayContribution;
-			FoundSection->MultiFind( *PropertyName, HierearchysArrayContribution, true );
+			FoundSection->MultiFind( *HierarchyPropertyName, HierearchysArrayContribution, true );
 
 			// Find array elements which should be removed.
 			for( TArray<FString>::TIterator PropertyIt(HierearchysArrayContribution); PropertyIt; ++PropertyIt )
@@ -1101,7 +1111,7 @@ void FConfigFile::ProcessPropertyAndWriteForDefaults( const TArray< FString >& I
 				}
 
 				// We need to remove this element from unprocessed if it exists on the list.
-				UnprocessedPropertyValues.Remove( PropertyValue );
+				UnprocessedPropertyValues.RemoveSingle( PropertyValue );
 			}
 		}
 	}
