Changed trackpad detection on Mac to report trackpad only if at least two fingers are touching it. This should limit issues with people accidentally touching trackpad while using mouse + WASD.

diff --git a/Engine/Source/Runtime/Core/Private/Mac/MacApplication.cpp b/Engine/Source/Runtime/Core/Private/Mac/MacApplication.cpp
--- a/Engine/Source/Runtime/Core/Private/Mac/MacApplication.cpp
+++ b/Engine/Source/Runtime/Core/Private/Mac/MacApplication.cpp
@@ -24,7 +24,7 @@ static int MTContactCallback(int32 Device, void* Data, int32 NumFingers, double
 {
 	if (MacApplication)
 	{
-		MacApplication->SetIsUsingTrackpad(NumFingers > 0);
+		MacApplication->SetIsUsingTrackpad(NumFingers > 1);
 	}
 	return 1;
 }
