Fix BlankProgram build.

diff --git a/Engine/Source/Runtime/Core/Private/Linux/LinuxPlatformProcess.cpp b/Engine/Source/Runtime/Core/Private/Linux/LinuxPlatformProcess.cpp
--- a/Engine/Source/Runtime/Core/Private/Linux/LinuxPlatformProcess.cpp
+++ b/Engine/Source/Runtime/Core/Private/Linux/LinuxPlatformProcess.cpp
@@ -6,6 +6,7 @@
 
 #include "CorePrivate.h"
 #include "LinuxPlatformRunnableThread.h"
+#include "../../Public/Modules/ModuleVersion.h"
 #include <spawn.h>
 #include <sys/wait.h>
 #include <sys/resource.h>
