Replace the assumption that a window without minimise and maximise support is modal with an explicit IsModalWindow flag. This makes handling modal windows on OS X far more robust.

diff --git a/Engine/Source/Runtime/Core/Private/Mac/MacWindow.cpp b/Engine/Source/Runtime/Core/Private/Mac/MacWindow.cpp
--- a/Engine/Source/Runtime/Core/Private/Mac/MacWindow.cpp
+++ b/Engine/Source/Runtime/Core/Private/Mac/MacWindow.cpp
@@ -854,7 +854,10 @@ void FMacWindow::Destroy()
 		// This makes it possible to correctly handle any NSEvent's sent to the window during destruction
 		// as the WindowHandle is still a valid NSWindow. Unlike HWNDs accessing a destructed NSWindow* is fatal.
 		FSlateCocoaWindow* Window = [WindowHandle retain];
-		GRunningModalWindows.Remove(WindowHandle);
+		if(Definition->IsModalWindow)
+		{
+			GRunningModalWindows.Remove(WindowHandle);
+		}
 		Window.bForwardEvents = false;
 		if( MacApplication->OnWindowDestroyed( Window ) )
 		{
@@ -917,7 +920,7 @@ void FMacWindow::Show()
 	SCOPED_AUTORELEASE_POOL;
 	if (!bIsVisible)
 	{
-		if(!Definition->SupportsMinimize && !Definition->SupportsMaximize && Definition->IsRegularWindow && !GRunningModalWindows.Contains(WindowHandle))
+		if(Definition->IsModalWindow && !GRunningModalWindows.Contains(WindowHandle))
 		{
 			GRunningModalWindows.Add(WindowHandle);
 		}
@@ -949,7 +952,7 @@ void FMacWindow::Hide()
 	{
 		SCOPED_AUTORELEASE_POOL;
 		bIsVisible = false;
-		if(!Definition->SupportsMinimize && !Definition->SupportsMaximize)
+		if(Definition->IsModalWindow)
 		{
 			GRunningModalWindows.Remove(WindowHandle);
 		}
diff --git a/Engine/Source/Runtime/Core/Public/GenericPlatform/GenericWindowDefinition.h b/Engine/Source/Runtime/Core/Public/GenericPlatform/GenericWindowDefinition.h
--- a/Engine/Source/Runtime/Core/Public/GenericPlatform/GenericWindowDefinition.h
+++ b/Engine/Source/Runtime/Core/Public/GenericPlatform/GenericWindowDefinition.h
@@ -22,6 +22,7 @@ struct CORE_API FGenericWindowDefinition
 	bool SupportsMinimize;
 	bool SupportsMaximize;
 
+    bool IsModalWindow;
 	bool IsRegularWindow;
 	bool HasSizingFrame;
 	bool SizeWillChangeOften;
diff --git a/Engine/Source/Runtime/Slate/Private/Framework/Application/SlateApplication.cpp b/Engine/Source/Runtime/Slate/Private/Framework/Application/SlateApplication.cpp
--- a/Engine/Source/Runtime/Slate/Private/Framework/Application/SlateApplication.cpp
+++ b/Engine/Source/Runtime/Slate/Private/Framework/Application/SlateApplication.cpp
@@ -1000,6 +1000,7 @@ TSharedRef< FGenericWindow > FSlateApplication::MakeWindow( TSharedRef<SWindow>
 	Definition->SupportsMinimize = InSlateWindow->HasMinimizeBox();
 	Definition->SupportsMaximize = InSlateWindow->HasMaximizeBox();
 
+    Definition->IsModalWindow = InSlateWindow->IsModalWindow();
 	Definition->IsRegularWindow = InSlateWindow->IsRegularWindow();
 	Definition->HasSizingFrame = InSlateWindow->HasSizingFrame();
 	Definition->SizeWillChangeOften = InSlateWindow->SizeWillChangeOften();
diff --git a/Engine/Source/Runtime/SlateCore/Private/Widgets/SWindow.cpp b/Engine/Source/Runtime/SlateCore/Private/Widgets/SWindow.cpp
--- a/Engine/Source/Runtime/SlateCore/Private/Widgets/SWindow.cpp
+++ b/Engine/Source/Runtime/SlateCore/Private/Widgets/SWindow.cpp
@@ -1622,6 +1622,7 @@ SWindow::SWindow()
 	, bHasMinimizeButton( false )
 	, bHasMaximizeButton( false )
 	, bHasSizingFrame( false )
+    , bIsModalWindow( false )
 	, InitialDesiredScreenPosition( FVector2D::ZeroVector )
 	, InitialDesiredSize( FVector2D::ZeroVector )
 	, ScreenPosition( FVector2D::ZeroVector )
diff --git a/Engine/Source/Runtime/SlateCore/Public/Widgets/SWindow.h b/Engine/Source/Runtime/SlateCore/Public/Widgets/SWindow.h
--- a/Engine/Source/Runtime/SlateCore/Public/Widgets/SWindow.h
+++ b/Engine/Source/Runtime/SlateCore/Public/Widgets/SWindow.h
@@ -514,10 +514,16 @@ public:
 	/** Set modal window related flags - called by Slate app code during FSlateApplication::AddModalWindow() */
 	void SetAsModalWindow()
 	{
+        bIsModalWindow = true;
 		bHasMaximizeButton = false;
 		bHasMinimizeButton = false;
 	}
 
+	bool IsModalWindow()
+    {
+        return bIsModalWindow;
+    }
+
 	void SetTitleBar( const TSharedPtr<IWindowTitleBar> InTitleBar )
 	{
 		TitleBar = InTitleBar;
@@ -694,6 +700,9 @@ protected:
 
 	/** True if this window displays thick edge that can be used to resize the window */
 	bool bHasSizingFrame : 1;
+    
+    /** True if the window is modal */
+    bool bIsModalWindow : 1;
 
 	/** Initial desired position of the window's content in screen space */
 	FVector2D InitialDesiredScreenPosition;
