Allow a window to be set to w:1 h:1 but defer subsequent changes until the frame size is increased. Some tooltips are set to 1x1 when they are not in use, so by deferring the updates we avoid the dramatically greater cost of moving & resizing windows on OS X.

diff --git a/Engine/Source/Runtime/Core/Private/Mac/MacWindow.cpp b/Engine/Source/Runtime/Core/Private/Mac/MacWindow.cpp
--- a/Engine/Source/Runtime/Core/Private/Mac/MacWindow.cpp
+++ b/Engine/Source/Runtime/Core/Private/Mac/MacWindow.cpp
@@ -251,9 +251,12 @@ TArray< FSlateCocoaWindow* > GRunningModalWindows;
 
 - (void)setFrame:(NSRect)FrameRect display:(BOOL)Flag
 {
-	if(!bRenderInitialised || ([self isVisible] && [super alphaValue] > 0.0f))
+	NSSize Size = [self frame].size;
+	NSSize NewSize = FrameRect.size;
+	if(!bRenderInitialised || ([self isVisible] && [super alphaValue] > 0.0f && (Size.width > 1 || Size.height > 1 || NewSize.width > 1 || NewSize.height > 1)))
 	{
 		[super setFrame:FrameRect display:Flag];
+		bDeferSetFrame = false;
 	}
 	else
 	{
@@ -268,9 +271,11 @@ TArray< FSlateCocoaWindow* > GRunningModalWindows;
 
 - (void)setFrameOrigin:(NSPoint)Point
 {
-	if(!bRenderInitialised || ([self isVisible] && [super alphaValue] > 0.0f))
+	NSSize Size = [self frame].size;
+	if(!bRenderInitialised || ([self isVisible] && [super alphaValue] > 0.0f && (Size.width > 1 || Size.height > 1)))
 	{
 		[super setFrameOrigin:Point];
+		bDeferSetOrigin = false;
 	}
 	else
 	{
