- Missed a file in teh Metal optional support checkin

diff --git a/Engine/Source/Runtime/Core/Private/IOS/IOSView.cpp b/Engine/Source/Runtime/Core/Private/IOS/IOSView.cpp
--- a/Engine/Source/Runtime/Core/Private/IOS/IOSView.cpp
+++ b/Engine/Source/Runtime/Core/Private/IOS/IOSView.cpp
@@ -24,13 +24,36 @@ id<MTLDevice> GMetalDevice = nil;
 + (Class)layerClass
 {
 #if HAS_METAL
+	// make sure the project setting has enabled Metal support (per-project user settings in the editor)
+	bool bSupportMetal = false;
+	GConfig->GetBool(TEXT("/Script/UnrealEd.CookerSettings"), TEXT("bSupportMetal"), bSupportMetal, GEngineIni);
+
+	// does commandline override?
+	bool bForceES2 = FParse::Param(FCommandLine::Get(), TEXT("ES2"));
+
+	bool bTriedToInit = false;
+
 	// the check for the function pointer itself is to determine if the Metal framework exists, before calling it
-	if (!FParse::Param(FCommandLine::Get(), TEXT("ES2")) && MTLCreateSystemDefaultDevice != NULL)
+	if (bSupportMetal && !bForceES2 && MTLCreateSystemDefaultDevice != NULL)
 	{
 		// if the device is unable to run with Metal (pre-A7), this will return nil
 		GMetalDevice = MTLCreateSystemDefaultDevice();
+
+		// just tracking for printout below
+		bTriedToInit = true;
 	}
 
+#if !UE_BUILD_SHIPPING
+	if (GMetalDevice == nil)
+	{
+		FPlatformMisc::LowLevelOutputDebugStringf(TEXT("Not using Metal because: [Project Settings Disabled Metal? %s :: Commandline Forced ES2? %s :: Older OS? %s :: Pre-A7 Device? %s]"),
+			bSupportMetal ? TEXT("No") : TEXT("Yes"),
+			bForceES2? TEXT("Yes") : TEXT("No"),
+			(MTLCreateSystemDefaultDevice == NULL) ? TEXT("Yes") : TEXT("No"),
+			bTriedToInit ? TEXT("Yes") : TEXT("Unknown (didn't test)"));
+	}
+#endif
+
 	if (GMetalDevice != nil)
 	{
 		return [CAMetalLayer class];
