Stable sort algorithm

Used in-place merge sort algorithm.

#ttp 340054

diff --git a/Engine/Source/Runtime/Core/Public/Containers/Array.h b/Engine/Source/Runtime/Core/Public/Containers/Array.h
--- a/Engine/Source/Runtime/Core/Public/Containers/Array.h
+++ b/Engine/Source/Runtime/Core/Public/Containers/Array.h
@@ -1503,6 +1503,27 @@ public:
 		::Sort( GetTypedData(), Num(), Predicate );
 	}
 
+	/**
+	 * Stable sorts the array assuming < operator is defined for the item type.
+	 * Stable sort is slower than non-stable algorithm.
+	 */
+	void StableSort()
+	{
+		::StableSort(GetTypedData(), Num());
+	}
+
+	/**
+	 * Stable sorts the array using user defined predicate class. Stable sort
+	 * is slower than non-stable algorithm.
+	 *
+	 * @param Predicate Predicate class instance
+	 */
+	template <class PREDICATE_CLASS>
+	void StableSort(const PREDICATE_CLASS& Predicate)
+	{
+		::StableSort(GetTypedData(), Num(), Predicate);
+	}
+
 #if defined(_MSC_VER)
 private:
 	/**
diff --git a/Engine/Source/Runtime/Core/Public/Templates/Sorting.h b/Engine/Source/Runtime/Core/Public/Templates/Sorting.h
--- a/Engine/Source/Runtime/Core/Public/Templates/Sorting.h
+++ b/Engine/Source/Runtime/Core/Public/Templates/Sorting.h
@@ -68,7 +68,6 @@ struct TDereferenceWrapper<T*, PREDICATE_CLASS>
  * @param	First	pointer to the first element to sort
  * @param	Num		the number of items to sort
  * @param Predicate predicate class
- * @param Deref dereference class used to dereference pointer types
  */
 template<class T, class PREDICATE_CLASS> 
 void SortInternal( T* First, const int32 Num, const PREDICATE_CLASS& Predicate )
@@ -209,3 +208,352 @@ void Sort( T** First, const int32 Num )
 	SortInternal( First, Num, TDereferenceWrapper<T*, TLess<T> >( TLess<T>() ) );
 }
 
+/**
+ * Stable merge to perform sort below. Stable sort is slower than non-stable
+ * algorithm.
+ *
+ * @param Out Pointer to the first element of output array.
+ * @param In Pointer to the first element to sort.
+ * @param Mid Middle point of the table, i.e. merge separator.
+ * @param Num Number of elements in the whole table.
+ * @param Predicate Predicate class.
+ */
+template<class T, class PREDICATE_CLASS>
+void Merge(T* Out, T* In, const int32 Mid, const int32 Num, const PREDICATE_CLASS& Predicate)
+{
+	int32 Merged = 0;
+	int32 Picked;
+	int32 A = 0, B = Mid;
+
+	while (Merged < Num)
+	{
+		if (Merged != B && (B >= Num || !Predicate(In[B], In[A])))
+		{
+			Picked = A++;
+		}
+		else
+		{
+			Picked = B++;
+		}
+
+		Out[Merged] = In[Picked];
+
+		++Merged;
+	}
+}
+
+/**
+ * Euclidean algorithm using modulo policy.
+ */
+class FEuclidDivisionGCD
+{
+public:
+	/**
+	 * Calculate GCD.
+	 *
+	 * @param A First parameter.
+	 * @param B Second parameter.
+	 *
+	 * @returns Greatest common divisor of A and B.
+	 */
+	static int32 GCD(int32 A, int32 B)
+	{
+		while (B != 0)
+		{
+			int32 Temp = B;
+			B = A % B;
+			A = Temp;
+		}
+
+		return A;
+	}
+};
+
+/**
+ * Array rotation using juggling technique.
+ *
+ * @template_param TGCDPolicy Policy for calculating greatest common divisor.
+ */
+template <class TGCDPolicy>
+class TJugglingRotation
+{
+public:
+	/**
+	 * Rotates array.
+	 *
+	 * @param First Pointer to the array.
+	 * @param From Rotation starting point.
+	 * @param To Rotation ending point.
+	 * @param Amount Amount of steps to rotate.
+	 */
+	template <class T>
+	static void Rotate(T* First, const int32 From, const int32 To, const int32 Amount)
+	{
+		if (Amount == 0)
+		{
+			return;
+		}
+
+		auto Num = To - From;
+		auto GCD = TGCDPolicy::GCD(Num, Amount);
+		auto CycleSize = Num / GCD;
+
+		for (int32 Index = 0; Index < GCD; ++Index)
+		{
+			T BufferObject = MoveTemp(First[From + Index]);
+			int32 IndexToFill = Index;
+
+			for (int32 InCycleIndex = 0; InCycleIndex < CycleSize; ++InCycleIndex)
+			{
+				IndexToFill = (IndexToFill + Amount) % Num;
+				Exchange(First[From + IndexToFill], BufferObject);
+			}
+		}
+	}
+};
+
+/**
+ * Merge policy for merge sort.
+ *
+ * @template_param TRotationPolicy Policy for array rotation algorithm.
+ */
+template <class TRotationPolicy>
+class TRotationInPlaceMerge
+{
+public:
+	/**
+	 * Two sorted arrays merging function.
+	 *
+	 * @param First Pointer to array.
+	 * @param Mid Middle point i.e. separation point of two arrays to merge.
+	 * @param Num Number of elements in array.
+	 * @param Predicate Predicate for comparison.
+	 */
+	template <class T, class PREDICATE_CLASS>
+	static void Merge(T* First, const int32 Mid, const int32 Num, const PREDICATE_CLASS& Predicate)
+	{
+		int32 AStart = 0;
+		int32 BStart = Mid;
+
+		while (AStart < BStart && BStart < Num)
+		{
+			int32 NewBOffset = BinarySearchFirst(First + BStart, Num - BStart, First[AStart], Predicate);
+
+			TRotationPolicy::Rotate(First, AStart, BStart + NewBOffset, NewBOffset);
+
+			BStart += NewBOffset;
+			AStart += BinarySearchLast(First + AStart, BStart - AStart, First[AStart + NewBOffset], Predicate) + 1;
+		}
+	}
+
+private:
+	/**
+	 * Performs binary search, resulting in position of the first element with given value in an array.
+	 *
+	 * @param First Pointer to array.
+	 * @param Num Number of elements in array.
+	 * @param Value Value to look for.
+	 * @param Predicate Predicate for comparison.
+	 *
+	 * @returns Position of the first element with value Value.
+	 */
+	template <class T, class PREDICATE_CLASS>
+	static int32 BinarySearchFirst(T* First, const int32 Num, const T& Value, const PREDICATE_CLASS& Predicate)
+	{
+		int32 Start = 0;
+		int32 End = Num;
+		
+		while (End - Start > 1)
+		{
+			int32 Mid = (Start + End) / 2;
+			bool bComparison = Predicate(First[Mid], Value);
+
+			Start = bComparison ? Mid : Start;
+			End = bComparison ? End : Mid;
+		}
+
+		return Predicate(First[Start], Value) ? Start + 1 : Start;
+	}
+
+	/**
+	 * Performs binary search, resulting in position of the last element with given value in an array.
+	 *
+	 * @param First Pointer to array.
+	 * @param Num Number of elements in array.
+	 * @param Value Value to look for.
+	 * @param Predicate Predicate for comparison.
+	 *
+	 * @returns Position of the last element with value Value.
+	 */
+	template <class T, class PREDICATE_CLASS>
+	static int32 BinarySearchLast(T* First, const int32 Num, const T& Value, const PREDICATE_CLASS& Predicate)
+	{
+		int32 Start = 0;
+		int32 End = Num;
+
+		while (End - Start > 1)
+		{
+			int32 Mid = (Start + End) / 2;
+			bool bComparison = !Predicate(Value, First[Mid]);
+			
+			Start = bComparison ? Mid : Start;
+			End = bComparison ? End : Mid;
+		}
+
+		return Predicate(Value, First[Start]) ? Start - 1 : Start;
+	}
+};
+
+/**
+ * Merge sort class.
+ *
+ * @template_param TMergePolicy Merging policy.
+ * @template_param MinMergeSubgroupSize Minimal size of the subgroup that should be merged.
+ */
+template <class TMergePolicy, int32 MinMergeSubgroupSize = 2>
+class TMergeSort
+{
+public:
+	/**
+	 * Sort the array.
+	 *
+	 * @param First Pointer to the array.
+	 * @param Num Number of elements in the array.
+	 * @param Predicate Predicate for comparison.
+	 */
+	template<class T, class PREDICATE_CLASS>
+	static void Sort(T* First, const int32 Num, const PREDICATE_CLASS& Predicate)
+	{
+		int32 SubgroupStart = 0;
+
+		if (MinMergeSubgroupSize > 1)
+		{
+			if (MinMergeSubgroupSize > 2)
+			{
+				// First pass with simple bubble-sort.
+				do
+				{
+					int32 GroupEnd = FPlatformMath::Min(SubgroupStart + MinMergeSubgroupSize, Num);
+					do
+					{
+						for (int32 It = SubgroupStart; It < GroupEnd - 1; ++It)
+						{
+							if (Predicate(First[It + 1], First[It]))
+							{
+								Exchange(First[It], First[It + 1]);
+							}
+						}
+						GroupEnd--;
+					} while (GroupEnd - SubgroupStart > 1);
+
+					SubgroupStart += MinMergeSubgroupSize;
+				} while (SubgroupStart < Num);
+			}
+			else
+			{
+				for (int32 SubgroupStart = 0; SubgroupStart < Num; SubgroupStart += 2)
+				{
+					if (Predicate(First[SubgroupStart + 1], First[SubgroupStart]))
+					{
+						Exchange(First[SubgroupStart], First[SubgroupStart + 1]);
+					}
+				}
+			}
+		}
+
+		if (Num > MinMergeSubgroupSize)
+		{
+			int32 SubgroupSize = MinMergeSubgroupSize << 1;
+			while (SubgroupSize < Num)
+			{
+				SubgroupStart = 0;
+				do
+				{
+					TMergePolicy::Merge(
+						First + SubgroupStart,
+						SubgroupSize,
+						FPlatformMath::Min(SubgroupSize << 1, Num - SubgroupStart),
+						Predicate);
+					SubgroupStart += SubgroupSize << 1;
+				} while (SubgroupStart < Num);
+
+				SubgroupSize <<= 1;
+			}
+		}
+	}
+};
+
+/**
+ * Stable sort elements using user defined predicate class. The sort is stable,
+ * meaning that the ordering of equal items is preserved, but it's slower than
+ * non-stable algorithm.
+ *
+ * This is the internal sorting function used by StableSort overrides.
+ *
+ * @param	First	pointer to the first element to sort
+ * @param	Num		the number of items to sort
+ * @param Predicate predicate class
+ */
+template<class T, class PREDICATE_CLASS>
+void StableSortInternal(T* First, const int32 Num, const PREDICATE_CLASS& Predicate)
+{
+	TMergeSort<TRotationInPlaceMerge<TJugglingRotation<FEuclidDivisionGCD> > >::Sort(First, Num, Predicate);
+}
+
+/**
+ * Stable sort elements using user defined predicate class. The sort is stable,
+ * meaning that the ordering of equal items is preserved, but it's slower than
+ * non-stable algorithm.
+ *
+ * @param	First	pointer to the first element to sort
+ * @param	Num		the number of items to sort
+ * @param Predicate predicate class
+ */
+template<class T, class PREDICATE_CLASS>
+void StableSort(T* First, const int32 Num, const PREDICATE_CLASS& Predicate)
+{
+	StableSortInternal(First, Num, TDereferenceWrapper<T, PREDICATE_CLASS>(Predicate));
+}
+
+/**
+ * Specialized version of the above StableSort function for pointers to elements.
+ * Stable sort is slower than non-stable algorithm.
+ *
+ * @param	First	pointer to the first element to sort
+ * @param	Num		the number of items to sort
+ * @param Predicate predicate class
+ */
+template<class T, class PREDICATE_CLASS>
+void StableSort(T** First, const int32 Num, const PREDICATE_CLASS& Predicate)
+{
+	StableSortInternal(First, Num, TDereferenceWrapper<T*, PREDICATE_CLASS>(Predicate));
+}
+
+/**
+ * Stable sort elements. The sort is stable, meaning that the ordering of equal
+ * items is preserved, but it's slower than non-stable algorithm.
+ *
+ * Assumes < operator is defined for the template type.
+ *
+ * @param	First	pointer to the first element to sort
+ * @param	Num		the number of items to sort
+ */
+template<class T>
+void StableSort(T* First, const int32 Num)
+{
+	StableSortInternal(First, Num, TDereferenceWrapper<T, TLess<T> >(TLess<T>()));
+}
+
+/**
+ * Specialized version of the above StableSort function for pointers to elements.
+ * Stable sort is slower than non-stable algorithm.
+ *
+ * @param	First	pointer to the first element to sort
+ * @param	Num		the number of items to sort
+ */
+template<class T>
+void StableSort(T** First, const int32 Num)
+{
+	StableSortInternal(First, Num, TDereferenceWrapper<T*, TLess<T> >(TLess<T>()));
+}
+
