Temporary comment out task graph stats cause they are making huge overhead. Solution is in progress.

diff --git a/Engine/Source/Runtime/Core/Public/Async/TaskGraphInterfaces.h b/Engine/Source/Runtime/Core/Public/Async/TaskGraphInterfaces.h
--- a/Engine/Source/Runtime/Core/Public/Async/TaskGraphInterfaces.h
+++ b/Engine/Source/Runtime/Core/Public/Async/TaskGraphInterfaces.h
@@ -817,6 +817,8 @@ public:
 	FCustomNameGraphTaskBase(const TCHAR* TaskClassName, const TCHAR* TaskName)
 	{
 #if STATS
+		// Commented out for now as its generating huge overhead.
+		/*
 		const FString LongName = FString(TaskClassName) + TEXT(".") + FString(TaskName);
 		const FName StatName = FName( *LongName );
 
@@ -831,6 +833,7 @@ public:
 			STAT_GROUP_TO_FStatGroup(STATGROUP_TaskGraphTasks)::GetGroupCategory(),
 			STAT_GROUP_TO_FStatGroup(STATGROUP_TaskGraphTasks)::DefaultEnable,
 			true, EStatDataType::ST_int64, *LongName, true);
+		*/
 #endif
 	}
 
@@ -842,14 +845,16 @@ public:
 	FORCEINLINE TStatId GetStatId() const
 	{
 #if STATS
-		return StatID;
+		// See the comment above.
+		return TStatId(); /* StatID; */
 #endif
 		return TStatId();
 	}
 
 private:
+	// See the comment above.
 	/** Stat id of this object. */
-	STAT(TStatId StatID;)
+	/* STAT(TStatId StatID;) */
 };
 
 /** 
