Make sure FGenericPlatformMisc::EngineDir() sets the correct working dir on desktop platforms before using DirectoryExists(), so it can be called before FEngineLoop::PreInit()

diff --git a/Engine/Source/Runtime/Core/Private/GenericPlatform/GenericPlatformMisc.cpp b/Engine/Source/Runtime/Core/Private/GenericPlatform/GenericPlatformMisc.cpp
--- a/Engine/Source/Runtime/Core/Private/GenericPlatform/GenericPlatformMisc.cpp
+++ b/Engine/Source/Runtime/Core/Private/GenericPlatform/GenericPlatformMisc.cpp
@@ -1,4 +1,4 @@
-﻿// Copyright 1998-2014 Epic Games, Inc. All Rights Reserved.
+// Copyright 1998-2014 Epic Games, Inc. All Rights Reserved.
 
 #include "CorePrivate.h"
 #include "MallocAnsi.h"
@@ -448,6 +448,8 @@ const TCHAR* FGenericPlatformMisc::EngineDir()
 		// See if we are a root-level project
 		FString DefaultEngineDir = TEXT("../../../Engine/");
 #if PLATFORM_DESKTOP
+		FPlatformProcess::SetCurrentWorkingDirectoryToBaseDir();
+
 		//@todo. Need to have a define specific for this scenario??
 		if (FPlatformFileManager::Get().GetPlatformFile().DirectoryExists(*(DefaultEngineDir / TEXT("Binaries"))))
 		{
