Use private MTDeviceIsBuiltIn() function to distinguish between trackpad and Magic Mouse when detecting if user is using trackpad.

diff --git a/Engine/Source/Runtime/Core/Private/Mac/MacApplication.cpp b/Engine/Source/Runtime/Core/Private/Mac/MacApplication.cpp
--- a/Engine/Source/Runtime/Core/Private/Mac/MacApplication.cpp
+++ b/Engine/Source/Runtime/Core/Private/Mac/MacApplication.cpp
@@ -15,16 +15,18 @@
 FMacApplication* MacApplication = NULL;
 
 #if WITH_EDITOR
-typedef int32 (*MTContactCallbackFunction)(int32, void*, int32, double, int32);
+typedef int32 (*MTContactCallbackFunction)(void*, void*, int32, double, int32);
 extern "C" CFMutableArrayRef MTDeviceCreateList(void);
 extern "C" void MTRegisterContactFrameCallback(void*, MTContactCallbackFunction);
 extern "C" void MTDeviceStart(void*, int);
+extern "C" bool MTDeviceIsBuiltIn(void*);
 
-static int MTContactCallback(int32 Device, void* Data, int32 NumFingers, double TimeStamp, int32 Frame)
+static int MTContactCallback(void* Device, void* Data, int32 NumFingers, double TimeStamp, int32 Frame)
 {
 	if (MacApplication)
 	{
-		MacApplication->SetIsUsingTrackpad(NumFingers > 1);
+		const bool bIsTrackpad = MTDeviceIsBuiltIn(Device);
+		MacApplication->SetIsUsingTrackpad(NumFingers > (bIsTrackpad ? 1 : 0));
 	}
 	return 1;
 }
