#UE4 #HTML5
             Process mouse move events correctly. StrategyGame window now drags correctly.

diff --git a/Engine/Source/Runtime/Core/Private/HTML5/HTML5InputInterface.cpp b/Engine/Source/Runtime/Core/Private/HTML5/HTML5InputInterface.cpp
--- a/Engine/Source/Runtime/Core/Private/HTML5/HTML5InputInterface.cpp
+++ b/Engine/Source/Runtime/Core/Private/HTML5/HTML5InputInterface.cpp
@@ -93,6 +93,7 @@ void FHTML5InputInterface::Tick(float DeltaTime, const SDL_Event& Event)
 		{
 			Cursor->SetPosition(Event.motion.x, Event.motion.y);
 			MessageHandler->OnRawMouseMove(Event.motion.xrel, -Event.motion.yrel);
+			MessageHandler->OnMouseMove(); 
 		}
 #if PLATFORM_HTML5_WIN32 && SDL_MAJOR_VERSION < 2
 		else if (Event.type == SDL_ACTIVEEVENT)
