Modified text logic to use a more appropriate whitespace check for internationalization.

diff --git a/Engine/Source/Runtime/Core/Private/Internationalization/Text.cpp b/Engine/Source/Runtime/Core/Private/Internationalization/Text.cpp
--- a/Engine/Source/Runtime/Core/Private/Internationalization/Text.cpp
+++ b/Engine/Source/Runtime/Core/Private/Internationalization/Text.cpp
@@ -740,7 +740,7 @@ bool FText::ShouldGatherForLocalization() const
 	{
 		for(int32 i = 0; i < String.Len(); ++i)
 		{
-			if( !FChar::IsWhitespace( String[i] ) )
+			if( !FText::IsWhitespace( String[i] ) )
 			{
 				return false;
 			}
