- And one more file for Metal beta3 -> main

diff --git a/Engine/Source/Runtime/Core/Private/IOS/IOSView.cpp b/Engine/Source/Runtime/Core/Private/IOS/IOSView.cpp
--- a/Engine/Source/Runtime/Core/Private/IOS/IOSView.cpp
+++ b/Engine/Source/Runtime/Core/Private/IOS/IOSView.cpp
@@ -203,7 +203,7 @@ id<MTLDevice> GMetalDevice = nil;
 #if HAS_METAL
 - (id<CAMetalDrawable>)MakeDrawable
 {
-	return [(CAMetalLayer*)self.layer newDrawable];
+	return [(CAMetalLayer*)self.layer nextDrawable];
 }
 #endif
 
