The Windows Application now has minimal TaskbarList support. Currently, the only feature implemented is the ability to present some sort of progress bar with the taskbar button.

diff --git a/Engine/Source/Runtime/Core/Private/Windows/WindowsApplication.cpp b/Engine/Source/Runtime/Core/Private/Windows/WindowsApplication.cpp
--- a/Engine/Source/Runtime/Core/Private/Windows/WindowsApplication.cpp
+++ b/Engine/Source/Runtime/Core/Private/Windows/WindowsApplication.cpp
@@ -595,6 +595,10 @@ LRESULT CALLBACK FWindowsApplication::AppWndProc(HWND hwnd, uint32 msg, WPARAM w
 
 int32 FWindowsApplication::ProcessMessage( HWND hwnd, uint32 msg, WPARAM wParam, LPARAM lParam )
 {
+	//The message id for when the taskbar button becoems created.
+	//This is set to s proper id once a WM_CREATE message is received.
+	static DWORD wmTaskbarButtonCreate = (DWORD)-1; 
+
 	TSharedPtr< FWindowsWindow > CurrentNativeEventWindowPtr = FindWindowByHWND( Windows, hwnd );
 
 	if( Windows.Num() && CurrentNativeEventWindowPtr.IsValid() )
@@ -1021,6 +1025,7 @@ int32 FWindowsApplication::ProcessMessage( HWND hwnd, uint32 msg, WPARAM wParam,
 		
 		case WM_CREATE:
 			{
+				wmTaskbarButtonCreate = RegisterWindowMessage(TEXT("TaskbarButtonCreated"));
 				return 0;
 			}
 			break;
@@ -1030,6 +1035,11 @@ int32 FWindowsApplication::ProcessMessage( HWND hwnd, uint32 msg, WPARAM wParam,
 				XInput->SetNeedsControllerStateUpdate(); 
 			}
 		}
+
+		if (wmTaskbarButtonCreate)
+		{
+			TaskbarList = FTaskbarList::Create();
+		}
 	}
 
 	return DefWindowProc(hwnd, msg, wParam, lParam);
@@ -1709,6 +1719,11 @@ void FWindowsApplication::AddExternalInputDevice(TSharedPtr<IInputDevice> InputD
 	}
 }
 
+TSharedPtr<FTaskbarList> FWindowsApplication::GetTaskbarList()
+{
+	return TaskbarList;
+}
+
 void FWindowsApplication::DeferDragDropOperation(const FDeferredWindowsDragDropOperation& DeferredDragDropOperation)
 {
 	DeferredDragDropOperations.Add(DeferredDragDropOperation);
@@ -1777,4 +1792,54 @@ HRESULT FWindowsApplication::OnOLEDrop( const HWND HWnd, const FDragDropOLEData&
 	return 0;
 }
 
+void FTaskbarList::Initialize()
+{
+	if (FWindowsPlatformMisc::CoInitialize())
+	{
+		CoCreateInstance(CLSID_TaskbarList, NULL, CLSCTX_INPROC_SERVER, IID_ITaskbarList3, (void **)&TaskBarList3);
+	}
+}
+
+FTaskbarList::FTaskbarList() 
+	: TaskBarList3(NULL)
+{
+
+}
+
+FTaskbarList::~FTaskbarList()
+{
+	if (FWindowsPlatformMisc::CoInitialize() && TaskBarList3)
+	{
+		TaskBarList3->Release();
+	}
+
+	TaskBarList3 = NULL;
+}
+
+void FTaskbarList::SetProgressValue(const TSharedRef<FGenericWindow>& NativeWindow, uint64 Current, uint64 Total)
+{
+	if (TaskBarList3)
+	{
+		const TSharedRef< FWindowsWindow > Window = StaticCastSharedRef< FWindowsWindow >(NativeWindow);
+		TaskBarList3->SetProgressValue(Window->GetHWnd(), (ULONGLONG)Current, (ULONGLONG)Total);
+	}
+}
+
+void FTaskbarList::SetProgressState(const TSharedRef<FGenericWindow>& NativeWindow, TaskbarProgressState::Type State)
+{
+	if (TaskBarList3)
+	{
+		const TSharedRef< FWindowsWindow > Window = StaticCastSharedRef< FWindowsWindow >(NativeWindow);
+		TaskBarList3->SetProgressState(Window->GetHWnd(), (TBPFLAG)State);
+	}
+}
+
+TSharedRef<FTaskbarList> FTaskbarList::Create()
+{
+	TSharedRef<FTaskbarList> TaskbarList = MakeShareable(new FTaskbarList());
+	TaskbarList->Initialize();
+
+	return TaskbarList;
+}
+
 #include "HideWindowsPlatformTypes.h"
diff --git a/Engine/Source/Runtime/Core/Public/Windows/WindowsApplication.h b/Engine/Source/Runtime/Core/Public/Windows/WindowsApplication.h
--- a/Engine/Source/Runtime/Core/Public/Windows/WindowsApplication.h
+++ b/Engine/Source/Runtime/Core/Public/Windows/WindowsApplication.h
@@ -7,6 +7,7 @@
 #include "AllowWindowsPlatformTypes.h"
 	#include "Ole2.h"
 	#include "OleIdl.h"
+	#include <ShObjIdl.h>
 #include "HideWindowsPlatformTypes.h"
 #include "IForceFeedbackSystem.h"
 #include "WindowsTextInputMethodSystem.h"
@@ -16,6 +17,62 @@ DECLARE_LOG_CATEGORY_EXTERN(LogWindowsDesktop, Log, All);
 class FWindowsWindow;
 class FGenericApplicationMessageHandler;
 
+namespace TaskbarProgressState
+{
+	enum Type
+	{
+		//Stops displaying progress and returns the button to its normal state.
+		NoProgress = 0x0,
+
+		//The progress indicator does not grow in size, but cycles repeatedly along the 
+		//length of the taskbar button. This indicates activity without specifying what 
+		//proportion of the progress is complete. Progress is taking place, but there is 
+		//no prediction as to how long the operation will take.
+		Indeterminate = 0x1,
+
+		//The progress indicator grows in size from left to right in proportion to the 
+		//estimated amount of the operation completed. This is a determinate progress 
+		//indicator; a prediction is being made as to the duration of the operation.
+		Normal = 0x2,
+
+		//The progress indicator turns red to show that an error has occurred in one of 
+		//the windows that is broadcasting progress. This is a determinate state. If the 
+		//progress indicator is in the indeterminate state, it switches to a red determinate 
+		//display of a generic percentage not indicative of actual progress.
+		Error = 0x4,
+
+		//The progress indicator turns yellow to show that progress is currently stopped in 
+		//one of the windows but can be resumed by the user. No error condition exists and 
+		//nothing is preventing the progress from continuing. This is a determinate state. 
+		//If the progress indicator is in the indeterminate state, it switches to a yellow 
+		//determinate display of a generic percentage not indicative of actual progress.
+		Paused = 0x8,
+	};
+}
+
+class FTaskbarList
+{
+public:
+
+	static TSharedRef<FTaskbarList> Create();
+
+	void SetProgressState(const TSharedRef<FGenericWindow>& NativeWindow, TaskbarProgressState::Type State);
+
+	void SetProgressValue(const TSharedRef<FGenericWindow>& NativeWindow, uint64 Current, uint64 Total);
+
+	~FTaskbarList();
+
+private:
+
+	FTaskbarList();
+
+	void Initialize();
+
+private:
+
+	ITaskbarList3* TaskBarList3;
+};
+
 struct FDeferredWindowsMessage
 {
 	FDeferredWindowsMessage( const TSharedPtr<FWindowsWindow>& InNativeWindow, HWND InHWnd, uint32 InMessage, WPARAM InWParam, LPARAM InLParam, int32 InX=0, int32 InY=0, uint32 InRawInputFlags = 0 )
@@ -209,6 +266,9 @@ public:
 
 	virtual void AddExternalInputDevice(TSharedPtr<class IInputDevice> InputDevice);
 
+	TSharedPtr<FTaskbarList> GetTaskbarList();
+
+
 protected:
 
 	static LRESULT CALLBACK AppWndProc(HWND hwnd, uint32 msg, WPARAM wParam, LPARAM lParam);
@@ -288,6 +348,8 @@ private:
 
 	TSharedPtr<FWindowsTextInputMethodSystem> TextInputMethodSystem;
 
+	TSharedPtr<FTaskbarList> TaskbarList;
+
 	// Accessibility shortcut keys
 	STICKYKEYS							StartupStickyKeys;
 	TOGGLEKEYS							StartupToggleKeys;
