Fixed the taskbarlist progress indicator to properly show indeterminate state.

diff --git a/Engine/Source/Runtime/Core/Private/Windows/WindowsApplication.cpp b/Engine/Source/Runtime/Core/Private/Windows/WindowsApplication.cpp
--- a/Engine/Source/Runtime/Core/Private/Windows/WindowsApplication.cpp
+++ b/Engine/Source/Runtime/Core/Private/Windows/WindowsApplication.cpp
@@ -1825,7 +1825,7 @@ void FTaskbarList::SetProgressValue(const TSharedRef<FGenericWindow>& NativeWind
 	}
 }
 
-void FTaskbarList::SetProgressState(const TSharedRef<FGenericWindow>& NativeWindow, TaskbarProgressState::Type State)
+void FTaskbarList::SetProgressState(const TSharedRef<FGenericWindow>& NativeWindow, ETaskbarProgressState::Type State)
 {
 	if (TaskBarList3)
 	{
diff --git a/Engine/Source/Runtime/Core/Public/Windows/WindowsApplication.h b/Engine/Source/Runtime/Core/Public/Windows/WindowsApplication.h
--- a/Engine/Source/Runtime/Core/Public/Windows/WindowsApplication.h
+++ b/Engine/Source/Runtime/Core/Public/Windows/WindowsApplication.h
@@ -17,7 +17,7 @@ DECLARE_LOG_CATEGORY_EXTERN(LogWindowsDesktop, Log, All);
 class FWindowsWindow;
 class FGenericApplicationMessageHandler;
 
-namespace TaskbarProgressState
+namespace ETaskbarProgressState
 {
 	enum Type
 	{
@@ -56,7 +56,7 @@ public:
 
 	static TSharedRef<FTaskbarList> Create();
 
-	void SetProgressState(const TSharedRef<FGenericWindow>& NativeWindow, TaskbarProgressState::Type State);
+	void SetProgressState(const TSharedRef<FGenericWindow>& NativeWindow, ETaskbarProgressState::Type State);
 
 	void SetProgressValue(const TSharedRef<FGenericWindow>& NativeWindow, uint64 Current, uint64 Total);
 
