Moved the Launcher's taskbar button code into a windows specific class. Refactored the application to allow platforms to inject special behavior easier.

diff --git a/Engine/Source/Runtime/Core/Private/Windows/WindowsApplication.cpp b/Engine/Source/Runtime/Core/Private/Windows/WindowsApplication.cpp
--- a/Engine/Source/Runtime/Core/Private/Windows/WindowsApplication.cpp
+++ b/Engine/Source/Runtime/Core/Private/Windows/WindowsApplication.cpp
@@ -1816,6 +1816,15 @@ FTaskbarList::~FTaskbarList()
 	TaskBarList3 = NULL;
 }
 
+void FTaskbarList::SetOverlayIcon(const TSharedRef<FGenericWindow>& NativeWindow, HICON Icon, FText Description)
+{
+	if (TaskBarList3)
+	{
+		const TSharedRef< FWindowsWindow > Window = StaticCastSharedRef< FWindowsWindow >(NativeWindow);
+		TaskBarList3->SetOverlayIcon(Window->GetHWnd(), Icon, *Description.ToString());
+	}
+}
+
 void FTaskbarList::SetProgressValue(const TSharedRef<FGenericWindow>& NativeWindow, uint64 Current, uint64 Total)
 {
 	if (TaskBarList3)
diff --git a/Engine/Source/Runtime/Core/Public/Windows/WindowsApplication.h b/Engine/Source/Runtime/Core/Public/Windows/WindowsApplication.h
--- a/Engine/Source/Runtime/Core/Public/Windows/WindowsApplication.h
+++ b/Engine/Source/Runtime/Core/Public/Windows/WindowsApplication.h
@@ -56,6 +56,8 @@ public:
 
 	static TSharedRef<FTaskbarList> Create();
 
+	void SetOverlayIcon(const TSharedRef<FGenericWindow>& NativeWindow, HICON Icon, FText Description);
+
 	void SetProgressState(const TSharedRef<FGenericWindow>& NativeWindow, ETaskbarProgressState::Type State);
 
 	void SetProgressValue(const TSharedRef<FGenericWindow>& NativeWindow, uint64 Current, uint64 Total);
