Added support for overriding config options from the commandline.
E.g. -ini:Engine:/Script/Engine.Engine.bSmoothFrameRate=False,TextureStreaming.PoolSize=100

Unshelved and fixed up Joshs UE3 port from the shelved CL# 1760099
- Brought over the commandline .ini override from UE3

Note, that this does not serialize the new value to the saved config. It is simply an override for that run.

diff --git a/Engine/Source/Runtime/Core/Private/Misc/ConfigCacheIni.cpp b/Engine/Source/Runtime/Core/Private/Misc/ConfigCacheIni.cpp
--- a/Engine/Source/Runtime/Core/Private/Misc/ConfigCacheIni.cpp
+++ b/Engine/Source/Runtime/Core/Private/Misc/ConfigCacheIni.cpp
@@ -555,6 +555,57 @@ bool FConfigFile::ShouldExportQuotedString(const FString& PropertyValue) const
 }
 
 
+#if !UE_BUILD_SHIPPING
+/**
+* Looks for any overrides on the commandline for this file
+*
+* @param File Config to possibly modify
+* @param Filename Name of the .ini file to look for overrides
+*/
+static void OverrideFromCommandline(FConfigFile* File, const FString& Filename)
+{
+	FString Settings;
+	// look for this filename on the commandline in the format:
+	//		-ini:IniName:Section1.Key1=Value1,Section2.Key2=Value2
+	// for example:
+	//		-ini:Engine:/Script/Engine.Engine.bSmoothFrameRate=False,TextureStreaming.PoolSize=100
+	//			(will update the cache after the final combined engine.ini)
+	//		-ini:DefaultEngine:/Script/Engine.Engine.bSmoothFrameRate=False,TextureStreaming.PoolSize=100
+	//			(will update the cache after the default engine.ini and before applying platform inis)
+	if (FParse::Value(FCommandLine::Get(), *FString::Printf(TEXT("-ini:%s:"), *FPaths::GetBaseFilename(Filename)), Settings, false))
+	{
+		// break apart on the commas
+		TArray<FString> SettingPairs;
+		Settings.ParseIntoArray(&SettingPairs, TEXT(","), true);
+		for (int32 Index = 0; Index < SettingPairs.Num(); Index++)
+		{
+			// set each one, by splitting on the =
+			FString SectionAndKey, Value;
+			if (SettingPairs[Index].Split(TEXT("="), &SectionAndKey, &Value))
+			{
+				// now we need to split off the key from the rest of the section name
+				int32 LastDot = INDEX_NONE;
+				SectionAndKey.FindLastChar(TCHAR('.'), LastDot);
+
+				// check for malformed string
+				if (LastDot == INDEX_NONE || LastDot == 0)
+				{
+					continue;
+				}
+
+				// break it apart
+				FString Section = SectionAndKey.Left(LastDot);
+				FString Key = SectionAndKey.Mid(LastDot + 1);
+
+				// now put it into this .ini
+				File->SetString(*Section, *Key, *Value);
+			}
+		}
+	}
+}
+#endif
+
+
 /**
  * This will completely load .ini file hierarchy into the passed in FConfigFile. The passed in FConfigFile will then
  * have the data after combining all of those .ini 
@@ -600,9 +651,10 @@ static bool LoadIniFileHierarchy(const TArray<FIniFilename>& HierarchyToLoad, FC
 	for( int32 IniIndex = 0; IniIndex < HierarchyToLoad.Num(); IniIndex++ )
 	{
 		const FIniFilename& IniToLoad = HierarchyToLoad[ IniIndex ];
+		const FString& IniFileName = IniToLoad.Filename;
 
 		// Spit out friendly error if there was a problem locating .inis (e.g. bad command line parameter or missing folder, ...).
-		if( IsUsingLocalIniFile(*IniToLoad.Filename, NULL) && (IFileManager::Get().FileSize( *IniToLoad.Filename ) < 0) )
+		if (IsUsingLocalIniFile(*IniFileName, NULL) && (IFileManager::Get().FileSize(*IniFileName) < 0))
 		{
 			if (IniToLoad.bRequired)
 			{
@@ -619,7 +671,12 @@ static bool LoadIniFileHierarchy(const TArray<FIniFilename>& HierarchyToLoad, FC
 		bool bDoCombine = (IniIndex != 0);
 		bool bDoWrite = false;
 		//UE_LOG(LogConfig, Log,  TEXT( "Combining configFile: %s" ), *IniList(IniIndex) );
-		ProcessIniContents(*HierarchyToLoad.Last().Filename, *IniToLoad.Filename, &ConfigFile, bDoEmptyConfig, bDoCombine, bDoWrite);
+		ProcessIniContents(*HierarchyToLoad.Last().Filename, *IniFileName, &ConfigFile, bDoEmptyConfig, bDoCombine, bDoWrite);
+
+#if !UE_BUILD_SHIPPING
+		// process any commandline overrides
+		OverrideFromCommandline(&ConfigFile, IniFileName);
+#endif
 	}
 
 	// Set this configs files source ini hierarchy to show where it was loaded from.
