[INTEGRATE] Latest from 4.3 branch.

diff --git a/Engine/Source/Programs/AutomationTool/IOS/IOSPlatform.Automation.cs b/Engine/Source/Programs/AutomationTool/IOS/IOSPlatform.Automation.cs
--- a/Engine/Source/Programs/AutomationTool/IOS/IOSPlatform.Automation.cs
+++ b/Engine/Source/Programs/AutomationTool/IOS/IOSPlatform.Automation.cs
@@ -149,7 +149,7 @@ public class IOSPlatform : Platform
 
 			IPPArguments += (cookonthefly ? " -cookonthefly" : "");
 			IPPArguments += " -stagedir \"" + CombinePaths(Params.BaseStageDirectory, "IOS") + "\"";
-			IPPArguments += " -projectdir \"" + Path.GetDirectoryName(Params.RawProjectPath) + "\"";
+			IPPArguments += " -project \"" + Params.RawProjectPath + "\"";
 
 			// rename the .ipa if not code based
 			if (!Params.IsCodeBasedProject)
diff --git a/Engine/Source/Programs/IOS/iPhonePackager/Config.cs b/Engine/Source/Programs/IOS/iPhonePackager/Config.cs
--- a/Engine/Source/Programs/IOS/iPhonePackager/Config.cs
+++ b/Engine/Source/Programs/IOS/iPhonePackager/Config.cs
@@ -40,9 +40,9 @@ namespace iPhonePackager
 		public static string RepackageStagingDirectory = "";
 
 		/// <summary>
-		/// The project root directory that is passed into IPP from UAT
+		/// The project file that is passed into IPP from UAT
 		/// </summary>
-		public static string ProjectRootDirectory = "";
+		public static string ProjectFile = "";
 
 		/// <summary>
 		/// The device to deploy or launch on
@@ -197,10 +197,10 @@ namespace iPhonePackager
 
 			string BinariesDir = BinariesDirectory;
 			string GameName = Program.GameName;
-			if (!String.IsNullOrEmpty(ProjectRootDirectory))
+			if (!String.IsNullOrEmpty(ProjectFile))
 			{
-				BinariesDir = Path.Combine(ProjectRootDirectory, "Binaries", "IOS");
-				GameName = ProjectRootDirectory.Substring(ProjectRootDirectory.LastIndexOfAny(new char[] { '\\', '/' }) + 1);
+				BinariesDir = Path.Combine(Path.GetDirectoryName(ProjectFile), "Binaries", "IOS");
+				GameName = Path.GetFileNameWithoutExtension(ProjectFile);
 			}
 			if (Program.GameConfiguration == "Development")
 			{
diff --git a/Engine/Source/Programs/IOS/iPhonePackager/MobileProvisionUtilities.cs b/Engine/Source/Programs/IOS/iPhonePackager/MobileProvisionUtilities.cs
--- a/Engine/Source/Programs/IOS/iPhonePackager/MobileProvisionUtilities.cs
+++ b/Engine/Source/Programs/IOS/iPhonePackager/MobileProvisionUtilities.cs
@@ -34,10 +34,9 @@ namespace iPhonePackager
 			string GameName = Program.GameName;
 			if (GameName == "UE4Game")
 			{
-				string[] project = Directory.GetFiles(Config.ProjectRootDirectory, "*.uproject");
-				if (project.Length == 1)
+				if(Config.ProjectFile.Length > 0)
 				{
-					GameName = Path.GetFileNameWithoutExtension(project[0]);
+					GameName = Path.GetFileNameWithoutExtension(Config.ProjectFile);
 				}
 			}
 
@@ -49,7 +48,7 @@ namespace iPhonePackager
 
 			#region remove after we provide an install mechanism
 			// copy all of the provisions from the game directory to the library
-			foreach (string Provision in Directory.EnumerateFiles(Config.ProjectRootDirectory + "/Build/IOS/", "*.mobileprovision", SearchOption.AllDirectories))
+			foreach (string Provision in Directory.EnumerateFiles(Path.GetDirectoryName(Config.ProjectFile) + "/Build/IOS/", "*.mobileprovision", SearchOption.AllDirectories))
 			{
 				if (!File.Exists(Config.ProvisionDirectory + Path.GetFileName(Provision)))
 				{
diff --git a/Engine/Source/Programs/IOS/iPhonePackager/iPhonePackager.cs b/Engine/Source/Programs/IOS/iPhonePackager/iPhonePackager.cs
--- a/Engine/Source/Programs/IOS/iPhonePackager/iPhonePackager.cs
+++ b/Engine/Source/Programs/IOS/iPhonePackager/iPhonePackager.cs
@@ -205,12 +205,12 @@ namespace iPhonePackager
 								return false;
 							}
 						}
-						else if (Arg == "-projectdir")
+						else if (Arg == "-project")
 						{
 							// make sure there's at least one more arg
 							if (Arguments.Length > ArgIndex + 1)
 							{
-								Config.ProjectRootDirectory = Arguments[++ArgIndex];
+								Config.ProjectFile = Arguments[++ArgIndex];
 							}
 							else
 							{
@@ -474,7 +474,7 @@ namespace iPhonePackager
 					Log("");
 					Log("Configuration switches:");
 					Log("	 -stagedir <path>		  sets the directory to copy staged files from (defaults to none)");
-					Log("	 -projectdir <path>		  sets the root directory of the project");
+					Log("	 -project <path>		  path to the project being packaged");
 					Log("	 -compress=fast|best|none  packaging compression level (defaults to none)");
 					Log("	 -strip					strip symbols during packaging");
 					Log("	 -config				   game configuration (e.g., Shipping, Development, etc...)");
diff --git a/Engine/Source/Programs/UnrealBuildTool/Configuration/UEBuildTarget.cs b/Engine/Source/Programs/UnrealBuildTool/Configuration/UEBuildTarget.cs
--- a/Engine/Source/Programs/UnrealBuildTool/Configuration/UEBuildTarget.cs
+++ b/Engine/Source/Programs/UnrealBuildTool/Configuration/UEBuildTarget.cs
@@ -1439,6 +1439,13 @@ namespace UnrealBuildTool
 					GlobalCompileEnvironment.Config.Definitions.Add("UE_ENGINE_DIRECTORY=" + EnginePath);
 				}
 
+				// Set the define for the project name. This allows the executable to locate the correct project file to use, which may not be the same as the game name or target.
+				if(UnrealBuildTool.HasUProjectFile())
+				{
+					string ProjectName = Path.GetFileNameWithoutExtension(UnrealBuildTool.GetUProjectPath());
+					GlobalCompileEnvironment.Config.Definitions.Add(String.Format("UE_PROJECT_NAME={0}", ProjectName));
+				}
+
 				// Generate static libraries for monolithic games in Rocket
 				if ((UnrealBuildTool.BuildingRocket() || UnrealBuildTool.RunningRocket()) && TargetRules.IsAGame(Rules.Type))
 				{
diff --git a/Engine/Source/Runtime/Core/Private/Apple/ApplePlatformFile.cpp b/Engine/Source/Runtime/Core/Private/Apple/ApplePlatformFile.cpp
--- a/Engine/Source/Runtime/Core/Private/Apple/ApplePlatformFile.cpp
+++ b/Engine/Source/Runtime/Core/Private/Apple/ApplePlatformFile.cpp
@@ -5,6 +5,7 @@
 =============================================================================*/
 
 #include "CorePrivate.h"
+#include <sys/stat.h>
 
 // make an FTimeSpan object that represents the "epoch" for time_t (from a stat struct)
 const FDateTime MacEpoch(1970, 1, 1);
@@ -503,7 +504,19 @@ bool FApplePlatformFile::IterateDirectory(const TCHAR* Directory, FDirectoryVisi
 				{
                     // Normalize any unicode forms so we match correctly
                     FString NormalizedFilename = UTF8_TO_TCHAR(([[[NSString stringWithUTF8String:Entry->d_name] precomposedStringWithCanonicalMapping] cStringUsingEncoding:NSUTF8StringEncoding]));
-                    Result = Visitor.Visit(*(FString(Directory) / NormalizedFilename), Entry->d_type == DT_DIR);
+					
+                    // Figure out whether it's a directory. Some protocols (like NFS) do not voluntairily return this as part of the directory entry, and need to be queried manually.
+                    bool bIsDirectory = (Entry->d_type == DT_DIR);
+                    if(Entry->d_type == DT_UNKNOWN)
+                    {
+                        struct stat StatInfo;
+                        if(stat(TCHAR_TO_UTF8(*(FString(Directory) / Entry->d_name)), &StatInfo) == 0)
+                        {
+                            bIsDirectory = S_ISDIR(StatInfo.st_mode);
+                        }
+                    }
+					
+                    Result = Visitor.Visit(*(FString(Directory) / NormalizedFilename), bIsDirectory);
 				}
 			}
 			closedir(Handle);
diff --git a/Engine/Source/Runtime/Core/Private/Internationalization/Internationalization.cpp b/Engine/Source/Runtime/Core/Private/Internationalization/Internationalization.cpp
--- a/Engine/Source/Runtime/Core/Private/Internationalization/Internationalization.cpp
+++ b/Engine/Source/Runtime/Core/Private/Internationalization/Internationalization.cpp
@@ -272,7 +272,7 @@ void FInternationalization::Initialize()
 	}
 	if( !FileManager.DirectoryExists(*DataDirectory) )
 	{
-		FString DataDirectory = FPaths::EngineContentDir() / DataDirectoryRelativeToContent; // Try engine content directory.
+		DataDirectory = FPaths::EngineContentDir() / DataDirectoryRelativeToContent; // Try engine content directory.
 	}
 	if( !FileManager.DirectoryExists(*DataDirectory) )
 	{
diff --git a/Engine/Source/Runtime/Core/Public/Modules/ModuleManager.h b/Engine/Source/Runtime/Core/Public/Modules/ModuleManager.h
--- a/Engine/Source/Runtime/Core/Public/Modules/ModuleManager.h
+++ b/Engine/Source/Runtime/Core/Public/Modules/ModuleManager.h
@@ -991,13 +991,14 @@ class FDefaultGameModuleImpl
 #else
 
 /** IMPLEMENT_PRIMARY_GAME_MODULE must be used for at least one game module in your game.  It sets the "name"
-    your game when compiling in monolithic mode. */
+    your game when compiling in monolithic mode. This is passed in by UBT from the .uproject name, and manually specifying a 
+	name is no longer necessary. */
 #if IS_MONOLITHIC
 	#if PLATFORM_DESKTOP
 
-		#define IMPLEMENT_PRIMARY_GAME_MODULE( ModuleImplClass, ModuleName, GameName ) \
+		#define IMPLEMENT_PRIMARY_GAME_MODULE( ModuleImplClass, ModuleName, DEPRECATED_GameName ) \
 			/* For monolithic builds, we must statically define the game's name string (See Core.h) */ \
-			TCHAR GGameName[64] = TEXT( GameName ); \
+			TCHAR GGameName[64] = TEXT( PREPROCESSOR_TO_STRING(UE_PROJECT_NAME) ); \
 			/* Implement the GIsGameAgnosticExe variable (See Core.h). */ \
 			bool GIsGameAgnosticExe = false; \
 			IMPLEMENT_DEBUGGAME() \
@@ -1012,9 +1013,9 @@ class FDefaultGameModuleImpl
 
 	#else	//PLATFORM_DESKTOP
 
-		#define IMPLEMENT_PRIMARY_GAME_MODULE( ModuleImplClass, ModuleName, GameName ) \
+		#define IMPLEMENT_PRIMARY_GAME_MODULE( ModuleImplClass, ModuleName, DEPRECATED_GameName ) \
 			/* For monolithic builds, we must statically define the game's name string (See Core.h) */ \
-			TCHAR GGameName[64] = TEXT( GameName ); \
+			TCHAR GGameName[64] = TEXT( PREPROCESSOR_TO_STRING(UE_PROJECT_NAME) ); \
 			IMPLEMENT_DEBUGGAME() \
 			PER_MODULE_BOILERPLATE \
 			IMPLEMENT_FOREIGN_ENGINE_DIR() \
diff --git a/Engine/Source/Runtime/Engine/Classes/Engine/Texture.h b/Engine/Source/Runtime/Engine/Classes/Engine/Texture.h
--- a/Engine/Source/Runtime/Engine/Classes/Engine/Texture.h
+++ b/Engine/Source/Runtime/Engine/Classes/Engine/Texture.h
@@ -593,7 +593,7 @@ public:
 	/**
 	 * Begins caching platform data in the background for the platform requested
 	 */
-	virtual void BeginCacheForCookedPlatformData(  const ITargetPlatform *TargetPlatform ) override;
+	ENGINE_API virtual void BeginCacheForCookedPlatformData(  const ITargetPlatform *TargetPlatform ) override;
 
 
 	/**
diff --git a/Engine/Source/Runtime/Engine/Private/UnrealEngine.cpp b/Engine/Source/Runtime/Engine/Private/UnrealEngine.cpp
--- a/Engine/Source/Runtime/Engine/Private/UnrealEngine.cpp
+++ b/Engine/Source/Runtime/Engine/Private/UnrealEngine.cpp
@@ -1440,6 +1440,7 @@ void UEngine::CleanupGameViewport()
 					Player->PlayerController->CleanupGameViewport();
 				}
 				Player->ViewportClient = NULL;
+				Player->PlayerRemoved();
 				Context.GamePlayers.RemoveAt(idx);
 			}
 		}
diff --git a/Engine/Source/Runtime/Launch/Private/LaunchEngineLoop.cpp b/Engine/Source/Runtime/Launch/Private/LaunchEngineLoop.cpp
--- a/Engine/Source/Runtime/Launch/Private/LaunchEngineLoop.cpp
+++ b/Engine/Source/Runtime/Launch/Private/LaunchEngineLoop.cpp
@@ -904,7 +904,6 @@ int32 FEngineLoop::PreInit( const TCHAR* CmdLine )
 	if ( !GIsGameAgnosticExe && FApp::HasGameName() && !FPaths::IsProjectFilePathSet() )
 	{
 		// If we are using a non-agnostic exe where a name was specified but we did not specify a project path. Assemble one based on the game name.
-		// @todo uproject This should not be necessary when uproject files are always specified.
 		const FString ProjectFilePath = FPaths::Combine(*FPaths::GameDir(), *FString::Printf(TEXT("%s.%s"), FApp::GetGameName(), *IProjectManager::GetProjectFileExtension()));
 		FPaths::SetProjectFilePath(ProjectFilePath);
 	}
diff --git a/Engine/Source/Runtime/Online/OnlineSubsystemUtils/Private/OnlineSessionClient.cpp b/Engine/Source/Runtime/Online/OnlineSubsystemUtils/Private/OnlineSessionClient.cpp
--- a/Engine/Source/Runtime/Online/OnlineSubsystemUtils/Private/OnlineSessionClient.cpp
+++ b/Engine/Source/Runtime/Online/OnlineSubsystemUtils/Private/OnlineSessionClient.cpp
@@ -88,6 +88,8 @@ void UOnlineSessionClient::ClearOnlineDelegates(UWorld* InWorld)
 			OnlineSessionInt->ClearOnSessionInviteAcceptedDelegate(ControllerId, OnSessionInviteAcceptedDelegate);
 		}
 	}
+
+	SessionInt = NULL;
 }
 
 /**
