#UE4 #HTML5

       Fix mouse wheel events for the browser platform - note : its broken on win32 simulator, we need to upgrade SDL on win32. StrategyGame now correctly zooms.

diff --git a/Engine/Source/Runtime/Core/Private/HTML5/HTML5InputInterface.cpp b/Engine/Source/Runtime/Core/Private/HTML5/HTML5InputInterface.cpp
--- a/Engine/Source/Runtime/Core/Private/HTML5/HTML5InputInterface.cpp
+++ b/Engine/Source/Runtime/Core/Private/HTML5/HTML5InputInterface.cpp
@@ -94,7 +94,15 @@ void FHTML5InputInterface::Tick(float DeltaTime, const SDL_Event& Event)
 			Cursor->SetPosition(Event.motion.x, Event.motion.y);
 			MessageHandler->OnRawMouseMove(Event.motion.xrel, -Event.motion.yrel);
 			MessageHandler->OnMouseMove(); 
+		} 
+#if PLATFORM_HTML5_BROWSER 
+		else if ( Event.type == SDL_MOUSEWHEEL )
+		{
+			SDL_MouseWheelEvent* w = (SDL_MouseWheelEvent*)&Event;
+			const float SpinFactor = 1 / 120.0f;
+			MessageHandler->OnMouseWheel(w->y * SpinFactor);
 		}
+#endif 
 #if PLATFORM_HTML5_WIN32 && SDL_MAJOR_VERSION < 2
 		else if (Event.type == SDL_ACTIVEEVENT)
 		{
