Refactored ICU initialization calls to provide a better diagnostic on failure.

diff --git a/Engine/Source/Runtime/Core/Private/Internationalization/Internationalization.cpp b/Engine/Source/Runtime/Core/Private/Internationalization/Internationalization.cpp
--- a/Engine/Source/Runtime/Core/Private/Internationalization/Internationalization.cpp
+++ b/Engine/Source/Runtime/Core/Private/Internationalization/Internationalization.cpp
@@ -262,24 +262,40 @@ void FInternationalization::Initialize()
 
 	u_setMemoryFunctions(NULL, &(FICUOverrides::Malloc), &(FICUOverrides::Realloc), &(FICUOverrides::Free), &(ICUStatus));
 
-	FString DataDirectory;
 	const FString DataDirectoryRelativeToContent = FString(TEXT("Localization")) / FString(TEXT("ICU"));
 	IFileManager& FileManager = IFileManager::Get();
-	DataDirectory = FPaths::GameContentDir() / DataDirectoryRelativeToContent; // Try game content directory.
-	if( !FileManager.DirectoryExists(*DataDirectory) )
-	{
-		DataDirectory = FString(FPlatformProcess::BaseDir()) / FPaths::GameContentDir() / DataDirectoryRelativeToContent; // Try game content directory with appropriate base.
-	}
-	if( !FileManager.DirectoryExists(*DataDirectory) )
+
+	const FString PotentialDataDirectories[] = 
+		{
+			FPaths::GameContentDir() / DataDirectoryRelativeToContent, // Try game content directory.
+			FPaths::EngineContentDir() / DataDirectoryRelativeToContent, // Try engine content directory.
+		};
+
+	bool HasFoundDataDirectory = false;
+	for(const auto& DataDirectoryString : PotentialDataDirectories)
 	{
-		DataDirectory = FPaths::EngineContentDir() / DataDirectoryRelativeToContent; // Try engine content directory.
+		if( FileManager.DirectoryExists( *DataDirectoryString ) )
+		{
+			u_setDataDirectory( StringCast<char>( *DataDirectoryString ).Get() );
+			HasFoundDataDirectory = true;
+			break;
+		}
 	}
-	if( !FileManager.DirectoryExists(*DataDirectory) )
+
+	auto GetPrioritizedDataDirectoriesString = [&]() -> FString
 	{
-		DataDirectory = FString(FPlatformProcess::BaseDir()) / FPaths::EngineContentDir() / DataDirectoryRelativeToContent; // Try engine content directory with appropriate base.
-	}
-	check( FileManager.DirectoryExists(*DataDirectory) );
-	u_setDataDirectory(StringCast<char>(*DataDirectory).Get());
+		FString Result;
+		for(const auto& DataDirectoryString : PotentialDataDirectories)
+		{
+			if(!Result.IsEmpty())
+			{
+				Result += TEXT("\n");
+			}
+			Result += DataDirectoryString;
+		}
+		return Result;
+	};
+	checkf( HasFoundDataDirectory, TEXT("ICU data directory was not discovered:\n%s"), *(GetPrioritizedDataDirectoriesString()) );
 
 	u_setDataFileFunctions(nullptr, &OpenDataFile, &CloseDataFile, &(ICUStatus));
 	u_init(&(ICUStatus));
