Custom splash screen now displays correctly

#ttp 341061 - LIVE: Custom splash screen is not displayed when opening the project.

#branch UE4

#change Make GetSplashPath return absolute paths, instead of relative ones (as the issue appears to be if UE4 is not installed in the primary drive, then it canÆt make a relative path from the userdir to the UE4 dir)

diff --git a/Engine/Source/Runtime/Core/Private/Mac/MacPlatformSplash.cpp b/Engine/Source/Runtime/Core/Private/Mac/MacPlatformSplash.cpp
--- a/Engine/Source/Runtime/Core/Private/Mac/MacPlatformSplash.cpp
+++ b/Engine/Source/Runtime/Core/Private/Mac/MacPlatformSplash.cpp
@@ -57,7 +57,7 @@ static bool GetSplashPath(const TCHAR* SplashFilename, FString& OutPath, bool& O
 	}
 	
 	// next look in Engine/Splash
-	OutPath = FPaths::EngineContentDir() + TEXT("Splash/") + SplashFilename;
+	OutPath = FPaths::ConvertRelativePathToFull(FPaths::EngineContentDir() + TEXT("Splash/") + SplashFilename);
 	OutIsCustom = false;
 	
 	// if this was found, then we're done
@@ -248,7 +248,7 @@ void FMacPlatformSplash::Show()
 			}
 		}
 
-		NSString *SplashScreenFileName = (NSString *)FPlatformString::TCHARToCFString(*(FString(FPlatformProcess::BaseDir()) / SplashPath));
+		NSString *SplashScreenFileName = (NSString *)FPlatformString::TCHARToCFString(*SplashPath);
 		GSplashScreenImage = [[NSImage alloc] initWithContentsOfFile: SplashScreenFileName];
 		[SplashScreenFileName release];
 
diff --git a/Engine/Source/Runtime/Core/Private/Windows/WindowsPlatformSplash.cpp b/Engine/Source/Runtime/Core/Private/Windows/WindowsPlatformSplash.cpp
--- a/Engine/Source/Runtime/Core/Private/Windows/WindowsPlatformSplash.cpp
+++ b/Engine/Source/Runtime/Core/Private/Windows/WindowsPlatformSplash.cpp
@@ -596,7 +596,7 @@ static bool GetSplashPath(const TCHAR* SplashFilename, FString& OutPath, bool& O
 	}
 
 	// next look in Engine/Splash
-	OutPath = FPaths::EngineContentDir() + TEXT("Splash/") + SplashFilename;
+	OutPath = FPaths::ConvertRelativePathToFull(FPaths::EngineContentDir() + TEXT("Splash/") + SplashFilename);
 	OutIsCustom = false;
 
 	// if this was found, then we're done
@@ -696,7 +696,7 @@ void FWindowsPlatformSplash::Show()
 				ShowWindow(GSplashScreenGuard, SW_SHOW); 
 			}
 			
-			GSplashScreenFileName = FString( FPlatformProcess::BaseDir() ) / SplashPath;
+			GSplashScreenFileName = SplashPath;
 			DWORD ThreadID = 0;
 			GSplashScreenThread = CreateThread(NULL, 128 * 1024, (LPTHREAD_START_ROUTINE)StartSplashScreenThread, (LPVOID)NULL, STACK_SIZE_PARAM_IS_A_RESERVATION, &ThreadID);
 #if	STATS
