Variadic support for TArray::Emplace, and 0-4 arguments supported for no-variadics compilers.

diff --git a/Engine/Source/Runtime/Core/Public/Containers/Array.h b/Engine/Source/Runtime/Core/Public/Containers/Array.h
--- a/Engine/Source/Runtime/Core/Public/Containers/Array.h
+++ b/Engine/Source/Runtime/Core/Public/Containers/Array.h
@@ -1179,13 +1179,58 @@ public:
 	 * @param Item	The item to add
 	 * @return		Index to the new item
 	 */
-	template <typename ArgsType>
-	int32 Emplace( ArgsType&& Args )
-	{
-		const int32 Index = AddUninitialized(1);
-		new(GetTypedData() + Index) ElementType(Forward<ArgsType>(Args));
-		return Index;
-	}
+	#if PLATFORM_COMPILER_HAS_VARIADIC_TEMPLATES
+
+		template <typename... ArgsType>
+		int32 Emplace(ArgsType&&... Args)
+		{
+			const int32 Index = AddUninitialized(1);
+			new(GetTypedData() + Index) ElementType(Forward<ArgsType>(Args)...);
+			return Index;
+		}
+
+	#else
+
+		int32 Emplace()
+		{
+			const int32 Index = AddUninitialized(1);
+			new(GetTypedData() + Index) ElementType();
+			return Index;
+		}
+
+		template <typename Arg0Type>
+		int32 Emplace(Arg0Type&& Arg0)
+		{
+			const int32 Index = AddUninitialized(1);
+			new(GetTypedData() + Index) ElementType(Forward<Arg0Type>(Arg0));
+			return Index;
+		}
+
+		template <typename Arg0Type, typename Arg1Type>
+		int32 Emplace(Arg0Type&& Arg0, Arg1Type&& Arg1)
+		{
+			const int32 Index = AddUninitialized(1);
+			new(GetTypedData() + Index) ElementType(Forward<Arg0Type>(Arg0), Forward<Arg1Type>(Arg1));
+			return Index;
+		}
+
+		template <typename Arg0Type, typename Arg1Type, typename Arg2Type>
+		int32 Emplace(Arg0Type&& Arg0, Arg1Type&& Arg1, Arg2Type&& Arg2)
+		{
+			const int32 Index = AddUninitialized(1);
+			new(GetTypedData() + Index) ElementType(Forward<Arg0Type>(Arg0), Forward<Arg1Type>(Arg1), Forward<Arg2Type>(Arg2));
+			return Index;
+		}
+
+		template <typename Arg0Type, typename Arg1Type, typename Arg2Type, typename Arg3Type>
+		int32 Emplace(Arg0Type&& Arg0, Arg1Type&& Arg1, Arg2Type&& Arg2, Arg3Type&& Arg3)
+		{
+			const int32 Index = AddUninitialized(1);
+			new(GetTypedData() + Index) ElementType(Forward<Arg0Type>(Arg0), Forward<Arg1Type>(Arg1), Forward<Arg2Type>(Arg2), Forward<Arg3Type>(Arg3));
+			return Index;
+		}
+
+	#endif
 
 	/**
 	 * Adds a new item to the end of the array, possibly reallocating the whole array to fit.
