Changed Mac window's redrawContents so it only redraws on resize.

diff --git a/Engine/Source/Runtime/Core/Private/Mac/MacWindow.cpp b/Engine/Source/Runtime/Core/Private/Mac/MacWindow.cpp
--- a/Engine/Source/Runtime/Core/Private/Mac/MacWindow.cpp
+++ b/Engine/Source/Runtime/Core/Private/Mac/MacWindow.cpp
@@ -146,10 +146,11 @@ TArray< FSlateCocoaWindow* > GRunningModalWindows;
 
 - (void)redrawContents
 {
-	if(self.bForwardEvents && ([self isVisible] && [super alphaValue] > 0.0f))
+	if(bNeedsRedraw && bForwardEvents && ([self isVisible] && [super alphaValue] > 0.0f))
 	{
 		MacApplication->OnWindowRedrawContents( self );
 	}
+	bNeedsRedraw = false;
 }
 
 - (void)setWindowMode:(EWindowMode::Type)NewWindowMode
@@ -443,6 +444,7 @@ TArray< FSlateCocoaWindow* > GRunningModalWindows;
 	{
 		MacApplication->OnWindowDidResize( self );
 	}
+	bNeedsRedraw = true;
 }
 
 - (void)windowWillClose:(NSNotification *)notification
diff --git a/Engine/Source/Runtime/Core/Public/Mac/MacWindow.h b/Engine/Source/Runtime/Core/Public/Mac/MacWindow.h
--- a/Engine/Source/Runtime/Core/Public/Mac/MacWindow.h
+++ b/Engine/Source/Runtime/Core/Public/Mac/MacWindow.h
@@ -22,6 +22,7 @@
 	bool bDeferSetFrame;
 	bool bDeferSetOrigin;
 	bool bRenderInitialised;
+	bool bNeedsRedraw;
 @public
 	bool bZoomed;
 }
