CrashReport - KernelBase.dll should no longer be parsed as the assert message

diff --git a/Engine/Source/Runtime/Core/Private/Windows/WindowsPlatformExceptionHandling.cpp b/Engine/Source/Runtime/Core/Private/Windows/WindowsPlatformExceptionHandling.cpp
--- a/Engine/Source/Runtime/Core/Private/Windows/WindowsPlatformExceptionHandling.cpp
+++ b/Engine/Source/Runtime/Core/Private/Windows/WindowsPlatformExceptionHandling.cpp
@@ -193,6 +193,10 @@ void SetReportParameters( HREPORT ReportHandle, EXCEPTION_POINTERS* ExceptionInf
 			}
 		}
 	}
+	else
+	{
+		FCString::Strcpy( LocalBuffer, TEXT("") );
+	}
 
 	StringCchPrintf( StringBuffer, LocalBufferSize, TEXT( "!%s!AssertLog=\"%s\"" ), FCommandLine::Get(), LocalBuffer );
 	Result = WerReportSetParameter( ReportHandle, WER_P8, TEXT( "Commandline" ), StringBuffer );
