Added support for overriding config options from the commandline.
E.g. -ini:Engine:/Script/Engine.Engine.bSmoothFrameRate=False,TextureStreaming.PoolSize=100

Unshelved and fixed up Joshs UE3 port from the shelved CL# 1760099
- Brought over the commandline .ini override from UE3

Note, that this does not serialize the new value to the saved config. It is simply an override for that run.

#Note, improves and fixes functionality from CL# 2211571.

diff --git a/Engine/Source/Runtime/Core/Private/Misc/ConfigCacheIni.cpp b/Engine/Source/Runtime/Core/Private/Misc/ConfigCacheIni.cpp
--- a/Engine/Source/Runtime/Core/Private/Misc/ConfigCacheIni.cpp
+++ b/Engine/Source/Runtime/Core/Private/Misc/ConfigCacheIni.cpp
@@ -570,8 +570,6 @@ static void OverrideFromCommandline(FConfigFile* File, const FString& Filename)
 	// for example:
 	//		-ini:Engine:/Script/Engine.Engine.bSmoothFrameRate=False,TextureStreaming.PoolSize=100
 	//			(will update the cache after the final combined engine.ini)
-	//		-ini:DefaultEngine:/Script/Engine.Engine.bSmoothFrameRate=False,TextureStreaming.PoolSize=100
-	//			(will update the cache after the default engine.ini and before applying platform inis)
 	if (FParse::Value(FCommandLine::Get(), *FString::Printf(TEXT("-ini:%s:"), *FPaths::GetBaseFilename(Filename)), Settings, false))
 	{
 		// break apart on the commas
@@ -593,12 +591,15 @@ static void OverrideFromCommandline(FConfigFile* File, const FString& Filename)
 					continue;
 				}
 
-				// break it apart
-				FString Section = SectionAndKey.Left(LastDot);
-				FString Key = SectionAndKey.Mid(LastDot + 1);
+				// Create the commandline override object
+				FConfigCommandlineOverride& CommandlineOption = File->CommandlineOptions[File->CommandlineOptions.Emplace()];
+				CommandlineOption.BaseFileName = *FPaths::GetBaseFilename(Filename);
+				CommandlineOption.Section = SectionAndKey.Left(LastDot);
+				CommandlineOption.PropertyKey = SectionAndKey.Mid(LastDot + 1);
+				CommandlineOption.PropertyValue = Value;
 
-				// now put it into this .ini
-				File->SetString(*Section, *Key, *Value);
+				// now put it into this into the cache
+				File->SetString(*CommandlineOption.Section, *CommandlineOption.PropertyKey, *CommandlineOption.PropertyValue);
 			}
 		}
 	}
@@ -672,11 +673,6 @@ static bool LoadIniFileHierarchy(const TArray<FIniFilename>& HierarchyToLoad, FC
 		bool bDoWrite = false;
 		//UE_LOG(LogConfig, Log,  TEXT( "Combining configFile: %s" ), *IniList(IniIndex) );
 		ProcessIniContents(*HierarchyToLoad.Last().Filename, *IniFileName, &ConfigFile, bDoEmptyConfig, bDoCombine, bDoWrite);
-
-#if !UE_BUILD_SHIPPING
-		// process any commandline overrides
-		OverrideFromCommandline(&ConfigFile, IniFileName);
-#endif
 	}
 
 	// Set this configs files source ini hierarchy to show where it was loaded from.
@@ -747,6 +743,37 @@ bool DoesConfigPropertyValueMatch( FConfigFile* InConfigFile, const FString& InS
 }
 
 
+/**
+ * Check if the provided property information was set as a commandline override
+ *
+ * @param InConfigFile		- The config file which we want to check had overridden values
+ * @param InSectionName		- The name of the section which we are checking for a match
+ * @param InPropertyName		- The name of the property which we are checking for a match
+ * @param InPropertyValue	- The value of the property which we are checking for a match
+ *
+ * @return True if a commandline option was set that matches the input parameters
+ */
+bool PropertySetFromCommandlineOption(const FConfigFile* InConfigFile, const FString& InSectionName, const FName& InPropertyName, const FString& InPropertyValue)
+{
+	bool bFromCommandline = false;
+
+#if !UE_BUILD_SHIPPING
+	for (const FConfigCommandlineOverride& CommandlineOverride : InConfigFile->CommandlineOptions)
+	{
+		if (CommandlineOverride.PropertyKey.Equals(InPropertyName.ToString(), ESearchCase::IgnoreCase) &&
+			CommandlineOverride.PropertyValue.Equals(InPropertyValue, ESearchCase::IgnoreCase) &&
+			CommandlineOverride.Section.Equals(InSectionName, ESearchCase::IgnoreCase) &&
+			CommandlineOverride.BaseFileName.Equals(FPaths::GetBaseFilename(InConfigFile->Name.ToString()), ESearchCase::IgnoreCase))
+		{
+			bFromCommandline = true;
+		}
+	}
+#endif // !UE_BUILD_SHIPPING
+
+	return bFromCommandline;
+}
+
+
 bool FConfigFile::Write( const FString& Filename, bool bDoRemoteWrite/* = true*/, const FString& InitialText/*=FString()*/ )
 {
 	if( !Dirty || NoSave || FParse::Param( FCommandLine::Get(), TEXT("nowrite")) || 
@@ -792,8 +819,11 @@ bool FConfigFile::Write( const FString& Filename, bool bDoRemoteWrite/* = true*/
 					}
 				}
 
+				// check whether the option we are attempting to write out, came from the commandline as a temporary override.
+				const bool bOptionIsFromCommandline = PropertySetFromCommandlineOption(this, SectionName, PropertyName, PropertyValue);
+
 				// Check if the property matches the source configs. We do not wanna write it out if so.
-				if( bDifferentNumberOfElements || !DoesConfigPropertyValueMatch( SourceConfigFile, SectionName, PropertyName, PropertyValue ) )
+				if( (bDifferentNumberOfElements || !DoesConfigPropertyValueMatch( SourceConfigFile, SectionName, PropertyName, PropertyValue )) && !bOptionIsFromCommandline )
 				{
 					// If this is the first property we are writing of this section, then print the section name
 					if( !bWroteASectionProperty )
@@ -2409,6 +2439,11 @@ static bool GenerateDestIniFile(FConfigFile& DestConfigFile, const FString& Dest
 	}
 	LoadAnIniFile(DestIniFilename, DestConfigFile);
 
+#if !UE_BUILD_SHIPPING
+	// process any commandline overrides
+	OverrideFromCommandline(&DestConfigFile, DestIniFilename);
+#endif
+
 	bool bForceRegenerate = false;
 	bool bShouldUpdate = FPlatformProperties::RequiresCookedData();
 
diff --git a/Engine/Source/Runtime/Core/Public/Misc/ConfigCacheIni.h b/Engine/Source/Runtime/Core/Public/Misc/ConfigCacheIni.h
--- a/Engine/Source/Runtime/Core/Public/Misc/ConfigCacheIni.h
+++ b/Engine/Source/Runtime/Core/Public/Misc/ConfigCacheIni.h
@@ -42,6 +42,16 @@ struct FIniFilename
 	{}
 };
 
+
+#if !UE_BUILD_SHIPPING
+// Options which stemmed from the commandline
+struct FConfigCommandlineOverride
+{
+	FString BaseFileName, Section, PropertyKey, PropertyValue;
+};
+#endif // !UE_BUILD_SHIPPING
+
+
 // One config file.
 class FConfigFile : public TMap<FString,FConfigSection>
 {
@@ -56,6 +66,9 @@ public:
 
 	/** The untainted config file which contains the coalesced base/default options. I.e. No Saved/ options*/
 	FConfigFile* SourceConfigFile;
+
+	/** The collection of overrides which stemmed from the commandline */
+	TArray<FConfigCommandlineOverride> CommandlineOptions;
 	
 	CORE_API FConfigFile();
 	FConfigFile( int32 ) {}	// @todo UE4 DLL: Workaround for instantiated TMap template during DLLExport (TMap::FindRef)
