Fix for builder, missing UE_BUILD_SHIPPING check

diff --git a/Engine/Source/Runtime/Core/Public/Misc/ConfigCacheIni.h b/Engine/Source/Runtime/Core/Public/Misc/ConfigCacheIni.h
--- a/Engine/Source/Runtime/Core/Public/Misc/ConfigCacheIni.h
+++ b/Engine/Source/Runtime/Core/Public/Misc/ConfigCacheIni.h
@@ -67,8 +67,10 @@ public:
 	/** The untainted config file which contains the coalesced base/default options. I.e. No Saved/ options*/
 	FConfigFile* SourceConfigFile;
 
+#if !UE_BUILD_SHIPPING
 	/** The collection of overrides which stemmed from the commandline */
 	TArray<FConfigCommandlineOverride> CommandlineOptions;
+#endif // !UE_BUILD_SHIPPING
 	
 	CORE_API FConfigFile();
 	FConfigFile( int32 ) {}	// @todo UE4 DLL: Workaround for instantiated TMap template during DLLExport (TMap::FindRef)
