Compile Times: Fixed Core module needing two separate PCH headers in monolithic compiles
- Changed Core module to use Core.h (removed CorePrivate.h)
- CorePrivate.h only included Core.h, so their PCHs were identical other than name
- Does not affect modular builds, as we need different defines for DLLExport/DLLImport for private and public PCH headers

diff --git a/Engine/Source/Runtime/Core/Private/Android/AndroidApplication.cpp b/Engine/Source/Runtime/Core/Private/Android/AndroidApplication.cpp
--- a/Engine/Source/Runtime/Core/Private/Android/AndroidApplication.cpp
+++ b/Engine/Source/Runtime/Core/Private/Android/AndroidApplication.cpp
@@ -1,6 +1,6 @@
 // Copyright 1998-2014 Epic Games, Inc. All Rights Reserved.
 
-#include "CorePrivate.h"
+#include "Core.h"
 #include "AndroidApplication.h"
 #include "AndroidInputInterface.h"
 #include "AndroidWindow.h"
diff --git a/Engine/Source/Runtime/Core/Private/Android/AndroidFile.cpp b/Engine/Source/Runtime/Core/Private/Android/AndroidFile.cpp
--- a/Engine/Source/Runtime/Core/Private/Android/AndroidFile.cpp
+++ b/Engine/Source/Runtime/Core/Private/Android/AndroidFile.cpp
@@ -4,7 +4,7 @@
 	AndroidFile.cpp: Android platform implementations of File functions
 =============================================================================*/
 
-#include "CorePrivate.h"
+#include "Core.h"
 
 #include <dirent.h>
 #include <jni.h>
