Ranged-based iterator support for TLinkedList or TDoubleLinkedList.

diff --git a/Engine/Source/Runtime/Core/Public/Containers/List.h b/Engine/Source/Runtime/Core/Public/Containers/List.h
--- a/Engine/Source/Runtime/Core/Public/Containers/List.h
+++ b/Engine/Source/Runtime/Core/Public/Containers/List.h
@@ -2,85 +2,89 @@
 
 #pragma once
 
-
-/**
- * Encapsulates a link in a single linked list with constant access time.
- */
-template<class ElementType> class TLinkedList
+template <class ContainerType, class ElementType>
+class TLinkedListIterator
 {
 public:
+	explicit TLinkedListIterator(ContainerType* FirstLink)
+		: CurrentLink(FirstLink)
+	{
+	}
 
 	/**
-	 * Used to iterate over the elements of a linked list.
-	 */
-	class TIterator
+		* Advances the iterator to the next element.
+		*/
+	void Next()
 	{
-	public:
-		TIterator(TLinkedList* FirstLink)
-		:	CurrentLink(FirstLink)
-		{}
-
-		/**
-		 * Advances the iterator to the next element.
-		 */
-		void Next()
-		{
-			checkSlow(CurrentLink);
-			CurrentLink = CurrentLink->NextLink;
-		}
+		checkSlow(CurrentLink);
+		CurrentLink = CurrentLink->GetNextLink();
+	}
 
-		TIterator& operator++()
-		{
-			Next();
-			return *this;
-		}
+	TLinkedListIterator& operator++()
+	{
+		Next();
+		return *this;
+	}
 
-		TIterator operator++(int)
-		{
-			TIterator Tmp(*this);
-			Next();
-			return Tmp;
-		}
+	TLinkedListIterator operator++(int)
+	{
+		auto Tmp = *this;
+		Next();
+		return Tmp;
+	}
 
-		SAFE_BOOL_OPERATORS(TIterator)
+	/** conversion to "bool" returning true if the iterator is valid. */
+	FORCEINLINE_EXPLICIT_OPERATOR_BOOL() const
+	{ 
+		return CurrentLink != nullptr;
+	}
 
-		/** conversion to "bool" returning true if the iterator is valid. */
-		FORCEINLINE_EXPLICIT_OPERATOR_BOOL() const
-		{ 
-			return CurrentLink != NULL; 
-		}
-		/** inverse of the "bool" operator */
-		FORCEINLINE bool operator !() const 
-		{
-			return !(bool)*this;
-		}
+	// Accessors.
+	ElementType& operator->() const
+	{
+		checkSlow(CurrentLink);
+		return **CurrentLink;
+	}
 
-		// Accessors.
-		ElementType& operator->() const
-		{
-			checkSlow(CurrentLink);
-			return CurrentLink->Element;
-		}
-		ElementType& operator*() const
-		{
-			checkSlow(CurrentLink);
-			return CurrentLink->Element;
-		}
+	ElementType& operator*() const
+	{
+		checkSlow(CurrentLink);
+		return **CurrentLink;
+	}
 
-	private:
-		TLinkedList* CurrentLink;
-	};
+private:
+	ContainerType* CurrentLink;
+
+	friend bool operator==(const TLinkedListIterator& Lhs, const TLinkedListIterator& Rhs) { return Lhs.CurrentLink == Rhs.CurrentLink; }
+	friend bool operator!=(const TLinkedListIterator& Lhs, const TLinkedListIterator& Rhs) { return Lhs.CurrentLink != Rhs.CurrentLink; }
+};
+
+/**
+ * Encapsulates a link in a single linked list with constant access time.
+ */
+template <class ElementType>
+class TLinkedList
+{
+public:
+	/**
+	 * Used to iterate over the elements of a linked list.
+	 */
+	typedef TLinkedListIterator<TLinkedList,       ElementType> TIterator;
+	typedef TLinkedListIterator<TLinkedList, const ElementType> TConstIterator;
 
 	// Constructors.
-	TLinkedList():
-		NextLink(NULL),
-		PrevLink(NULL)
-	{}
-	TLinkedList(const ElementType& InElement):
-		Element(InElement),
-		NextLink(NULL),
-		PrevLink(NULL)
-	{}
+	TLinkedList()
+		: NextLink(nullptr)
+		, PrevLink(nullptr)
+	{
+	}
+
+	explicit TLinkedList(const ElementType& InElement)
+		: Element (InElement)
+		, NextLink(nullptr)
+		, PrevLink(nullptr)
+	{
+	}
 
 	/**
 	 * Removes this element from the list in constant time.
@@ -128,6 +132,16 @@ public:
 		return PrevLink != NULL;
 	}
 
+	TLinkedList* GetPrevLink() const
+	{
+		return PrevLink;
+	}
+
+	TLinkedList* GetNextLink() const
+	{
+		return NextLink;
+	}
+
 	// Accessors.
 	ElementType* operator->()
 	{
@@ -149,26 +163,99 @@ public:
 	{
 		return NextLink;
 	}
+
 private:
-	ElementType Element;
-	TLinkedList* NextLink;
+	ElementType   Element;
+	TLinkedList*  NextLink;
 	TLinkedList** PrevLink;
+
+	friend TIterator      begin(      TLinkedList& List) { return TIterator     (&List); }
+	friend TConstIterator begin(const TLinkedList& List) { return TConstIterator(const_cast<TLinkedList*>(&List)); }
+	friend TIterator      end  (      TLinkedList& List) { return TIterator     (nullptr); }
+	friend TConstIterator end  (const TLinkedList& List) { return TConstIterator(nullptr); }
 };
 
 
+template <class NodeType, class ElementType>
+class TDoubleLinkedListIterator
+{
+public:
+	explicit TDoubleLinkedListIterator(NodeType* StartingNode)
+		: CurrentNode(StartingNode)
+	{
+	}
+
+	/** conversion to "bool" returning true if the iterator is valid. */
+	FORCEINLINE_EXPLICIT_OPERATOR_BOOL() const
+	{ 
+		return CurrentNode != nullptr; 
+	}
+
+	TDoubleLinkedListIterator& operator++()
+	{
+		checkSlow(CurrentNode);
+		CurrentNode = CurrentNode->GetNextNode();
+		return *this;
+	}
+
+	TDoubleLinkedListIterator operator++(int)
+	{
+		auto Tmp = *this;
+		++(*this);
+		return Tmp;
+	}
+
+	TDoubleLinkedListIterator& operator--()
+	{
+		checkSlow(CurrentNode);
+		CurrentNode = CurrentNode->GetPrevNode();
+		return *this;
+	}
+
+	TDoubleLinkedListIterator operator--(int)
+	{
+		auto Tmp = *this;
+		--(*this);
+		return Tmp;
+	}
+
+	// Accessors.
+	ElementType& operator->() const
+	{
+		checkSlow(CurrentNode);
+		return CurrentNode->GetValue();
+	}
+
+	ElementType& operator*() const
+	{
+		checkSlow(CurrentNode);
+		return CurrentNode->GetValue();
+	}
+
+	NodeType* GetNode() const
+	{
+		checkSlow(CurrentNode);
+		return CurrentNode;
+	}
+
+private:
+	NodeType* CurrentNode;
+
+	friend bool operator==(const TDoubleLinkedListIterator& Lhs, const TDoubleLinkedListIterator& Rhs) { return Lhs.CurrentNode == Rhs.CurrentNode; }
+	friend bool operator!=(const TDoubleLinkedListIterator& Lhs, const TDoubleLinkedListIterator& Rhs) { return Lhs.CurrentNode != Rhs.CurrentNode; }
+};
+
 /**
  * Double linked list.
  */
-template<class ElementType> class TDoubleLinkedList
+template <class ElementType>
+class TDoubleLinkedList
 {
 public:
-
-	class TIterator;
 	class TDoubleLinkedListNode
 	{
 	public:
 		friend class TDoubleLinkedList;
-		friend class TIterator;
 
 		/** Constructor */
 		TDoubleLinkedListNode( const ElementType& InValue )
@@ -197,87 +284,24 @@ public:
 		}
 
 	protected:
-		ElementType				Value;
-		TDoubleLinkedListNode*	NextNode;
-		TDoubleLinkedListNode*	PrevNode;
+		ElementType            Value;
+		TDoubleLinkedListNode* NextNode;
+		TDoubleLinkedListNode* PrevNode;
 	};
 
 	/**
 	 * Used to iterate over the elements of a linked list.
 	 */
-	class TIterator
-	{
-	public:
-
-		TIterator(TDoubleLinkedListNode* StartingNode)
-		: CurrentNode(StartingNode)
-		{}
-
-		SAFE_BOOL_OPERATORS(TIterator)
-		/** conversion to "bool" returning true if the iterator is valid. */
-		FORCEINLINE_EXPLICIT_OPERATOR_BOOL() const
-		{ 
-			return CurrentNode != NULL; 
-		}
-		/** inverse of the "bool" operator */
-		FORCEINLINE bool operator !() const 
-		{
-			return !(bool)*this;
-		}
-
-		TIterator& operator++()
-		{
-			checkSlow(CurrentNode);
-			CurrentNode = CurrentNode->NextNode;
-			return *this;
-		}
-
-		TIterator operator++(int)
-		{
-			TIterator Tmp(*this);
-			++(*this);
-			return Tmp;
-		}
-
-		TIterator& operator--()
-		{
-			checkSlow(CurrentNode);
-			CurrentNode = CurrentNode->PrevNode;
-			return *this;
-		}
-
-		TIterator operator--(int)
-		{
-			TIterator Tmp(*this);
-			--(*this);
-			return Tmp;
-		}
-
-		// Accessors.
-		ElementType& operator->() const
-		{
-			checkSlow(CurrentNode);
-			return CurrentNode->GetValue();
-		}
-		ElementType& operator*() const
-		{
-			checkSlow(CurrentNode);
-			return CurrentNode->GetValue();
-		}
-		TDoubleLinkedListNode* GetNode() const
-		{
-			checkSlow(CurrentNode);
-			return CurrentNode;
-		}
-
-	private:
-
-		TDoubleLinkedListNode*	CurrentNode;
-	};
+	typedef TDoubleLinkedListIterator<TDoubleLinkedListNode,       ElementType> TIterator;
+	typedef TDoubleLinkedListIterator<TDoubleLinkedListNode, const ElementType> TConstIterator;
 
 	/** Constructors. */
-	TDoubleLinkedList() : HeadNode(NULL), TailNode(NULL), ListSize(0)
-	{}
+	TDoubleLinkedList()
+		: HeadNode(nullptr)
+		, TailNode(nullptr)
+		, ListSize(0)
+	{
+	}
 
 	/** Destructor */
 	virtual ~TDoubleLinkedList()
@@ -514,6 +538,11 @@ private:
 
 	TDoubleLinkedList(const TDoubleLinkedList&);
 	TDoubleLinkedList& operator=(const TDoubleLinkedList&);
+
+	friend TIterator      begin(      TDoubleLinkedList& List) { return TIterator     (List.GetHead()); }
+	friend TConstIterator begin(const TDoubleLinkedList& List) { return TConstIterator(List.GetHead()); }
+	friend TIterator      end  (      TDoubleLinkedList& List) { return TIterator     (nullptr); }
+	friend TConstIterator end  (const TDoubleLinkedList& List) { return TConstIterator(nullptr); }
 };
 
 
