Fix SteamController header dependency.

Do not include SteamController headers if module isn't being compiled.

diff --git a/Engine/Source/Runtime/Core/Public/Linux/LinuxApplication.h b/Engine/Source/Runtime/Core/Public/Linux/LinuxApplication.h
--- a/Engine/Source/Runtime/Core/Public/Linux/LinuxApplication.h
+++ b/Engine/Source/Runtime/Core/Public/Linux/LinuxApplication.h
@@ -5,7 +5,7 @@
 #include "GenericApplication.h"
 #include "LinuxWindow.h"
 
-#define STEAM_CONTROLLER_SUPPORT				(WITH_STEAMWORKS && WITH_ENGINE && !UE_SERVER)
+#define STEAM_CONTROLLER_SUPPORT				(WITH_STEAMWORKS && WITH_ENGINE && WITH_STEAMCONTROLLER)
 
 typedef SDL_GameController* SDL_HController;
 
diff --git a/Engine/Source/Runtime/Linux/SteamController/SteamController.Build.cs b/Engine/Source/Runtime/Linux/SteamController/SteamController.Build.cs
--- a/Engine/Source/Runtime/Linux/SteamController/SteamController.Build.cs
+++ b/Engine/Source/Runtime/Linux/SteamController/SteamController.Build.cs
@@ -8,6 +8,7 @@ public class SteamController : ModuleRules
     {
         PrivateIncludePathModuleNames.Add("TargetPlatform");
 
+        Definitions.Add("WITH_STEAMCONTROLLER=1");
         PrivateDependencyModuleNames.AddRange(new string[]
         {
 			"Core",
