Fixed IMM based IME crash due to starting a compostiion with no active context

https://udn.unrealengine.com/questions/213482/bug-reportcrash-on-fwindowstextinputmethodsystempr.html

It seems that we still get the IMM based IME messages even when we have no active context set (meaning that you could start composing at any time), which could lead to a crash if you switched a text-box, since you'd get an WM_IME_COMPOSITION message but miss the WM_IME_STARTCOMPOSITION message.

This change just has the composition immediately cancel if there's no active context set.

diff --git a/Engine/Source/Runtime/Core/Private/Windows/WindowsTextInputMethodSystem.cpp b/Engine/Source/Runtime/Core/Private/Windows/WindowsTextInputMethodSystem.cpp
--- a/Engine/Source/Runtime/Core/Private/Windows/WindowsTextInputMethodSystem.cpp
+++ b/Engine/Source/Runtime/Core/Private/Windows/WindowsTextInputMethodSystem.cpp
@@ -649,6 +649,13 @@ int32 FWindowsTextInputMethodSystem::ProcessMessage(HWND hwnd, uint32 msg, WPARA
 				
 				UE_LOG(LogWindowsTextInputMethodSystem, Verbose, TEXT("Beginning IMM composition."));
 			}
+			else
+			{
+				// No valid context, so don't allow this composition to start
+				HIMC IMMContext = ::ImmGetContext(hwnd);
+				::ImmNotifyIME(IMMContext, NI_COMPOSITIONSTR, CPS_CANCEL, 0);
+				::ImmReleaseContext(hwnd, IMMContext);
+			}
 
 			return DefWindowProc(hwnd, msg, wParam, lParam);
 		}
