[GitHub] 279 : FMath::RInterpConstantTo, Essential Const Fix, Can't use FApp:GetDeltaTime() without it!

TTP# 341128

diff --git a/Engine/Source/Runtime/Core/Private/Math/UnrealMath.cpp b/Engine/Source/Runtime/Core/Private/Math/UnrealMath.cpp
--- a/Engine/Source/Runtime/Core/Private/Math/UnrealMath.cpp
+++ b/Engine/Source/Runtime/Core/Private/Math/UnrealMath.cpp
@@ -1953,7 +1953,7 @@ CORE_API FVector2D FMath::Vector2DInterpTo( const FVector2D& Current, const FVec
 	return Current + DeltaMove;
 }
 
-CORE_API FRotator FMath::RInterpConstantTo( const FRotator& Current, const FRotator& Target, float& DeltaTime, float InterpSpeed )
+CORE_API FRotator FMath::RInterpConstantTo( const FRotator& Current, const FRotator& Target, float DeltaTime, float InterpSpeed )
 {
 	// if DeltaTime is 0, do not perform any interpolation (Location was already calculated for that frame)
 	if( DeltaTime == 0.f || Current == Target )
@@ -1977,7 +1977,7 @@ CORE_API FRotator FMath::RInterpConstantTo( const FRotator& Current, const FRota
 	return Result.GetNormalized();
 }
 
-CORE_API FRotator FMath::RInterpTo( const FRotator& Current, const FRotator& Target, const float& DeltaTime, float InterpSpeed)
+CORE_API FRotator FMath::RInterpTo( const FRotator& Current, const FRotator& Target, float DeltaTime, float InterpSpeed)
 {
 	// if DeltaTime is 0, do not perform any interpolation (Location was already calculated for that frame)
 	if( DeltaTime == 0.f || Current == Target )
diff --git a/Engine/Source/Runtime/Core/Public/Math/UnrealMathUtility.h b/Engine/Source/Runtime/Core/Public/Math/UnrealMathUtility.h
--- a/Engine/Source/Runtime/Core/Public/Math/UnrealMathUtility.h
+++ b/Engine/Source/Runtime/Core/Public/Math/UnrealMathUtility.h
@@ -594,10 +594,10 @@ struct FMath : public FPlatformMath
 	static CORE_API FVector2D Vector2DInterpTo( const FVector2D& Current, const FVector2D& Target, float DeltaTime, float InterpSpeed );
 
 	/** Interpolate rotator from Current to Target with constant step */
-	static CORE_API FRotator RInterpConstantTo( const FRotator& Current, const FRotator& Target, float& DeltaTime, float InterpSpeed);
+	static CORE_API FRotator RInterpConstantTo( const FRotator& Current, const FRotator& Target, float DeltaTime, float InterpSpeed);
 
 	/** Interpolate rotator from Current to Target. Scaled by distance to Target, so it has a strong start speed and ease out. */
-	static CORE_API FRotator RInterpTo( const FRotator& Current, const FRotator& Target, const float& DeltaTime, float InterpSpeed);
+	static CORE_API FRotator RInterpTo( const FRotator& Current, const FRotator& Target, float DeltaTime, float InterpSpeed);
 
 	/** Interpolate float from Current to Target with constant step */
 	static CORE_API float FInterpConstantTo( float Current, float Target, float DeltaTime, float InterpSpeed );
