[GitHub] 275 : FVector::ToDirectionAndLength not marked as const

TTP# 340773

diff --git a/Engine/Source/Runtime/Core/Public/Math/Vector.h b/Engine/Source/Runtime/Core/Public/Math/Vector.h
--- a/Engine/Source/Runtime/Core/Public/Math/Vector.h
+++ b/Engine/Source/Runtime/Core/Public/Math/Vector.h
@@ -438,7 +438,7 @@ public:
 	 * @param OutDir Reference passed in to store unit direction vector.
 	 * @param OutLength Reference passed in to store length of the vector.
 	 */
-	void ToDirectionAndLength(FVector &OutDir, float &OutLength);
+	void ToDirectionAndLength(FVector &OutDir, float &OutLength) const;
 
 	/**
 	 * Get a copy of the vector as sign only.
@@ -1349,7 +1349,7 @@ FORCEINLINE bool FVector::IsNormalized() const
 	return (FMath::Abs(1.f - SizeSquared()) < THRESH_VECTOR_NORMALIZED);
 }
 
-FORCEINLINE void FVector::ToDirectionAndLength(FVector &OutDir, float &OutLength)
+FORCEINLINE void FVector::ToDirectionAndLength(FVector &OutDir, float &OutLength) const
 {
 	OutLength = Size();
 	if (OutLength > SMALL_NUMBER)
