Improvements for the commandline ini overrides

Better parsing of CVars. the new syntax for specifying an override is:
-ini:IniName:[Section1]:Key1=Value1,[Section2]:Key2=Value2

diff --git a/Engine/Source/Runtime/Core/Private/Misc/ConfigCacheIni.cpp b/Engine/Source/Runtime/Core/Private/Misc/ConfigCacheIni.cpp
--- a/Engine/Source/Runtime/Core/Private/Misc/ConfigCacheIni.cpp
+++ b/Engine/Source/Runtime/Core/Private/Misc/ConfigCacheIni.cpp
@@ -556,6 +556,18 @@ bool FConfigFile::ShouldExportQuotedString(const FString& PropertyValue) const
 
 
 #if !UE_BUILD_SHIPPING
+
+/** A collection of identifiers which will help us parse the commandline opions. */
+namespace CommandlineOverrideSpecifiers
+{
+	// -ini:IniName:[Section1]:Key1=Value1,[Section2]:Key2=Value2
+	FString	IniSwitchIdentifier		= TEXT("-ini:");
+	FString	IniNameEndIdentifier	= TEXT(":[");
+	TCHAR	SectionStartIdentifier	= TCHAR('[');
+	FString PropertyStartIdentifier	= TEXT("]:");
+	TCHAR	PropertySeperator		= TCHAR(',');
+}
+
 /**
 * Looks for any overrides on the commandline for this file
 *
@@ -568,13 +580,13 @@ static void OverrideFromCommandline(FConfigFile* File, const FString& Filename)
 	// look for this filename on the commandline in the format:
 	//		-ini:IniName:Section1.Key1=Value1,Section2.Key2=Value2
 	// for example:
-	//		-ini:Engine:/Script/Engine.Engine.bSmoothFrameRate=False,TextureStreaming.PoolSize=100
+	//		-ini:Engine:/Script/Engine.Engine:bSmoothFrameRate=False,TextureStreaming:PoolSize=100
 	//			(will update the cache after the final combined engine.ini)
-	if (FParse::Value(FCommandLine::Get(), *FString::Printf(TEXT("-ini:%s:"), *FPaths::GetBaseFilename(Filename)), Settings, false))
+	if (FParse::Value(FCommandLine::Get(), *FString::Printf(TEXT("%s%s"), *CommandlineOverrideSpecifiers::IniSwitchIdentifier, *FPaths::GetBaseFilename(Filename)), Settings, false))
 	{
 		// break apart on the commas
 		TArray<FString> SettingPairs;
-		Settings.ParseIntoArray(&SettingPairs, TEXT(","), true);
+		Settings.ParseIntoArray(&SettingPairs, &CommandlineOverrideSpecifiers::PropertySeperator, true);
 		for (int32 Index = 0; Index < SettingPairs.Num(); Index++)
 		{
 			// set each one, by splitting on the =
@@ -582,11 +594,9 @@ static void OverrideFromCommandline(FConfigFile* File, const FString& Filename)
 			if (SettingPairs[Index].Split(TEXT("="), &SectionAndKey, &Value))
 			{
 				// now we need to split off the key from the rest of the section name
-				int32 LastDot = INDEX_NONE;
-				SectionAndKey.FindLastChar(TCHAR('.'), LastDot);
-
+				int32 SectionNameEndIndex = SectionAndKey.Find(CommandlineOverrideSpecifiers::PropertyStartIdentifier, ESearchCase::IgnoreCase, ESearchDir::FromEnd);
 				// check for malformed string
-				if (LastDot == INDEX_NONE || LastDot == 0)
+				if (SectionNameEndIndex == INDEX_NONE || SectionNameEndIndex == 0)
 				{
 					continue;
 				}
@@ -594,8 +604,14 @@ static void OverrideFromCommandline(FConfigFile* File, const FString& Filename)
 				// Create the commandline override object
 				FConfigCommandlineOverride& CommandlineOption = File->CommandlineOptions[File->CommandlineOptions.Emplace()];
 				CommandlineOption.BaseFileName = *FPaths::GetBaseFilename(Filename);
-				CommandlineOption.Section = SectionAndKey.Left(LastDot);
-				CommandlineOption.PropertyKey = SectionAndKey.Mid(LastDot + 1);
+				CommandlineOption.Section = SectionAndKey.Left(SectionNameEndIndex);
+				
+				// Remove commandline syntax from the section name.
+				CommandlineOption.Section = CommandlineOption.Section.Replace(*CommandlineOverrideSpecifiers::IniNameEndIdentifier, TEXT(""));
+				CommandlineOption.Section = CommandlineOption.Section.Replace(*CommandlineOverrideSpecifiers::PropertyStartIdentifier, TEXT(""));
+				CommandlineOption.Section = CommandlineOption.Section.Replace(&CommandlineOverrideSpecifiers::SectionStartIdentifier, TEXT(""));
+
+				CommandlineOption.PropertyKey = SectionAndKey.Mid(SectionNameEndIndex + CommandlineOverrideSpecifiers::PropertyStartIdentifier.Len());
 				CommandlineOption.PropertyValue = Value;
 
 				// now put it into this into the cache
