UE4 - reworked memstack allocator to share blocks with all threads. 64k blocks are recycled; larger blocks are not.

diff --git a/Engine/Source/Runtime/Core/Private/Misc/MemStack.cpp b/Engine/Source/Runtime/Core/Private/Misc/MemStack.cpp
--- a/Engine/Source/Runtime/Core/Private/Misc/MemStack.cpp
+++ b/Engine/Source/Runtime/Core/Private/Misc/MemStack.cpp
@@ -8,41 +8,36 @@
 
 DECLARE_THREAD_SINGLETON( FMemStack );
 
-DECLARE_MEMORY_STAT(TEXT("MemStack Allocated (all threads)"), STAT_MemStackAllocated,STATGROUP_Memory);
-DECLARE_MEMORY_STAT(TEXT("MemStack Used (all threads)"), STAT_MemStackUsed,STATGROUP_Memory);
+DECLARE_MEMORY_STAT(TEXT("MemStack Large Block"), STAT_MemStackLargeBLock,STATGROUP_Memory);
+DECLARE_MEMORY_STAT(TEXT("PageAllocator Free"), STAT_PageAllocatorFree, STATGROUP_Memory);
+DECLARE_MEMORY_STAT(TEXT("PageAllocator Used"), STAT_PageAllocatorUsed, STATGROUP_Memory);
 
+
+TLockFreeFixedSizeAllocator<FPageAllocator::PageSize, FThreadSafeCounter> FPageAllocator::TheAllocator;
+
+#if STATS
+void FPageAllocator::UpdateStats()
+{
+	SET_MEMORY_STAT(STAT_PageAllocatorFree, BytesFree());
+	SET_MEMORY_STAT(STAT_PageAllocatorUsed, BytesUsed());
+}
+#endif
 /*-----------------------------------------------------------------------------
 	FMemStack implementation.
 -----------------------------------------------------------------------------*/
 
-FMemStack::FMemStack( int32 InDefaultChunkSize /*= DEFAULT_CHUNK_SIZE*/ )
-:	Top(NULL)
-,	End(NULL)
-,	DefaultChunkSize(InDefaultChunkSize)
-,	TopChunk(NULL)
-,	TopMark(NULL)
-,	UnusedChunks(NULL)
+FMemStack::FMemStack()
+:	Top(nullptr)
+,	End(nullptr)
+,	TopChunk(nullptr)
+,	TopMark(nullptr)
 ,	NumMarks(0)
 {
 }
 
 FMemStack::~FMemStack()
 {
-	check(GIsCriticalError || !NumMarks);
-
-	Tick();
-	while( UnusedChunks )
-	{
-		FTaggedMemory* Old = UnusedChunks;
-		UnusedChunks = UnusedChunks->Next;
-		DEC_MEMORY_STAT_BY( STAT_MemStackAllocated, Old->DataSize + sizeof(FTaggedMemory) );
-		FMemory::Free( Old );
-	}
-}
-
-void FMemStack::Tick() const
-{
-	check(TopChunk==NULL);
+	check((GIsCriticalError || !NumMarks) && IsEmpty());
 }
 
 int32 FMemStack::GetByteCount() const
@@ -61,44 +56,23 @@ int32 FMemStack::GetByteCount() const
 	}
 	return Count;
 }
-int32 FMemStack::GetUnusedByteCount() const
-{
-	int32 Count = 0;
-	for( FTaggedMemory* Chunk=UnusedChunks; Chunk; Chunk=Chunk->Next )
-	{
-		Count += Chunk->DataSize;
-	}
-	return Count;
-}
 
 void FMemStack::AllocateNewChunk( int32 MinSize )
 {
-	FTaggedMemory* Chunk=NULL;
-	for( FTaggedMemory** Link=&UnusedChunks; *Link; Link=&(*Link)->Next )
-	{
-		// Find existing chunk.
-		if( (*Link)->DataSize >= MinSize )
-		{
-			Chunk = *Link;
-			*Link = (*Link)->Next;
-			break;
-		}
-	}
-
-	if( !Chunk )
+	FTaggedMemory* Chunk=nullptr;
+	// Create new chunk.
+	const int32 DataSize = AlignArbitrary<int32>(MinSize + (int32)sizeof(FTaggedMemory), FPageAllocator::PageSize) - sizeof(FTaggedMemory);
+	const uint32 AllocSize = DataSize + sizeof(FTaggedMemory);
+	if (AllocSize == FPageAllocator::PageSize)
 	{
-		// Create new chunk.
-		const int32 DataSize   = AlignArbitrary<int32>( MinSize + (int32)sizeof(FTaggedMemory), DefaultChunkSize ) - sizeof(FTaggedMemory);
-		const uint32 AllocSize = DataSize + sizeof(FTaggedMemory);
-		Chunk                  = (FTaggedMemory*)FMemory::Malloc( AllocSize );
-		Chunk->DataSize        = DataSize;
-		INC_MEMORY_STAT_BY( STAT_MemStackAllocated, AllocSize );	
+		Chunk = (FTaggedMemory*)FPageAllocator::Alloc();
 	}
-
-	if( Chunk != TopChunk )
+	else
 	{
-		INC_MEMORY_STAT_BY( STAT_MemStackUsed, Chunk->DataSize );
+		Chunk = (FTaggedMemory*)FMemory::Malloc(AllocSize);
+		INC_MEMORY_STAT_BY(STAT_MemStackLargeBLock, AllocSize);
 	}
+	Chunk->DataSize        = DataSize;
 
 	Chunk->Next = TopChunk;
 	TopChunk    = Chunk;
@@ -112,13 +86,20 @@ void FMemStack::FreeChunks( FTaggedMemory* NewTopChunk )
 	{
 		FTaggedMemory* RemoveChunk = TopChunk;
 		TopChunk                   = TopChunk->Next;
-		RemoveChunk->Next          = UnusedChunks;
-		UnusedChunks               = RemoveChunk;
 
-		DEC_MEMORY_STAT_BY( STAT_MemStackUsed, RemoveChunk->DataSize );
+		if (RemoveChunk->DataSize + sizeof(FTaggedMemory) == FPageAllocator::PageSize)
+		{
+			FPageAllocator::Free(RemoveChunk);
+		}
+		else
+		{
+			DEC_MEMORY_STAT_BY(STAT_MemStackLargeBLock, RemoveChunk->DataSize + sizeof(FTaggedMemory));
+			FMemory::Free(RemoveChunk);
+		}
+
 	}
-	Top = NULL;
-	End = NULL;
+	Top = nullptr;
+	End = nullptr;
 	if( TopChunk )
 	{
 		Top = TopChunk->Data;
diff --git a/Engine/Source/Runtime/Core/Public/Async/TaskGraphInterfaces.h b/Engine/Source/Runtime/Core/Public/Async/TaskGraphInterfaces.h
--- a/Engine/Source/Runtime/Core/Public/Async/TaskGraphInterfaces.h
+++ b/Engine/Source/Runtime/Core/Public/Async/TaskGraphInterfaces.h
@@ -669,6 +669,7 @@ private:
 #endif
 
 			Task.~TTask();
+			checkSlow(ENamedThreads::GetThreadIndex(CurrentThread) <= ENamedThreads::RenderThread || FMemStack::Get().IsEmpty()); // you must mark and pop memstacks if you use them in tasks! Named threads are excepted.
 		}
 		
 		TaskConstructed = false;
diff --git a/Engine/Source/Runtime/Core/Public/Containers/LockFreeFixedSizeAllocator.h b/Engine/Source/Runtime/Core/Public/Containers/LockFreeFixedSizeAllocator.h
--- a/Engine/Source/Runtime/Core/Public/Containers/LockFreeFixedSizeAllocator.h
+++ b/Engine/Source/Runtime/Core/Public/Containers/LockFreeFixedSizeAllocator.h
@@ -4,37 +4,8 @@
 
 #include "LockFreeList.h"
 
-/** If this is set to one, then the current allocation state is monitored with thread safe pointers **/
-#define MONITOR_FIXED_ALLOCATION (0)
-
-#define USE_RECYCLING (1)
-
-
-#if !USE_RECYCLING
-
 /** Thread safe, lock free pooling allocator of fixed size blocks that never returns free space until program shutdown. **/
-template<int32 SIZE>
-class TLockFreeFixedSizeAllocator	// alignment isn't handled, assumes FMemory::Malloc will work
-{
-public:
-	/** Returns a memory block of size SIZE **/
-	void* Allocate()
-	{
-		return FMemory::Malloc(SIZE);
-	}
-	/** Given a memory block (obtained from Allocate) free it **/
-	void Free(void *Item)
-	{
-		FMemory::Free(Item);
-	}
-};
-
-
-#else //!USE_RECYCLING
-
-
-/** Thread safe, lock free pooling allocator of fixed size blocks that never returns free space until program shutdown. **/
-template<int32 SIZE>
+template<int32 SIZE, typename TTrackingCounter = FNoopCounter>
 class TLockFreeFixedSizeAllocator	// alignment isn't handled, assumes FMemory::Malloc will work
 {
 public:
@@ -45,6 +16,7 @@ public:
 		while (void* Mem = FreeList.Pop())
 		{
 			FMemory::Free(Mem);
+			NumFree.Decrement();
 		}
 		check(!NumFree.GetValue());
 	}
@@ -70,23 +42,27 @@ public:
 		FreeList.Push(Item);
 		NumFree.Increment();
 	}
+
+	const TTrackingCounter& GetNumUsed() const
+	{
+		return NumUsed;
+	}
+
+	const TTrackingCounter& GetNumFree() const
+	{
+		return NumFree;
+	}
+
 private:
 	/** Lock free list of free memory blocks **/
 	TLockFreePointerList<void>		FreeList;
 
-#if MONITOR_FIXED_ALLOCATION
 	/** Total number of blocks outstanding and not in the free list **/
-	FThreadSafeCounter	NumUsed; 
+	TTrackingCounter	NumUsed; 
 	/** Total number of blocks in the free list **/
-	FThreadSafeCounter	NumFree;
-#else
-	/** If we aren't monitoring allocation, these do nothing **/
-	FNoopCounter		NumUsed; 
-	FNoopCounter		NumFree;
-#endif
+	TTrackingCounter	NumFree;
 };
 
-#endif //!USE_RECYCLING
 
 
 /** Thread safe, lock free pooling allocator of memory for instances of T. Never returns free space until program shutdown. **/
diff --git a/Engine/Source/Runtime/Core/Public/Misc/MemStack.h b/Engine/Source/Runtime/Core/Public/Misc/MemStack.h
--- a/Engine/Source/Runtime/Core/Public/Misc/MemStack.h
+++ b/Engine/Source/Runtime/Core/Public/Misc/MemStack.h
@@ -5,15 +5,47 @@
 =============================================================================*/
 
 #pragma once
+#include "LockFreeFixedSizeAllocator.h"
 
-/*-----------------------------------------------------------------------------
-	Globals.
------------------------------------------------------------------------------*/
 
 // Enums for specifying memory allocation type.
 enum EMemZeroed {MEM_Zeroed=1};
 enum EMemOned   {MEM_Oned  =1};
 
+class CORE_API FPageAllocator
+{
+public:
+	enum
+	{
+		PageSize = 64 * 1024
+	};
+	FORCEINLINE static void *Alloc()
+	{
+		void *Result = TheAllocator.Allocate();
+		STAT(UpdateStats());
+		return Result;
+	}
+	FORCEINLINE static void Free(void *Mem)
+	{
+		TheAllocator.Free(Mem);
+		STAT(UpdateStats());
+	}
+	static uint64 BytesUsed()
+	{
+		return uint64(TheAllocator.GetNumUsed().GetValue()) * PageSize;
+	}
+	static uint64 BytesFree()
+	{
+		return uint64(TheAllocator.GetNumFree().GetValue()) * PageSize;
+	}
+private:
+
+#if STATS
+	static void UpdateStats();
+#endif
+	static TLockFreeFixedSizeAllocator<PageSize, FThreadSafeCounter> TheAllocator;
+};
+
 /**
  * Simple linear-allocation memory stack.
  * Items are allocated via PushBytes() or the specialized operator new()s.
@@ -24,16 +56,18 @@ class CORE_API FMemStack : public TThreadSingleton<FMemStack>
 public:
 
 	/** Initialization constructor. */
-	FMemStack(int32 DefaultChunkSize = DEFAULT_CHUNK_SIZE);
+	FMemStack();
 
 	/** Destructor. */
 	~FMemStack();
 
 	// Get bytes.
-	FORCEINLINE uint8* PushBytes( int32 AllocSize, int32 Alignment )
+	FORCEINLINE uint8* PushBytes(int32 AllocSize, int32 Alignment)
+	{
+		return (uint8*)Alloc(AllocSize, FMath::Max(AllocSize >= 16 ? (int32)16 : (int32)8, Alignment));
+	}
+	FORCEINLINE void* Alloc(int32 AllocSize, int32 Alignment)
 	{
-		Alignment = FMath::Max(AllocSize >= 16 ? (int32)16 : (int32)8, Alignment);
-
 		// Debug checks.
 		checkSlow(AllocSize>=0);
 		checkSlow((Alignment&(Alignment-1))==0);
@@ -60,15 +94,15 @@ public:
 		return Result;
 	}
 
-	/** Timer tick. Makes sure the memory stack is empty. */
-	void Tick() const;
+	/** return true if this stack is empty. */
+	FORCEINLINE bool IsEmpty() const
+	{
+		return TopChunk == nullptr;
+	}
 
 	/** @return the number of bytes allocated for this FMemStack that are currently in use. */
 	int32 GetByteCount() const;
 
-	/** @return the number of bytes allocated for this FMemStack that are currently unused and available. */
-	int32 GetUnusedByteCount() const;
-
 	// Returns true if the pointer was allocated using this allocator
 	bool ContainsPointer(const void* Pointer) const;
 
@@ -90,25 +124,15 @@ public:
 	};
 
 private:
-	// Constants.
-	enum 
-	{
-		MAX_CHUNKS=1024,
-		DEFAULT_CHUNK_SIZE=65536,
-	};
 
 	// Variables.
 	uint8*			Top;				// Top of current chunk (Top<=End).
 	uint8*			End;				// End of current chunk.
-	int32			DefaultChunkSize;	// Maximum chunk size to allocate.
 	FTaggedMemory*	TopChunk;			// Only chunks 0..ActiveChunks-1 are valid.
 
 	/** The top mark on the stack. */
 	class FMemMark*	TopMark;
 
-	/** The memory chunks that have been allocated but are currently unused. */
-	FTaggedMemory*	UnusedChunks;
-
 	/** The number of marks on this stack. */
 	int32 NumMarks;
 
@@ -204,7 +228,7 @@ public:
 
 		/** Default constructor. */
 		ForElementType():
-			Data(NULL)
+			Data(nullptr)
 		{}
 
 		// FContainerAllocatorInterface
@@ -301,7 +325,7 @@ public:
 			Mem.TopMark = NextTopmostMark;
 
 			// Ensure that the mark is only popped once by clearing the top pointer.
-			Top = NULL;
+			Top = nullptr;
 		}
 	}
 
diff --git a/Engine/Source/Runtime/Engine/Private/UnrealEngine.cpp b/Engine/Source/Runtime/Engine/Private/UnrealEngine.cpp
--- a/Engine/Source/Runtime/Engine/Private/UnrealEngine.cpp
+++ b/Engine/Source/Runtime/Engine/Private/UnrealEngine.cpp
@@ -3932,8 +3932,9 @@ bool UEngine::HandleMemCommand( const TCHAR* Cmd, FOutputDevice& Ar )
 	if( bDetailed || bReport)
 	{
 		Ar.CategorizedLogf( CategoryName, ELogVerbosity::Log, TEXT("Memory Stats:") );
-		Ar.CategorizedLogf( CategoryName, ELogVerbosity::Log, TEXT("FMemStack (gamethread) allocation size [used/ unused] = [%.2f / %.2f] MB"), FMemStack::Get().GetByteCount() / (1024.0f * 1024.0f), FMemStack::Get().GetUnusedByteCount() / (1024.0f * 1024.0f)  );
-		Ar.CategorizedLogf( CategoryName, ELogVerbosity::Log, TEXT("Nametable memory usage = %.2f MB"), FName::GetNameTableMemorySize() / (1024.0f * 1024.0f) );
+		Ar.CategorizedLogf( CategoryName, ELogVerbosity::Log, TEXT("FMemStack (gamethread) current size = %.2f MB"), FMemStack::Get().GetByteCount() / (1024.0f * 1024.0f));
+		Ar.CategorizedLogf(CategoryName, ELogVerbosity::Log, TEXT("FPageAllocator (all threads) allocation size [used/ unused] = [%.2f / %.2f] MB"), (FPageAllocator::BytesUsed()) / (1024.0f * 1024.0f), (FPageAllocator::BytesFree()) / (1024.0f * 1024.0f));
+		Ar.CategorizedLogf(CategoryName, ELogVerbosity::Log, TEXT("Nametable memory usage = %.2f MB"), FName::GetNameTableMemorySize() / (1024.0f * 1024.0f));
 
 #if STATS
 		TArray<FStatMessage> Stats;
