UE4 - refactor MemStack to allow it to be used without the TLS stuff

diff --git a/Engine/Source/Runtime/Core/Private/Misc/MemStack.cpp b/Engine/Source/Runtime/Core/Private/Misc/MemStack.cpp
--- a/Engine/Source/Runtime/Core/Private/Misc/MemStack.cpp
+++ b/Engine/Source/Runtime/Core/Private/Misc/MemStack.cpp
@@ -26,21 +26,7 @@ void FPageAllocator::UpdateStats()
 	FMemStack implementation.
 -----------------------------------------------------------------------------*/
 
-FMemStack::FMemStack()
-:	Top(nullptr)
-,	End(nullptr)
-,	TopChunk(nullptr)
-,	TopMark(nullptr)
-,	NumMarks(0)
-{
-}
-
-FMemStack::~FMemStack()
-{
-	check((GIsCriticalError || !NumMarks) && IsEmpty());
-}
-
-int32 FMemStack::GetByteCount() const
+int32 FMemStackBase::GetByteCount() const
 {
 	int32 Count = 0;
 	for( FTaggedMemory* Chunk=TopChunk; Chunk; Chunk=Chunk->Next )
@@ -57,7 +43,7 @@ int32 FMemStack::GetByteCount() const
 	return Count;
 }
 
-void FMemStack::AllocateNewChunk( int32 MinSize )
+void FMemStackBase::AllocateNewChunk(int32 MinSize)
 {
 	FTaggedMemory* Chunk=nullptr;
 	// Create new chunk.
@@ -80,7 +66,7 @@ void FMemStack::AllocateNewChunk( int32 MinSize )
 	End         = Top + Chunk->DataSize;
 }
 
-void FMemStack::FreeChunks( FTaggedMemory* NewTopChunk )
+void FMemStackBase::FreeChunks(FTaggedMemory* NewTopChunk)
 {
 	while( TopChunk!=NewTopChunk )
 	{
@@ -107,7 +93,7 @@ void FMemStack::FreeChunks( FTaggedMemory* NewTopChunk )
 	}
 }
 
-bool FMemStack::ContainsPointer( const void* Pointer ) const
+bool FMemStackBase::ContainsPointer(const void* Pointer) const
 {
 	const uint8* Ptr = (const uint8*)Pointer;
 	for (const FTaggedMemory* Chunk = TopChunk; Chunk; Chunk = Chunk->Next)
diff --git a/Engine/Source/Runtime/Core/Public/Misc/MemStack.h b/Engine/Source/Runtime/Core/Public/Misc/MemStack.h
--- a/Engine/Source/Runtime/Core/Public/Misc/MemStack.h
+++ b/Engine/Source/Runtime/Core/Public/Misc/MemStack.h
@@ -51,15 +51,25 @@ private:
  * Items are allocated via PushBytes() or the specialized operator new()s.
  * Items are freed en masse by using FMemMark to Pop() them.
  **/
-class CORE_API FMemStack : public TThreadSingleton<FMemStack>
+class CORE_API FMemStackBase
 {
 public:
 
-	/** Initialization constructor. */
-	FMemStack();
+	FMemStackBase()
+		: Top(nullptr)
+		, End(nullptr)
+		, TopChunk(nullptr)
+		, TopMark(nullptr)
+		, NumMarks(0)
+	{
+	}
+
+	~FMemStackBase()
+	{
+		check((GIsCriticalError || !NumMarks));
+		FreeChunks(nullptr);
+	}
 
-	/** Destructor. */
-	~FMemStack();
 
 	// Get bytes.
 	FORCEINLINE uint8* PushBytes(int32 AllocSize, int32 Alignment)
@@ -108,12 +118,12 @@ public:
 
 	// Friends.
 	friend class FMemMark;
-	friend void* operator new( size_t Size, FMemStack& Mem, int32 Count, int32 Align );
-	friend void* operator new( size_t Size, FMemStack& Mem, EMemZeroed Tag, int32 Count, int32 Align );
-	friend void* operator new( size_t Size, FMemStack& Mem, EMemOned Tag, int32 Count, int32 Align );
-	friend void* operator new[]( size_t Size, FMemStack& Mem, int32 Count, int32 Align );
-	friend void* operator new[]( size_t Size, FMemStack& Mem, EMemZeroed Tag, int32 Count, int32 Align );
-	friend void* operator new[]( size_t Size, FMemStack& Mem, EMemOned Tag, int32 Count, int32 Align );
+	friend void* operator new(size_t Size, FMemStackBase& Mem, int32 Count, int32 Align);
+	friend void* operator new(size_t Size, FMemStackBase& Mem, EMemZeroed Tag, int32 Count, int32 Align);
+	friend void* operator new(size_t Size, FMemStackBase& Mem, EMemOned Tag, int32 Count, int32 Align);
+	friend void* operator new[](size_t Size, FMemStackBase& Mem, int32 Count, int32 Align);
+	friend void* operator new[](size_t Size, FMemStackBase& Mem, EMemZeroed Tag, int32 Count, int32 Align);
+	friend void* operator new[](size_t Size, FMemStackBase& Mem, EMemOned Tag, int32 Count, int32 Align);
 
 	// Types.
 	struct FTaggedMemory
@@ -146,22 +156,27 @@ private:
 	void FreeChunks( FTaggedMemory* NewTopChunk );
 };
 
+class CORE_API FMemStack : public TThreadSingleton<FMemStack>, public FMemStackBase
+{
+
+};
+
 /*-----------------------------------------------------------------------------
 	FMemStack templates.
 -----------------------------------------------------------------------------*/
 
 // Operator new for typesafe memory stack allocation.
-template <class T> inline T* New( FMemStack& Mem, int32 Count=1, int32 Align=DEFAULT_ALIGNMENT )
+template <class T> inline T* New(FMemStackBase& Mem, int32 Count = 1, int32 Align = DEFAULT_ALIGNMENT)
 {
 	return (T*)Mem.PushBytes( Count*sizeof(T), Align );
 }
-template <class T> inline T* NewZeroed( FMemStack& Mem, int32 Count=1, int32 Align=DEFAULT_ALIGNMENT )
+template <class T> inline T* NewZeroed(FMemStackBase& Mem, int32 Count = 1, int32 Align = DEFAULT_ALIGNMENT)
 {
 	uint8* Result = Mem.PushBytes( Count*sizeof(T), Align );
 	FMemory::Memzero( Result, Count*sizeof(T) );
 	return (T*)Result;
 }
-template <class T> inline T* NewOned( FMemStack& Mem, int32 Count=1, int32 Align=DEFAULT_ALIGNMENT )
+template <class T> inline T* NewOned(FMemStackBase& Mem, int32 Count = 1, int32 Align = DEFAULT_ALIGNMENT)
 {
 	uint8* Result = Mem.PushBytes( Count*sizeof(T), Align );
 	FMemory::Memset( Result, 0xff, Count*sizeof(T) );
@@ -173,38 +188,38 @@ template <class T> inline T* NewOned( FMemStack& Mem, int32 Count=1, int32 Align
 -----------------------------------------------------------------------------*/
 
 // Operator new for typesafe memory stack allocation.
-inline void* operator new( size_t Size, FMemStack& Mem, int32 Count=1, int32 Align=DEFAULT_ALIGNMENT )
+inline void* operator new(size_t Size, FMemStackBase& Mem, int32 Count = 1, int32 Align = DEFAULT_ALIGNMENT)
 {
 	// Get uninitialized memory.
 	return Mem.PushBytes( Size*Count, Align );
 }
-inline void* operator new( size_t Size, FMemStack& Mem, EMemZeroed Tag, int32 Count=1, int32 Align=DEFAULT_ALIGNMENT )
+inline void* operator new(size_t Size, FMemStackBase& Mem, EMemZeroed Tag, int32 Count = 1, int32 Align = DEFAULT_ALIGNMENT)
 {
 	// Get zero-filled memory.
 	uint8* Result = Mem.PushBytes( Size*Count, Align );
 	FMemory::Memzero( Result, Size*Count );
 	return Result;
 }
-inline void* operator new( size_t Size, FMemStack& Mem, EMemOned Tag, int32 Count=1, int32 Align=DEFAULT_ALIGNMENT )
+inline void* operator new(size_t Size, FMemStackBase& Mem, EMemOned Tag, int32 Count = 1, int32 Align = DEFAULT_ALIGNMENT)
 {
 	// Get one-filled memory.
 	uint8* Result = Mem.PushBytes( Size*Count, Align );
 	FMemory::Memset( Result, 0xff, Size*Count );
 	return Result;
 }
-inline void* operator new[]( size_t Size, FMemStack& Mem, int32 Count=1, int32 Align=DEFAULT_ALIGNMENT )
+inline void* operator new[](size_t Size, FMemStackBase& Mem, int32 Count = 1, int32 Align = DEFAULT_ALIGNMENT)
 {
 	// Get uninitialized memory.
 	return Mem.PushBytes( Size*Count, Align );
 }
-inline void* operator new[]( size_t Size, FMemStack& Mem, EMemZeroed Tag, int32 Count=1, int32 Align=DEFAULT_ALIGNMENT )
+inline void* operator new[](size_t Size, FMemStackBase& Mem, EMemZeroed Tag, int32 Count = 1, int32 Align = DEFAULT_ALIGNMENT)
 {
 	// Get zero-filled memory.
 	uint8* Result = Mem.PushBytes( Size*Count, Align );
 	FMemory::Memzero( Result, Size*Count );
 	return Result;
 }
-inline void* operator new[]( size_t Size, FMemStack& Mem, EMemOned Tag, int32 Count=1, int32 Align=DEFAULT_ALIGNMENT )
+inline void* operator new[](size_t Size, FMemStackBase& Mem, EMemOned Tag, int32 Count = 1, int32 Align = DEFAULT_ALIGNMENT)
 {
 	// Get one-filled memory.
 	uint8* Result = Mem.PushBytes( Size*Count, Align );
@@ -284,7 +299,7 @@ class FMemMark
 {
 public:
 	// Constructors.
-	FMemMark( FMemStack& InMem )
+	FMemMark(FMemStackBase& InMem)
 	:	Mem(InMem)
 	,	Top(InMem.Top)
 	,	SavedChunk(InMem.TopChunk)
@@ -331,9 +346,9 @@ public:
 
 private:
 	// Implementation variables.
-	FMemStack& Mem;
+	FMemStackBase& Mem;
 	uint8* Top;
-	FMemStack::FTaggedMemory* SavedChunk;
+	FMemStackBase::FTaggedMemory* SavedChunk;
 	bool bPopped;
 	FMemMark* NextTopmostMark;
 };
