Added PumpMessages call to Mac SetCapture to solve problems with mouse move events accumulated during engine startup.

diff --git a/Engine/Source/Runtime/Core/Private/Mac/MacApplication.cpp b/Engine/Source/Runtime/Core/Private/Mac/MacApplication.cpp
--- a/Engine/Source/Runtime/Core/Private/Mac/MacApplication.cpp
+++ b/Engine/Source/Runtime/Core/Private/Mac/MacApplication.cpp
@@ -727,6 +727,7 @@ void FMacApplication::PollGameDeviceState( const float TimeDelta )
 
 void FMacApplication::SetCapture( const TSharedPtr< FGenericWindow >& InWindow )
 {
+	FPlatformMisc::PumpMessages(true);
 	bIsMouseCaptureEnabled = InWindow.IsValid();
 	UpdateMouseCaptureWindow( bIsMouseCaptureEnabled ? ((FMacWindow*)InWindow.Get())->GetWindowHandle() : NULL );
 }
