Added Close item to Window menu in Mac menu bar

diff --git a/Engine/Source/Runtime/Core/Private/Mac/MacPlatformMisc.cpp b/Engine/Source/Runtime/Core/Private/Mac/MacPlatformMisc.cpp
--- a/Engine/Source/Runtime/Core/Private/Mac/MacPlatformMisc.cpp
+++ b/Engine/Source/Runtime/Core/Private/Mac/MacPlatformMisc.cpp
@@ -317,10 +317,12 @@ void FMacPlatformMisc::UpdateWindowMenu()
 
 	NSMenuItem* MinimizeItem = [[[NSMenuItem alloc] initWithTitle:@"Minimize" action:@selector(miniaturize:) keyEquivalent:@"m"] autorelease];
 	NSMenuItem* ZoomItem = [[[NSMenuItem alloc] initWithTitle:@"Zoom" action:@selector(performZoom:) keyEquivalent:@""] autorelease];
+	NSMenuItem* CloseItem = [[[NSMenuItem alloc] initWithTitle:@"Close" action:@selector(performClose:) keyEquivalent:@""] autorelease];
 	NSMenuItem* BringAllToFrontItem = [[[NSMenuItem alloc] initWithTitle:@"Bring All to Front" action:@selector(arrangeInFront:) keyEquivalent:@""] autorelease];
 
 	[WindowMenu addItem:MinimizeItem];
 	[WindowMenu addItem:ZoomItem];
+	[WindowMenu addItem:CloseItem];
 	[WindowMenu addItem:[NSMenuItem separatorItem]];
 	[WindowMenu addItem:BringAllToFrontItem];
 }
