fix for CIS errors
#ue4
#ios

diff --git a/Engine/Source/Runtime/Core/Private/IOS/IOSPlatformFile.cpp b/Engine/Source/Runtime/Core/Private/IOS/IOSPlatformFile.cpp
--- a/Engine/Source/Runtime/Core/Private/IOS/IOSPlatformFile.cpp
+++ b/Engine/Source/Runtime/Core/Private/IOS/IOSPlatformFile.cpp
@@ -9,7 +9,7 @@
 //#endif
 
 // make an FTimeSpan object that represents the "epoch" for time_t (from a stat struct)
-const FDateTime MacEpoch(1970, 1, 1);
+const FDateTime IOSEpoch(1970, 1, 1);
 
 /** 
  * iOS file handle implementation
@@ -219,7 +219,7 @@ FDateTime FIOSPlatformFile::GetTimeStamp(const TCHAR* Filename)
 
 	// convert _stat time to FDateTime
 	FTimespan TimeSinceEpoch(0, 0, FileInfo.st_mtime);
-	return MacEpoch + TimeSinceEpoch;
+	return IOSEpoch + TimeSinceEpoch;
 
 }
 
@@ -236,7 +236,7 @@ void FIOSPlatformFile::SetTimeStamp(const TCHAR* Filename, const FDateTime DateT
 	// change the modification time only
 	struct utimbuf Times;
 	Times.actime = FileInfo.st_atime;
-	Times.modtime = (DateTime - MacEpoch).GetTotalSeconds();
+	Times.modtime = (DateTime - IOSEpoch).GetTotalSeconds();
 	utime(TCHAR_TO_UTF8(*IOSFilename), &Times);
 }
 
@@ -257,7 +257,7 @@ FDateTime FIOSPlatformFile::GetAccessTimeStamp(const TCHAR* Filename)
 
 	// convert _stat time to FDateTime
 	FTimespan TimeSinceEpoch(0, 0, FileInfo.st_atime);
-	return MacEpoch + TimeSinceEpoch;
+	return IOSEpoch + TimeSinceEpoch;
 }
 
 IFileHandle* FIOSPlatformFile::OpenRead(const TCHAR* Filename)
diff --git a/Engine/Source/Runtime/Core/Private/IOS/IOSPlatformMisc.cpp b/Engine/Source/Runtime/Core/Private/IOS/IOSPlatformMisc.cpp
--- a/Engine/Source/Runtime/Core/Private/IOS/IOSPlatformMisc.cpp
+++ b/Engine/Source/Runtime/Core/Private/IOS/IOSPlatformMisc.cpp
@@ -9,6 +9,7 @@
 #include "SecureHash.h"
 #include "IOSApplication.h"
 #include "IOSAppDelegate.h"
+#include "IOSView.h"
 
 /** Amount of free memory in MB reported by the system at starup */
 CORE_API int32 GStartupFreeMemoryMB;
