Fixed a problem in Mac version with determining event window in case of receiving mouse events for mouse capture window that is no longer active (which can happen with low framerate)

diff --git a/Engine/Source/Runtime/Core/Private/Mac/MacApplication.cpp b/Engine/Source/Runtime/Core/Private/Mac/MacApplication.cpp
--- a/Engine/Source/Runtime/Core/Private/Mac/MacApplication.cpp
+++ b/Engine/Source/Runtime/Core/Private/Mac/MacApplication.cpp
@@ -646,14 +646,14 @@ FSlateCocoaWindow* FMacApplication::FindEventWindow( NSEvent* Event )
 		{
 			EventWindow = DraggedWindow;
 		}
-		else if( MouseCaptureWindow && MouseCaptureWindow == [Event window] )
+		else if( MouseCaptureWindow && MouseCaptureWindow == [Event window] && [MouseCaptureWindow targetWindow] )
 		{
 			EventWindow = [MouseCaptureWindow targetWindow];
 		}
 		else
 		{
-			const NSPoint CursorPos = [NSEvent mouseLocation];
-			TSharedPtr<FMacWindow> WindowUnderCursor = StaticCastSharedPtr<FMacWindow>(LocateWindowUnderCursor(FVector2D(CursorPos.x, FPlatformMisc::ConvertSlateYPositionToCocoa(CursorPos.y))));
+			const FVector2D CursorPos = bUsingHighPrecisionMouseInput ? HighPrecisionMousePos : FVector2D([NSEvent mouseLocation].x, FPlatformMisc::ConvertSlateYPositionToCocoa([NSEvent mouseLocation].y));
+			TSharedPtr<FMacWindow> WindowUnderCursor = StaticCastSharedPtr<FMacWindow>(LocateWindowUnderCursor(CursorPos));
 			if (WindowUnderCursor.IsValid())
 			{
 				EventWindow = WindowUnderCursor->GetWindowHandle();
