Removed redundant setting of GUseCrashReportClient on Mac and some unused code from FMacPlatformMisc::SubmitErrorReport()

diff --git a/Engine/Source/Programs/CrashReporter/CrashReportClient/Private/Mac/CrashReportClientMainMac.cpp b/Engine/Source/Programs/CrashReporter/CrashReportClient/Private/Mac/CrashReportClientMainMac.cpp
--- a/Engine/Source/Programs/CrashReporter/CrashReportClient/Private/Mac/CrashReportClientMainMac.cpp
+++ b/Engine/Source/Programs/CrashReporter/CrashReportClient/Private/Mac/CrashReportClientMainMac.cpp
@@ -50,9 +50,6 @@ static FString GSavedCommandLine;
     NSAppleEventManager* appleEventManager = [NSAppleEventManager sharedAppleEventManager];
     [appleEventManager setEventHandler:self andSelector:@selector(handleQuitEvent:withReplyEvent:) forEventClass:kCoreEventClass andEventID:kAEQuitApplication];
 	
-	// OS X always uses the CrashReportClient, since the other paths aren't reliable
-	GUseCrashReportClient = true;
-	
 	FPlatformMisc::SetGracefulTerminationHandler();
 	FPlatformMisc::SetCrashHandler(CrashReporterCrashHandler);
 	
diff --git a/Engine/Source/Programs/SlateViewer/Private/Mac/SlateViewerMainMac.cpp b/Engine/Source/Programs/SlateViewer/Private/Mac/SlateViewerMainMac.cpp
--- a/Engine/Source/Programs/SlateViewer/Private/Mac/SlateViewerMainMac.cpp
+++ b/Engine/Source/Programs/SlateViewer/Private/Mac/SlateViewerMainMac.cpp
@@ -25,8 +25,6 @@ static FString GSavedCommandLine;
     NSAppleEventManager* appleEventManager = [NSAppleEventManager sharedAppleEventManager];
     [appleEventManager setEventHandler:self andSelector:@selector(handleQuitEvent:withReplyEvent:) forEventClass:kCoreEventClass andEventID:kAEQuitApplication];
 	
-	// OS X always uses the CrashReportClient, since the other paths aren't reliable
-	GUseCrashReportClient = true;
 	FPlatformMisc::SetGracefulTerminationHandler();
 	FPlatformMisc::SetCrashHandler(NULL);
 	
diff --git a/Engine/Source/Programs/UnrealFrontend/Private/Mac/MacUnrealFrontendMain.cpp b/Engine/Source/Programs/UnrealFrontend/Private/Mac/MacUnrealFrontendMain.cpp
--- a/Engine/Source/Programs/UnrealFrontend/Private/Mac/MacUnrealFrontendMain.cpp
+++ b/Engine/Source/Programs/UnrealFrontend/Private/Mac/MacUnrealFrontendMain.cpp
@@ -29,8 +29,6 @@ static FString GSavedCommandLine;
     NSAppleEventManager* appleEventManager = [NSAppleEventManager sharedAppleEventManager];
     [appleEventManager setEventHandler:self andSelector:@selector(handleQuitEvent:withReplyEvent:) forEventClass:kCoreEventClass andEventID:kAEQuitApplication];
 	
-	// OS X always uses the CrashReportClient, since the other paths aren't reliable
-	GUseCrashReportClient = true;
 	FPlatformMisc::SetGracefulTerminationHandler();
 	FPlatformMisc::SetCrashHandler(NULL);
 	
diff --git a/Engine/Source/Runtime/Core/Private/Mac/MacPlatformMisc.cpp b/Engine/Source/Runtime/Core/Private/Mac/MacPlatformMisc.cpp
--- a/Engine/Source/Runtime/Core/Private/Mac/MacPlatformMisc.cpp
+++ b/Engine/Source/Runtime/Core/Private/Mac/MacPlatformMisc.cpp
@@ -447,182 +447,19 @@ TArray<uint8> FMacPlatformMisc::GetMacAddress()
 	return Result;
 }
 
-
-void FMacPlatformMisc::SubmitErrorReport( const TCHAR* InErrorHist, EErrorReportMode::Type InMode )
+void FMacPlatformMisc::SubmitErrorReport(const TCHAR* InErrorHist, EErrorReportMode::Type InMode)
 {
-	if ( !FPlatformMisc::IsDebuggerPresent() || GAlwaysReportCrash )
+	if (GUseCrashReportClient && (!FPlatformMisc::IsDebuggerPresent() || GAlwaysReportCrash))
 	{
-		if (GUseCrashReportClient)
-		{
-			int32 FromCommandLine = 0;
-			FParse::Value( FCommandLine::Get(), TEXT("AutomatedPerfTesting="), FromCommandLine );
-			if(( FApp::IsUnattended() == true ) && ( FromCommandLine != 0 ) && ( FParse::Param(FCommandLine::Get(), TEXT("KillAllPopUpBlockingWindows")) == true ))
-			{
-				abort();
-			}
-			return;
-		}
-		
-		TCHAR ReportDumpVersion[] = TEXT("3");
-        
-		FString ReportDumpPath;
-		{
-			const TCHAR ReportDumpFilename[] = TEXT("UnrealAutoReportDump");
-			ReportDumpPath = FPaths::CreateTempFilename( *FPaths::GameLogDir(), ReportDumpFilename, TEXT( ".txt" ) );
-		}
-        
-		FString IniDumpPath;
-		if (GGameName[0])
-		{
-			const TCHAR IniDumpFilename[] = TEXT("UnrealAutoReportIniDump");
-			IniDumpPath = FPaths::CreateTempFilename( *FPaths::GameLogDir(), IniDumpFilename, TEXT( ".txt" ) );
-			//build the ini dump
-			FOutputDeviceFile AutoReportIniFile(*IniDumpPath);
-			GConfig->Dump(AutoReportIniFile);
-			AutoReportIniFile.Flush();
-			AutoReportIniFile.TearDown();
-		}
-        
-		TCHAR AutoReportExe[] = TEXT("../DotNET/AutoReporter.exe");
-        
-		FArchive * AutoReportFile = IFileManager::Get().CreateFileWriter(*ReportDumpPath, FILEWRITE_EvenIfReadOnly);
-		if (AutoReportFile != NULL)
+		int32 FromCommandLine = 0;
+		FParse::Value( FCommandLine::Get(), TEXT("AutomatedPerfTesting="), FromCommandLine );
+		if (FApp::IsUnattended() && FromCommandLine != 0 && FParse::Param(FCommandLine::Get(), TEXT("KillAllPopUpBlockingWindows")))
 		{
-			TCHAR CompName[256];
-			FCString::Strcpy(CompName, FPlatformProcess::ComputerName());
-			TCHAR UserName[256];
-			FCString::Strcpy(UserName, FPlatformProcess::UserName());
-			TCHAR GameName[256];
-			FCString::Strcpy(GameName, *FString::Printf(TEXT("%s %s"), TEXT(BRANCH_NAME), FApp::GetGameName()));
-			TCHAR PlatformName[32];
-#if PLATFORM_64BITS
-			FCString::Strcpy(PlatformName, TEXT("Mac 64-bit"));
-#else	//PLATFORM_64BITS
-			FCString::Strcpy(PlatformName, TEXT("Mac 32-bit"));
-#endif	//PLATFORM_64BITS
-			TCHAR CultureName[10];
-			FCString::Strcpy(CultureName, *FInternationalization::Get().GetCurrentCulture()->GetName());
-			TCHAR SystemTime[256];
-			FCString::Strcpy(SystemTime, *FDateTime::Now().ToString());
-			TCHAR EngineVersionStr[32];
-			FCString::Strcpy(EngineVersionStr, *FString::FromInt(GEngineVersion.GetChangelist()));
-            
-			TCHAR ChangelistVersionStr[32];
-			int32 ChangelistFromCommandLine = 0;
-			const bool bFoundAutomatedBenchMarkingChangelist = FParse::Value( FCommandLine::Get(), TEXT("-gABC="), ChangelistFromCommandLine );
-			if( bFoundAutomatedBenchMarkingChangelist == true )
-			{
-				FCString::Strcpy(ChangelistVersionStr, *FString::FromInt(ChangelistFromCommandLine));
-			}
-			// we are not passing in the changelist to use so use the one that was stored in the ObjectVersion
-			else
-			{
-				FCString::Strcpy(ChangelistVersionStr, *FString::FromInt(GEngineVersion.GetChangelist()));
-			}
-            
-			TCHAR CmdLine[2048];
-			FCString::Strcpy(CmdLine, FCommandLine::Get());
-			TCHAR BaseDir[260];
-			FCString::Strcpy(BaseDir, FPlatformProcess::BaseDir());
-			TCHAR separator = 0;
-            
-			TCHAR EngineMode[64];
-			if( IsRunningCommandlet() )
-			{
-				FCString::Strcpy(EngineMode, TEXT("Commandlet"));
-			}
-			else if( GIsEditor )
-			{
-				FCString::Strcpy(EngineMode, TEXT("Editor"));
-			}
-			else if( GIsServer && !GIsClient )
-			{
-				FCString::Strcpy(EngineMode, TEXT("Server"));
-			}
-			else
-			{
-				FCString::Strcpy(EngineMode, TEXT("Game"));
-			}
-            
-			//build the report dump file
-			AutoReportFile->Serialize(ReportDumpVersion, FCString::Strlen(ReportDumpVersion) * sizeof(TCHAR));
-			AutoReportFile->Serialize(&separator, sizeof(TCHAR));
-			AutoReportFile->Serialize(CompName, FCString::Strlen(CompName) * sizeof(TCHAR));
-			AutoReportFile->Serialize(&separator, sizeof(TCHAR));
-			AutoReportFile->Serialize(UserName, FCString::Strlen(UserName) * sizeof(TCHAR));
-			AutoReportFile->Serialize(&separator, sizeof(TCHAR));
-			AutoReportFile->Serialize(GameName, FCString::Strlen(GameName) * sizeof(TCHAR));
-			AutoReportFile->Serialize(&separator, sizeof(TCHAR));
-			AutoReportFile->Serialize(PlatformName, FCString::Strlen(PlatformName) * sizeof(TCHAR));
-			AutoReportFile->Serialize(&separator, sizeof(TCHAR));
-			AutoReportFile->Serialize(CultureName, FCString::Strlen(CultureName) * sizeof(TCHAR));
-			AutoReportFile->Serialize(&separator, sizeof(TCHAR));
-			AutoReportFile->Serialize(SystemTime, FCString::Strlen(SystemTime) * sizeof(TCHAR));
-			AutoReportFile->Serialize(&separator, sizeof(TCHAR));
-			AutoReportFile->Serialize(EngineVersionStr, FCString::Strlen(EngineVersionStr) * sizeof(TCHAR));
-			AutoReportFile->Serialize(&separator, sizeof(TCHAR));
-			AutoReportFile->Serialize(ChangelistVersionStr, FCString::Strlen(ChangelistVersionStr) * sizeof(TCHAR));
-			AutoReportFile->Serialize(&separator, sizeof(TCHAR));
-			AutoReportFile->Serialize(CmdLine, FCString::Strlen(CmdLine) * sizeof(TCHAR));
-			AutoReportFile->Serialize(&separator, sizeof(TCHAR));
-			AutoReportFile->Serialize(BaseDir, FCString::Strlen(BaseDir) * sizeof(TCHAR));
-			AutoReportFile->Serialize(&separator, sizeof(TCHAR));
-            
-			TCHAR* NonConstErrorHist = const_cast< TCHAR* >( InErrorHist );
-			AutoReportFile->Serialize(NonConstErrorHist, FCString::Strlen(NonConstErrorHist) * sizeof(TCHAR));
-            
-			AutoReportFile->Serialize(&separator, sizeof(TCHAR));
-			AutoReportFile->Serialize(EngineMode, FCString::Strlen(EngineMode) * sizeof(TCHAR));
-			AutoReportFile->Serialize(&separator, sizeof(TCHAR));
-			AutoReportFile->Close();
-			
-			FString CrashVideoPath = FPaths::GameLogDir() + TEXT("CrashVideo.avi");
-            
-			// Get the paths that the files will actually have been saved to
-			FString UserIniDumpPath = IFileManager::Get().ConvertToAbsolutePathForExternalAppForWrite(*IniDumpPath);
-			FString LogDirectory = FPaths::GameLogDir();
-			TCHAR CommandlineLogFile[MAX_SPRINTF]=TEXT("");
-            
-			// Use the log file specified on the commandline if there is one
-			if (FParse::Value(FCommandLine::Get(), TEXT("LOG="), CommandlineLogFile, ARRAY_COUNT(CommandlineLogFile)))
-			{
-				LogDirectory += CommandlineLogFile;
-			}
-			else if (FCString::Strlen(GGameName) != 0)
-			{
-				// If the app name is defined, use it as the log filename
-				LogDirectory += FString::Printf(TEXT("%s.Log"), GGameName);
-			}
-			else
-			{
-				// Revert to hardcoded UE4.log
-				LogDirectory += TEXT("UE4.Log");
-			}
-            
-			FString UserLogFile = IFileManager::Get().ConvertToAbsolutePathForExternalAppForWrite(*LogDirectory);
-			FString UserReportDumpPath = IFileManager::Get().ConvertToAbsolutePathForExternalAppForWrite(*ReportDumpPath);
-			
-            FString UserCrashVideoPath = IFileManager::Get().ConvertToAbsolutePathForExternalAppForWrite(*CrashVideoPath);
-            
-			// Start up the auto reporting app, passing the report dump file path, the games' log file,
-			// the ini dump path, the minidump path, and the crashvideo path
-			// protect against spaces in paths breaking them up on the commandline
-            // On Mac, AutoReporter.exe is a command line argument (for Mono)
-			FString CallingCommandLine = FString::Printf(TEXT("\"%s\" %d \"%s\" \"%s\" \"%s\" \"%s\" \"%s\""),
-														 AutoReportExe, (uint32)(getpid()), *UserReportDumpPath, *UserLogFile, *UserIniDumpPath,
-														 MiniDumpFilenameW, *UserCrashVideoPath);
-            
-			// Only unattended mode is supported on Mac/Mono
-			CallingCommandLine += TEXT( " -unattended" );
-
-            // Start Mono
-			if (!FPlatformProcess::CreateProc(TEXT("/usr/bin/mono"), *CallingCommandLine, true, false, false, NULL, 0, NULL, NULL).IsValid())
-			{
-				UE_LOG(LogWindows, Warning, TEXT("Couldn't start up the Auto Reporting process!"));
-				FMessageDialog::Open( EAppMsgType::Ok, FText::FromString( InErrorHist ) );
-			}
+			abort();
 		}
 	}
+
+	// @todo Mac
 }
 
 void FMacPlatformMisc::PumpMessages( bool bFromMainLoop )
diff --git a/Engine/Source/Runtime/Launch/Private/Mac/LaunchMac.cpp b/Engine/Source/Runtime/Launch/Private/Mac/LaunchMac.cpp
--- a/Engine/Source/Runtime/Launch/Private/Mac/LaunchMac.cpp
+++ b/Engine/Source/Runtime/Launch/Private/Mac/LaunchMac.cpp
@@ -153,9 +153,6 @@ void EngineCrashHandler(const FGenericCrashContext & GenericContext)
 		GAlwaysReportCrash = true;
 	}
 #endif
-
-	// OS X always uses the CrashReportClient, since the other paths aren't reliable
-	GUseCrashReportClient = true;
 	
 #if WITH_EDITOR
 	bHasFinishedLaunching = true;
