Adding IsMouseAttached function to Applicatoin base class.
Added function to detect connected mice on Win64/Win32
Fixes ttp #339741

diff --git a/Engine/Source/Programs/UnrealBuildTool/Windows/UEBuildWindows.cs b/Engine/Source/Programs/UnrealBuildTool/Windows/UEBuildWindows.cs
--- a/Engine/Source/Programs/UnrealBuildTool/Windows/UEBuildWindows.cs
+++ b/Engine/Source/Programs/UnrealBuildTool/Windows/UEBuildWindows.cs
@@ -467,6 +467,7 @@ namespace UnrealBuildTool
                 InBuildTarget.GlobalLinkEnvironment.Config.AdditionalLibraries.Add("netapi32.lib");
                 InBuildTarget.GlobalLinkEnvironment.Config.AdditionalLibraries.Add("iphlpapi.lib");
                 InBuildTarget.GlobalLinkEnvironment.Config.AdditionalLibraries.Add("setupapi.lib"); //  Required for access monitor device enumeration
+                InBuildTarget.GlobalLinkEnvironment.Config.AdditionalLibraries.Add("hid.lib");
 
                 // Windows Vista/7 Desktop Windows Manager API for Slate Windows Compliance
                 if( !SupportWindowsXP )		// Windows XP does not support DWM
diff --git a/Engine/Source/Runtime/Core/Private/Windows/WindowsApplication.cpp b/Engine/Source/Runtime/Core/Private/Windows/WindowsApplication.cpp
--- a/Engine/Source/Runtime/Core/Private/Windows/WindowsApplication.cpp
+++ b/Engine/Source/Runtime/Core/Private/Windows/WindowsApplication.cpp
@@ -20,6 +20,8 @@
 	#include <shlobj.h>
 	#include <objbase.h>
 	#include <SetupApi.h>
+	#include <Hidsdi.h>
+	#include <devguid.h>
 
 // This might not be defined by Windows when maintaining backwards-compatibility to pre-Vista builds
 #ifndef WM_MOUSEHWHEEL
@@ -42,6 +44,7 @@ FWindowsApplication::FWindowsApplication( const HINSTANCE HInstance, const HICON
 	: GenericApplication( MakeShareable( new FWindowsCursor() ) )
 	, InstanceHandle( HInstance )
 	, bUsingHighPrecisionMouseInput( false )
+	, bIsMouseAttached( false )
 	, XInput( XInputInterface::Create( MessageHandler ) )
 	, CVarDeferMessageProcessing( 
 		TEXT( "Slate.DeferWindowsMessageProcessing" ),
@@ -122,6 +125,8 @@ FWindowsApplication::FWindowsApplication( const HINSTANCE HInstance, const HICON
 
 	// Disable accessibility shortcuts
 	AllowAccessibilityShortcutKeys(false);
+
+	QueryConnectedMice();
 }
 
 void FWindowsApplication::AllowAccessibilityShortcutKeys(const bool bAllowKeys)
@@ -1033,6 +1038,7 @@ int32 FWindowsApplication::ProcessMessage( HWND hwnd, uint32 msg, WPARAM wParam,
 		case WM_DEVICECHANGE:
 			{
 				XInput->SetNeedsControllerStateUpdate(); 
+				QueryConnectedMice();
 			}
 		}
 
@@ -1792,6 +1798,52 @@ HRESULT FWindowsApplication::OnOLEDrop( const HWND HWnd, const FDragDropOLEData&
 	return 0;
 }
 
+void FWindowsApplication::QueryConnectedMice()
+{
+	GUID HIDGuid;
+	GUID MouseGuid = GUID_DEVCLASS_MOUSE;
+	HDEVINFO hDevInfo;
+	BYTE TempBuffer[4096];
+
+	HidD_GetHidGuid(&HIDGuid);
+
+	hDevInfo = SetupDiGetClassDevs(&HIDGuid, NULL, NULL, DIGCF_PRESENT | DIGCF_DEVICEINTERFACE);
+
+	if (hDevInfo == INVALID_HANDLE_VALUE)
+	{
+		// in this case, on a desktop it's proably safer to assume that a mouse exists rather than doesn't
+		bIsMouseAttached = true;
+		return;
+	}
+
+	bIsMouseAttached = false;
+
+	for (DWORD DeviceIndex = 0; ; DeviceIndex++)
+	{
+		SP_DEVICE_INTERFACE_DATA DevInterface;
+		SP_DEVINFO_DATA DevInfo;
+		PSP_DEVICE_INTERFACE_DETAIL_DATA DevDetail;
+
+		DevInterface.cbSize = sizeof(DevInterface);
+		if (!SetupDiEnumDeviceInterfaces(hDevInfo, NULL, &HIDGuid, DeviceIndex, &DevInterface))
+		{
+			break;
+		}
+		DevDetail = (PSP_DEVICE_INTERFACE_DETAIL_DATA)TempBuffer;
+		DevDetail->cbSize = sizeof(SP_DEVICE_INTERFACE_DETAIL_DATA);
+		DevInfo.cbSize = sizeof(DevInfo);
+		if (!SetupDiGetDeviceInterfaceDetail(hDevInfo, &DevInterface, DevDetail, sizeof(TempBuffer), NULL, &DevInfo))
+		{
+			break;
+		}
+		if (IsEqualGUID(DevInfo.ClassGuid, MouseGuid))
+		{
+			bIsMouseAttached = true;
+			break;
+		}
+	}
+}
+
 void FTaskbarList::Initialize()
 {
 	if (FWindowsPlatformMisc::CoInitialize())
diff --git a/Engine/Source/Runtime/Core/Public/GenericPlatform/GenericApplication.h b/Engine/Source/Runtime/Core/Public/GenericPlatform/GenericApplication.h
--- a/Engine/Source/Runtime/Core/Public/GenericPlatform/GenericApplication.h
+++ b/Engine/Source/Runtime/Core/Public/GenericPlatform/GenericApplication.h
@@ -275,6 +275,8 @@ public:
 	
 	virtual bool IsUsingTrackpad() const { return false; }
 
+	virtual bool IsMouseAttached() const { return true; }
+
 	virtual FPlatformRect GetWorkArea( const FPlatformRect& CurrentWindow ) const
 	{
 		FPlatformRect OutRect;
diff --git a/Engine/Source/Runtime/Core/Public/Windows/WindowsApplication.h b/Engine/Source/Runtime/Core/Public/Windows/WindowsApplication.h
--- a/Engine/Source/Runtime/Core/Public/Windows/WindowsApplication.h
+++ b/Engine/Source/Runtime/Core/Public/Windows/WindowsApplication.h
@@ -240,6 +240,8 @@ public:
 
 	virtual bool IsUsingHighPrecisionMouseMode() const override { return bUsingHighPrecisionMouseInput; }
 
+	virtual bool IsMouseAttached() const override { return bIsMouseAttached; }
+
 	virtual FModifierKeysState GetModifierKeys() const override;
 
 	virtual FPlatformRect GetWorkArea( const FPlatformRect& CurrentWindow ) const override;
@@ -317,6 +319,8 @@ private:
 
 	void AllowAccessibilityShortcutKeys(const bool bAllowKeys);
 
+	void QueryConnectedMice();
+
 private:
 
 	static const FIntPoint MinimizedWindowPosition;
@@ -325,6 +329,8 @@ private:
 
 	bool bUsingHighPrecisionMouseInput;
 
+	bool bIsMouseAttached;
+
 	TArray< FDeferredWindowsMessage > DeferredMessages;
 
 	TArray< FDeferredWindowsDragDropOperation > DeferredDragDropOperations;
