CIS fix

diff --git a/Engine/Source/Runtime/Core/Private/Windows/WindowsApplication.cpp b/Engine/Source/Runtime/Core/Private/Windows/WindowsApplication.cpp
--- a/Engine/Source/Runtime/Core/Private/Windows/WindowsApplication.cpp
+++ b/Engine/Source/Runtime/Core/Private/Windows/WindowsApplication.cpp
@@ -264,7 +264,7 @@ FModifierKeysState FWindowsApplication::GetModifierKeys() const
 	const bool bIsLeftAltDown = ( ::GetAsyncKeyState( VK_LMENU ) & 0x8000 ) != 0;
 	const bool bIsRightAltDown = ( ::GetAsyncKeyState( VK_RMENU ) & 0x8000 ) != 0;
 
-	return FModifierKeysState( bIsLeftShiftDown, bIsRightShiftDown, bIsLeftControlDown, bIsRightControlDown, bIsLeftAltDown, bIsRightAltDown );
+	return FModifierKeysState( bIsLeftShiftDown, bIsRightShiftDown, bIsLeftControlDown, bIsRightControlDown, bIsLeftAltDown, bIsRightAltDown, false, false ); // Win key is ignored
 }
 
 void FWindowsApplication::SetCapture( const TSharedPtr< FGenericWindow >& InWindow )
