Fix for builds using Windows Kits 8.0

diff --git a/Engine/Source/Runtime/Core/Private/Windows/WindowsApplication.cpp b/Engine/Source/Runtime/Core/Private/Windows/WindowsApplication.cpp
--- a/Engine/Source/Runtime/Core/Private/Windows/WindowsApplication.cpp
+++ b/Engine/Source/Runtime/Core/Private/Windows/WindowsApplication.cpp
@@ -20,7 +20,11 @@
 	#include <shlobj.h>
 	#include <objbase.h>
 	#include <SetupApi.h>
+// The version of the Windows Kits 8.0 are missing this guard, 8.1 has a fix
+extern "C" 
+{	
 	#include <Hidsdi.h>
+}
 	#include <devguid.h>
 
 // This might not be defined by Windows when maintaining backwards-compatibility to pre-Vista builds
