One more CIS fix

diff --git a/Engine/Source/Runtime/Core/Private/WinRT/WinRTInputInterface.cpp b/Engine/Source/Runtime/Core/Private/WinRT/WinRTInputInterface.cpp
--- a/Engine/Source/Runtime/Core/Private/WinRT/WinRTInputInterface.cpp
+++ b/Engine/Source/Runtime/Core/Private/WinRT/WinRTInputInterface.cpp
@@ -142,6 +142,8 @@ void FWinRTInputInterface::ConditionalScanForKeyboardChanges( float DeltaTime )
 		(WinRTKeyState[VK_LCONTROL] == true) ? true : false,
 		(WinRTKeyState[VK_RCONTROL] == true) ? true : false,
 		false,
+		false,
+		false,
 		false
 		);
 
