Remove duplicate cvar definition for "r.allowhighqualitylightmaps"

diff --git a/Engine/Source/Runtime/Core/Private/HAL/ConsoleManager.cpp b/Engine/Source/Runtime/Core/Private/HAL/ConsoleManager.cpp
--- a/Engine/Source/Runtime/Core/Private/HAL/ConsoleManager.cpp
+++ b/Engine/Source/Runtime/Core/Private/HAL/ConsoleManager.cpp
@@ -1842,12 +1842,6 @@ static TAutoConsoleVariable<int32> CVarPreTileTextures(
 	TEXT("If set to 1, textures will be tiled during cook and are expected to be cooked at runtime"),
 	ECVF_RenderThreadSafe);
 
-static TAutoConsoleVariable<int32> CVarAllowHighQualityLightMaps(
-	TEXT("r.HighQualityLightMaps"),
-	1,
-	TEXT("If set to 1, allow high quality lightmaps which don't bake in direct lighting of stationary lights"),
-	ECVF_RenderThreadSafe);
-
 static TAutoConsoleVariable<int32> CVarVirtualTextureReducedMemoryEnabled(
 	TEXT("r.VirtualTextureReducedMemory"),
 	0,
