StableSort fix for small arrays. It was throwing an error for a corner case.

diff --git a/Engine/Source/Runtime/Core/Public/Templates/Sorting.h b/Engine/Source/Runtime/Core/Public/Templates/Sorting.h
--- a/Engine/Source/Runtime/Core/Public/Templates/Sorting.h
+++ b/Engine/Source/Runtime/Core/Public/Templates/Sorting.h
@@ -453,7 +453,7 @@ public:
 			{
 				for (int32 SubgroupStart = 0; SubgroupStart < Num; SubgroupStart += 2)
 				{
-					if (Predicate(First[SubgroupStart + 1], First[SubgroupStart]))
+					if (SubgroupStart + 1 < Num && Predicate(First[SubgroupStart + 1], First[SubgroupStart]))
 					{
 						Exchange(First[SubgroupStart], First[SubgroupStart + 1]);
 					}
@@ -461,24 +461,21 @@ public:
 			}
 		}
 
-		if (Num > MinMergeSubgroupSize)
+		int32 SubgroupSize = MinMergeSubgroupSize;
+		while (SubgroupSize < Num)
 		{
-			int32 SubgroupSize = MinMergeSubgroupSize << 1;
-			while (SubgroupSize < Num)
+			SubgroupStart = 0;
+			do
 			{
-				SubgroupStart = 0;
-				do
-				{
-					TMergePolicy::Merge(
-						First + SubgroupStart,
-						SubgroupSize,
-						FPlatformMath::Min(SubgroupSize << 1, Num - SubgroupStart),
-						Predicate);
-					SubgroupStart += SubgroupSize << 1;
-				} while (SubgroupStart < Num);
-
-				SubgroupSize <<= 1;
-			}
+				TMergePolicy::Merge(
+					First + SubgroupStart,
+					SubgroupSize,
+					FPlatformMath::Min(SubgroupSize << 1, Num - SubgroupStart),
+					Predicate);
+				SubgroupStart += SubgroupSize << 1;
+			} while (SubgroupStart < Num);
+
+			SubgroupSize <<= 1;
 		}
 	}
 };
