Fixed a problem with window that is already open not being brought to front on Mac if user used a menu option that opens this window.

diff --git a/Engine/Source/Runtime/Core/Private/Mac/MacWindow.cpp b/Engine/Source/Runtime/Core/Private/Mac/MacWindow.cpp
--- a/Engine/Source/Runtime/Core/Private/Mac/MacWindow.cpp
+++ b/Engine/Source/Runtime/Core/Private/Mac/MacWindow.cpp
@@ -172,9 +172,8 @@ TArray< FSlateCocoaWindow* > GRunningModalWindows;
 {
 	if ([NSApp isHidden] == NO)
 	{
-
-		bool bModal = GRunningModalWindows.Num() == 0 || GRunningModalWindows.Last() == self || [self styleMask] == NSBorderlessWindowMask;
-		if (![self isVisible] && bModal)
+		bool bBringToFront = GRunningModalWindows.Num() == 0 || GRunningModalWindows.Last() == self || [self styleMask] == NSBorderlessWindowMask;
+		if (bBringToFront)
 		{
 			[self orderFront:nil];
 		}
@@ -848,7 +847,7 @@ void FMacWindow::BringToFront( bool bForce )
 	if (bIsVisible)
 	{
 		SCOPED_AUTORELEASE_POOL;
-		[WindowHandle orderFrontAndMakeMain:false andKey:false];
+		[WindowHandle orderFrontAndMakeMain:true andKey:true];
 	}
 }
 
