Integrate force feedback in to XInputInterface

diff --git a/Engine/Source/Runtime/Core/Private/Windows/WindowsApplication.cpp b/Engine/Source/Runtime/Core/Private/Windows/WindowsApplication.cpp
--- a/Engine/Source/Runtime/Core/Private/Windows/WindowsApplication.cpp
+++ b/Engine/Source/Runtime/Core/Private/Windows/WindowsApplication.cpp
@@ -1714,6 +1714,10 @@ void FWindowsApplication::SetChannelValue (int32 ControllerId, FForceFeedbackCha
 
 void FWindowsApplication::SetChannelValues (int32 ControllerId, const FForceFeedbackValues &Values)
 {
+	const FForceFeedbackValues* InternalValues = &Values;
+ 
+ 	XInput->SetChannelValues( ControllerId, *InternalValues );
+ 
 	// send vibration to externally-implemented devices
 	for( auto DeviceIt = ExternalInputDevices.CreateIterator(); DeviceIt; ++DeviceIt )
 	{
diff --git a/Engine/Source/Runtime/Core/Private/Windows/XInputInterface.cpp b/Engine/Source/Runtime/Core/Private/Windows/XInputInterface.cpp
--- a/Engine/Source/Runtime/Core/Private/Windows/XInputInterface.cpp
+++ b/Engine/Source/Runtime/Core/Private/Windows/XInputInterface.cpp
@@ -204,7 +204,14 @@ void XInputInterface::SendControllerEvents()
 					// Update the state for next time
 					ControllerState.ButtonStates[ButtonIndex] = CurrentStates[ButtonIndex];
 				}	
-			}
+
+ 				// apply force feedback
+ 				XINPUT_VIBRATION VibrationState;
+ 
+ 				VibrationState.wLeftMotorSpeed = ( ::WORD ) ( ControllerState.ForceFeedback.LeftLarge * 65535.0f );
+ 				VibrationState.wRightMotorSpeed = ( ::WORD ) ( ControllerState.ForceFeedback.RightSmall * 65535.0f );
+ 
+ 				XInputSetState( ( ::DWORD ) ControllerState.ControllerId, &VibrationState );			}
 		}
 	}
 
@@ -215,3 +222,40 @@ void XInputInterface::SetMessageHandler( const TSharedRef< FGenericApplicationMe
 {
 	MessageHandler = InMessageHandler;
 }
+
+void XInputInterface::SetChannelValue( const int32 ControllerId, const FForceFeedbackChannelType ChannelType, const float Value )
+{
+	FControllerState& ControllerState = ControllerStates[ ControllerId ];
+
+	if( ControllerState.bIsConnected )
+	{
+		switch( ChannelType )
+		{
+			case FF_CHANNEL_LEFT_LARGE:
+				ControllerState.ForceFeedback.LeftLarge = Value;
+				break;
+
+			case FF_CHANNEL_LEFT_SMALL:
+				ControllerState.ForceFeedback.LeftSmall = Value;
+				break;
+
+			case FF_CHANNEL_RIGHT_LARGE:
+				ControllerState.ForceFeedback.RightLarge = Value;
+				break;
+
+			case FF_CHANNEL_RIGHT_SMALL:
+				ControllerState.ForceFeedback.RightSmall = Value;
+				break;
+		}
+	}
+}
+
+void XInputInterface::SetChannelValues( const int32 ControllerId, const FForceFeedbackValues &Values )
+{
+	FControllerState& ControllerState = ControllerStates[ ControllerId ];
+
+	if( ControllerState.bIsConnected )
+	{
+		ControllerState.ForceFeedback = Values;
+	}
+}
\ No newline at end of file
diff --git a/Engine/Source/Runtime/Core/Private/Windows/XInputInterface.h b/Engine/Source/Runtime/Core/Private/Windows/XInputInterface.h
--- a/Engine/Source/Runtime/Core/Private/Windows/XInputInterface.h
+++ b/Engine/Source/Runtime/Core/Private/Windows/XInputInterface.h
@@ -27,6 +27,25 @@ public:
 
 	void SetNeedsControllerStateUpdate() { bNeedsControllerStateUpdate = true; }
 
+	/**
+	* Sets the strength/speed of the given channel for the given controller id.
+	* NOTE: If the channel is not supported, the call will silently fail
+	*
+	* @param ControllerId the id of the controller whose value is to be set
+	* @param ChannelType the type of channel whose value should be set
+	* @param Value strength or speed of feedback, 0.0f to 1.0f. 0.0f will disable
+	*/
+	void SetChannelValue( const int32 ControllerId, const FForceFeedbackChannelType ChannelType, const float Value );
+
+	/**
+	* Sets the strength/speed of all the channels for the given controller id.
+	* NOTE: Unsupported channels are silently ignored
+	*
+	* @param ControllerId the id of the controller whose value is to be set
+	* @param Values strength or speed of feedback for all channels
+	*/
+	void SetChannelValues( const int32 ControllerId, const FForceFeedbackValues& Values );
+
 private:
 
 	XInputInterface( const TSharedRef< FGenericApplicationMessageHandler >& MessageHandler );
@@ -64,6 +83,8 @@ private:
 		/** If the controller is currently connected */
 		bool bIsConnected;
 		
+		/** Current force feedback values */
+		FForceFeedbackValues ForceFeedback;
 	};
 
 	/** If we've been notified by the system that the controller state may have changed */
