Logging verbosity docs

diff --git a/Engine/Source/Runtime/Core/Public/Misc/OutputDevice.h b/Engine/Source/Runtime/Core/Public/Misc/OutputDevice.h
--- a/Engine/Source/Runtime/Core/Public/Misc/OutputDevice.h
+++ b/Engine/Source/Runtime/Core/Public/Misc/OutputDevice.h
@@ -178,14 +178,44 @@ namespace ELogVerbosity
 {
 	enum Type
 	{
+		/** Not used */
 		NoLogging		= 0,
+
+		/** Always prints s fatal error to console (and log file) and crashes (even if logging is disabled) */
 		Fatal,
+
+		/** 
+		 * Prints an error to console (and log file). 
+		 * Commandlets and the editor collect and report errors. Error messages result in commandlet failure.
+		 */
 		Error,
+
+		/** 
+		 * Prints a warning to console (and log file).
+		 * Commandlets and the editor collect and report warnings. Warnings can be treated as an error.
+		 */
 		Warning,
+
+		/** Prints a message to console (and log file) */
 		Display,
+
+		/** Prints a message to a log file (does not print to console) */
 		Log,
+
+		/** 
+		 * Prints a verbose message to a log file (if Verbose logging is enabled for the given category, 
+		 * usually used for detailed logging) 
+		 */
 		Verbose,
+
+		/** 
+		 * Prints a verbose message to a log file (if VeryVerbose logging is enabled, 
+		 * usually used for detailed logging that would otherwise spam output) 
+		 */
 		VeryVerbose,
+
+		// Log masks and special Enum values
+
 		All				= VeryVerbose,
 		NumVerbosity,
 		VerbosityMask	= 0xf,
