Stats - Fixed an issue with not displaying self/otherchildren when using 'stat dumpframe'

diff --git a/Engine/Source/Runtime/Core/Private/Stats/StatsData.cpp b/Engine/Source/Runtime/Core/Private/Stats/StatsData.cpp
--- a/Engine/Source/Runtime/Core/Private/Stats/StatsData.cpp
+++ b/Engine/Source/Runtime/Core/Private/Stats/StatsData.cpp
@@ -1355,7 +1355,10 @@ FString FStatsUtils::DebugPrint(FStatMessage const& Item)
 
 	if( Desc.Len() && Desc != ShortName )
 	{
-		Desc += TEXT(" - ");
+		Desc += TEXT( " - " );
+	}
+	if( ShortName.Len() )
+	{
 		Desc += ShortName;
 	}
 
