Stats - Fixed previous check-in

diff --git a/Engine/Source/Runtime/Core/Private/Stats/StatsData.cpp b/Engine/Source/Runtime/Core/Private/Stats/StatsData.cpp
--- a/Engine/Source/Runtime/Core/Private/Stats/StatsData.cpp
+++ b/Engine/Source/Runtime/Core/Private/Stats/StatsData.cpp
@@ -1353,12 +1353,12 @@ FString FStatsUtils::DebugPrint(FStatMessage const& Item)
 	FString Desc = Item.NameAndInfo.GetDescription();
 	Desc.Trim();
 
-	if( Desc.Len() && Desc != ShortName )
-	{
-		Desc += TEXT( " - " );
-	}
-	if( ShortName.Len() )
+	if( Desc != ShortName )
 	{
+		if( Desc.Len() )
+		{
+			Desc += TEXT( " - " );
+		}
 		Desc += ShortName;
 	}
 
