Changed FMacWindow::BringToFront() so it only activates regular windows

diff --git a/Engine/Source/Runtime/Core/Private/Mac/MacWindow.cpp b/Engine/Source/Runtime/Core/Private/Mac/MacWindow.cpp
--- a/Engine/Source/Runtime/Core/Private/Mac/MacWindow.cpp
+++ b/Engine/Source/Runtime/Core/Private/Mac/MacWindow.cpp
@@ -847,7 +847,7 @@ void FMacWindow::BringToFront( bool bForce )
 	if (bIsVisible)
 	{
 		SCOPED_AUTORELEASE_POOL;
-		[WindowHandle orderFrontAndMakeMain:true andKey:true];
+		[WindowHandle orderFrontAndMakeMain:IsRegularWindow() andKey:IsRegularWindow()];
 	}
 }
 
