Disabled OpenGL info querying in MacPlatformSurvey for apps that don't use OpenGL at all

diff --git a/Engine/Source/Runtime/Core/Private/Mac/MacPlatformSurvey.cpp b/Engine/Source/Runtime/Core/Private/Mac/MacPlatformSurvey.cpp
--- a/Engine/Source/Runtime/Core/Private/Mac/MacPlatformSurvey.cpp
+++ b/Engine/Source/Runtime/Core/Private/Mac/MacPlatformSurvey.cpp
@@ -110,64 +110,67 @@ void FMacPlatformSurvey::TickSurveyHardware( FHardwareSurveyResults& OutResults
 	uint64_t TotalPhys = FreeMem + UsedMem;
 	OutResults.MemoryMB = uint32(float(TotalPhys/1024.0/1024.0)+ .1f);
 
-	// Get OpenGL version
-	GLint MajorVersion = 0;
-	GLint MinorVersion = 0;
-	glGetIntegerv(GL_MAJOR_VERSION, &MajorVersion);
-	glGetIntegerv(GL_MINOR_VERSION, &MinorVersion);
-	FString OpenGLVerString = FString::Printf(TEXT("%d.%d"), MajorVersion, MinorVersion);
-	WriteFStringToResults(OutResults.MultimediaAPI, OpenGLVerString);
-
-	// Display devices
-	OutResults.DisplayCount = 0;
-	CGDirectDisplayID Displays[FHardwareSurveyResults::MaxDisplayCount];
-	if (CGGetActiveDisplayList(FHardwareSurveyResults::MaxDisplayCount, Displays, &OutResults.DisplayCount) == CGDisplayNoErr)
+	if ([NSOpenGLContext currentContext])
 	{
-		for (uint32 Index = 0; Index < OutResults.DisplayCount; Index++)
+		// Get OpenGL version
+		GLint MajorVersion = 0;
+		GLint MinorVersion = 0;
+		glGetIntegerv(GL_MAJOR_VERSION, &MajorVersion);
+		glGetIntegerv(GL_MINOR_VERSION, &MinorVersion);
+		FString OpenGLVerString = FString::Printf(TEXT("%d.%d"), MajorVersion, MinorVersion);
+		WriteFStringToResults(OutResults.MultimediaAPI, OpenGLVerString);
+
+		// Display devices
+		OutResults.DisplayCount = 0;
+		CGDirectDisplayID Displays[FHardwareSurveyResults::MaxDisplayCount];
+		if (CGGetActiveDisplayList(FHardwareSurveyResults::MaxDisplayCount, Displays, &OutResults.DisplayCount) == CGDisplayNoErr)
 		{
-			FHardwareDisplay& Display = OutResults.Displays[Index];
-
-			Display.CurrentModeWidth = CGDisplayPixelsWide(Displays[Index]);
-			Display.CurrentModeHeight = CGDisplayPixelsHigh(Displays[Index]);
-
-			WriteFStringToResults(Display.GPUCardName, FString(StringCast<TCHAR>((const char*)glGetString(GL_RENDERER)).Get()));
-			WriteFStringToResults(Display.GPUDriverVersion, FString(StringCast<TCHAR>((const char*)glGetString(GL_VERSION)).Get()));
-
-			CGLRendererInfoObj InfoObject;
-			GLint NumRenderers = 0;
-			if (CGLQueryRendererInfo(CGDisplayIDToOpenGLDisplayMask(Displays[Index]), &InfoObject, &NumRenderers) == kCGLNoError && NumRenderers > 0)
-			{
-				CGLDescribeRenderer(InfoObject, 0, kCGLRPVideoMemoryMegabytes, (GLint*)&Display.GPUDedicatedMemoryMB);
-				CGLDestroyRendererInfo(InfoObject);
-			}
-			else
+			for (uint32 Index = 0; Index < OutResults.DisplayCount; Index++)
 			{
-				UE_LOG(LogMac, Warning, TEXT("FMacPlatformSurvey::TickSurveyHardware() failed to query renderer info for disaplay %u"), Index );
-				OutResults.ErrorCount++;
-				WriteFStringToResults(OutResults.LastSurveyError, TEXT("Failed to query renderer info"));
-				WriteFStringToResults(OutResults.LastSurveyErrorDetail, TEXT(""));
+				FHardwareDisplay& Display = OutResults.Displays[Index];
+
+				Display.CurrentModeWidth = CGDisplayPixelsWide(Displays[Index]);
+				Display.CurrentModeHeight = CGDisplayPixelsHigh(Displays[Index]);
+
+				WriteFStringToResults(Display.GPUCardName, FString(StringCast<TCHAR>((const char*)glGetString(GL_RENDERER)).Get()));
+				WriteFStringToResults(Display.GPUDriverVersion, FString(StringCast<TCHAR>((const char*)glGetString(GL_VERSION)).Get()));
+
+				CGLRendererInfoObj InfoObject;
+				GLint NumRenderers = 0;
+				if (CGLQueryRendererInfo(CGDisplayIDToOpenGLDisplayMask(Displays[Index]), &InfoObject, &NumRenderers) == kCGLNoError && NumRenderers > 0)
+				{
+					CGLDescribeRenderer(InfoObject, 0, kCGLRPVideoMemoryMegabytes, (GLint*)&Display.GPUDedicatedMemoryMB);
+					CGLDestroyRendererInfo(InfoObject);
+				}
+				else
+				{
+					UE_LOG(LogMac, Warning, TEXT("FMacPlatformSurvey::TickSurveyHardware() failed to query renderer info for disaplay %u"), Index );
+					OutResults.ErrorCount++;
+					WriteFStringToResults(OutResults.LastSurveyError, TEXT("Failed to query renderer info"));
+					WriteFStringToResults(OutResults.LastSurveyErrorDetail, TEXT(""));
+				}
 			}
 		}
-	}
-	else
-	{
-		UE_LOG(LogMac, Warning, TEXT("FMacPlatformSurvey::TickSurveyHardware() failed to get active displays list") );
-		OutResults.ErrorCount++;
-		WriteFStringToResults(OutResults.LastSurveyError, TEXT("Failed to get active displays list"));
-		WriteFStringToResults(OutResults.LastSurveyErrorDetail, TEXT(""));
-	}
+		else
+		{
+			UE_LOG(LogMac, Warning, TEXT("FMacPlatformSurvey::TickSurveyHardware() failed to get active displays list") );
+			OutResults.ErrorCount++;
+			WriteFStringToResults(OutResults.LastSurveyError, TEXT("Failed to get active displays list"));
+			WriteFStringToResults(OutResults.LastSurveyErrorDetail, TEXT(""));
+		}
 
-	if (OutResults.DisplayCount == 0)
-	{
-		OutResults.ErrorCount++;
-		WriteFStringToResults(OutResults.LastSurveyError, TEXT("Display count zero"));
-		WriteFStringToResults(OutResults.LastSurveyErrorDetail, TEXT(""));
-	}
-	else if (OutResults.DisplayCount > 3)
-	{
-		OutResults.ErrorCount++;
-		WriteFStringToResults(OutResults.LastSurveyError, FString::Printf(TEXT("Display count %d"), OutResults.DisplayCount));
-		WriteFStringToResults(OutResults.LastSurveyErrorDetail, TEXT(""));
+		if (OutResults.DisplayCount == 0)
+		{
+			OutResults.ErrorCount++;
+			WriteFStringToResults(OutResults.LastSurveyError, TEXT("Display count zero"));
+			WriteFStringToResults(OutResults.LastSurveyErrorDetail, TEXT(""));
+		}
+		else if (OutResults.DisplayCount > 3)
+		{
+			OutResults.ErrorCount++;
+			WriteFStringToResults(OutResults.LastSurveyError, FString::Printf(TEXT("Display count %d"), OutResults.DisplayCount));
+			WriteFStringToResults(OutResults.LastSurveyErrorDetail, TEXT(""));
+		}
 	}
 
 	// Get CPU count
