Linux: make module manager actually load modules.

- part of GitHub pull request #242, with changes
- new platform define for conditionally using cpuid, other platforms may want to follow suit
- better detection of debugger, IsDebuggerPresent() is on again.

diff --git a/Engine/Source/Runtime/Core/Private/Linux/LinuxPlatformMisc.cpp b/Engine/Source/Runtime/Core/Private/Linux/LinuxPlatformMisc.cpp
--- a/Engine/Source/Runtime/Core/Private/Linux/LinuxPlatformMisc.cpp
+++ b/Engine/Source/Runtime/Core/Private/Linux/LinuxPlatformMisc.cpp
@@ -19,10 +19,13 @@
 #include <libdwarf.h>
 #include <dwarf.h>
 
+#include "ModuleManager.h"
+
 #include "../../Launch/Resources/Version.h"
 #include <sys/utsname.h>	// for uname()
-#include <sys/ptrace.h>		// for ptrace()
-#include <sys/wait.h>		// for waitpid()
+
+// define for glibc 2.12.2 and lower (which is shipped with CentOS 6.x and which we target by default)
+#define __secure_getenv getenv
 
 namespace PlatformMiscLimits
 {
@@ -58,7 +61,7 @@ namespace
 	}
 }
 
-bool FLinuxMisc::ControlScreensaver(EScreenSaverAction Action)
+bool FLinuxPlatformMisc::ControlScreensaver(EScreenSaverAction Action)
 {
 	if (Action == FGenericPlatformMisc::EScreenSaverAction::Disable)
 	{
@@ -72,13 +75,13 @@ bool FLinuxMisc::ControlScreensaver(EScreenSaverAction Action)
 	return true;
 }
 
-const TCHAR* FLinuxMisc::RootDir()
+const TCHAR* FLinuxPlatformMisc::RootDir()
 {
 	const TCHAR * TrueRootDir = FGenericPlatformMisc::RootDir();
 	return TrueRootDir;
 }
 
-void FLinuxMisc::NormalizePath(FString& InPath)
+void FLinuxPlatformMisc::NormalizePath(FString& InPath)
 {
 	if (InPath.Contains(TEXT("~"), ESearchCase::CaseSensitive))	// case sensitive is quicker, and our substring doesn't care
 	{
@@ -119,7 +122,7 @@ void FLinuxMisc::NormalizePath(FString& InPath)
 	}
 }
 
-void FLinuxMisc::PlatformInit()
+void FLinuxPlatformMisc::PlatformInit()
 {
 	// install a platform-specific signal handler
 	InstallChildExitedSignalHanlder();
@@ -145,12 +148,12 @@ void FLinuxMisc::PlatformInit()
 	UE_LOG(LogInit, Log, TEXT(" -virtmemkb=NUMBER - sets process virtual memory (address space) limit (overrides VirtualMemoryLimitInKB value from .ini)"));
 }
 
-GenericApplication* FLinuxMisc::CreateApplication()
+GenericApplication* FLinuxPlatformMisc::CreateApplication()
 {
 	return FLinuxApplication::CreateLinuxApplication();
 }
 
-void FLinuxMisc::PumpMessages( bool bFromMainLoop )
+void FLinuxPlatformMisc::PumpMessages( bool bFromMainLoop )
 {
 	if( bFromMainLoop )
 	{
@@ -166,12 +169,12 @@ void FLinuxMisc::PumpMessages( bool bFromMainLoop )
 	}
 }
 
-uint32 FLinuxMisc::GetCharKeyMap(uint16* KeyCodes, FString* KeyNames, uint32 MaxMappings)
+uint32 FLinuxPlatformMisc::GetCharKeyMap(uint16* KeyCodes, FString* KeyNames, uint32 MaxMappings)
 {
 	return FGenericPlatformMisc::GetStandardPrintableKeyMap(KeyCodes, KeyNames, MaxMappings, false, true);
 }
 
-uint32 FLinuxMisc::GetKeyMap( uint16* KeyCodes, FString* KeyNames, uint32 MaxMappings )
+uint32 FLinuxPlatformMisc::GetKeyMap( uint16* KeyCodes, FString* KeyNames, uint32 MaxMappings )
 {
 #define ADDKEYMAP(KeyCode, KeyName)		if (NumMappings<MaxMappings) { KeyCodes[NumMappings]=KeyCode; KeyNames[NumMappings]=KeyName; ++NumMappings; };
 
@@ -227,7 +230,7 @@ uint32 FLinuxMisc::GetKeyMap( uint16* KeyCodes, FString* KeyNames, uint32 MaxMap
 	return NumMappings;
 }
 
-int32 FLinuxMisc::NumberOfCores()
+int32 FLinuxPlatformMisc::NumberOfCores()
 {
 	cpu_set_t AvailableCpusMask;
 	CPU_ZERO(&AvailableCpusMask);
@@ -275,7 +278,7 @@ int32 FLinuxMisc::NumberOfCores()
 	return NumCoreIds;
 }
 
-int32 FLinuxMisc::NumberOfCoresIncludingHyperthreads()
+int32 FLinuxPlatformMisc::NumberOfCoresIncludingHyperthreads()
 {
 	cpu_set_t AvailableCpusMask;
 	CPU_ZERO(&AvailableCpusMask);
@@ -288,6 +291,66 @@ int32 FLinuxMisc::NumberOfCoresIncludingHyperthreads()
 	return CPU_COUNT(&AvailableCpusMask);
 }
 
+void FLinuxPlatformMisc::LoadPreInitModules()
+{
+#if WITH_EDITOR
+	FModuleManager::Get().LoadModule(TEXT("OpenGLDrv"));
+#endif // WITH_EDITOR
+}
+
+void FLinuxPlatformMisc::LoadStartupModules()
+{
+#if !IS_PROGRAM
+	FModuleManager::Get().LoadModule(TEXT("ALAudio"));
+	FModuleManager::Get().LoadModule(TEXT("SteamController"));
+	FModuleManager::Get().LoadModule(TEXT("HeadMountedDisplay"));
+#endif // !IS_PROGRAM
+
+#if WITH_EDITOR
+	FModuleManager::Get().LoadModule(TEXT("SourceCodeAccess"));
+#endif	//WITH_EDITOR
+}
+
+const TCHAR* FLinuxPlatformMisc::GetNullRHIShaderFormat()
+{
+	return TEXT("GLSL_150");
+}
+
+#if PLATFORM_HAS_CPUID
+FString FLinuxPlatformMisc::GetCPUVendor()
+{
+	union
+	{
+		char Buffer[12+1];
+		struct
+		{
+			int dw0;
+			int dw1;
+			int dw2;
+		} Dw;
+	} VendorResult;
+
+
+	int32 Args[4];
+	asm( "cpuid" : "=a" (Args[0]), "=b" (Args[1]), "=c" (Args[2]), "=d" (Args[3]) : "a" (0));
+
+	VendorResult.Dw.dw0 = Args[1];
+	VendorResult.Dw.dw1 = Args[3];
+	VendorResult.Dw.dw2 = Args[2];
+	VendorResult.Buffer[12] = 0;
+
+	return ANSI_TO_TCHAR(VendorResult.Buffer);
+}
+
+uint32 FLinuxPlatformMisc::GetCPUInfo()
+{
+	uint32 Args[4];
+	asm( "cpuid" : "=a" (Args[0]), "=b" (Args[1]), "=c" (Args[2]), "=d" (Args[3]) : "a" (1));
+
+	return Args[0];
+}
+#endif // PLATFORM_HAS_CPUID
+
 FString DescribeSignal(int32 Signal, siginfo_t* Info)
 {
 	FString ErrorString;
@@ -1067,7 +1130,7 @@ void PlatformCrashHandler(int32 Signal, siginfo_t* Info, void* Context)
 	}
 }
 
-void FLinuxMisc::SetGracefulTerminationHandler()
+void FLinuxPlatformMisc::SetGracefulTerminationHandler()
 {
 	struct sigaction Action;
 	FMemory::Memzero(&Action, sizeof(struct sigaction));
@@ -1079,7 +1142,7 @@ void FLinuxMisc::SetGracefulTerminationHandler()
 	sigaction(SIGHUP, &Action, NULL);	//  this should actually cause the server to just re-read configs (restart?)
 }
 
-void FLinuxMisc::SetCrashHandler(void (* CrashHandler)(const FGenericCrashContext & Context))
+void FLinuxPlatformMisc::SetCrashHandler(void (* CrashHandler)(const FGenericCrashContext & Context))
 {
 	GCrashHandlerPointer = CrashHandler;
 
@@ -1097,75 +1160,40 @@ void FLinuxMisc::SetCrashHandler(void (* CrashHandler)(const FGenericCrashContex
 }
 
 #if !UE_BUILD_SHIPPING
-
-/** Heavy! */
-bool FLinuxMisc::IsDebuggerPresent()
+bool FLinuxPlatformMisc::IsDebuggerPresent()
 {
-	// @TODO: Make this less heavy. It currently causes a crash loop when using "debug crash"
-	return false;
-	/*
-	int ChildPid = fork();
-	if (ChildPid == -1)
-	{
-		UE_LOG(LogInit, Fatal, TEXT("IsDebuggerPresent(): fork() failed with errno=%d (%s)."), errno, strerror(errno));
-
-		// unreachable
-		return false;
-	}
+	// If a process is tracing this one then TracerPid in /proc/self/status will
+	// be the id of the tracing process. Use SignalHandler safe functions 
 
-	if (ChildPid == 0)
+	int StatusFile = open("/proc/self/status", O_RDONLY);
+	if (StatusFile == -1) 
 	{
-		// we are child, try attaching to parent
-		int Parent = getppid();
-		if (ptrace(PTRACE_ATTACH, Parent, NULL, NULL) != 0)
-		{
-			// couldn't attach, assume gdb
-			exit(1);
-			
-			// unreachable
-			return 1;
-		}
-
-		// succeeded - no gdb
-
-		// traced process will be stopped at this point, wait for it, continue and detach
-		waitpid(Parent, NULL, 0);
-		ptrace(PTRACE_CONT, NULL, NULL);
-		ptrace(PTRACE_DETACH, Parent, NULL, NULL);
-		exit(0);
-		
-		// unreachable
-		return 0;
+		// Failed - unknown debugger status.
+		return false;
 	}
 
-	// parent
-	check(ChildPid != 0);
+	char Buffer[256];	
+	ssize_t Length = read(StatusFile, Buffer, sizeof(Buffer));
+	
+	bool bDebugging = false;
+	const char* TracerString = "TracerPid:\t";
+	const ssize_t LenTracerString = strlen(TracerString);
+	int i = 0;
 
-	int Status = 0;
-	for(;;)
+	while((Length - i) > LenTracerString)
 	{
-		if (waitpid(ChildPid, &Status, 0) == -1)
-		{
-			// error waiting for child! EINTR is allowed though, will retry
-			if (errno == EINTR)
-			{
-				continue;
-			}
-
-			UE_LOG(LogInit, Fatal, TEXT("IsDebuggerPresent(): waitpid() failed with errno=%d (%s)."), errno, strerror(errno));
-
-			// unreachable
-			return false;
-		}
-		else
+		// TracerPid is found
+		if (strncmp(&Buffer[i], TracerString, LenTracerString) == 0)
 		{
+			// 0 if no process is tracing.
+			bDebugging = Buffer[i + LenTracerString] != '0';
 			break;
 		}
+
+		++i;
 	}
 
-	int RetCode = WEXITSTATUS(Status);
-	return RetCode != 0;
-	*/
+	close(StatusFile);
+	return bDebugging;
 }
-
 #endif // !UE_BUILD_SHIPPING
diff --git a/Engine/Source/Runtime/Core/Private/Modules/ModuleManager.cpp b/Engine/Source/Runtime/Core/Private/Modules/ModuleManager.cpp
--- a/Engine/Source/Runtime/Core/Private/Modules/ModuleManager.cpp
+++ b/Engine/Source/Runtime/Core/Private/Modules/ModuleManager.cpp
@@ -350,7 +350,7 @@ TSharedPtr<IModuleInterface> FModuleManager::LoadModuleWithFailureReason( const
 			UE_LOG(LogModuleManager, Verbose, TEXT( "ModuleManager: Load Module '%s' DLL '%s'" ), *InModuleName.ToString(), *ModuleInfo->Filename);
 
 			// Determine which file to load for this module.
-			const FString ModuleFileToLoad = ModuleInfo->Filename;
+			const FString ModuleFileToLoad = FPaths::ConvertRelativePathToFull(ModuleInfo->Filename);
 
 			// Clear the handle and set it again below if the module is successfully loaded
 			ModuleInfo->Handle = NULL;
@@ -828,6 +828,9 @@ FString FModuleManager::GetSolutionFilepath()
 {
 #if PLATFORM_MAC
 	const TCHAR* Postfix = TEXT(".xcodeproj/project.pbxproj");
+#elif PLATFORM_LINUX
+	UE_LOG(LogModuleManager, Warning, TEXT("STUBBED: solution file path for Linux"));
+	const TCHAR* Postfix = TEXT("/stubbed/path/to.solution");
 #else
 	const TCHAR* Postfix = TEXT(".sln");
 #endif
@@ -1162,7 +1165,7 @@ void FModuleManager::FindModulePathsInDirectory(const FString& InDirectoryName,
 	}
 
 	// Get the base name for modules of this application
-	FString ModulePrefix = FPaths::GetBaseFilename(FPlatformProcess::ExecutableName());
+	FString ModulePrefix = FPlatformProcess::GetModulePrefix() + FPaths::GetBaseFilename(FPlatformProcess::ExecutableName());
 	if (ModulePrefix.Contains(TEXT("-")))
 	{
 		ModulePrefix = ModulePrefix.Left(ModulePrefix.Find(TEXT("-")) + 1);
@@ -1334,6 +1337,8 @@ bool FModuleManager::GenerateCodeProjectFiles( const FString& ProjectFilename, F
 
 #if PLATFORM_MAC
 	FString CmdLineParams = TEXT("-xcodeprojectfile");
+#elif PLATFORM_LINUX
+	FString CmdLineParams = TEXT("-makefile");
 #else
 	FString CmdLineParams = TEXT("-projectfiles");
 #endif
@@ -1898,6 +1903,8 @@ bool FUBTInvoker::BuildUnrealBuildTool(FOutputDevice &Ar)
 	FString ScriptPath = FPaths::ConvertRelativePathToFull(FPaths::EngineDir() / TEXT("Build/BatchFiles/Mac/RunXBuild.sh"));
 	CompilerExecutableFilename = TEXT("/bin/sh");
 	CmdLineParams = FString::Printf(TEXT("\"%s\" /property:Configuration=Development %s"), *ScriptPath, *CsProjLocation);
+#elif PLATFORM_LINUX
+	printf("ModuleManager.cpp: TODO: Linux BuildUBT");
 #else
 	Ar.Log(TEXT("Unknown platform, unable to build UnrealBuildTool."));
 #endif
diff --git a/Engine/Source/Runtime/Core/Public/HAL/Platform.h b/Engine/Source/Runtime/Core/Public/HAL/Platform.h
--- a/Engine/Source/Runtime/Core/Public/HAL/Platform.h
+++ b/Engine/Source/Runtime/Core/Public/HAL/Platform.h
@@ -114,6 +114,13 @@
 #ifndef PLATFORM_ENABLE_VECTORINTRINSICS
 	#define PLATFORM_ENABLE_VECTORINTRINSICS	0
 #endif
+#ifndef PLATFORM_HAS_CPUID
+	#if defined(_M_IX86) || defined(__i386__)
+		#define PLATFORM_HAS_CPUID				1
+	#else
+		#define PLATFORM_HAS_CPUID				0
+	#endif
+#endif	
 #ifndef PLATFORM_ENABLE_VECTORINTRINSICS_NEON
 	#define PLATFORM_ENABLE_VECTORINTRINSICS_NEON	0
 #endif
diff --git a/Engine/Source/Runtime/Core/Public/Linux/LinuxPlatformMisc.h b/Engine/Source/Runtime/Core/Public/Linux/LinuxPlatformMisc.h
--- a/Engine/Source/Runtime/Core/Public/Linux/LinuxPlatformMisc.h
+++ b/Engine/Source/Runtime/Core/Public/Linux/LinuxPlatformMisc.h
@@ -83,7 +83,7 @@ struct CORE_API FLinuxCrashContext : public FGenericCrashContext
 /**
  * Linux implementation of the misc OS functions
  */
-struct CORE_API FLinuxMisc : public FGenericPlatformMisc
+struct CORE_API FLinuxPlatformMisc : public FGenericPlatformMisc
 {
 	static void PlatformInit();
 	static void SetGracefulTerminationHandler();
@@ -120,6 +120,40 @@ struct CORE_API FLinuxMisc : public FGenericPlatformMisc
 
 	static int32 NumberOfCores();
 	static int32 NumberOfCoresIncludingHyperthreads();
+	static void LoadPreInitModules();
+	static void LoadStartupModules();
+
+	/**
+	 * Determines the shader format for the platform
+	 *
+	 * @return	Returns the shader format to be used by that platform
+	 */
+	static const TCHAR* GetNullRHIShaderFormat();
+
+#if PLATFORM_HAS_CPUID
+	/**
+	 * Uses cpuid instruction to get the vendor string
+	 *
+	 * @return	CPU vendor name
+	 */
+	static FString GetCPUVendor();
+
+	/**
+	 * Uses cpuid instruction to get the vendor string
+	 *
+	 * @return	CPU info bitfield
+	 *
+	 *			Bits 0-3	Stepping ID
+	 *			Bits 4-7	Model
+	 *			Bits 8-11	Family
+	 *			Bits 12-13	Processor type (Intel) / Reserved (AMD)
+	 *			Bits 14-15	Reserved
+	 *			Bits 16-19	Extended model
+	 *			Bits 20-27	Extended family
+	 *			Bits 28-31	Reserved
+	 */
+	static uint32 GetCPUInfo();
+#endif // PLATFORM_HAS_CPUID
 
 	static const TCHAR* EngineDir()
 	{
@@ -127,4 +161,4 @@ struct CORE_API FLinuxMisc : public FGenericPlatformMisc
 	}
 };
 
-typedef FLinuxMisc FPlatformMisc;
+typedef FLinuxPlatformMisc FPlatformMisc;
