Call FGenericApplication::BroadcastDisplayMetricsChanged when the display area has changed on OS X.

diff --git a/Engine/Source/Runtime/Core/Private/Mac/MacApplication.cpp b/Engine/Source/Runtime/Core/Private/Mac/MacApplication.cpp
--- a/Engine/Source/Runtime/Core/Private/Mac/MacApplication.cpp
+++ b/Engine/Source/Runtime/Core/Private/Mac/MacApplication.cpp
@@ -41,6 +41,14 @@ FMacApplication* FMacApplication::CreateMacApplication()
 void FMacApplication::OnDisplayReconfiguration(CGDirectDisplayID Display, CGDisplayChangeSummaryFlags Flags, void* UserInfo)
 {
 	FMacApplication* App = (FMacApplication*)UserInfo;
+	if(Flags & kCGDisplayDesktopShapeChangedFlag)
+	{
+		// Slate needs to know when desktop size changes.
+		FDisplayMetrics DisplayMetrics;
+		App->GetDisplayMetrics( DisplayMetrics );
+		App->BroadcastDisplayMetricsChanged(DisplayMetrics);
+	}
+	
 	for (int32 WindowIndex=0; WindowIndex < App->Windows.Num(); ++WindowIndex)
 	{
 		TSharedRef< FMacWindow > WindowRef = App->Windows[ WindowIndex ];
