Fix potential UBT stall when invoking UBT syncronously.

diff --git a/Engine/Source/Runtime/Core/Private/Modules/ModuleManager.cpp b/Engine/Source/Runtime/Core/Private/Modules/ModuleManager.cpp
--- a/Engine/Source/Runtime/Core/Private/Modules/ModuleManager.cpp
+++ b/Engine/Source/Runtime/Core/Private/Modules/ModuleManager.cpp
@@ -1951,10 +1951,14 @@ bool FUBTInvoker::InvokeUnrealBuildToolSync(const FString& InCmdLineParams, FOut
 	FProcHandle ProcHandle = InvokeUnrealBuildToolAsync(InCmdLineParams, Ar, PipeRead, PipeWrite, bSkipBuildUBT);
 	if (ProcHandle.IsValid())
 	{
-		FPlatformProcess::WaitForProc(ProcHandle);
+		// rather than waiting, we must flush the read pipe or UBT will stall if it writes out a ton of text to the console.
+		while (FPlatformProcess::IsProcRunning(ProcHandle))
+		{
+			OutProcOutput += FPlatformProcess::ReadPipe(PipeRead);
+			FPlatformProcess::Sleep(0.1f);
+		}		
 		bInvoked = true;
-		bool bGotReturnCode = FPlatformProcess::GetProcReturnCode(ProcHandle, &OutReturnCode);
-		OutProcOutput = FPlatformProcess::ReadPipe(PipeRead);
+		bool bGotReturnCode = FPlatformProcess::GetProcReturnCode(ProcHandle, &OutReturnCode);		
 		check(bGotReturnCode);
 	}
 	else
