Fixing iOS build - VectorGetComponent was missing from UnrealMathNeon

diff --git a/Engine/Source/Runtime/Core/Public/Math/UnrealMathNeon.h b/Engine/Source/Runtime/Core/Public/Math/UnrealMathNeon.h
--- a/Engine/Source/Runtime/Core/Public/Math/UnrealMathNeon.h
+++ b/Engine/Source/Runtime/Core/Public/Math/UnrealMathNeon.h
@@ -611,6 +611,23 @@ FORCEINLINE VectorRegister VectorNormalize( const VectorRegister& Vector )
 */
 #define VectorSet_W1( Vec )		VectorSetComponent( Vec, 3, 1.0f )
 
+
+/**
+* Returns a component from a vector.
+*
+* @param Vec				Vector register
+* @param ComponentIndex	Which component to get, X=0, Y=1, Z=2, W=3
+* @return					The component as a float
+*/
+FORCEINLINE float VectorGetComponent(VectorRegister Vec, uint32 ComponentIndex)
+{
+	float Tmp[4];
+	VectorStore(Vec, Tmp);
+	return Tmp[ComponentIndex];
+}
+
+
+
 /**
  * Multiplies two 4x4 matrices.
  *
