Fixed few warnings in MacPlatformMisc.cpp

diff --git a/Engine/Source/Runtime/Core/Private/Mac/MacPlatformMisc.cpp b/Engine/Source/Runtime/Core/Private/Mac/MacPlatformMisc.cpp
--- a/Engine/Source/Runtime/Core/Private/Mac/MacPlatformMisc.cpp
+++ b/Engine/Source/Runtime/Core/Private/Mac/MacPlatformMisc.cpp
@@ -424,7 +424,7 @@ TArray<uint8> FMacPlatformMisc::GetMacAddress()
 	}
 
 	io_object_t InterfaceService;
-	while ( (InterfaceService = IOIteratorNext(InterfaceIterator)) != NULL)
+	while ( (InterfaceService = IOIteratorNext(InterfaceIterator)) != 0)
 	{
 		io_object_t ControllerService;
 		if (IORegistryEntryGetParentEntry(InterfaceService, kIOServicePlane, &ControllerService) == KERN_SUCCESS)
@@ -1107,7 +1107,7 @@ void FMacPlatformMisc::SetCrashHandler(void (* CrashHandler)(const FGenericCrash
 void FMacCrashContext::GenerateReport(char const* DiagnosticsPath) const
 {
 	int ReportFile = open(DiagnosticsPath, O_CREAT|O_WRONLY, 0766);
-	if (ReportFile != NULL)
+	if (ReportFile != -1)
 	{
 		char Line[PATH_MAX] = {};
 		
@@ -1183,7 +1183,7 @@ void FMacCrashContext::GenerateReport(char const* DiagnosticsPath) const
 void FMacCrashContext::GenerateWindowsErrorReport(char const* WERPath) const
 {
 	int ReportFile = open(WERPath, O_CREAT|O_WRONLY, 0766);
-	if (ReportFile != NULL)
+	if (ReportFile != -1)
 	{
 		// write BOM
 		static uint16 ByteOrderMarker = 0xFEFF;
@@ -1305,7 +1305,7 @@ void FMacCrashContext::GenerateWindowsErrorReport(char const* WERPath) const
 void FMacCrashContext::GenerateMinidump(char const* MinidumpCallstackInfo, char const* Path) const
 {
 	int ReportFile = open(Path, O_CREAT|O_WRONLY, 0766);
-	if (ReportFile != NULL)
+	if (ReportFile != -1)
 	{
 		// write BOM
 		static uint16 ByteOrderMarker = 0xFEFF;
@@ -1337,7 +1337,7 @@ void FMacCrashContext::GenerateCrashInfoAndLaunchReporter() const
 		FCStringAnsi::Strncpy(FilePath, CrashInfoFolder, PATH_MAX);
 		FCStringAnsi::Strcat(FilePath, PATH_MAX, "/report.wer");
 		int ReportFile = open(FilePath, O_CREAT|O_WRONLY, 0766);
-		if (ReportFile != NULL)
+		if (ReportFile != -1)
 		{
 			// write BOM
 			static uint16 ByteOrderMarker = 0xFEFF;
@@ -1369,7 +1369,7 @@ void FMacCrashContext::GenerateCrashInfoAndLaunchReporter() const
 		FCStringAnsi::Strncpy(FilePath, CrashInfoFolder, PATH_MAX);
 		FCStringAnsi::Strcat(FilePath, PATH_MAX, "/info.txt");
 		ReportFile = open(FilePath, O_CREAT|O_WRONLY, 0766);
-		if (ReportFile != NULL)
+		if (ReportFile != -1)
 		{
 			WriteUTF16String(ReportFile, TEXT("GameName UE4-"));
 			WriteLine(ReportFile, *GMacAppInfo.AppName);
