Implementing stub VectorGetComponent for  DX math to fix build breakage; actual implementation will follow

diff --git a/Engine/Source/Runtime/Core/Public/Math/UnrealMathDirectX.h b/Engine/Source/Runtime/Core/Public/Math/UnrealMathDirectX.h
--- a/Engine/Source/Runtime/Core/Public/Math/UnrealMathDirectX.h
+++ b/Engine/Source/Runtime/Core/Public/Math/UnrealMathDirectX.h
@@ -193,6 +193,20 @@ FORCEINLINE VectorRegister MakeVectorRegister( float X, float Y, float Z, float
  */
 #define VectorStoreFloat1( Vec, Ptr )	DirectX::XMStoreFloat((float*)(Ptr), Vec )
  
+
+/**
+* Returns an component from a vector.  May be slow, use with caution.
+*
+* @param Vec				Vector register
+* @param ComponentIndex	Which component to get, X=0, Y=1, Z=2, W=3
+* @return					The component as a float
+*/
+FORCEINLINE float VectorGetComponent(VectorRegister Vec, uint32 ComponentIndex)
+{
+	// stub for now, just to fix build failure
+}
+
+
 /**
  * Replicates one element into all four elements and returns the new vector.
  *
