fix for Android build error with VectorGetComponent
#ue4
#android

diff --git a/Engine/Source/Runtime/Core/Public/Math/UnrealMathFPU.h b/Engine/Source/Runtime/Core/Public/Math/UnrealMathFPU.h
--- a/Engine/Source/Runtime/Core/Public/Math/UnrealMathFPU.h
+++ b/Engine/Source/Runtime/Core/Public/Math/UnrealMathFPU.h
@@ -931,6 +931,15 @@ FORCEINLINE uint32 VectorAnyGreaterThan(const VectorRegister& Vec1, const Vector
  */
 #define VectorGetControlRegister()		0
 
+/**
+ * Returns an component from a vector.
+ *
+ * @param Vec				Vector register
+ * @param ComponentIndex	Which component to get, X=0, Y=1, Z=2, W=3
+ * @return					The component as a float
+ */
+#define VectorGetComponent( Vec, ComponentIndex )	(((float*) &(Vec))[ComponentIndex])
+
 /**
  * Sets the control register.
  *
