Fix for map emplacement where the key has an explicit constructor.

diff --git a/Engine/Source/Runtime/Core/Public/Containers/Map.h b/Engine/Source/Runtime/Core/Public/Containers/Map.h
--- a/Engine/Source/Runtime/Core/Public/Containers/Map.h
+++ b/Engine/Source/Runtime/Core/Public/Containers/Map.h
@@ -284,7 +284,7 @@ public:
 		// This is only necessary if the TSet allows duplicate keys; otherwise TSet::Add replaces the existing key-value pair.
 		if(KeyFuncs::bAllowDuplicateKeys)
 		{
-			for(typename PairSetType::TKeyIterator It(Pairs,InKey);It;++It)
+			for(typename PairSetType::TKeyIterator It(Pairs,KeyType(InKey));It;++It)
 			{
 				It.RemoveCurrent();
 			}
