implementation of VectorGetComponent (CIS fix)
#ue4

diff --git a/Engine/Source/Runtime/Core/Public/Math/UnrealMathDirectX.h b/Engine/Source/Runtime/Core/Public/Math/UnrealMathDirectX.h
--- a/Engine/Source/Runtime/Core/Public/Math/UnrealMathDirectX.h
+++ b/Engine/Source/Runtime/Core/Public/Math/UnrealMathDirectX.h
@@ -195,19 +195,31 @@ FORCEINLINE VectorRegister MakeVectorRegister( float X, float Y, float Z, float
  
 
 /**
-* Returns an component from a vector.  May be slow, use with caution.
-*
-* @param Vec				Vector register
-* @param ComponentIndex	Which component to get, X=0, Y=1, Z=2, W=3
-* @return					The component as a float
-*/
-FORCEINLINE float VectorGetComponent(VectorRegister Vec, uint32 ComponentIndex)
+ * Returns an component from a vector.
+ *
+ * @param Vec				Vector register
+ * @param ComponentIndex	Which component to get, X=0, Y=1, Z=2, W=3
+ * @return					The component as a float
+ */
+FORCEINLINE float VectorGetComponent( VectorRegoster Vec, uint32 ComponentIndex )
 {
-	// stub for now, just to fix build failure
+	switch (ComponentIndex)
+	{
+	case 0:
+		return XMVectorGetX(Vec);
+	case 1:
+		return XMVectorGetY(Vec);
+	case 2:
+		return XMVectorGetZ(Vec);
+	case 3:
+		return XMVectorGetW(Vec);
+	}
+
+	return 0.0f;
 }
 
 
-/**
+ /**
  * Replicates one element into all four elements and returns the new vector.
  *
  * @param Vec			Source vector
