Fix Xbox compile

diff --git a/Engine/Source/Runtime/Core/Public/Math/UnrealMathDirectX.h b/Engine/Source/Runtime/Core/Public/Math/UnrealMathDirectX.h
--- a/Engine/Source/Runtime/Core/Public/Math/UnrealMathDirectX.h
+++ b/Engine/Source/Runtime/Core/Public/Math/UnrealMathDirectX.h
@@ -201,18 +201,18 @@ FORCEINLINE VectorRegister MakeVectorRegister( float X, float Y, float Z, float
  * @param ComponentIndex	Which component to get, X=0, Y=1, Z=2, W=3
  * @return					The component as a float
  */
-FORCEINLINE float VectorGetComponent( VectorRegoster Vec, uint32 ComponentIndex )
+FORCEINLINE float VectorGetComponent( VectorRegister Vec, uint32 ComponentIndex )
 {
 	switch (ComponentIndex)
 	{
 	case 0:
-		return XMVectorGetX(Vec);
+		return DirectX::XMVectorGetX(Vec);
 	case 1:
-		return XMVectorGetY(Vec);
+		return DirectX::XMVectorGetY(Vec);
 	case 2:
-		return XMVectorGetZ(Vec);
+		return DirectX::XMVectorGetZ(Vec);
 	case 3:
-		return XMVectorGetW(Vec);
+		return DirectX::XMVectorGetW(Vec);
 	}
 
 	return 0.0f;
