Don't forward events to a destructed FMacApplication - causes a crash when closing to reset the editor layout (TTP #342458).

diff --git a/Engine/Source/Runtime/Core/Private/Mac/MacWindow.cpp b/Engine/Source/Runtime/Core/Private/Mac/MacWindow.cpp
--- a/Engine/Source/Runtime/Core/Private/Mac/MacWindow.cpp
+++ b/Engine/Source/Runtime/Core/Private/Mac/MacWindow.cpp
@@ -448,7 +448,7 @@ TArray< FSlateCocoaWindow* > GRunningModalWindows;
 
 - (void)windowWillClose:(NSNotification *)notification
 {
-	if(self.bForwardEvents)
+	if(self.bForwardEvents && MacApplication)
 	{
 		MacApplication->OnWindowDidClose( self );
 	}
