Moved key down and up events processing to window delegate's keyDown and keyUp functions to solve problems with Slate receiving events for shortcuts handled by system (for example the menu bar items)

diff --git a/Engine/Source/Runtime/Core/Private/Mac/MacPlatformMisc.cpp b/Engine/Source/Runtime/Core/Private/Mac/MacPlatformMisc.cpp
--- a/Engine/Source/Runtime/Core/Private/Mac/MacPlatformMisc.cpp
+++ b/Engine/Source/Runtime/Core/Private/Mac/MacPlatformMisc.cpp
@@ -470,13 +470,14 @@ void FMacPlatformMisc::PumpMessages( bool bFromMainLoop )
 
 		while( NSEvent *Event = [NSApp nextEventMatchingMask: NSAnyEventMask untilDate: nil inMode: NSDefaultRunLoopMode dequeue: true] )
 		{
-			const bool bIsMouseClickEvent = [Event type] == NSLeftMouseDown || [Event type] == NSLeftMouseUp
+			const bool bIsMouseClickOrKeyEvent = [Event type] == NSLeftMouseDown || [Event type] == NSLeftMouseUp
 										 || [Event type] == NSRightMouseDown || [Event type] == NSRightMouseUp
-										 || [Event type] == NSOtherMouseDown || [Event type] == NSOtherMouseUp;
+										 || [Event type] == NSOtherMouseDown || [Event type] == NSOtherMouseUp
+										 || [Event type] == NSKeyDown || [Event type] == NSKeyUp;
 
 			if( MacApplication )
 			{
-				if( !bIsMouseClickEvent || [Event window] == NULL )
+				if( !bIsMouseClickOrKeyEvent || [Event window] == NULL )
 				{
 					MacApplication->ProcessEvent( Event );
 				}
diff --git a/Engine/Source/Runtime/Core/Private/Mac/MacWindow.cpp b/Engine/Source/Runtime/Core/Private/Mac/MacWindow.cpp
--- a/Engine/Source/Runtime/Core/Private/Mac/MacWindow.cpp
+++ b/Engine/Source/Runtime/Core/Private/Mac/MacWindow.cpp
@@ -284,13 +284,20 @@ TArray< FSlateCocoaWindow* > GRunningModalWindows;
 	}
 }
 
-// keyDown and keyUp are empty, but having them lets Cocoa know we handle the keys ourselves
 - (void)keyDown:(NSEvent *)Event
 {
+	if(self.bForwardEvents)
+	{
+		MacApplication->ProcessEvent( Event );
+	}
 }
 
 - (void)keyUp:(NSEvent *)Event
 {
+	if(self.bForwardEvents)
+	{
+		MacApplication->ProcessEvent( Event );
+	}
 }
 
 - (void)windowDidEnterFullScreen:(NSNotification *)notification
