- Fixed unity error thing

diff --git a/Engine/Source/Runtime/Core/Private/IOS/IOSPlatformSurvey.cpp b/Engine/Source/Runtime/Core/Private/IOS/IOSPlatformSurvey.cpp
--- a/Engine/Source/Runtime/Core/Private/IOS/IOSPlatformSurvey.cpp
+++ b/Engine/Source/Runtime/Core/Private/IOS/IOSPlatformSurvey.cpp
@@ -5,7 +5,7 @@
 =============================================================================*/
 #include "Core.h"
 #include "IOSAppDelegate.h"
-
+#include "IOSWindow.h"
 
 bool FIOSPlatformSurvey::GetSurveyResults(FHardwareSurveyResults& OutResults, bool bWait)
 {
