Fixed faulty optimization to text localization manager that was potentially falsely setting text localization to not "LEETify" unlocalized strings.

diff --git a/Engine/Source/Runtime/Core/Private/Internationalization/TextLocalizationManager.cpp b/Engine/Source/Runtime/Core/Private/Internationalization/TextLocalizationManager.cpp
--- a/Engine/Source/Runtime/Core/Private/Internationalization/TextLocalizationManager.cpp
+++ b/Engine/Source/Runtime/Core/Private/Internationalization/TextLocalizationManager.cpp
@@ -456,7 +456,17 @@ TSharedRef<FString, ESPMode::ThreadSafe> FTextLocalizationManager::GetString(con
 
 #if ENABLE_LOC_TESTING
 	const bool bShouldLEETIFYAll = bIsInitialized && FInternationalization::Get().GetCurrentCulture()->GetName() == TEXT("LEET");
-	static const bool bShouldLEETIFYUnlocalizedString = FCommandLine::IsInitialized() && FParse::Param(FCommandLine::Get(), TEXT("LEETIFYUnlocalized"));
+
+	// Attempt to set bShouldLEETIFYUnlocalizedString appropriately, only once, after the commandline is initialized and parsed
+	static bool bShouldLEETIFYUnlocalizedString = false;
+	{
+		static bool bHasParsedCommandLine = false;
+		if(!bHasParsedCommandLine && FCommandLine::IsInitialized())
+		{
+			bShouldLEETIFYUnlocalizedString = FParse::Param(FCommandLine::Get(), TEXT("LEETIFYUnlocalized"));
+			bHasParsedCommandLine = true;
+		}
+	}
 #endif
 
 	// Find namespace's key table.
