Missing file from Hot-Reload improvements check in

diff --git a/Engine/Source/Runtime/Core/Private/Modules/ModuleManager.cpp b/Engine/Source/Runtime/Core/Private/Modules/ModuleManager.cpp
--- a/Engine/Source/Runtime/Core/Private/Modules/ModuleManager.cpp
+++ b/Engine/Source/Runtime/Core/Private/Modules/ModuleManager.cpp
@@ -1247,7 +1247,8 @@ bool FModuleManager::RecompileModulesAsync( const TArray< FName > ModuleNames, c
 
 	// Kick off compilation!
 	const FString AdditionalArguments = MakeUBTArgumentsForModuleCompiling();
-	bool bWasSuccessful = StartCompilingModuleDLLs( FApp::GetGameName(), ModulesToRecompile, InRecompileModulesCallback, Ar, true, AdditionalArguments );
+	const bool bFailIfGeneratedCodeChanges = false;
+	bool bWasSuccessful = StartCompilingModuleDLLs( FApp::GetGameName(), ModulesToRecompile, InRecompileModulesCallback, Ar, bFailIfGeneratedCodeChanges, AdditionalArguments );
 	if (bWasSuccessful)
 	{
 		// Go ahead and check for completion right away.  This is really just so that we can handle the case
