#ttp 338022 TappyChicken: Android: Add back button
Bug Fix : Events not being correctly routed to "Android Back" Blueprint node

diff --git a/Engine/Source/Runtime/Core/Private/Android/AndroidInputInterface.cpp b/Engine/Source/Runtime/Core/Private/Android/AndroidInputInterface.cpp
--- a/Engine/Source/Runtime/Core/Private/Android/AndroidInputInterface.cpp
+++ b/Engine/Source/Runtime/Core/Private/Android/AndroidInputInterface.cpp
@@ -48,6 +48,7 @@ FAndroidInputInterface::FAndroidInputInterface( const TSharedRef< FGenericApplic
 	ButtonMapping[13] = EControllerButtons::DPadDown;
 	ButtonMapping[14] = EControllerButtons::DPadLeft;
 	ButtonMapping[15] = EControllerButtons::DPadRight;
+	ButtonMapping[16] = EControllerButtons::AndroidBack;  // Technically just an alias for SpecialLeft
 
 	InitialButtonRepeatDelay = 0.2f;
 	ButtonRepeatDelay = 0.1f;
@@ -262,7 +263,7 @@ void FAndroidInputInterface::JoystickButtonEvent(int32 deviceId, int32 buttonId,
 		case AKEYCODE_BUTTON_START:
 		case AKEYCODE_MENU:          NewControllerData[deviceId].ButtonStates[ 6] = buttonDown; break;
 		case AKEYCODE_BUTTON_SELECT: 
-		case AKEYCODE_BACK:          NewControllerData[deviceId].ButtonStates[ 7] = buttonDown; break;
+		case AKEYCODE_BACK:          NewControllerData[deviceId].ButtonStates[7] = buttonDown; NewControllerData[deviceId].ButtonStates[16] = buttonDown;  break;
 		case AKEYCODE_BUTTON_THUMBL: NewControllerData[deviceId].ButtonStates[ 8] = buttonDown; break;
 		case AKEYCODE_BUTTON_THUMBR: NewControllerData[deviceId].ButtonStates[ 9] = buttonDown; break;
 		case AKEYCODE_BUTTON_L2:     NewControllerData[deviceId].ButtonStates[10] = buttonDown; break;
diff --git a/Engine/Source/Runtime/Core/Public/Android/AndroidInputInterface.h b/Engine/Source/Runtime/Core/Public/Android/AndroidInputInterface.h
--- a/Engine/Source/Runtime/Core/Public/Android/AndroidInputInterface.h
+++ b/Engine/Source/Runtime/Core/Public/Android/AndroidInputInterface.h
@@ -84,7 +84,7 @@ struct TouchInput
 };
 
 #define MAX_NUM_CONTROLLERS					1  // Only one for now
-#define MAX_NUM_CONTROLLER_BUTTONS			16
+#define MAX_NUM_CONTROLLER_BUTTONS			17
 #define MAX_DEFERRED_MESSAGE_QUEUE_SIZE		128
 
 struct FAndroidControllerData
diff --git a/Engine/Source/Runtime/Core/Public/GenericPlatform/GenericApplicationMessageHandler.h b/Engine/Source/Runtime/Core/Public/GenericPlatform/GenericApplicationMessageHandler.h
--- a/Engine/Source/Runtime/Core/Public/GenericPlatform/GenericApplicationMessageHandler.h
+++ b/Engine/Source/Runtime/Core/Public/GenericPlatform/GenericApplicationMessageHandler.h
@@ -71,6 +71,8 @@ namespace EControllerButtons
 		GlobalPlay,
 		GlobalBack,
 
+		AndroidBack,
+
 		Invalid,
 	};
 }
diff --git a/Engine/Source/Runtime/Slate/Private/Framework/Application/SlateApplication.cpp b/Engine/Source/Runtime/Slate/Private/Framework/Application/SlateApplication.cpp
--- a/Engine/Source/Runtime/Slate/Private/Framework/Application/SlateApplication.cpp
+++ b/Engine/Source/Runtime/Slate/Private/Framework/Application/SlateApplication.cpp
@@ -3875,6 +3875,8 @@ FKey TranslateControllerButtonToKey( EControllerButtons::Type Button )
 	case EControllerButtons::GlobalPlay: Key = EKeys::Global_Play; break;
 	case EControllerButtons::GlobalBack: Key = EKeys::Global_Back; break;
 
+	case EControllerButtons::AndroidBack: Key = EKeys::Android_Back; break;
+
 	case EControllerButtons::Invalid: Key = EKeys::Invalid; break;
 	}
 
