Fixed a number of compile warnings reported by Clang on Mac (just some for now, there's many more that will need to be fixed in the future)

diff --git a/Engine/Source/Programs/UnrealHeaderTool/Private/HeaderParser.cpp b/Engine/Source/Programs/UnrealHeaderTool/Private/HeaderParser.cpp
--- a/Engine/Source/Programs/UnrealHeaderTool/Private/HeaderParser.cpp
+++ b/Engine/Source/Programs/UnrealHeaderTool/Private/HeaderParser.cpp
@@ -6773,7 +6773,7 @@ void FHeaderParser::SimplifiedClassParse(const TCHAR* InBuffer, bool& bIsInterfa
 				bKeepPreprocessorDirectives = false;
 				bNotCPP = true;
 			}
-			else if (bIf && FParse::Command(&Str,TEXT("WITH_EDITORONLY_DATA")) || FParse::Command(&Str,TEXT("WITH_EDITOR")))
+			else if (bIf && (FParse::Command(&Str,TEXT("WITH_EDITORONLY_DATA")) || FParse::Command(&Str,TEXT("WITH_EDITOR"))))
 			{
 				Target = &ClassHeaderTextStrippedOfCppText;
 				bUnknownDirective = true;
diff --git a/Engine/Source/Runtime/Core/Private/Mac/MacApplication.cpp b/Engine/Source/Runtime/Core/Private/Mac/MacApplication.cpp
--- a/Engine/Source/Runtime/Core/Private/Mac/MacApplication.cpp
+++ b/Engine/Source/Runtime/Core/Private/Mac/MacApplication.cpp
@@ -892,7 +892,7 @@ void FMacApplication::GetDisplayMetrics( FDisplayMetrics& OutDisplayMetrics ) co
 	OutDisplayMetrics.PrimaryDisplayHeight = ScreenFrame.size.height;
 
 	// Virtual desktop area
-	NSRect WholeWorkspace = {0};
+	NSRect WholeWorkspace = {{0,0},{0,0}};
 	for (NSScreen* Screen in AllScreens)
 	{
 		if (Screen)
diff --git a/Engine/Source/Runtime/Core/Private/Mac/MacPlatformFeedbackContextPrivate.h b/Engine/Source/Runtime/Core/Private/Mac/MacPlatformFeedbackContextPrivate.h
--- a/Engine/Source/Runtime/Core/Private/Mac/MacPlatformFeedbackContextPrivate.h
+++ b/Engine/Source/Runtime/Core/Private/Mac/MacPlatformFeedbackContextPrivate.h
@@ -20,10 +20,11 @@ public:
 
 	// Constructor.
 	FFeedbackContextMac()
-	: FFeedbackContext()
-	, SlowTaskCount( 0 )
-	, Context( NULL )
+		: FFeedbackContext()
+		, Context( NULL )
+		, SlowTaskCount( 0 )
 	{}
+
 	void Serialize( const TCHAR* V, ELogVerbosity::Type Verbosity, const class FName& Category )
 	{
 		// if we set the color for warnings or errors, then reset at the end of the function
@@ -73,30 +74,10 @@ public:
 			Serialize(COLOR_NONE, ELogVerbosity::SetColor, Category);
 		}
 	}
-	VARARG_BODY( bool, YesNof, const TCHAR*, VARARG_NONE )
-	{
-		TCHAR TempStr[4096];
-		GET_VARARGS( TempStr, ARRAY_COUNT(TempStr), ARRAY_COUNT(TempStr)-1, Fmt, Fmt );
-		if( ( GIsClient || GIsEditor ) && ( ( GIsSilent != true ) && ( FApp::IsUnattended() != true ) ) )
-		{
-			CFStringRef Message = FPlatformString::TCHARToCFString( TempStr );
-			CFStringRef Caption = FPlatformString::TCHARToCFString( *NSLOCTEXT( "UnrealEd", "Question", "Question" ).ToString() );
-			CFStringRef Yes     = FPlatformString::TCHARToCFString( *NSLOCTEXT( "UnrealEd", "Yes",      "Yes"      ).ToString() );
-			CFStringRef No      = FPlatformString::TCHARToCFString( *NSLOCTEXT( "UnrealEd", "No",       "No"       ).ToString() );
-
-			int32 Result = NSRunInformationalAlertPanel( (NSString*)Caption, (NSString*)Message, (NSString*)Yes, (NSString*)No, NULL );
-
-			CFRelease( No );
-			CFRelease( Yes );
-			CFRelease( Caption );
-			CFRelease( Message );
 
-			return Result != 0;
-		}
-		else
-		{
-			return false;
-		}
+	bool YesNof( const FText& Question )
+	{
+		return FPlatformMisc::MessageBoxExt(EAppMsgType::YesNo, *Question.ToString(), TEXT("")) == EAppReturnType::Yes;
 	}
 
 	void BeginSlowTask( const FText& Task, bool ShowProgressDialog, bool bShowCancelButton=false )
diff --git a/Engine/Source/Runtime/Core/Private/Mac/MacPlatformMisc.cpp b/Engine/Source/Runtime/Core/Private/Mac/MacPlatformMisc.cpp
--- a/Engine/Source/Runtime/Core/Private/Mac/MacPlatformMisc.cpp
+++ b/Engine/Source/Runtime/Core/Private/Mac/MacPlatformMisc.cpp
@@ -960,7 +960,7 @@ int32 FMacPlatformMisc::ConvertSlateYPositionToCocoa(int32 YPosition)
 	NSArray* AllScreens = [NSScreen screens];
 	NSScreen* PrimaryScreen = (NSScreen*)[AllScreens objectAtIndex: 0];
 	NSRect ScreenFrame = [PrimaryScreen frame];
-	NSRect WholeWorkspace = {0};
+	NSRect WholeWorkspace = {{0,0},{0,0}};
 	for(NSScreen* Screen in AllScreens)
 	{
 		if(Screen)
diff --git a/Engine/Source/Runtime/Core/Private/Mac/MacPlatformOutputDevices.cpp b/Engine/Source/Runtime/Core/Private/Mac/MacPlatformOutputDevices.cpp
--- a/Engine/Source/Runtime/Core/Private/Mac/MacPlatformOutputDevices.cpp
+++ b/Engine/Source/Runtime/Core/Private/Mac/MacPlatformOutputDevices.cpp
@@ -112,40 +112,35 @@ void FOutputDeviceMacError::HandleError()
 	// Trigger the OnSystemFailure hook if it exists
 	FCoreDelegates::OnHandleSystemError.Broadcast();
 
-	try
-	{
-		GIsGuarded				= 0;
-		GIsRunning				= 0;
-		GIsCriticalError		= 1;
-		GLogConsole				= NULL;
-		GErrorHist[ARRAY_COUNT(GErrorHist)-1]=0;
+	GIsGuarded				= 0;
+	GIsRunning				= 0;
+	GIsCriticalError		= 1;
+	GLogConsole				= NULL;
+	GErrorHist[ARRAY_COUNT(GErrorHist)-1]=0;
 
-		// Dump the error and flush the log.
-		UE_LOG(LogMac, Log, TEXT("=== Critical error: ===") LINE_TERMINATOR TEXT("%s") LINE_TERMINATOR, GErrorExceptionDescription);
-		UE_LOG(LogMac, Log, GErrorHist);
+	// Dump the error and flush the log.
+	UE_LOG(LogMac, Log, TEXT("=== Critical error: ===") LINE_TERMINATOR TEXT("%s") LINE_TERMINATOR, GErrorExceptionDescription);
+	UE_LOG(LogMac, Log, GErrorHist);
 
-		GLog->Flush();
+	GLog->Flush();
 
-		// Unhide the mouse.
-		// @TODO: Remove usage of deprecated CGCursorIsVisible function
+	// Unhide the mouse.
+	// @TODO: Remove usage of deprecated CGCursorIsVisible function
 #pragma clang diagnostic push
 #pragma clang diagnostic ignored "-Wdeprecated-declarations"
-		while (!CGCursorIsVisible())
-		{
-			CGDisplayShowCursor(kCGDirectMainDisplay);
-		}
+	while (!CGCursorIsVisible())
+	{
+		CGDisplayShowCursor(kCGDirectMainDisplay);
+	}
 #pragma clang diagnostic pop
-		// Release capture and allow mouse to freely roam around.
-		CGAssociateMouseAndMouseCursorPosition(true);
+	// Release capture and allow mouse to freely roam around.
+	CGAssociateMouseAndMouseCursorPosition(true);
 
-		FPlatformMisc::ClipboardCopy(GErrorHist);
+	FPlatformMisc::ClipboardCopy(GErrorHist);
 
-		FPlatformMisc::SubmitErrorReport( GErrorHist, EErrorReportMode::Interactive );
+	FPlatformMisc::SubmitErrorReport( GErrorHist, EErrorReportMode::Interactive );
 
-		FCoreDelegates::OnShutdownAfterError.Broadcast();
-	}
-	catch( ... )
-	{}
+	FCoreDelegates::OnShutdownAfterError.Broadcast();
 }
 
 ////////////////////////////////////////
diff --git a/Engine/Source/Runtime/Core/Private/Mac/MacPlatformProcess.cpp b/Engine/Source/Runtime/Core/Private/Mac/MacPlatformProcess.cpp
--- a/Engine/Source/Runtime/Core/Private/Mac/MacPlatformProcess.cpp
+++ b/Engine/Source/Runtime/Core/Private/Mac/MacPlatformProcess.cpp
@@ -210,10 +210,6 @@ bool FMacPlatformProcess::ExecProcess( const TCHAR* URL, const TCHAR* Params, in
 			NSBundle* Bundle = [NSBundle bundleWithPath:LaunchPath];
 			LaunchPath = Bundle ? [Bundle executablePath] : NULL;
 		}
-		else
-		{
-			LaunchPath = LaunchPath;
-		}
 	}
 	else
 	{
@@ -340,10 +336,6 @@ FProcHandle FMacPlatformProcess::CreateProc( const TCHAR* URL, const TCHAR* Parm
 			NSBundle* Bundle = [NSBundle bundleWithPath:LaunchPath];
 			LaunchPath = Bundle ? [Bundle executablePath] : NULL;
 		}
-		else
-		{
-			LaunchPath = LaunchPath;
-		}
 	}
 	else
 	{
diff --git a/Engine/Source/Runtime/Core/Private/Mac/MacTextInputMethodSystem.cpp b/Engine/Source/Runtime/Core/Private/Mac/MacTextInputMethodSystem.cpp
--- a/Engine/Source/Runtime/Core/Private/Mac/MacTextInputMethodSystem.cpp
+++ b/Engine/Source/Runtime/Core/Private/Mac/MacTextInputMethodSystem.cpp
@@ -405,12 +405,12 @@ DEFINE_LOG_CATEGORY_STATIC(LogMacTextInputMethodSystem, Log, All);
 		const float PrimaryScreenHeight = [PrimaryScreen visibleFrame].size.height;
 		Position.Y = -(Position.Y - PrimaryScreenHeight + 1);
 		
-		NSRect GlyphBox = {Position.X,Position.Y,Size.X,Size.Y};
+		NSRect GlyphBox = {{Position.X,Position.Y},{Size.X,Size.Y}};
 		return GlyphBox;
 	}
 	else
 	{
-		NSRect GlyphBox = {0,0,0,0};
+		NSRect GlyphBox = {{0,0},{0,0}};
 		return GlyphBox;
 	}
 }
diff --git a/Engine/Source/Runtime/Core/Public/Mac/MacPlatformMisc.h b/Engine/Source/Runtime/Core/Public/Mac/MacPlatformMisc.h
--- a/Engine/Source/Runtime/Core/Public/Mac/MacPlatformMisc.h
+++ b/Engine/Source/Runtime/Core/Public/Mac/MacPlatformMisc.h
@@ -29,7 +29,10 @@ struct CORE_API FMacPlatformMisc : public FGenericPlatformMisc
 	{
 		// Based on http://developer.apple.com/library/mac/#qa/qa1361/_index.html
 
+#pragma clang diagnostic push
+#pragma clang diagnostic ignored "-Wmissing-braces"
 		struct kinfo_proc Info = { 0 };
+#pragma clang diagnostic pop
 		int32 Mib[] = { CTL_KERN, KERN_PROC, KERN_PROC_PID, getpid() };
 		SIZE_T Size = sizeof(Info);
 
diff --git a/Engine/Source/Runtime/Core/Public/Mac/MacTextInputMethodSystem.h b/Engine/Source/Runtime/Core/Public/Mac/MacTextInputMethodSystem.h
--- a/Engine/Source/Runtime/Core/Public/Mac/MacTextInputMethodSystem.h
+++ b/Engine/Source/Runtime/Core/Public/Mac/MacTextInputMethodSystem.h
@@ -96,6 +96,7 @@ class FSlateTextView;
 class FMacTextInputMethodSystem : public ITextInputMethodSystem
 {
 public:
+	virtual ~FMacTextInputMethodSystem() {}
 	bool Initialize();
 	void Terminate();
 	
diff --git a/Engine/Source/Runtime/Engine/Classes/EdGraph/EdGraph.h b/Engine/Source/Runtime/Engine/Classes/EdGraph/EdGraph.h
--- a/Engine/Source/Runtime/Engine/Classes/EdGraph/EdGraph.h
+++ b/Engine/Source/Runtime/Engine/Classes/EdGraph/EdGraph.h
@@ -223,7 +223,7 @@ struct FGraphNodeCreator
 {
 public:
 	FGraphNodeCreator(UEdGraph& InGraph)
-		: bPlaced(false), Graph(InGraph), Node(NULL)
+		: Node(NULL), Graph(InGraph), bPlaced(false)
 	{
 	}
 
diff --git a/Engine/Source/Runtime/OpenGLDrv/Private/Mac/OpenGLMac.cpp b/Engine/Source/Runtime/OpenGLDrv/Private/Mac/OpenGLMac.cpp
--- a/Engine/Source/Runtime/OpenGLDrv/Private/Mac/OpenGLMac.cpp
+++ b/Engine/Source/Runtime/OpenGLDrv/Private/Mac/OpenGLMac.cpp
@@ -1498,13 +1498,13 @@ void FMacOpenGLTimer::CacheResult(void)
 }
 
 FMacOpenGLQuery::FMacOpenGLQuery(FPlatformOpenGLContext* InContext)
-: Context(InContext)
-, Name(0)
-, Target(0)
-, Result(0)
-, Available(0)
-, Running(0)
-, Cached(false)
+	: Name(0)
+	, Target(0)
+	, Context(InContext)
+	, Result(0)
+	, Available(0)
+	, Running(0)
+	, Cached(false)
 {
 	check(Context);
 }
