WinXP+ compatible version of QueryConnectedMice() that doesn't require a WDK install
Fixes TTP#342552

diff --git a/Engine/Source/Programs/UnrealBuildTool/Windows/UEBuildWindows.cs b/Engine/Source/Programs/UnrealBuildTool/Windows/UEBuildWindows.cs
--- a/Engine/Source/Programs/UnrealBuildTool/Windows/UEBuildWindows.cs
+++ b/Engine/Source/Programs/UnrealBuildTool/Windows/UEBuildWindows.cs
@@ -467,7 +467,6 @@ namespace UnrealBuildTool
                 InBuildTarget.GlobalLinkEnvironment.Config.AdditionalLibraries.Add("netapi32.lib");
                 InBuildTarget.GlobalLinkEnvironment.Config.AdditionalLibraries.Add("iphlpapi.lib");
                 InBuildTarget.GlobalLinkEnvironment.Config.AdditionalLibraries.Add("setupapi.lib"); //  Required for access monitor device enumeration
-                InBuildTarget.GlobalLinkEnvironment.Config.AdditionalLibraries.Add("hid.lib");
 
                 // Windows Vista/7 Desktop Windows Manager API for Slate Windows Compliance
                 if( !SupportWindowsXP )		// Windows XP does not support DWM
diff --git a/Engine/Source/Runtime/Core/Private/Windows/WindowsApplication.cpp b/Engine/Source/Runtime/Core/Private/Windows/WindowsApplication.cpp
--- a/Engine/Source/Runtime/Core/Private/Windows/WindowsApplication.cpp
+++ b/Engine/Source/Runtime/Core/Private/Windows/WindowsApplication.cpp
@@ -20,11 +20,6 @@
 	#include <shlobj.h>
 	#include <objbase.h>
 	#include <SetupApi.h>
-// The version of the Windows Kits 8.0 are missing this guard, 8.1 has a fix
-extern "C" 
-{	
-	#include <Hidsdi.h>
-}
 	#include <devguid.h>
 
 // This might not be defined by Windows when maintaining backwards-compatibility to pre-Vista builds
@@ -1817,48 +1812,51 @@ HRESULT FWindowsApplication::OnOLEDrop( const HWND HWnd, const FDragDropOLEData&
 
 void FWindowsApplication::QueryConnectedMice()
 {
-	GUID HIDGuid;
-	GUID MouseGuid = GUID_DEVCLASS_MOUSE;
-	HDEVINFO hDevInfo;
-	BYTE TempBuffer[4096];
-
-	HidD_GetHidGuid(&HIDGuid);
+	TArray<RAWINPUTDEVICELIST> DeviceList;
+	UINT DeviceCount;
 
-	hDevInfo = SetupDiGetClassDevs(&HIDGuid, NULL, NULL, DIGCF_PRESENT | DIGCF_DEVICEINTERFACE);
-
-	if (hDevInfo == INVALID_HANDLE_VALUE)
+	GetRawInputDeviceList(nullptr, &DeviceCount, sizeof(RAWINPUTDEVICELIST));
+	if (DeviceCount == 0)
 	{
-		// in this case, on a desktop it's proably safer to assume that a mouse exists rather than doesn't
-		bIsMouseAttached = true;
+		bIsMouseAttached = false;
 		return;
 	}
 
-	bIsMouseAttached = false;
-
-	for (DWORD DeviceIndex = 0; ; DeviceIndex++)
+	DeviceList.AddUninitialized(DeviceCount);
+	GetRawInputDeviceList(DeviceList.GetData(), &DeviceCount, sizeof(RAWINPUTDEVICELIST));
+	
+	int32 MouseCount = 0;
+	for (const auto& Device : DeviceList)
 	{
-		SP_DEVICE_INTERFACE_DATA DevInterface;
-		SP_DEVINFO_DATA DevInfo;
-		PSP_DEVICE_INTERFACE_DETAIL_DATA DevDetail;
+		UINT NameLen;
+		TAutoPtr<char> Name;
+		if (Device.dwType != RIM_TYPEMOUSE)
+			continue;
+		//Force the use of ANSI versions of these calls
+		if (GetRawInputDeviceInfoA(Device.hDevice, RIDI_DEVICENAME, nullptr, &NameLen) < 0)
+			continue;
 
-		DevInterface.cbSize = sizeof(DevInterface);
-		if (!SetupDiEnumDeviceInterfaces(hDevInfo, NULL, &HIDGuid, DeviceIndex, &DevInterface))
-		{
-			break;
-		}
-		DevDetail = (PSP_DEVICE_INTERFACE_DETAIL_DATA)TempBuffer;
-		DevDetail->cbSize = sizeof(SP_DEVICE_INTERFACE_DETAIL_DATA);
-		DevInfo.cbSize = sizeof(DevInfo);
-		if (!SetupDiGetDeviceInterfaceDetail(hDevInfo, &DevInterface, DevDetail, sizeof(TempBuffer), NULL, &DevInfo))
-		{
-			break;
-		}
-		if (IsEqualGUID(DevInfo.ClassGuid, MouseGuid))
+		Name.Reset(new char[NameLen+1]);
+		if (GetRawInputDeviceInfoA(Device.hDevice, RIDI_DEVICENAME, Name.GetOwnedPointer(), &NameLen) < 0)
+			continue;
+
+		Name[NameLen] = 0;
+		FString WName = ANSI_TO_TCHAR(Name);
+		WName.ReplaceInline(TEXT("#"), TEXT("\\"));
+		/*
+		 * Name XP starts with \??\, vista+ starts \\?\ 
+		 * In the device list exists a fake Mouse device with the device name of RDP_MOU
+		 * This is used for Remote Desktop so ignore it.
+		 */
+		if (WName.StartsWith(TEXT("\\??\\ROOT\\RDP_MOU\\")) || WName.StartsWith(TEXT("\\\\?\\ROOT\\RDP_MOU\\")))
 		{
-			bIsMouseAttached = true;
-			break;
+			continue;
 		}
+
+		++MouseCount;
 	}
+
+	bIsMouseAttached = MouseCount > 0;
 }
 
 void FTaskbarList::Initialize()
