Improved messages regarding falling back to the base language of a culture when there is no localization data for the specified culture.

diff --git a/Engine/Source/Developer/InternationalizationSettings/Private/InternationalizationSettingsModelDetails.cpp b/Engine/Source/Developer/InternationalizationSettings/Private/InternationalizationSettingsModelDetails.cpp
--- a/Engine/Source/Developer/InternationalizationSettings/Private/InternationalizationSettingsModelDetails.cpp
+++ b/Engine/Source/Developer/InternationalizationSettings/Private/InternationalizationSettingsModelDetails.cpp
@@ -248,7 +248,7 @@ void FInternationalizationSettingsModelDetails::CustomizeDetails( IDetailLayoutB
 void FInternationalizationSettingsModelDetails::RefreshAvailableCultures()
 {
 	AvailableCultures.Empty();
-	FInternationalization::Get().GetCulturesWithAvailableLocalization(FPaths::GetEditorLocalizationPaths(), AvailableCultures);
+	FInternationalization::Get().GetCulturesWithAvailableLocalization(FPaths::GetEditorLocalizationPaths(), AvailableCultures, true);
 
 	// Update our selected culture based on the available choices
 	if ( !AvailableCultures.Contains( SelectedCulture ) )
diff --git a/Engine/Source/Runtime/Core/Private/Internationalization/Internationalization.cpp b/Engine/Source/Runtime/Core/Private/Internationalization/Internationalization.cpp
--- a/Engine/Source/Runtime/Core/Private/Internationalization/Internationalization.cpp
+++ b/Engine/Source/Runtime/Core/Private/Internationalization/Internationalization.cpp
@@ -1,4 +1,4 @@
-// Copyright 1998-2014 Epic Games, Inc. All Rights Reserved.
+﻿// Copyright 1998-2014 Epic Games, Inc. All Rights Reserved.
 
 #include "Core.h"
 #include "InvariantCulture.h"
@@ -566,7 +566,7 @@ void FInternationalization::GetCultureNames(TArray<FString>& CultureNames) const
 	}
 }
 
-void FInternationalization::GetCulturesWithAvailableLocalization(const TArray<FString>& InLocalizationPaths, TArray< TSharedPtr<FCulture, ESPMode::ThreadSafe> >& OutAvailableCultures) const
+void FInternationalization::GetCulturesWithAvailableLocalization(const TArray<FString>& InLocalizationPaths, TArray< TSharedPtr<FCulture, ESPMode::ThreadSafe> >& OutAvailableCultures, const bool bAllowFallback) const
 {
 	OutAvailableCultures.Reset();
 
@@ -610,7 +610,7 @@ void FInternationalization::GetCulturesWithAvailableLocalization(const TArray<FS
 		for(const FString& LocalizationFolder : AllLocalizationFolders)
 		{
 			// Check the full name, but if it doesn't match, see if the base language is present
-			if(LocalizationFolder == Culture->GetName() || LocalizationFolder == Culture->GetTwoLetterISOLanguageName())
+			if(LocalizationFolder == Culture->GetName() || (bAllowFallback && LocalizationFolder == Culture->GetTwoLetterISOLanguageName()))
 			{
 				OutAvailableCultures.AddUnique(Culture);
 			}
diff --git a/Engine/Source/Runtime/Core/Private/Internationalization/TextLocalizationManager.cpp b/Engine/Source/Runtime/Core/Private/Internationalization/TextLocalizationManager.cpp
--- a/Engine/Source/Runtime/Core/Private/Internationalization/TextLocalizationManager.cpp
+++ b/Engine/Source/Runtime/Core/Private/Internationalization/TextLocalizationManager.cpp
@@ -36,19 +36,19 @@ void EndInitTextLocalization()
 	else
 #endif
 	{
-		FString CultureName;
-		if (FParse::Value(FCommandLine::Get(), TEXT("CULTUREFORCOOKING="), CultureName))
+		FString TargetCultureName;
+		if (FParse::Value(FCommandLine::Get(), TEXT("CULTUREFORCOOKING="), TargetCultureName))
 		{
 			// Write the culture passed in if first install...
 			if (FParse::Param(FCommandLine::Get(), TEXT("firstinstall")))
 			{
-				GConfig->SetString(TEXT("Internationalization"), TEXT("Culture"), *CultureName, GEngineIni);
+				GConfig->SetString(TEXT("Internationalization"), TEXT("Culture"), *TargetCultureName, GEngineIni);
 			}
 		}
 		else
 #if !UE_BUILD_SHIPPING
 		// Use culture override specified on commandline.
-		if (FParse::Value(FCommandLine::Get(), TEXT("CULTURE="), CultureName))
+		if (FParse::Value(FCommandLine::Get(), TEXT("CULTURE="), TargetCultureName))
 		{
 			//UE_LOG(LogInit, Log, TEXT("Overriding culture %s w/ command-line option".), *CultureName);
 		}
@@ -56,20 +56,20 @@ void EndInitTextLocalization()
 #endif // !UE_BUILD_SHIPPING
 #if WITH_EDITOR
 		// See if we've been provided a culture override in the editor
-		if(GConfig->GetString( TEXT("Internationalization"), TEXT("Culture"), CultureName, GEditorGameAgnosticIni ))
+		if(GConfig->GetString( TEXT("Internationalization"), TEXT("Culture"), TargetCultureName, GEditorGameAgnosticIni ))
 		{
 			//UE_LOG(LogInit, Log, TEXT("Overriding culture %s w/ editor configuration."), *CultureName);
 		}
 		else
 #endif // WITH_EDITOR
 		// Use culture specified in engine configuration.
-		if(GConfig->GetString( TEXT("Internationalization"), TEXT("Culture"), CultureName, GEngineIni ))
+		if(GConfig->GetString( TEXT("Internationalization"), TEXT("Culture"), TargetCultureName, GEngineIni ))
 		{
 			//UE_LOG(LogInit, Log, TEXT("Overriding culture %s w/ engine configuration."), *CultureName);
 		}
 		else
 		{
-			CultureName = I18N.GetDefaultCulture()->GetName();
+			TargetCultureName = I18N.GetDefaultCulture()->GetName();
 		}
 
 		{
@@ -90,16 +90,37 @@ void EndInitTextLocalization()
 			TArray< TSharedPtr<FCulture, ESPMode::ThreadSafe> > AvailableCultures;
 			I18N.GetCulturesWithAvailableLocalization(LocalizationPaths, AvailableCultures);
 
-			// If we don't have an available translation for the current culture, fallback to generic English
-			TSharedPtr<FCulture, ESPMode::ThreadSafe> TargetCulture = I18N.GetCulture(CultureName);
-			if(!TargetCulture.IsValid() || !AvailableCultures.Contains(TargetCulture))
+			const FString RequestedCultureName = TargetCultureName;
+			TSharedPtr<FCulture, ESPMode::ThreadSafe> RequestedCulture = I18N.GetCulture(TargetCultureName);
+			if(RequestedCulture.IsValid())
 			{
-				UE_LOG(LogTextLocalizationManager, Warning, TEXT("The selected culture '%s' is not available; falling back to 'en'"), *CultureName);
-				CultureName = "en";
+				if(!AvailableCultures.Contains(RequestedCulture))
+				{
+					// Try base language.
+					TargetCultureName = RequestedCulture->GetTwoLetterISOLanguageName();
+					TSharedPtr<FCulture, ESPMode::ThreadSafe> TargetCulture = I18N.GetCulture(TargetCultureName);
+					if(AvailableCultures.Contains(TargetCulture))
+					{
+						// Fallback to base language.
+						UE_LOG(LogTextLocalizationManager, Warning, TEXT("The selected culture '%s' has no localization data; falling back to base language '%s'"), *RequestedCultureName, *TargetCultureName);
+					}
+					else
+					{
+						// Fallback to English.
+						UE_LOG(LogTextLocalizationManager, Warning, TEXT("The selected culture '%s' has no localization data; falling back to 'en'"), *RequestedCultureName);
+						TargetCultureName = "en";
+					}
+				}
+			}
+			// Fallback to English.
+			else
+			{
+				UE_LOG(LogTextLocalizationManager, Warning, TEXT("The selected culture '%s' is not valid; falling back to 'en'"), *RequestedCultureName);
+				TargetCultureName = "en";
 			}
 		}
 
-		I18N.SetCurrentCulture(CultureName);
+		I18N.SetCurrentCulture(TargetCultureName);
 	}
 
 	FTextLocalizationManager::Get().LoadResources(ShouldLoadEditor, ShouldLoadGame);
diff --git a/Engine/Source/Runtime/Core/Public/Internationalization/Internationalization.h b/Engine/Source/Runtime/Core/Public/Internationalization/Internationalization.h
--- a/Engine/Source/Runtime/Core/Public/Internationalization/Internationalization.h
+++ b/Engine/Source/Runtime/Core/Public/Internationalization/Internationalization.h
@@ -68,8 +68,8 @@ public:
 
 	CORE_API void GetCultureNames(TArray<FString>& CultureNames) const;
 
-	// Given some paths to look at, populate a list of cultures that we have available localization information for
-	CORE_API void GetCulturesWithAvailableLocalization(const TArray<FString>& InLocalizationPaths, TArray< TSharedPtr<FCulture, ESPMode::ThreadSafe> >& OutAvailableCultures) const;
+	// Given some paths to look at, populate a list of cultures that we have available localization information for. If bAllowFallback, include cultures for which we have non-region-specific data for.
+	CORE_API void GetCulturesWithAvailableLocalization(const TArray<FString>& InLocalizationPaths, TArray< TSharedPtr<FCulture, ESPMode::ThreadSafe> >& OutAvailableCultures, const bool bAllowFallback = false) const;
 
 private:
 	FInternationalization();
