Fix for enum default values parsing in UHT.
TTP #342183: Enum default values do not parse correctly

diff --git a/Engine/Source/Programs/UnrealHeaderTool/Private/HeaderParser.cpp b/Engine/Source/Programs/UnrealHeaderTool/Private/HeaderParser.cpp
--- a/Engine/Source/Programs/UnrealHeaderTool/Private/HeaderParser.cpp
+++ b/Engine/Source/Programs/UnrealHeaderTool/Private/HeaderParser.cpp
@@ -7128,7 +7128,7 @@ bool FHeaderParser::DefaultValueStringCppFormatToInnerFormat(const UProperty* Pr
 			const UEnum* Enum = CastChecked<UByteProperty>(Property)->Enum;
 			if( NULL != Enum )
 			{
-				OutForm = FDefaultValueHelper::RemoveWhitespaces( CppForm );
+				OutForm = FDefaultValueHelper::GetUnqualifiedEnumValue(FDefaultValueHelper::RemoveWhitespaces(CppForm));
 				return ( INDEX_NONE != Enum->FindEnumIndex( *OutForm ) );
 			}
 			int32 Value;
diff --git a/Engine/Source/Runtime/Core/Private/Misc/DefaultValueHelper.cpp b/Engine/Source/Runtime/Core/Private/Misc/DefaultValueHelper.cpp
--- a/Engine/Source/Runtime/Core/Private/Misc/DefaultValueHelper.cpp
+++ b/Engine/Source/Runtime/Core/Private/Misc/DefaultValueHelper.cpp
@@ -86,6 +86,17 @@ bool FDefaultValueHelper::Trim(const TCHAR* & Start, const TCHAR* End )
 	return Start < End;
 }
 
+FString FDefaultValueHelper::GetUnqualifiedEnumValue(const FString& Source)
+{
+	auto SeparatorPosition = Source.Find(FString("::"));
+	if (SeparatorPosition == INDEX_NONE)
+	{
+		return Source;
+	}
+
+	return Source.RightChop(SeparatorPosition + 2);
+}
+
 bool FDefaultValueHelper::HasWhitespaces(const FString& Source)
 {
 	for(int Pos = 0; Pos < Source.Len(); ++Pos)
diff --git a/Engine/Source/Runtime/Core/Public/Misc/DefaultValueHelper.h b/Engine/Source/Runtime/Core/Public/Misc/DefaultValueHelper.h
--- a/Engine/Source/Runtime/Core/Public/Misc/DefaultValueHelper.h
+++ b/Engine/Source/Runtime/Core/Public/Misc/DefaultValueHelper.h
@@ -45,6 +45,9 @@ class CORE_API FDefaultValueHelper
 	static bool IsStringValidInteger(const TCHAR* Start, const TCHAR* End, int32& OutBase);
 
 public:
+
+	static FString GetUnqualifiedEnumValue(const FString& Source);
+
 	static bool HasWhitespaces(const FString& Source);
 
 	static FString RemoveWhitespaces(const FString& Source);
