Moved FText::Revision to the history classes to be managed there.

diff --git a/Engine/Source/Runtime/Core/Private/Internationalization/Text.cpp b/Engine/Source/Runtime/Core/Private/Internationalization/Text.cpp
--- a/Engine/Source/Runtime/Core/Private/Internationalization/Text.cpp
+++ b/Engine/Source/Runtime/Core/Private/Internationalization/Text.cpp
@@ -235,7 +235,6 @@ const FText FText::SerializationFailureError = LOCTEXT("Error_SerializationFailu
 FText::FText()
 	: DisplayString( new FString() )
 	, Flags(0)
-	, Revision(FTextLocalizationManager::Get().GetHeadCultureRevision())
 {
 }
 
@@ -249,7 +248,6 @@ FText::FText(const FText& Source)
 	: DisplayString(Source.DisplayString)
 	, History(Source.History)
 	, Flags(Source.Flags)
-	, Revision(Source.Revision)
 {
 }
 
@@ -257,7 +255,6 @@ FText::FText(FText&& Source)
 	: DisplayString(MoveTemp(Source.DisplayString))
 	, History(MoveTemp(Source.History))
 	, Flags(MoveTemp(Source.Flags))
-	, Revision(MoveTemp(Source.Revision))
 {
 }
 
@@ -266,7 +263,6 @@ FText& FText::operator=(const FText& Source)
 	DisplayString = Source.DisplayString;
 	History = Source.History;
 	Flags = Source.Flags;
-	Revision = Source.Revision;
 
 	return *this;
 }
@@ -276,7 +272,6 @@ FText& FText::operator=(FText&& Source)
 	DisplayString = MoveTemp(Source.DisplayString);
 	History = MoveTemp(Source.History);
 	Flags = MoveTemp(Source.Flags);
-	Revision = MoveTemp(Source.Revision);
 
 	return *this;
 }
@@ -285,7 +280,6 @@ FText& FText::operator=(FText&& Source)
 FText::FText( FString InSourceString )
 	: DisplayString( new FString( MoveTemp(InSourceString) ))
 	, Flags(0)
-	, Revision(FTextLocalizationManager::Get().GetHeadCultureRevision())
 {
 	History = MakeShareable(new FTextHistory_Base(DisplayString));
 }
@@ -293,7 +287,6 @@ FText::FText( FString InSourceString )
 FText::FText( FString InSourceString, FString InNamespace, FString InKey, int32 InFlags )
 	: DisplayString( FTextLocalizationManager::Get().GetString(InNamespace, InKey, &InSourceString) )
 	, Flags(InFlags)
-	, Revision(FTextLocalizationManager::Get().GetHeadCultureRevision())
 {
 	History = MakeShareable(new FTextHistory_Base(InSourceString));
 }
@@ -635,8 +628,6 @@ CORE_API FArchive& operator<<( FArchive& Ar, FText& Value )
 
 	if(Ar.IsLoading())
 	{
-		// Keeps the revision up-to-date
-		Value.Revision = INDEX_NONE;
 		Value.Rebuild();
 	}
 
@@ -739,11 +730,9 @@ FString FText::BuildSourceString() const
 
 void FText::Rebuild() const
 {
-	if(History.IsValid() && History->IsOutOfDate(Revision))
+	if(History.IsValid())
 	{
-		Revision = FTextLocalizationManager::Get().GetHeadCultureRevision();
-		
-		DisplayString.Get() = History->ToText(false).DisplayString.Get();
+		History->Rebuild(DisplayString);
 	}
 }
 
diff --git a/Engine/Source/Runtime/Core/Private/Internationalization/TextHistory.cpp b/Engine/Source/Runtime/Core/Private/Internationalization/TextHistory.cpp
--- a/Engine/Source/Runtime/Core/Private/Internationalization/TextHistory.cpp
+++ b/Engine/Source/Runtime/Core/Private/Internationalization/TextHistory.cpp
@@ -8,9 +8,16 @@
 // FTextHistory
 
 /** Base class for all FText history types */
-bool FTextHistory::IsOutOfDate(int32 InRevision)
+
+FTextHistory::FTextHistory()
+	: Revision(FTextLocalizationManager::Get().GetHeadCultureRevision())
 {
-	return InRevision < FTextLocalizationManager::Get().GetHeadCultureRevision();
+
+}
+
+bool FTextHistory::IsOutOfDate()
+{
+	return Revision < FTextLocalizationManager::Get().GetHeadCultureRevision();
 }
 
 TSharedPtr< FString, ESPMode::ThreadSafe > FTextHistory::GetSourceString() const
@@ -22,6 +29,9 @@ void FTextHistory::SerializeForDisplayString(FArchive& Ar, TSharedRef<FString, E
 {
 	if(Ar.IsLoading())
 	{
+		// We will definitely need to do a rebuild later
+		Revision = INDEX_NONE;
+
 		//When duplicating, the CDO is used as the template, then values for the instance are assigned.
 		//If we don't duplicate the string, the CDO and the instance are both pointing at the same thing.
 		//This would result in all subsequently duplicated objects stamping over formerly duplicated ones.
@@ -29,8 +39,18 @@ void FTextHistory::SerializeForDisplayString(FArchive& Ar, TSharedRef<FString, E
 	}
 }
 
+void FTextHistory::Rebuild(TSharedRef< FString, ESPMode::ThreadSafe > InDisplayString)
+{
+	if(IsOutOfDate())
+	{
+		Revision = FTextLocalizationManager::Get().GetHeadCultureRevision();
+
+		InDisplayString.Get() = FTextInspector::GetDisplayString(ToText(false));
+	}
+}
+
 ///////////////////////////////////////
-// FTextHistory_NamedFormat
+// FTextHistory_Base
 
 FTextHistory_Base::FTextHistory_Base(FString InSourceString)
 	: SourceString(new FString( MoveTemp( InSourceString ) ))
@@ -80,6 +100,9 @@ void FTextHistory_Base::SerializeForDisplayString(FArchive& Ar, TSharedRef<FStri
 {
 	if(Ar.IsLoading())
 	{
+		// We will definitely need to do a rebuild later
+		Revision = INDEX_NONE;
+
 		FString Namespace;
 		FString Key;
 		FString SourceStringRaw;
diff --git a/Engine/Source/Runtime/Core/Private/Internationalization/TextHistory.h b/Engine/Source/Runtime/Core/Private/Internationalization/TextHistory.h
--- a/Engine/Source/Runtime/Core/Private/Internationalization/TextHistory.h
+++ b/Engine/Source/Runtime/Core/Private/Internationalization/TextHistory.h
@@ -21,6 +21,8 @@ enum class ETextHistoryType
 class CORE_API FTextHistory
 {
 public:
+	FTextHistory();
+
 	virtual ~FTextHistory() {};
 
 	/** Rebuilds the FText from the hierarchical history, the result should be in the current locale */
@@ -33,10 +35,17 @@ public:
 	virtual void SerializeForDisplayString(FArchive& Ar, TSharedRef<FString, ESPMode::ThreadSafe>& InOutDisplayString);
 
 	/** Returns TRUE if the Revision is out of date */
-	virtual bool IsOutOfDate(int32 InRevision);
+	virtual bool IsOutOfDate();
 
 	/** Returns the source string managed by the history (if any). */
 	virtual TSharedPtr< FString, ESPMode::ThreadSafe > GetSourceString() const;
+
+	/** Will rebuild the display string if out of date. */
+	void Rebuild(TSharedRef< FString, ESPMode::ThreadSafe > InDisplayString);
+
+protected:
+	/** Revision index of this history, rebuilds when it is out of sync with the FTextLocalizationManager */
+	int32 Revision;
 };
 
 /** No complexity to it, just holds the source string. */
@@ -52,7 +61,7 @@ public:
 	virtual FText ToText(bool bInAsSource) const override;
 	virtual void Serialize(FArchive& Ar) override;
 	virtual void SerializeForDisplayString(FArchive& Ar, TSharedRef<FString, ESPMode::ThreadSafe>& InOutDisplayString) override;
-	virtual bool IsOutOfDate(int32 InRevision) override { return false; }
+	virtual bool IsOutOfDate() override { return false; }
 	virtual TSharedPtr< FString, ESPMode::ThreadSafe > GetSourceString() const override;
 	// End FTextHistory interface
 
diff --git a/Engine/Source/Runtime/Core/Public/Internationalization/Text.h b/Engine/Source/Runtime/Core/Public/Internationalization/Text.h
--- a/Engine/Source/Runtime/Core/Public/Internationalization/Text.h
+++ b/Engine/Source/Runtime/Core/Public/Internationalization/Text.h
@@ -333,9 +333,6 @@ private:
 	/** Flags with various information on what sort of FText this is */
 	int32 Flags;
 
-	/** Revision index of this FText, rebuilds when it is out of sync with the FTextLocalizationManager */
-	mutable int32 Revision;
-
 	static bool bEnableErrorCheckingResults;
 	static bool bSuppressWarnings;
 
