Added IterateDirectory functionality to FileManager and GenericFileManager, mirrored from IPlatformFile, and forwarded to call to the appropriate IPlatformFile. Fixed improper iteration over directories when trying to enumerate which cultures have localization data.

diff --git a/Engine/Source/Runtime/Core/Private/HAL/FileManagerGeneric.cpp b/Engine/Source/Runtime/Core/Private/HAL/FileManagerGeneric.cpp
--- a/Engine/Source/Runtime/Core/Private/HAL/FileManagerGeneric.cpp
+++ b/Engine/Source/Runtime/Core/Private/HAL/FileManagerGeneric.cpp
@@ -386,6 +386,16 @@ void FFileManagerGeneric::FindFiles( TArray<FString>& Result, const TCHAR* InFil
 	GetLowLevel().IterateDirectory( *FPaths::GetPath(Filename), FileMatch );
 }
 
+bool FFileManagerGeneric::IterateDirectory(const TCHAR* Directory, IPlatformFile::FDirectoryVisitor& Visitor)
+{
+	return GetLowLevel().IterateDirectory( Directory, Visitor );
+}
+
+bool FFileManagerGeneric::IterateDirectoryRecursively(const TCHAR* Directory, IPlatformFile::FDirectoryVisitor& Visitor)
+{
+	return GetLowLevel().IterateDirectoryRecursively( Directory, Visitor );
+}
+
 double FFileManagerGeneric::GetFileAgeSeconds( const TCHAR* Filename )
 {
 	// make sure it exists
diff --git a/Engine/Source/Runtime/Core/Private/Internationalization/Internationalization.cpp b/Engine/Source/Runtime/Core/Private/Internationalization/Internationalization.cpp
--- a/Engine/Source/Runtime/Core/Private/Internationalization/Internationalization.cpp
+++ b/Engine/Source/Runtime/Core/Private/Internationalization/Internationalization.cpp
@@ -571,7 +571,7 @@ void FInternationalization::GetCulturesWithAvailableLocalization(const TArray<FS
 	OutAvailableCultures.Reset();
 
 	TArray<FString> AllLocalizationFolders;
-	IPlatformFile& PlatformFile = IPlatformFile::GetPlatformPhysical();
+	IFileManager& FileManager = IFileManager::Get();
 	for(const auto& LocalizationPath : InLocalizationPaths)
 	{
 		/* Visitor class used to enumerate directories of culture */
@@ -601,7 +601,7 @@ void FInternationalization::GetCulturesWithAvailableLocalization(const TArray<FS
 		};
 
 		FCultureEnumeratorVistor CultureEnumeratorVistor(AllLocalizationFolders);
-		PlatformFile.IterateDirectory(*LocalizationPath, CultureEnumeratorVistor);	
+		FileManager.IterateDirectory(*LocalizationPath, CultureEnumeratorVistor);
 	}
 
 	// Find any cultures that are a complete or partial match for the languages we have translations for
diff --git a/Engine/Source/Runtime/Core/Public/HAL/FileManager.h b/Engine/Source/Runtime/Core/Public/HAL/FileManager.h
--- a/Engine/Source/Runtime/Core/Public/HAL/FileManager.h
+++ b/Engine/Source/Runtime/Core/Public/HAL/FileManager.h
@@ -100,6 +100,22 @@ public:
 	/** Finds file or directories recursively. */
 	virtual void FindFilesRecursive( TArray<FString>& FileNames, const TCHAR* StartDirectory, const TCHAR* Filename, bool Files, bool Directories, bool bClearFileNames=true) = 0; // utility
 
+	/** 
+	 * Call the Visit function of the visitor once for each file or directory in a single directory. This function does not explore subdirectories.
+	 * @param Directory		The directory to iterate the contents of.
+	 * @param Visitor		Visitor to call for each element of the directory
+	 * @return				false if the directory did not exist or if the visitor returned false.
+	**/
+	virtual bool IterateDirectory(const TCHAR* Directory, IPlatformFile::FDirectoryVisitor& Visitor) = 0;
+
+	/** 
+	 * Call the Visit function of the visitor once for each file or directory in a directory tree. This function explores subdirectories.
+	 * @param Directory		The directory to iterate the contents of, recursively.
+	 * @param Visitor		Visitor to call for each element of the directory and each element of all subdirectories.
+	 * @return				false if the directory did not exist or if the visitor returned false.
+	**/
+	virtual bool IterateDirectoryRecursively(const TCHAR* Directory, IPlatformFile::FDirectoryVisitor& Visitor) = 0;
+
 	/** Gets the age of a file measured in seconds. */
 	virtual double GetFileAgeSeconds( const TCHAR* Filename )=0;
 
diff --git a/Engine/Source/Runtime/Core/Public/HAL/FileManagerGeneric.h b/Engine/Source/Runtime/Core/Public/HAL/FileManagerGeneric.h
--- a/Engine/Source/Runtime/Core/Public/HAL/FileManagerGeneric.h
+++ b/Engine/Source/Runtime/Core/Public/HAL/FileManagerGeneric.h
@@ -69,6 +69,22 @@ public:
 	virtual bool	MakeDirectory( const TCHAR* Path, bool Tree=0 );
 	virtual bool	DeleteDirectory( const TCHAR* Path, bool RequireExists=0, bool Tree=0 );
 
+	/** 
+	 * Call the Visit function of the visitor once for each file or directory in a single directory. This function does not explore subdirectories.
+	 * @param Directory		The directory to iterate the contents of.
+	 * @param Visitor		Visitor to call for each element of the directory
+	 * @return				false if the directory did not exist or if the visitor returned false.
+	**/
+	bool IterateDirectory(const TCHAR* Directory, IPlatformFile::FDirectoryVisitor& Visitor) override;
+
+	/** 
+	 * Call the Visit function of the visitor once for each file or directory in a directory tree. This function explores subdirectories.
+	 * @param Directory		The directory to iterate the contents of, recursively.
+	 * @param Visitor		Visitor to call for each element of the directory and each element of all subdirectories.
+	 * @return				false if the directory did not exist or if the visitor returned false.
+	**/
+	bool IterateDirectoryRecursively(const TCHAR* Directory, IPlatformFile::FDirectoryVisitor& Visitor) override;
+
 	/**
 	 * Converts passed in filename to use a relative path.
 	 *
