Core: Renamed FCOMPtr to TComPtr

diff --git a/Engine/Source/Runtime/Core/Private/Windows/TextStoreACP.cpp b/Engine/Source/Runtime/Core/Private/Windows/TextStoreACP.cpp
--- a/Engine/Source/Runtime/Core/Private/Windows/TextStoreACP.cpp
+++ b/Engine/Source/Runtime/Core/Private/Windows/TextStoreACP.cpp
@@ -102,14 +102,14 @@ STDAPI FTextStoreACP::AdviseSink(__RPC__in REFIID riid, __RPC__in_opt IUnknown *
 	}
 	else
 	{
-		FCOMPtr<IUnknown> OurSinkId(nullptr);
+		TComPtr<IUnknown> OurSinkId(nullptr);
 		if( FAILED( OurSinkId.FromQueryInterface(IID_IUnknown, punk) ) )
 		{
 			
 			return E_UNEXPECTED;
 		}
 
-		FCOMPtr<IUnknown> TheirSinkId(nullptr);
+		TComPtr<IUnknown> TheirSinkId(nullptr);
 		if( FAILED( TheirSinkId.FromQueryInterface(IID_IUnknown, punk) ) )
 		{
 			return E_UNEXPECTED;
@@ -144,13 +144,13 @@ STDAPI FTextStoreACP::UnadviseSink(__RPC__in_opt IUnknown *punk)
 		return CONNECT_E_NOCONNECTION;
 	}
 
-	FCOMPtr<IUnknown> OurSinkId(nullptr);
+	TComPtr<IUnknown> OurSinkId(nullptr);
 	if( FAILED( OurSinkId.FromQueryInterface(IID_IUnknown, punk) ) )
 	{
 		return E_UNEXPECTED;
 	}
 
-	FCOMPtr<IUnknown> TheirSinkId(nullptr);
+	TComPtr<IUnknown> TheirSinkId(nullptr);
 	if( FAILED( TheirSinkId.FromQueryInterface(IID_IUnknown, punk) ) )
 	{
 		return E_UNEXPECTED;
@@ -782,7 +782,7 @@ STDAPI FTextStoreACP::OnUpdateComposition(__RPC__in_opt ITfCompositionView *pCom
 
 	if(pRangeNew)
 	{
-		FCOMPtr<ITfRangeACP> RangeACP;
+		TComPtr<ITfRangeACP> RangeACP;
 		RangeACP.FromQueryInterface(IID_ITfRangeACP, pRangeNew);
 
 		LONG BeginIndex;
diff --git a/Engine/Source/Runtime/Core/Private/Windows/WindowsTextInputMethodSystem.cpp b/Engine/Source/Runtime/Core/Private/Windows/WindowsTextInputMethodSystem.cpp
--- a/Engine/Source/Runtime/Core/Private/Windows/WindowsTextInputMethodSystem.cpp
+++ b/Engine/Source/Runtime/Core/Private/Windows/WindowsTextInputMethodSystem.cpp
@@ -8,6 +8,7 @@ DEFINE_LOG_CATEGORY_STATIC(LogWindowsTextInputMethodSystem, Log, All);
 
 #include "AllowWindowsPlatformTypes.h"
 
+
 namespace
 {
 	FString GetIMMStringAsFString(HIMC IMMContext, const DWORD StringType)
@@ -28,10 +29,11 @@ namespace
 		return OutString;
 	}
 
-	class FTextInputMethodChangeNotifier : public ITextInputMethodChangeNotifier
+	class FTextInputMethodChangeNotifier
+		: public ITextInputMethodChangeNotifier
 	{
 	public:
-		FTextInputMethodChangeNotifier(const FCOMPtr<FTextStoreACP>& InTextStoreACP)
+		FTextInputMethodChangeNotifier(const TComPtr<FTextStoreACP>& InTextStoreACP)
 			:	TextStoreACP(InTextStoreACP)
 		{}
 
@@ -43,7 +45,7 @@ namespace
 		virtual void CancelComposition() override;
 
 	private:
-		const FCOMPtr<FTextStoreACP> TextStoreACP;
+		const TComPtr<FTextStoreACP> TextStoreACP;
 	};
 
 	void FTextInputMethodChangeNotifier::NotifyLayoutChanged(const ELayoutChangeType ChangeType)
@@ -321,7 +323,7 @@ bool FWindowsTextInputMethodSystem::InitializeTSF()
 		}
 
 		// Get source from thread manager, needed to install profile processor related sinks.
-		FCOMPtr<ITfSource> TSFSource;
+		TComPtr<ITfSource> TSFSource;
 		Result = TSFSource.FromQueryInterface(IID_ITfSource, TSFThreadManager);
 		if(FAILED(Result))
 		{
@@ -410,7 +412,7 @@ void FWindowsTextInputMethodSystem::Terminate()
 	HRESULT Result;
 
 	// Get source from thread manager, needed to uninstall profile processor related sinks.
-	FCOMPtr<ITfSource> TSFSource;
+	TComPtr<ITfSource> TSFSource;
 	Result = TSFSource.FromQueryInterface(IID_ITfSource, TSFThreadManager);
 	if(FAILED(Result) || !TSFSource)
 	{
@@ -459,7 +461,7 @@ TSharedPtr<ITextInputMethodChangeNotifier> FWindowsTextInputMethodSystem::Regist
 	FInternalContext& InternalContext = ContextToInternalContextMap.Add(Context);
 
 	// TSF Implementation
-	FCOMPtr<FTextStoreACP>& TextStore = InternalContext.TSFContext;
+	TComPtr<FTextStoreACP>& TextStore = InternalContext.TSFContext;
 	TextStore.Attach(new FTextStoreACP(Context));
 
 	Result = TSFThreadManager->CreateDocumentMgr(&(TextStore->TSFDocumentManager));
@@ -514,7 +516,7 @@ void FWindowsTextInputMethodSystem::UnregisterContext(const TSharedRef<ITextInpu
 	FInternalContext& InternalContext = ContextToInternalContextMap[Context];
 
 	// TSF Implementation
-	FCOMPtr<FTextStoreACP>& TextStore = InternalContext.TSFContext;
+	TComPtr<FTextStoreACP>& TextStore = InternalContext.TSFContext;
 
 	Result = TextStore->TSFDocumentManager->Pop(TF_POPF_ALL);
 	if(FAILED(Result))
@@ -542,7 +544,7 @@ void FWindowsTextInputMethodSystem::ActivateContext(const TSharedRef<ITextInputM
 	InternalContext.IMMContext.IsDeactivating = false;
 
 	// TSF Implementation
-	FCOMPtr<FTextStoreACP>& TextStore = InternalContext.TSFContext;
+	TComPtr<FTextStoreACP>& TextStore = InternalContext.TSFContext;
 	Result = TSFThreadManager->SetFocus(TextStore->TSFDocumentManager);
 	if(FAILED(Result))
 	{
@@ -560,7 +562,7 @@ void FWindowsTextInputMethodSystem::DeactivateContext(const TSharedRef<ITextInpu
 	FInternalContext& InternalContext = ContextToInternalContextMap[Context];
 
 	// TSF Implementation
-	FCOMPtr<FTextStoreACP> TextStore = InternalContext.TSFContext;
+	TComPtr<FTextStoreACP> TextStore = InternalContext.TSFContext;
 
 	ITfDocumentMgr* FocusedDocumentManger = nullptr;
 	Result = TSFThreadManager->GetFocus(&(FocusedDocumentManger));
diff --git a/Engine/Source/Runtime/Core/Public/Windows/COMPointer.h b/Engine/Source/Runtime/Core/Public/Windows/COMPointer.h
--- a/Engine/Source/Runtime/Core/Public/Windows/COMPointer.h
+++ b/Engine/Source/Runtime/Core/Public/Windows/COMPointer.h
@@ -6,17 +6,17 @@
 * Smart COM object pointer.
 */
 template<typename T>
-class FCOMPtr
+class TComPtr
 {
 public:
 	typedef T PointerType;
 
 public:
-	FCOMPtr() :	RawPointer(nullptr)
+	TComPtr() :	RawPointer(nullptr)
 	{
 	}
 
-	FCOMPtr(PointerType* const Source) : RawPointer(Source)
+	TComPtr(PointerType* const Source) : RawPointer(Source)
 	{
 		if(RawPointer)
 		{
@@ -24,7 +24,7 @@ public:
 		}
 	}
 
-	FCOMPtr(const FCOMPtr<PointerType>& Source) : RawPointer(Source.RawPointer)
+	TComPtr(const TComPtr<PointerType>& Source) : RawPointer(Source.RawPointer)
 	{
 		if(RawPointer)
 		{
@@ -32,12 +32,12 @@ public:
 		}
 	}
 
-	FCOMPtr(FCOMPtr<PointerType>&& Source) : RawPointer(Source.RawPointer)
+	TComPtr(TComPtr<PointerType>&& Source) : RawPointer(Source.RawPointer)
 	{	
 		Source.RawPointer = nullptr;
 	}	
 
-	FCOMPtr<PointerType>& operator=(PointerType* const Source) 
+	TComPtr<PointerType>& operator=(PointerType* const Source) 
 	{
 		if(RawPointer != Source)
 		{
@@ -54,7 +54,7 @@ public:
 		return *this;
 	}
 
-	FCOMPtr<PointerType>& operator=(const FCOMPtr<PointerType>& Source) 
+	TComPtr<PointerType>& operator=(const TComPtr<PointerType>& Source) 
 	{
 		if(RawPointer != Source.RawPointer)
 		{
@@ -71,7 +71,7 @@ public:
 		return *this;
 	}
 
-	FCOMPtr<PointerType>& operator=(FCOMPtr<PointerType>&& Source) 
+	TComPtr<PointerType>& operator=(TComPtr<PointerType>&& Source) 
 	{			
 		if(RawPointer != Source.RawPointer)
 		{
@@ -85,7 +85,7 @@ public:
 		return *this;
 	}
 
-	~FCOMPtr() 
+	~TComPtr() 
 	{
 		if(RawPointer)
 		{
diff --git a/Engine/Source/Runtime/Core/Public/Windows/TextStoreACP.h b/Engine/Source/Runtime/Core/Public/Windows/TextStoreACP.h
--- a/Engine/Source/Runtime/Core/Public/Windows/TextStoreACP.h
+++ b/Engine/Source/Runtime/Core/Public/Windows/TextStoreACP.h
@@ -103,9 +103,9 @@ public:
 		FAdviseSinkObject() : TextStoreACPSink(nullptr), SinkFlags(0) {}
 
 		// Sink object for ITextStoreACP Implementation
-		FCOMPtr<ITextStoreACPSink> TextStoreACPSink;
+		TComPtr<ITextStoreACPSink> TextStoreACPSink;
 
-		FCOMPtr<ITextStoreACPServices> TextStoreACPServices;
+		TComPtr<ITextStoreACPServices> TextStoreACPServices;
 
 		// Flags defining what events the sink object should be notified of.
 		DWORD SinkFlags;
@@ -117,19 +117,19 @@ public:
 		{}
 
 		// Composition view object for managing compositions.
-		FCOMPtr<ITfCompositionView> TSFCompositionView;
+		TComPtr<ITfCompositionView> TSFCompositionView;
 
 	} Composition;
 
 public:
 	// Document manager object for managing contexts.
-	FCOMPtr<ITfDocumentMgr> TSFDocumentManager;
+	TComPtr<ITfDocumentMgr> TSFDocumentManager;
 
 	// Context object for pushing context on to document manager.
-	FCOMPtr<ITfContext> TSFContext;
+	TComPtr<ITfContext> TSFContext;
 
 	// Context owner composition services object for terminating compositions.
-	FCOMPtr<ITfContextOwnerCompositionServices> TSFContextOwnerCompositionServices;
+	TComPtr<ITfContextOwnerCompositionServices> TSFContextOwnerCompositionServices;
 
 	TfEditCookie TSFEditCookie;
 };
diff --git a/Engine/Source/Runtime/Core/Public/Windows/WindowsTextInputMethodSystem.h b/Engine/Source/Runtime/Core/Public/Windows/WindowsTextInputMethodSystem.h
--- a/Engine/Source/Runtime/Core/Public/Windows/WindowsTextInputMethodSystem.h
+++ b/Engine/Source/Runtime/Core/Public/Windows/WindowsTextInputMethodSystem.h
@@ -82,16 +82,16 @@ private:
 	} CurrentAPI;
 
 	// TSF Implementation
-	FCOMPtr<ITfInputProcessorProfiles> TSFInputProcessorProfiles;
-	FCOMPtr<ITfInputProcessorProfileMgr> TSFInputProcessorProfileManager;
-	FCOMPtr<ITfThreadMgr> TSFThreadManager;
+	TComPtr<ITfInputProcessorProfiles> TSFInputProcessorProfiles;
+	TComPtr<ITfInputProcessorProfileMgr> TSFInputProcessorProfileManager;
+	TComPtr<ITfThreadMgr> TSFThreadManager;
 	TfClientId TSFClientId;
-	FCOMPtr<ITfDocumentMgr> TSFDisabledDocumentManager;
-	FCOMPtr<FTSFActivationProxy> TSFActivationProxy;
+	TComPtr<ITfDocumentMgr> TSFDisabledDocumentManager;
+	TComPtr<FTSFActivationProxy> TSFActivationProxy;
 
 	struct FInternalContext
 	{
-		FCOMPtr<FTextStoreACP> TSFContext;
+		TComPtr<FTextStoreACP> TSFContext;
 		struct
 		{
 			bool IsComposing;
