Core: Removed file header comments; NULL to nullptr; added missing documentation

diff --git a/Engine/Source/Runtime/Core/Private/Windows/WindowsEvent.h b/Engine/Source/Runtime/Core/Private/Windows/WindowsEvent.h
--- a/Engine/Source/Runtime/Core/Private/Windows/WindowsEvent.h
+++ b/Engine/Source/Runtime/Core/Private/Windows/WindowsEvent.h
@@ -1,13 +1,10 @@
 // Copyright 1998-2014 Epic Games, Inc. All Rights Reserved.
 
-/*=============================================================================
-	WindowsEvent.h: Declares the FEventWin class.
-=============================================================================*/
-
 #pragma once
 
 #include "AllowWindowsPlatformTypes.h"
 
+
 /**
  * Implements the Windows version of the FEvent interface.
  */
@@ -19,48 +16,47 @@ public:
 	/**
 	 * Default constructor.
 	 */
-	FEventWin ()
-		: Event(NULL)
+	FEventWin( )
+		: Event(nullptr)
 	{
 	}
 
 	/**
-	 * Destructor.
+	 * Virtual destructor.
 	 */
-	virtual ~FEventWin ()
+	virtual ~FEventWin( )
 	{
-		if (Event != NULL)
+		if (Event != nullptr)
 		{
 			CloseHandle(Event);
 		}
 	}
 
-
 public:
 
-	virtual bool Create (bool bIsManualReset = false) override
+	virtual bool Create( bool bIsManualReset = false ) override
 	{
 		// Create the event and default it to non-signaled
-		Event = CreateEvent(NULL, bIsManualReset, 0, NULL);
+		Event = CreateEvent(nullptr, bIsManualReset, 0, nullptr);
 
-		return Event != NULL;
+		return Event != nullptr;
 	}
 
-	virtual void Trigger () override
+	virtual void Trigger( ) override
 	{
 		check(Event);
 
 		SetEvent(Event);
 	}
 
-	virtual void Reset () override
+	virtual void Reset( ) override
 	{
 		check(Event);
 
 		ResetEvent(Event);
 	}
 
-	virtual bool Wait (uint32 WaitTime, const bool bIgnoreThreadIdleStats = false) override;
+	virtual bool Wait( uint32 WaitTime, const bool bIgnoreThreadIdleStats = false ) override;
 
 private:
 
@@ -68,4 +64,5 @@ private:
 	HANDLE Event;
 };
 
+
 #include "HideWindowsPlatformTypes.h"
diff --git a/Engine/Source/Runtime/Core/Private/Windows/WindowsPlatformAtomics.cpp b/Engine/Source/Runtime/Core/Private/Windows/WindowsPlatformAtomics.cpp
--- a/Engine/Source/Runtime/Core/Private/Windows/WindowsPlatformAtomics.cpp
+++ b/Engine/Source/Runtime/Core/Private/Windows/WindowsPlatformAtomics.cpp
@@ -1,11 +1,8 @@
 // Copyright 1998-2014 Epic Games, Inc. All Rights Reserved.
 
-/*=============================================================================
-	WindowsPlatformAtomics.cpp: Windows implementations of atomic functions
-=============================================================================*/
-
 #include "Core.h"
 
+
 void FWindowsPlatformAtomics::HandleAtomicsFailure( const TCHAR* InFormat, ... )
 {	
 	TCHAR TempStr[1024];
diff --git a/Engine/Source/Runtime/Core/Private/Windows/WindowsPlatformExceptionHandling.cpp b/Engine/Source/Runtime/Core/Private/Windows/WindowsPlatformExceptionHandling.cpp
--- a/Engine/Source/Runtime/Core/Private/Windows/WindowsPlatformExceptionHandling.cpp
+++ b/Engine/Source/Runtime/Core/Private/Windows/WindowsPlatformExceptionHandling.cpp
@@ -1,9 +1,5 @@
 // Copyright 1998-2014 Epic Games, Inc. All Rights Reserved.
 
-/*=============================================================================
-	WindowsPlatformExceptionHandling.cpp: Exception handling for functions that want to create crash dumps.
-=============================================================================*/
-
 #include "Core.h"
 #include "ExceptionHandling.h"
 #include "../../Launch/Resources/Version.h"
diff --git a/Engine/Source/Runtime/Core/Private/Windows/WindowsPlatformFile.cpp b/Engine/Source/Runtime/Core/Private/Windows/WindowsPlatformFile.cpp
--- a/Engine/Source/Runtime/Core/Private/Windows/WindowsPlatformFile.cpp
+++ b/Engine/Source/Runtime/Core/Private/Windows/WindowsPlatformFile.cpp
@@ -1,12 +1,9 @@
 // Copyright 1998-2014 Epic Games, Inc. All Rights Reserved.
 
-/*=============================================================================
-	WindowsPlatformFile.cpp: Windows implementations of File functions
-=============================================================================*/
-
 #include "Core.h"
 #include <sys/utime.h>
 
+
 // make an FTimeSpan object that represents the "epoch" for time_t (from a _stat struct)
 const FDateTime WindowsEpoch(1970, 1, 1);
 
diff --git a/Engine/Source/Runtime/Core/Private/Windows/WindowsPlatformMemory.cpp b/Engine/Source/Runtime/Core/Private/Windows/WindowsPlatformMemory.cpp
--- a/Engine/Source/Runtime/Core/Private/Windows/WindowsPlatformMemory.cpp
+++ b/Engine/Source/Runtime/Core/Private/Windows/WindowsPlatformMemory.cpp
@@ -1,9 +1,5 @@
 // Copyright 1998-2014 Epic Games, Inc. All Rights Reserved.
 
-/*=============================================================================
-	WindowsPlatformMemory.cpp: Windows platform memory functions
-=============================================================================*/
-
 #include "Core.h"
 
 #include "MallocTBB.h"
diff --git a/Engine/Source/Runtime/Core/Private/Windows/WindowsPlatformMisc.cpp b/Engine/Source/Runtime/Core/Private/Windows/WindowsPlatformMisc.cpp
--- a/Engine/Source/Runtime/Core/Private/Windows/WindowsPlatformMisc.cpp
+++ b/Engine/Source/Runtime/Core/Private/Windows/WindowsPlatformMisc.cpp
@@ -1,9 +1,5 @@
 // Copyright 1998-2014 Epic Games, Inc. All Rights Reserved.
 
-/*=============================================================================
-	WindowsPlatformMisc.cpp: Windows implementations of misc functions
-=============================================================================*/
-
 #include "Core.h"
 #include "ExceptionHandling.h"
 #include "SecureHash.h"
diff --git a/Engine/Source/Runtime/Core/Private/Windows/WindowsPlatformNamedPipe.cpp b/Engine/Source/Runtime/Core/Private/Windows/WindowsPlatformNamedPipe.cpp
--- a/Engine/Source/Runtime/Core/Private/Windows/WindowsPlatformNamedPipe.cpp
+++ b/Engine/Source/Runtime/Core/Private/Windows/WindowsPlatformNamedPipe.cpp
@@ -1,10 +1,5 @@
 // Copyright 1998-2014 Epic Games, Inc. All Rights Reserved.
 
-/*=============================================================================
-	WindowsPlatformNamedPipe.cpp: Windows implementations of Named Pipe functions
-=============================================================================*/
-
-
 #include "Core.h"
 
 static void VerifyWinResult(BOOL bResult, const TCHAR* InMessage)
diff --git a/Engine/Source/Runtime/Core/Private/Windows/WindowsPlatformOutputDevices.cpp b/Engine/Source/Runtime/Core/Private/Windows/WindowsPlatformOutputDevices.cpp
--- a/Engine/Source/Runtime/Core/Private/Windows/WindowsPlatformOutputDevices.cpp
+++ b/Engine/Source/Runtime/Core/Private/Windows/WindowsPlatformOutputDevices.cpp
@@ -1,10 +1,5 @@
 // Copyright 1998-2014 Epic Games, Inc. All Rights Reserved.
 
-/*=============================================================================
-	WindowsPlatformOutputDevices.cpp: Windows implementations of OutputDevices functions
-=============================================================================*/
-
-
 #include "Core.h"
 
 #include "FeedbackContextAnsi.h"
@@ -12,12 +7,16 @@
 #include "../Private/Windows/WindowsPlatformFeedbackContextPrivate.h"
 
 #include "AllowWindowsPlatformTypes.h"
+
 namespace OutputDeviceConstants
 {
 	uint32 WIN_STD_OUTPUT_HANDLE = STD_OUTPUT_HANDLE;
 	uint32 WIN_ATTACH_PARENT_PROCESS = ATTACH_PARENT_PROCESS;
 }
+
 #include "HideWindowsPlatformTypes.h"
+
+
 //////////////////////////////////
 // FWindowsPlatformOutputDevices
 //////////////////////////////////
diff --git a/Engine/Source/Runtime/Core/Private/Windows/WindowsPlatformOutputDevicesPrivate.h b/Engine/Source/Runtime/Core/Private/Windows/WindowsPlatformOutputDevicesPrivate.h
--- a/Engine/Source/Runtime/Core/Private/Windows/WindowsPlatformOutputDevicesPrivate.h
+++ b/Engine/Source/Runtime/Core/Private/Windows/WindowsPlatformOutputDevicesPrivate.h
@@ -1,9 +1,5 @@
 // Copyright 1998-2014 Epic Games, Inc. All Rights Reserved.
 
-/*=============================================================================
-	WindowsPlatformOutputDevicesPrivate.h: Windows implementations of output devices
-=============================================================================*/
-
 #pragma once
 
 
@@ -15,12 +11,14 @@
 	#endif
 #endif
 
+
 #if WANTS_WINDOWS_EVENT_LOGGING
+
 /**
  * Output device that writes to Windows Event Log
  */
-class FOutputDeviceEventLog :
-	public FOutputDevice
+class FOutputDeviceEventLog
+	: public FOutputDevice
 {
 	/** Handle to the event log object */
 	HANDLE EventLog;
diff --git a/Engine/Source/Runtime/Core/Private/Windows/WindowsPlatformSplash.cpp b/Engine/Source/Runtime/Core/Private/Windows/WindowsPlatformSplash.cpp
--- a/Engine/Source/Runtime/Core/Private/Windows/WindowsPlatformSplash.cpp
+++ b/Engine/Source/Runtime/Core/Private/Windows/WindowsPlatformSplash.cpp
@@ -1,9 +1,5 @@
 // Copyright 1998-2014 Epic Games, Inc. All Rights Reserved.
 
-/*================================================================================
-	SplashScreen.cpp: Splash screen for game/editor startup
-================================================================================*/
-
 #include "Core.h"
 
 #include "EngineVersion.h"
diff --git a/Engine/Source/Runtime/Core/Private/Windows/WindowsPlatformStackWalk.cpp b/Engine/Source/Runtime/Core/Private/Windows/WindowsPlatformStackWalk.cpp
--- a/Engine/Source/Runtime/Core/Private/Windows/WindowsPlatformStackWalk.cpp
+++ b/Engine/Source/Runtime/Core/Private/Windows/WindowsPlatformStackWalk.cpp
@@ -1,9 +1,5 @@
 // Copyright 1998-2014 Epic Games, Inc. All Rights Reserved.
 
-/*=============================================================================
-	WindowsPlatformStackWalk.cpp: Windows implementations of stack walk functions
-=============================================================================*/
-
 #include "Core.h"
 
 #include "AllowWindowsPlatformTypes.h"
@@ -14,6 +10,7 @@
 
 #include "ModuleManager.h"
 
+
 /*-----------------------------------------------------------------------------
 	Stack walking.
 -----------------------------------------------------------------------------*/
diff --git a/Engine/Source/Runtime/Core/Private/Windows/WindowsPlatformSurvey.cpp b/Engine/Source/Runtime/Core/Private/Windows/WindowsPlatformSurvey.cpp
--- a/Engine/Source/Runtime/Core/Private/Windows/WindowsPlatformSurvey.cpp
+++ b/Engine/Source/Runtime/Core/Private/Windows/WindowsPlatformSurvey.cpp
@@ -1,9 +1,5 @@
 // Copyright 1998-2014 Epic Games, Inc. All Rights Reserved.
 
-/*=================================================================================
-	WindowsPlatformSurvey.cpp: Microsoft Windows platform hardware-survey classes
-=================================================================================*/
-
 #include "Core.h"
 
 #define USING_WINSAT_API	1
diff --git a/Engine/Source/Runtime/Core/Private/Windows/WindowsPlatformTime.cpp b/Engine/Source/Runtime/Core/Private/Windows/WindowsPlatformTime.cpp
--- a/Engine/Source/Runtime/Core/Private/Windows/WindowsPlatformTime.cpp
+++ b/Engine/Source/Runtime/Core/Private/Windows/WindowsPlatformTime.cpp
@@ -1,13 +1,11 @@
 // Copyright 1998-2014 Epic Games, Inc. All Rights Reserved.
 
-/*=============================================================================
-	WindowsPlatformTime.cpp: Windows implementations of time functions
-=============================================================================*/
-
 #include "Core.h"
 
+
 float FWindowsPlatformTime::CPUTimePctRelative = 0.0f;
 
+
 double FWindowsPlatformTime::InitTiming(void)
 {
 	LARGE_INTEGER Frequency;
@@ -24,6 +22,7 @@ double FWindowsPlatformTime::InitTiming(void)
 	return FPlatformTime::Seconds();
 }
 
+
 void FWindowsPlatformTime::SystemTime( int32& Year, int32& Month, int32& DayOfWeek, int32& Day, int32& Hour, int32& Min, int32& Sec, int32& MSec )
 {
 	SYSTEMTIME st;
@@ -39,6 +38,7 @@ void FWindowsPlatformTime::SystemTime( int32& Year, int32& Month, int32& DayOfWe
 	MSec		= st.wMilliseconds;
 }
 
+
 void FWindowsPlatformTime::UtcTime( int32& Year, int32& Month, int32& DayOfWeek, int32& Day, int32& Hour, int32& Min, int32& Sec, int32& MSec )
 {
 	SYSTEMTIME st;
@@ -54,6 +54,7 @@ void FWindowsPlatformTime::UtcTime( int32& Year, int32& Month, int32& DayOfWeek,
 	MSec		= st.wMilliseconds;
 }
 
+
 /** Holds Windows filetime misc functions. */
 struct FFiletimeMisc
 {
@@ -75,6 +76,7 @@ struct FFiletimeMisc
 	}
 };
 
+
 bool FWindowsPlatformTime::UpdateCPUTime( float /*DeltaTime*/ )
 {
 	static double LastTotalProcessTime = 0.0f;
@@ -108,6 +110,7 @@ bool FWindowsPlatformTime::UpdateCPUTime( float /*DeltaTime*/ )
 	return true;
 }
 
+
 FCPUTime FWindowsPlatformTime::GetCPUTime()
 {
 	return FCPUTime( CPUTimePctRelative / (float)FPlatformMisc::NumberOfCoresIncludingHyperthreads(), CPUTimePctRelative );
diff --git a/Engine/Source/Runtime/Core/Private/Windows/WindowsRunnableThread.h b/Engine/Source/Runtime/Core/Private/Windows/WindowsRunnableThread.h
--- a/Engine/Source/Runtime/Core/Private/Windows/WindowsRunnableThread.h
+++ b/Engine/Source/Runtime/Core/Private/Windows/WindowsRunnableThread.h
@@ -2,45 +2,33 @@
 
 #pragma once
 
+
 /**
  * This is the base interface for all runnable thread classes. It specifies the
  * methods used in managing its life cycle.
  */
-class FRunnableThreadWin : public FRunnableThread
+class FRunnableThreadWin
+	: public FRunnableThread
 {
-	/**
-	 * The thread handle for the thread
-	 */
+	/** The thread handle for the thread. */
 	HANDLE Thread;
 
-	/**
-	 * The runnable object to execute on this thread
-	 */
+	/** The runnable object to execute on this thread. */
 	FRunnable* Runnable;
 
-	/** 
-	 * Sync event to make sure that Init() has been completed before allowing the main thread to continue
-	 */
+	/** Sync event to make sure that Init() has been completed before allowing the main thread to continue. */
 	FEvent* ThreadInitSyncEvent;
 
-	/**
-	 * The priority to run the thread at
-	 */
+	/** The priority to run the thread at. */
 	EThreadPriority ThreadPriority;
 
-	/**
-	 * The Affinity to run the thread with
-	 */
+	/** The Affinity to run the thread with. */
 	uint64 ThreadAffintyMask;
 
-	/**
-	* ID set during thread creation
-	*/
+	/** ID set during thread creation. */
 	uint32 ThreadID;
 
-	/**
-	 * Name of the thread.
-	 */
+	/** Name of the thread. */
 	FString ThreadName;
 
 	/**
@@ -92,7 +80,7 @@ class FRunnableThreadWin : public FRunnableThread
 	 * The thread entry point. Simply forwards the call on to the right
 	 * thread main function
 	 */
-	static ::DWORD STDCALL _ThreadProc(LPVOID pThis)
+	static ::DWORD STDCALL _ThreadProc( LPVOID pThis )
 	{
 		check(pThis);
 		return ((FRunnableThreadWin*)pThis)->GuardedRun();
@@ -108,7 +96,8 @@ class FRunnableThreadWin : public FRunnableThread
 	uint32 Run();
 
 public:
-	FRunnableThreadWin()
+
+	FRunnableThreadWin( )
 		: Thread(NULL)
 		, Runnable(NULL)
 		, ThreadInitSyncEvent(NULL)
@@ -118,7 +107,7 @@ public:
 
 	}
 
-	~FRunnableThreadWin()
+	~FRunnableThreadWin( )
 	{
 		// Clean up our thread if it is still active
 		if (Thread != NULL)
@@ -128,7 +117,7 @@ public:
 		FRunnableThread::GetThreadRegistry().Remove(ThreadID);
 	}
 	
-	virtual void SetThreadPriority(EThreadPriority NewPriority) override
+	virtual void SetThreadPriority( EThreadPriority NewPriority ) override
 	{
 		// Don't bother calling the OS if there is no need
 		if (NewPriority != ThreadPriority)
@@ -142,7 +131,7 @@ public:
 		}
 	}
 
-	virtual void Suspend(bool bShouldPause = 1) override
+	virtual void Suspend( bool bShouldPause = 1 ) override
 	{
 		check(Thread);
 		if (bShouldPause == true)
@@ -155,7 +144,7 @@ public:
 		}
 	}
 
-	virtual bool Kill(bool bShouldWait = false) override
+	virtual bool Kill( bool bShouldWait = false ) override
 	{
 		check(Thread && "Did you forget to call Create()?");
 		bool bDidExitOK = true;
@@ -182,26 +171,27 @@ public:
 		return bDidExitOK;
 	}
 
-	virtual void WaitForCompletion() override
+	virtual void WaitForCompletion( ) override
 	{
 		// Block until this thread exits
 		WaitForSingleObject(Thread,INFINITE);
 	}
 
-	virtual uint32 GetThreadID() override
+	virtual uint32 GetThreadID( ) override
 	{
 		return ThreadID;
 	}
 
-	virtual FString GetThreadName() override
+	virtual FString GetThreadName( ) override
 	{
 		return ThreadName;
 	}
 
 protected:
-	virtual bool CreateInternal(FRunnable* InRunnable, const TCHAR* InThreadName,
+
+	virtual bool CreateInternal( FRunnable* InRunnable, const TCHAR* InThreadName,
 		uint32 InStackSize = 0,
-		EThreadPriority InThreadPri = TPri_Normal, uint64 InThreadAffinityMask = 0) override
+		EThreadPriority InThreadPri = TPri_Normal, uint64 InThreadAffinityMask = 0 ) override
 	{
 		check(InRunnable);
 		Runnable = InRunnable;
@@ -212,6 +202,7 @@ protected:
 
 		// Create the new thread
 		Thread = CreateThread(NULL,InStackSize,_ThreadProc,this,STACK_SIZE_PARAM_IS_A_RESERVATION,(::DWORD *)&ThreadID);
+		
 		// If it fails, clear all the vars
 		if (Thread == NULL)
 		{
diff --git a/Engine/Source/Runtime/Core/Private/Windows/XInputInterface.h b/Engine/Source/Runtime/Core/Private/Windows/XInputInterface.h
--- a/Engine/Source/Runtime/Core/Private/Windows/XInputInterface.h
+++ b/Engine/Source/Runtime/Core/Private/Windows/XInputInterface.h
@@ -2,11 +2,14 @@
 
 #pragma once
 
+
+/** Max number of controllers. */
 #define MAX_NUM_XINPUT_CONTROLLERS 4
 
 /** Max number of controller buttons.  Must be < 256*/
 #define MAX_NUM_CONTROLLER_BUTTONS 24
 
+
 /**
  * Interface class for XInput devices (xbox 360 controller)                 
  */
@@ -106,5 +109,5 @@ private:
 	EControllerButtons::Type Buttons[MAX_NUM_CONTROLLER_BUTTONS];
 
 	/**  */
-	TSharedRef< FGenericApplicationMessageHandler > MessageHandler;
-};
\ No newline at end of file
+	TSharedRef<FGenericApplicationMessageHandler> MessageHandler;
+};
diff --git a/Engine/Source/Runtime/Core/Public/HAL/ConsoleManager.h b/Engine/Source/Runtime/Core/Public/HAL/ConsoleManager.h
--- a/Engine/Source/Runtime/Core/Public/HAL/ConsoleManager.h
+++ b/Engine/Source/Runtime/Core/Public/HAL/ConsoleManager.h
@@ -1,13 +1,10 @@
 // Copyright 1998-2014 Epic Games, Inc. All Rights Reserved.
 
-/*=============================================================================
-ConsoleManager.h: console command handling
-=============================================================================*/
-
 #pragma once
 
 #include "IConsoleManager.h"
 
+
 class FConsoleVariableBase;
 
 class CORE_API FConsoleManager :public IConsoleManager
diff --git a/Engine/Source/Runtime/Core/Public/HAL/ExceptionHandling.h b/Engine/Source/Runtime/Core/Public/HAL/ExceptionHandling.h
--- a/Engine/Source/Runtime/Core/Public/HAL/ExceptionHandling.h
+++ b/Engine/Source/Runtime/Core/Public/HAL/ExceptionHandling.h
@@ -1,13 +1,10 @@
 // Copyright 1998-2014 Epic Games, Inc. All Rights Reserved.
 
-/*=============================================================================
-	ExceptionHandling.h: Exception handling for functions that want to create crash dumps.
-=============================================================================*/
-
 #pragma once
 
 #include "Core.h"
 
+
 /** Whether we should generate crash reports even if the debugger is attached. */
 extern CORE_API bool GAlwaysReportCrash;
 
diff --git a/Engine/Source/Runtime/Core/Public/HAL/FeedbackContextAnsi.h b/Engine/Source/Runtime/Core/Public/HAL/FeedbackContextAnsi.h
--- a/Engine/Source/Runtime/Core/Public/HAL/FeedbackContextAnsi.h
+++ b/Engine/Source/Runtime/Core/Public/HAL/FeedbackContextAnsi.h
@@ -1,11 +1,8 @@
 // Copyright 1998-2014 Epic Games, Inc. All Rights Reserved.
 
-/*=============================================================================
-	FeedbackContextAnsi.h: Unreal Ansi user interface interaction.
-=============================================================================*/
-
 #pragma once
 
+
 /*-----------------------------------------------------------------------------
 	FFeedbackContextAnsi.
 -----------------------------------------------------------------------------*/
@@ -36,9 +33,9 @@ public:
 	// Constructor.
 	FFeedbackContextAnsi()
 	: FFeedbackContext()
-	, SlowTaskCount( 0 )
-	, Context( NULL )
-	, AuxOut( NULL )
+	, SlowTaskCount(0)
+	, Context(nullptr)
+	, AuxOut(nullptr)
 	{}
 	void Serialize( const TCHAR* V, ELogVerbosity::Type Verbosity, const class FName& Category )
 	{
@@ -95,7 +92,7 @@ public:
 			else
 			{
 				char InputText[256];
-				if (fgets( InputText, sizeof(InputText), stdin ) != NULL)
+				if (fgets( InputText, sizeof(InputText), stdin ) != nullptr)
 				{
 					return (InputText[0]=='Y' || InputText[0]=='y');
 				}
diff --git a/Engine/Source/Runtime/Core/Public/HAL/FileManager.h b/Engine/Source/Runtime/Core/Public/HAL/FileManager.h
--- a/Engine/Source/Runtime/Core/Public/HAL/FileManager.h
+++ b/Engine/Source/Runtime/Core/Public/HAL/FileManager.h
@@ -1,14 +1,12 @@
 // Copyright 1998-2014 Epic Games, Inc. All Rights Reserved.
 
-/*=============================================================================
-	FileManager.h: IFileManager Interface
-=============================================================================*/
-
 #pragma once
 
+
 // Maximum length of any filename.  For now, we have no restriction. We would probably use shortening rules if we have to.
 #define MAX_UNREAL_FILENAME_LENGTH (PLATFORM_MAX_FILEPATH_LENGTH)
 
+
 enum EFileWrite
 {
 	FILEWRITE_NoFail            = 0x01,
@@ -17,11 +15,15 @@ enum EFileWrite
 	FILEWRITE_Append			= 0x08,
 	FILEWRITE_AllowRead			= 0x10
 };
+
+
 enum EFileRead
 {
 	FILEREAD_NoFail             = 0x01,
 	FILEREAD_Silent				= 0x02,
 };
+
+
 enum ECopyResult
 {
 	COPY_OK						= 0x00,
@@ -29,6 +31,7 @@ enum ECopyResult
 	COPY_Canceled				= 0x02,
 };
 
+
 struct FCopyProgress
 {
 	virtual bool Poll( float Fraction )=0;
@@ -42,6 +45,7 @@ enum EFileOpenFlags
 	IO_APPEND		= 0x40,					// When writing, keep the existing data, set the filepointer to the end of the existing data
 };
 
+
 class CORE_API IFileManager
 {
 protected:
@@ -80,7 +84,7 @@ public:
 	virtual bool Delete( const TCHAR* Filename, bool RequireExists=0, bool EvenReadOnly=0, bool Quiet=0 )=0;
 
 	/** Copies a file. */
-	virtual uint32 Copy( const TCHAR* Dest, const TCHAR* Src, bool Replace=1, bool EvenIfReadOnly=0, bool Attributes=0, FCopyProgress* Progress=NULL )=0; // utility
+	virtual uint32 Copy( const TCHAR* Dest, const TCHAR* Src, bool Replace=1, bool EvenIfReadOnly=0, bool Attributes=0, FCopyProgress* Progress = nullptr )=0; // utility
 
 	/** Moves/renames a file. */
 	virtual bool Move( const TCHAR* Dest, const TCHAR* Src, bool Replace=1, bool EvenIfReadOnly=0, bool Attributes=0, bool bDoNotRetryOrError=0 )=0;
@@ -179,5 +183,3 @@ public:
 	 */
 	virtual bool SendMessageToServer(const TCHAR* Message, IPlatformFile::IFileServerMessageHandler* Handler)=0;
 };
-
-
diff --git a/Engine/Source/Runtime/Core/Public/HAL/FileManagerGeneric.h b/Engine/Source/Runtime/Core/Public/HAL/FileManagerGeneric.h
--- a/Engine/Source/Runtime/Core/Public/HAL/FileManagerGeneric.h
+++ b/Engine/Source/Runtime/Core/Public/HAL/FileManagerGeneric.h
@@ -1,20 +1,17 @@
 // Copyright 1998-2014 Epic Games, Inc. All Rights Reserved.
 
-/*=============================================================================
-	FileManagerGeneric.h: Unreal generic file manager support code.
-
-	This base class simplifies IFileManager implementations by providing
-	simple, unoptimized implementations of functions whose implementations
-	can be derived from other functions.
-=============================================================================*/
-
 #pragma once
 
-/*-----------------------------------------------------------------------------
-	File Manager.
------------------------------------------------------------------------------*/
 
-class CORE_API FFileManagerGeneric : public IFileManager
+/**
+ * Base class for file managers.
+ *
+ * This base class simplifies IFileManager implementations by providing
+ * simple, unoptimized implementations of functions whose implementations
+ * can be derived from other functions.
+ */
+class CORE_API FFileManagerGeneric
+	: public IFileManager
 {
 	// instead of caching the LowLevel, we call the singleton each time to never be incorrect
 	FORCEINLINE IPlatformFile& GetLowLevel() const
@@ -24,13 +21,19 @@ class CORE_API FFileManagerGeneric : public IFileManager
 
 public:
 
-	FFileManagerGeneric()
-	{
-	}
+	/**
+	 * Default constructor.
+	 */
+	FFileManagerGeneric( ) { }
 
-	virtual ~FFileManagerGeneric()
-	{
-	}
+	/**
+	 * Virtual destructor.
+	 */
+	virtual ~FFileManagerGeneric( ) { }
+
+public:
+
+	// IFileManager interface
 
 	virtual void ProcessCommandLineOptions() override;
 
@@ -89,7 +92,6 @@ public:
 	 * Converts passed in filename to use a relative path.
 	 *
 	 * @param	Filename	filename to convert to use a relative path
-	 * 
 	 * @return	filename using relative path
 	 */
 	static FString DefaultConvertToRelativePath( const TCHAR* Filename );
@@ -98,7 +100,6 @@ public:
 	 * Converts passed in filename to use a relative path.
 	 *
 	 * @param	Filename	filename to convert to use a relative path
-	 * 
 	 * @return	filename using relative path
 	 */
 	FString ConvertToRelativePath( const TCHAR* Filename ) override;
@@ -107,7 +108,6 @@ public:
 	 * Converts passed in filename to use an absolute path (for reading)
 	 *
 	 * @param	Filename	filename to convert to use an absolute path, safe to pass in already using absolute path
-	 * 
 	 * @return	filename using absolute path
 	 */
 	FString ConvertToAbsolutePathForExternalAppForRead( const TCHAR* Filename ) override;
@@ -116,7 +116,6 @@ public:
 	 * Converts passed in filename to use an absolute path (for writing)
 	 *
 	 * @param	Filename	filename to convert to use an absolute path, safe to pass in already using absolute path
-	 * 
 	 * @return	filename using absolute path
 	 */
 	FString ConvertToAbsolutePathForExternalAppForWrite( const TCHAR* Filename ) override;
@@ -134,7 +133,6 @@ public:
 	 * immediately if the file manager doesn't support talking to a server.
 	 *
 	 * @param Message	The string message to send to the server
-	 *
 	 * @return			true if the message was sent to server and it returned success, or false if there is no server, or the command failed
 	 */
 	virtual bool SendMessageToServer(const TCHAR* Message, IPlatformFile::IFileServerMessageHandler* Handler) override
@@ -152,6 +150,7 @@ private:
 	void FindFilesRecursiveInternal( TArray<FString>& FileNames, const TCHAR* StartDirectory, const TCHAR* Filename, bool Files, bool Directories);
 };
 
+
 /*-----------------------------------------------------------------------------
 	FArchiveFileReaderGeneric
 -----------------------------------------------------------------------------*/
@@ -250,6 +249,4 @@ protected:
 	int64            BufferCount;
 	TAutoPtr<IFileHandle>		 Handle;
 	uint8            Buffer[4096];
-	
 };
-
diff --git a/Engine/Source/Runtime/Core/Public/HAL/IConsoleManager.h b/Engine/Source/Runtime/Core/Public/HAL/IConsoleManager.h
--- a/Engine/Source/Runtime/Core/Public/HAL/IConsoleManager.h
+++ b/Engine/Source/Runtime/Core/Public/HAL/IConsoleManager.h
@@ -1,11 +1,8 @@
 // Copyright 1998-2014 Epic Games, Inc. All Rights Reserved.
 
-/*=============================================================================
-IConsoleManager.h: console command interface
-=============================================================================*/
-
 #pragma once
 
+
 /**
  * Console variable usage guide:
  *
@@ -359,7 +356,7 @@ struct CORE_API IConsoleManager
 	/**
 	 * Register a console command that takes no arguments
 	 *
-	 * @param	Name		The name of this command (must not be NULL)
+	 * @param	Name		The name of this command (must not be nullptr)
 	 * @param	Help		Help text for this command
 	 * @param	Command		The user function to call when this command is executed
 	 * @param	Flags		Optional flags bitmask
@@ -369,7 +366,7 @@ struct CORE_API IConsoleManager
 	/**
 	 * Register a console command that takes arguments
 	 *
-	 * @param	Name		The name of this command (must not be NULL)
+	 * @param	Name		The name of this command (must not be nullptr)
 	 * @param	Help		Help text for this command
 	 * @param	Command		The user function to call when this command is executed
 	 * @param	Flags		Optional flags bitmask
@@ -379,7 +376,7 @@ struct CORE_API IConsoleManager
 	/**
 	 * Register a console command that takes arguments
 	 *
-	 * @param	Name		The name of this command (must not be NULL)
+	 * @param	Name		The name of this command (must not be nullptr)
 	 * @param	Help		Help text for this command
 	 * @param	Command		The user function to call when this command is executed
 	 * @param	Flags		Optional flags bitmask
@@ -389,7 +386,7 @@ struct CORE_API IConsoleManager
 	/**
 	 * Register a console command that takes arguments
 	 *
-	 * @param	Name		The name of this command (must not be NULL)
+	 * @param	Name		The name of this command (must not be nullptr)
 	 * @param	Help		Help text for this command
 	 * @param	Command		The user function to call when this command is executed
 	 * @param	Flags		Optional flags bitmask
@@ -399,7 +396,7 @@ struct CORE_API IConsoleManager
 	/**
 	 * Register a console command that takes arguments
 	 *
-	 * @param	Name		The name of this command (must not be NULL)
+	 * @param	Name		The name of this command (must not be nullptr)
 	 * @param	Help		Help text for this command
 	 * @param	Command		The user function to call when this command is executed
 	 * @param	Flags		Optional flags bitmask
@@ -409,7 +406,7 @@ struct CORE_API IConsoleManager
 	/**
 	 * Register a console command that is handles by an Exec functions (for auto completion)
 	 *
-	 * @param	Name		The name of this command (must not be NULL)
+	 * @param	Name		The name of this command (must not be nullptr)
 	 * @param	Help		Help text for this command
 	 * @param	Flags		Optional flags bitmask
 	 */
@@ -499,7 +496,7 @@ struct CORE_API IConsoleManager
 		if (!Singleton)
 		{
 			SetupSingleton();
-			check(Singleton != NULL);
+			check(Singleton != nullptr);
 		}
 		return *Singleton;
 	}
@@ -789,7 +786,7 @@ public:
 	/**
 	 * Register a console command that takes no arguments
 	 *
-	 * @param	Name		The name of this command (must not be NULL)
+	 * @param	Name		The name of this command (must not be nullptr)
 	 * @param	Help		Help text for this command
 	 * @param	Command		The user function to call when this command is executed
 	 * @param	Flags		Optional flags bitmask
@@ -802,7 +799,7 @@ public:
 	/**
 	 * Register a console command that takes arguments
 	 *
-	 * @param	Name		The name of this command (must not be NULL)
+	 * @param	Name		The name of this command (must not be nullptr)
 	 * @param	Help		Help text for this command
 	 * @param	Command		The user function to call when this command is executed
 	 * @param	Flags		Optional flags bitmask
@@ -823,7 +820,7 @@ public:
 	/**
 	 * Register a console command that takes a world argument
 	 *
-	 * @param	Name		The name of this command (must not be NULL)
+	 * @param	Name		The name of this command (must not be nullptr)
 	 * @param	Help		Help text for this command
 	 * @param	Command		The user function to call when this command is executed
 	 * @param	Flags		Optional flags bitmask
@@ -845,7 +842,7 @@ public:
 	/**
 	 * Register a console command that takes arguments and a world argument
 	 *
-	 * @param	Name		The name of this command (must not be NULL)
+	 * @param	Name		The name of this command (must not be nullptr)
 	 * @param	Help		Help text for this command
 	 * @param	Command		The user function to call when this command is executed
 	 * @param	Flags		Optional flags bitmask
@@ -865,7 +862,7 @@ public:
 	/**
 	 * Register a console command that takes an output device
 	 *
-	 * @param	Name		The name of this command (must not be NULL)
+	 * @param	Name		The name of this command (must not be nullptr)
 	 * @param	Help		Help text for this command
 	 * @param	Command		The user function to call when this command is executed
 	 * @param	Flags		Optional flags bitmask
diff --git a/Engine/Source/Runtime/Core/Public/HAL/IOBase.h b/Engine/Source/Runtime/Core/Public/HAL/IOBase.h
--- a/Engine/Source/Runtime/Core/Public/HAL/IOBase.h
+++ b/Engine/Source/Runtime/Core/Public/HAL/IOBase.h
@@ -1,11 +1,8 @@
 // Copyright 1998-2014 Epic Games, Inc. All Rights Reserved.
 
-/*=============================================================================
-	IOBase.h: General-purpose file utilities.
-=============================================================================*/
-
 #pragma once
 
+
 /**
  * Enum for async IO priorities.
  */
@@ -19,6 +16,7 @@ enum EAsyncIOPriority
 	AIOP_MAX
 };
 
+
 /*-----------------------------------------------------------------------------
 	FIOSystem.
 -----------------------------------------------------------------------------*/
@@ -48,7 +46,7 @@ struct CORE_API FIOSystem
 	 * @param	Offset		Offset into file
 	 * @param	Size		Size of load request
 	 * @param	Dest		Pointer to load data into
-	 * @param	Counter		Thread safe counter to decrement when loading has finished, can be NULL
+	 * @param	Counter		Thread safe counter to decrement when loading has finished, can be nullptr
 	 * @param	Priority	Priority of request
 	 *
 	 * @return Returns an index to the request that can be used for canceling or 0 if the request failed.
@@ -70,7 +68,7 @@ struct CORE_API FIOSystem
 	 * @param	UncompressedSize	Size of uncompressed data
 	 * @param	Dest				Pointer to load data into
 	 * @param	CompressionFlags	Flags controlling data decompression
-	 * @param	Counter				Thread safe counter to decrement when loading has finished, can be NULL
+	 * @param	Counter				Thread safe counter to decrement when loading has finished, can be nullptr
 	 * @param	Priority			Priority of request
 	 *
 	 * @return Returns an index to the request that can be used for canceling or 0 if the request failed.
diff --git a/Engine/Source/Runtime/Core/Public/HAL/IPlatformFileCachedWrapper.h b/Engine/Source/Runtime/Core/Public/HAL/IPlatformFileCachedWrapper.h
--- a/Engine/Source/Runtime/Core/Public/HAL/IPlatformFileCachedWrapper.h
+++ b/Engine/Source/Runtime/Core/Public/HAL/IPlatformFileCachedWrapper.h
@@ -217,13 +217,13 @@ public:
 	}
 
 	FCachedReadPlatformFile()
-		: LowerLevel(NULL)
+		: LowerLevel(nullptr)
 	{
 	}
 	virtual bool Initialize(IPlatformFile* Inner, const TCHAR* CommandLineParam) override
 	{
 		// Inner is required.
-		check(Inner != NULL);
+		check(Inner != nullptr);
 		LowerLevel = Inner;
 		return !!LowerLevel;
 	}
@@ -281,7 +281,7 @@ public:
 		IFileHandle* InnerHandle=LowerLevel->OpenRead(Filename);
 		if (!InnerHandle)
 		{
-			return NULL;
+			return nullptr;
 		}
 		return new FCachedFileHandle(InnerHandle, true, false);
 	}
@@ -290,7 +290,7 @@ public:
 		IFileHandle* InnerHandle=LowerLevel->OpenWrite(Filename, bAppend, bAllowRead);
 		if (!InnerHandle)
 		{
-			return NULL;
+			return nullptr;
 		}
 		return new FCachedFileHandle(InnerHandle, bAllowRead, true);
 	}
diff --git a/Engine/Source/Runtime/Core/Public/HAL/IPlatformFileLogWrapper.h b/Engine/Source/Runtime/Core/Public/HAL/IPlatformFileLogWrapper.h
--- a/Engine/Source/Runtime/Core/Public/HAL/IPlatformFileLogWrapper.h
+++ b/Engine/Source/Runtime/Core/Public/HAL/IPlatformFileLogWrapper.h
@@ -1,8 +1,8 @@
 // Copyright 1998-2014 Epic Games, Inc. All Rights Reserved.
 
-
 #pragma once
 
+
 /**
  * Wrapper to log the low level file system
 **/
@@ -97,7 +97,7 @@ public:
 	}
 
 	FLoggedPlatformFile()
-		: LowerLevel(NULL)
+		: LowerLevel(nullptr)
 	{
 	}
 
@@ -108,7 +108,7 @@ public:
 	virtual bool Initialize(IPlatformFile* Inner, const TCHAR* CommandLineParam) override
 	{
 		// Inner is required.
-		check(Inner != NULL);
+		check(Inner != nullptr);
 		LowerLevel = Inner;
 		return !!LowerLevel;
 	}
@@ -308,5 +308,3 @@ public:
 		return Result;
 	}
 };
-
-
diff --git a/Engine/Source/Runtime/Core/Public/HAL/IPlatformFileModule.h b/Engine/Source/Runtime/Core/Public/HAL/IPlatformFileModule.h
--- a/Engine/Source/Runtime/Core/Public/HAL/IPlatformFileModule.h
+++ b/Engine/Source/Runtime/Core/Public/HAL/IPlatformFileModule.h
@@ -1,16 +1,13 @@
 // Copyright 1998-2014 Epic Games, Inc. All Rights Reserved.
 
-
-/*=============================================================================================
-	IPlatformFileModule.h: Platform file module interface
-==============================================================================================*/
-
 #pragma once
 
+
 /**
  * Platform File Module Interface
  */
-class IPlatformFileModule : public IModuleInterface
+class IPlatformFileModule
+	: public IModuleInterface
 {
 public:
 
@@ -20,4 +17,4 @@ public:
 	 * @return Platform file instance.
 	 */
 	virtual IPlatformFile* GetPlatformFile() = 0;
-};
\ No newline at end of file
+};
diff --git a/Engine/Source/Runtime/Core/Public/HAL/IPlatformFileProfilerWrapper.h b/Engine/Source/Runtime/Core/Public/HAL/IPlatformFileProfilerWrapper.h
--- a/Engine/Source/Runtime/Core/Public/HAL/IPlatformFileProfilerWrapper.h
+++ b/Engine/Source/Runtime/Core/Public/HAL/IPlatformFileProfilerWrapper.h
@@ -1,8 +1,8 @@
 // Copyright 1998-2014 Epic Games, Inc. All Rights Reserved.
 
-
 #pragma once
 
+
 #if !UE_BUILD_SHIPPING
 
 /**
@@ -202,7 +202,7 @@ protected:
 	FCriticalSection SynchronizationObject;
 
 	FProfiledPlatformFile()
-		: LowerLevel(NULL)
+		: LowerLevel(nullptr)
 		, StartTime(0.0)
 	{
 	}
@@ -221,7 +221,7 @@ public:
 	virtual bool Initialize(IPlatformFile* Inner, const TCHAR* CommandLineParam) override
 	{
 		// Inner is required.
-		check(Inner != NULL);
+		check(Inner != nullptr);
 		LowerLevel = Inner;
 		StartTime = FPlatformTime::Seconds() * 1000.0;
 		return !!LowerLevel;
@@ -252,7 +252,7 @@ class TProfiledPlatformFile : public FProfiledPlatformFile
 		FScopeLock ScopeLock(&SynchronizationObject);
 
 		TSharedPtr< FProfiledFileStatsFileBase >* ExistingStat = Stats.Find( Path );
-		if( ExistingStat != NULL )
+		if( ExistingStat != nullptr )
 		{
 			return (StatsType*)(ExistingStat->Get());
 		}
@@ -273,7 +273,7 @@ public:
 	}
 	static const TCHAR* GetTypeName()
 	{
-		return NULL;
+		return nullptr;
 	}
 	virtual const TCHAR* GetName() const override
 	{
@@ -501,7 +501,7 @@ protected:
 public:
 
 	FPlatformFileReadStats()
-		: LowerLevel(NULL)
+		: LowerLevel(nullptr)
 		, Timer(0.f)
 	{
 	}
@@ -612,5 +612,5 @@ public:
 	}
 };
 
-#endif // !UE_BUILD_SHIPPING
 
+#endif // !UE_BUILD_SHIPPING
diff --git a/Engine/Source/Runtime/Core/Public/HAL/MallocAnsi.h b/Engine/Source/Runtime/Core/Public/HAL/MallocAnsi.h
--- a/Engine/Source/Runtime/Core/Public/HAL/MallocAnsi.h
+++ b/Engine/Source/Runtime/Core/Public/HAL/MallocAnsi.h
@@ -1,26 +1,25 @@
 // Copyright 1998-2014 Epic Games, Inc. All Rights Reserved.
 
-/*=============================================================================
-	MallocAnsi.h: ANSI memory allocator.
-=============================================================================*/
-
 #pragma once
 
+
 #if _MSC_VER || PLATFORM_MAC
-#define USE_ALIGNED_MALLOC 1
+	#define USE_ALIGNED_MALLOC 1
 #else
-//@todo gcc: this should be implemented more elegantly on other platforms
-#define USE_ALIGNED_MALLOC 0
+	//@todo gcc: this should be implemented more elegantly on other platforms
+	#define USE_ALIGNED_MALLOC 0
 #endif
 
 #if PLATFORM_IOS
-#include "mach/mach.h"
+	#include "mach/mach.h"
 #endif
 
+
 //
 // ANSI C memory allocator.
 //
-class FMallocAnsi : public FMalloc
+class FMallocAnsi
+	: public FMalloc
 {
 	
 public:
@@ -52,7 +51,7 @@ public:
 		*((SIZE_T*)( (uint8*)Result - sizeof(void*) - sizeof(SIZE_T)	))	= Size;
 #endif
 
-		if (Result == NULL)
+		if (Result == nullptr)
 		{
 			FPlatformMemory::OnOutOfMemory(Size, Alignment);
 		}
@@ -69,14 +68,14 @@ public:
 		{
 			Result = _aligned_realloc( Ptr, NewSize, Alignment );
 		}
-		else if( Ptr == NULL )
+		else if( Ptr == nullptr )
 		{
 			Result = _aligned_malloc( NewSize, Alignment );
 		}
 		else
 		{
 			_aligned_free( Ptr );
-			Result = NULL;
+			Result = nullptr;
 		}
 #else
 		if( Ptr && NewSize )
@@ -86,17 +85,17 @@ public:
 			FMemory::Memcpy( Result, Ptr, FMath::Min(NewSize, *((SIZE_T*)( (uint8*)Ptr - sizeof(void*) - sizeof(SIZE_T))) ) );
 			Free( Ptr );
 		}
-		else if( Ptr == NULL )
+		else if( Ptr == nullptr )
 		{
 			Result = Malloc( NewSize, Alignment);
 		}
 		else
 		{
 			free( *((void**)((uint8*)Ptr-sizeof(void*))) );
-			Result = NULL;
+			Result = nullptr;
 		}
 #endif
-		if (Result == NULL && NewSize != 0)
+		if (Result == nullptr && NewSize != 0)
 		{
 			FPlatformMemory::OnOutOfMemory(NewSize, Alignment);
 		}
@@ -188,4 +187,3 @@ public:
 
 	virtual const TCHAR * GetDescriptiveName() override { return TEXT("ANSI"); }
 };
-
diff --git a/Engine/Source/Runtime/Core/Public/HAL/MallocBinned.h b/Engine/Source/Runtime/Core/Public/HAL/MallocBinned.h
--- a/Engine/Source/Runtime/Core/Public/HAL/MallocBinned.h
+++ b/Engine/Source/Runtime/Core/Public/HAL/MallocBinned.h
@@ -1,11 +1,8 @@
 // Copyright 1998-2014 Epic Games, Inc. All Rights Reserved.
 
-/*=============================================================================
-	MallocBinned.h: Binned memory allocator, refactoring the Windows allocator
-=============================================================================*/
-
 #pragma once
 
+
 #define MEM_TIME(st)
 
 //#define USE_LOCKFREE_DELETE
@@ -44,6 +41,7 @@ DECLARE_MEMORY_STAT_EXTERN(TEXT("Binned Current Allocs"),STAT_Binned_CurrentAllo
 DECLARE_MEMORY_STAT_EXTERN(TEXT("Binned Total Allocs"),	STAT_Binned_TotalAllocs,STATGROUP_MemoryAllocator, CORE_API);
 DECLARE_MEMORY_STAT_EXTERN(TEXT("Binned Slack Current"),	STAT_Binned_SlackCurrent,STATGROUP_MemoryAllocator, CORE_API);
 
+
 //
 // Optimized virtual memory allocator.
 //
@@ -156,7 +154,7 @@ private:
 
 		FFreePageBlock() 
 		{
-			Ptr = NULL;
+			Ptr = nullptr;
 			ByteSize = 0;
 		}
 	};
@@ -197,8 +195,8 @@ private:
 		uint64				TotalWaste;
 #endif
 		FPoolTable()
-			: FirstPool(NULL)
-			, ExhaustedPool(NULL)
+			: FirstPool(nullptr)
+			, ExhaustedPool(nullptr)
 			, BlockSize(0)
 #if STATS
 			, NumActivePools(0)
@@ -226,7 +224,7 @@ private:
 		PoolHashBucket()
 		{
 			Key=0;
-			FirstPool=NULL;
+			FirstPool=nullptr;
 			Prev=this;
 			Next=this;
 		}
@@ -413,7 +411,7 @@ private:
 			Ptr = ((Ptr-(PageSize*NextStep))-1)&~((UPTRINT)PageSize-1);
 		}
 		AllocationBase=0;
-		return NULL;
+		return nullptr;
 	}
 
 	FORCEINLINE FPoolInfo* FindPoolInfoInternal(UPTRINT Ptr, uint16& JumpOffset)
@@ -433,14 +431,14 @@ private:
 				if (!collision->FirstPool[PoolIndex].AllocSize)
 				{
 					JumpOffset = collision->FirstPool[PoolIndex].TableIndex;
-					return NULL;
+					return nullptr;
 				}
 				return &collision->FirstPool[PoolIndex];
 			}
 			collision=collision->Next;
 		} while (collision!=&HashBuckets[Hash]);
 
-		return NULL;
+		return nullptr;
 	}
 
 	/**
@@ -483,9 +481,9 @@ private:
 		PoolHashBucket* NextFree=HashBucketFreeList->Next;
 		PoolHashBucket* Free=HashBucketFreeList;
 		Free->Unlink();
-		if (NextFree==Free) 
+		if (NextFree == Free) 
 		{
-			NextFree=NULL;
+			NextFree = nullptr;
 		}
 		HashBucketFreeList=NextFree;
 		return Free;
@@ -501,7 +499,7 @@ private:
 		checkSlow(Blocks * Table->BlockSize <= Bytes && PoolSize >= Bytes);
 
 		// Allocate memory.
-		FFreeMem* Free = NULL;
+		FFreeMem* Free = nullptr;
 		SIZE_T ActualPoolSize; //TODO: use this to reduce waste?
 		Free = (FFreeMem*)OSAlloc(OsBytes, ActualPoolSize);
 
@@ -541,7 +539,7 @@ private:
 #endif
 		// Create first free item.
 		Free->NumFreeBlocks = Blocks;
-		Free->Next       = NULL;
+		Free->Next       = nullptr;
 
 		return Pool;
 	}
@@ -769,12 +767,12 @@ public:
 	FMallocBinned(uint32 InPageSize, uint64 AddressLimit)
 		:	TableAddressLimit(AddressLimit)
 #ifdef USE_LOCKFREE_DELETE
-		,	PendingFreeList(NULL)
+		,	PendingFreeList(nullptr)
 		,	bFlushingFrees(false)
 		,	bDoneFreeListInit(false)
 #endif
-		,	HashBuckets(NULL)
-		,	HashBucketFreeList(NULL)
+		,	HashBuckets(nullptr)
+		,	HashBucketFreeList(nullptr)
 		,	PageSize		(InPageSize)
 #ifdef CACHE_FREED_OS_ALLOCS
 		,	FreedPageBlocksNum(0)
@@ -820,17 +818,17 @@ public:
 
 
 		// Init tables.
-		OsTable.FirstPool = NULL;
-		OsTable.ExhaustedPool = NULL;
+		OsTable.FirstPool = nullptr;
+		OsTable.ExhaustedPool = nullptr;
 		OsTable.BlockSize = 0;
 
 		/** The following options are not valid for page sizes less than 64k. They are here to reduce waste*/
-		PagePoolTable[0].FirstPool = NULL;
-		PagePoolTable[0].ExhaustedPool = NULL;
+		PagePoolTable[0].FirstPool = nullptr;
+		PagePoolTable[0].ExhaustedPool = nullptr;
 		PagePoolTable[0].BlockSize = PageSize == PAGE_SIZE_LIMIT ? BinnedSizeLimit+(BinnedSizeLimit/2) : 0;
 
-		PagePoolTable[1].FirstPool = NULL;
-		PagePoolTable[1].ExhaustedPool = NULL;
+		PagePoolTable[1].FirstPool = nullptr;
+		PagePoolTable[1].ExhaustedPool = nullptr;
 		PagePoolTable[1].BlockSize = PageSize == PAGE_SIZE_LIMIT ? PageSize+BinnedSizeLimit : 0;
 
 		// Block sizes are based around getting the maximum amount of allocations per pool, with as little alignment waste as possible.
@@ -847,8 +845,8 @@ public:
 
 		for( uint32 i = 0; i < POOL_COUNT; i++ )
 		{
-			PoolTable[i].FirstPool = NULL;
-			PoolTable[i].ExhaustedPool = NULL;
+			PoolTable[i].FirstPool = nullptr;
+			PoolTable[i].ExhaustedPool = nullptr;
 			PoolTable[i].BlockSize = BlockSizes[i];
 #if STATS
 			PoolTable[i].MinRequest = PoolTable[i].BlockSize;
@@ -1039,14 +1037,14 @@ public:
 				}
 			}
 		}
-		else if( Ptr == NULL )
+		else if( Ptr == nullptr )
 		{
 			NewPtr = Malloc( NewSize, Alignment );
 		}
 		else
 		{
 			Free( Ptr );
-			NewPtr = NULL;
+			NewPtr = nullptr;
 		}
 
 		MEM_TIME(MemTime += FPlatformTime::Seconds());
@@ -1211,7 +1209,7 @@ public:
 			uint32 TotalPools = 0;
 			uint32 TotalSlack = 0;
 
-			FPoolTable* Table = NULL;
+			FPoolTable* Table = nullptr;
 			for( int32 i = 0; i < BinnedSizeLimit + EXTENED_PAGE_POOL_ALLOCATION_COUNT; i++ )
 			{
 				if( Table == MemSizeToPoolTable[i] || MemSizeToPoolTable[i]->BlockSize == 0 )
@@ -1295,6 +1293,3 @@ protected:
 	 */
 	uint32 BinnedGetPageSize();
 };
-
-
-
diff --git a/Engine/Source/Runtime/Core/Public/HAL/MallocDebug.h b/Engine/Source/Runtime/Core/Public/HAL/MallocDebug.h
--- a/Engine/Source/Runtime/Core/Public/HAL/MallocDebug.h
+++ b/Engine/Source/Runtime/Core/Public/HAL/MallocDebug.h
@@ -1,11 +1,8 @@
 // Copyright 1998-2014 Epic Games, Inc. All Rights Reserved.
 
-/*=============================================================================
-	MallocDebug.h: Debug memory allocator.
-=============================================================================*/
-
 #pragma once
 
+
 // Debug memory allocator.
 class FMallocDebug : public FMalloc
 {
@@ -47,7 +44,7 @@ private:
 public:
 	// FMalloc interface.
 	FMallocDebug()
-	:	GFirstDebug( NULL )
+	:	GFirstDebug( nullptr )
 	,	TotalAllocationSize( 0 )
 	,	TotalWasteSize( 0 )
 	{}
@@ -58,7 +55,7 @@ public:
 	virtual void* Malloc( SIZE_T Size, uint32 Alignment ) override
 	{
 		check(Alignment == DEFAULT_ALIGNMENT && "Alignment currently unsupported in this allocator");
-		FMemDebug* Ptr = NULL;
+		FMemDebug* Ptr = nullptr;
 		Ptr = (FMemDebug*)malloc( AllocatorOverhead + Size );
 		check(Ptr);
 		uint8* AlignedPtr = Align( (uint8*) Ptr + sizeof(FMemDebug) + sizeof(FMemDebug*) + sizeof(int32), ALLOCATION_ALIGNMENT );
@@ -98,14 +95,14 @@ public:
 			Free( InPtr );
 			return Result;
 		}
-		else if( InPtr == NULL )
+		else if( InPtr == nullptr )
 		{
 			return Malloc( NewSize, Alignment );
 		}
 		else
 		{
 			Free( InPtr );
-			return NULL;
+			return nullptr;
 		}
 	}
 
@@ -218,4 +215,3 @@ public:
 	virtual const TCHAR * GetDescriptiveName() override { return TEXT("debug"); }
 
 };
-
diff --git a/Engine/Source/Runtime/Core/Public/HAL/MallocJemalloc.h b/Engine/Source/Runtime/Core/Public/HAL/MallocJemalloc.h
--- a/Engine/Source/Runtime/Core/Public/HAL/MallocJemalloc.h
+++ b/Engine/Source/Runtime/Core/Public/HAL/MallocJemalloc.h
@@ -2,37 +2,32 @@
 
 #pragma once
 
+
 // Only use for supported platforms
 #if PLATFORM_SUPPORTS_JEMALLOC
 
 #define MEM_TIME(st)
 
+
 //
 // Jason Evans' malloc (default in FreeBSD, NetBSD, used in Firefox, Facebook servers) 
 // http://www.canonware.com/jemalloc/
 //
-class FMallocJemalloc : public FMalloc
+class FMallocJemalloc
+	: public FMalloc
 {
-
-private:
-
-	/** Time spent on memory tasks, in seconds */
-	double		MemTime;
-
-	// Implementation. 
-	void OutOfMemory()
-	{
-		UE_LOG(LogHAL, Fatal, TEXT("%s"), TEXT("Ran out of virtual memory. To prevent this condition, you must free up more space on your primary hard disk.") );
-	}
-
 public:
 
-	FMallocJemalloc() :
-		MemTime		( 0.0 )
-	{
-	}
+	/**
+	 * Default constructor.
+	 */
+	FMallocJemalloc()
+		: MemTime(0.0)
+	{ }
 	
-	// Begin FMalloc interface.
+public:
+
+	// FMalloc interface.
 
 	virtual void* Malloc( SIZE_T Size, uint32 Alignment ) override;
 	virtual void* Realloc( void* Ptr, SIZE_T NewSize, uint32 Alignment ) override;
@@ -42,8 +37,18 @@ public:
 	virtual bool IsInternallyThreadSafe() const override {  return true; }
 	virtual const TCHAR * GetDescriptiveName() override { return TEXT("jemalloc"); }
 
-	// End FMalloc interface.
+protected:
 
+	void OutOfMemory( )
+	{
+		UE_LOG(LogHAL, Fatal, TEXT("%s"), TEXT("Ran out of virtual memory. To prevent this condition, you must free up more space on your primary hard disk.") );
+	}
+
+private:
+
+	/** Time spent on memory tasks, in seconds */
+	double MemTime;
 };
 
+
 #endif // PLATFORM_SUPPORTS_JEMALLOC
diff --git a/Engine/Source/Runtime/Core/Public/HAL/MallocTBB.h b/Engine/Source/Runtime/Core/Public/HAL/MallocTBB.h
--- a/Engine/Source/Runtime/Core/Public/HAL/MallocTBB.h
+++ b/Engine/Source/Runtime/Core/Public/HAL/MallocTBB.h
@@ -1,66 +1,51 @@
 // Copyright 1998-2014 Epic Games, Inc. All Rights Reserved.
 
-/*=============================================================================
-	MallocTBB.h: Intel-TBB 64-bit scalable memory allocator.
-=============================================================================*/
-
 #pragma once
 
-//
-// TBB 64-bit scalable memory allocator.
-//
-class FMallocTBB : public FMalloc
-{
-private:
-	double		MemTime;
 
-	// Implementation. 
-	void OutOfMemory(uint64 Size, uint32 Alignment)
-	{
-		// this is expected not to return
-		FPlatformMemory::OnOutOfMemory(Size, Alignment);
-	}
+/**
+ * TBB 64-bit scalable memory allocator.
+ */
+class FMallocTBB
+	: public FMalloc
+{
+public:
+	
+	/**
+	 * Default constructor.
+	 */
+	FMallocTBB() :
+		MemTime(0.0)
+	{ }
 
 public:
+
 	// FMalloc interface.
-	FMallocTBB() :
-		MemTime		( 0.0 )
-	{
-	}
 	
-	/** 
-	 * Malloc
-	 */
 	virtual void* Malloc( SIZE_T Size, uint32 Alignment ) override;
-	
-	/** 
-	 * Realloc
-	 */
 	virtual void* Realloc( void* Ptr, SIZE_T NewSize, uint32 Alignment ) override;
-	
-	/** 
-	 * Free
-	 */
 	virtual void Free( void* Ptr ) override;
+	virtual bool GetAllocationSize( void *Original, SIZE_T &SizeOut ) override;
 
-	/**
-	 * If possible determine the size of the memory allocated at the given address
-	 *
- 	 * @param Original - Pointer to memory we are checking the size of
-	 * @param SizeOut - If possible, this value is set to the size of the passed in pointer
-	 * @return true if succeeded
-	 */
-	virtual bool GetAllocationSize(void *Original, SIZE_T &SizeOut) override;
-
-	/**
-	 * Returns if the allocator is guaranteed to be thread-safe and therefore
-	 * doesn't need a unnecessary thread-safety wrapper around it.
-	 */
-	virtual bool IsInternallyThreadSafe() const override
+	virtual bool IsInternallyThreadSafe( ) const override
 	{ 
 		return true;
 	}
 
-	virtual const TCHAR * GetDescriptiveName() override { return TEXT("TBB"); }
-};
+	virtual const TCHAR * GetDescriptiveName( ) override
+	{
+		return TEXT("TBB");
+	}
 
+protected:
+
+	void OutOfMemory( uint64 Size, uint32 Alignment )
+	{
+		// this is expected not to return
+		FPlatformMemory::OnOutOfMemory(Size, Alignment);
+	}
+
+private:
+
+	double MemTime;
+};
diff --git a/Engine/Source/Runtime/Core/Public/HAL/MallocThreadSafeProxy.h b/Engine/Source/Runtime/Core/Public/HAL/MallocThreadSafeProxy.h
--- a/Engine/Source/Runtime/Core/Public/HAL/MallocThreadSafeProxy.h
+++ b/Engine/Source/Runtime/Core/Public/HAL/MallocThreadSafeProxy.h
@@ -1,12 +1,8 @@
 // Copyright 1998-2014 Epic Games, Inc. All Rights Reserved.
 
-/*=============================================================================
-	FMallocThreadSafeProxy.h: FMalloc proxy used to render any FMalloc thread
-							  safe.
-=============================================================================*/
-
 #pragma once
 
+
 /**
  * FMalloc proxy that synchronizes access, making the used malloc thread safe.
  */
@@ -133,4 +129,3 @@ public:
 		return UsedMalloc->GetDescriptiveName(); 
 	}
 };
-
diff --git a/Engine/Source/Runtime/Core/Public/HAL/MemoryBase.h b/Engine/Source/Runtime/Core/Public/HAL/MemoryBase.h
--- a/Engine/Source/Runtime/Core/Public/HAL/MemoryBase.h
+++ b/Engine/Source/Runtime/Core/Public/HAL/MemoryBase.h
@@ -1,11 +1,8 @@
 // Copyright 1998-2014 Epic Games, Inc. All Rights Reserved.
 
 #pragma once
-#include "UnrealMemory.h"
 
-/*=============================================================================
-	MemoryBase.h: Base memory management definitions.
-=============================================================================*/
+#include "UnrealMemory.h"
 
 #ifndef __MEMORYBASE_H__
 #define __MEMORYBASE_H__
@@ -20,6 +17,7 @@ MALLOC_PROFILER( CORE_API extern class FMallocProfiler* GMallocProfiler; )
 /** Holds generic memory stats, internally implemented as a map. */
 struct FGenericMemoryStats;
 
+
 /**
  * Inherit from FUseSystemMallocForNew if you want your objects to be placed in memory
  * alloced by the system malloc routines, bypassing GMalloc. This is e.g. used by FMalloc
@@ -32,7 +30,7 @@ public:
 	 * Overloaded new operator using the system allocator.
 	 *
 	 * @param	Size	Amount of memory to allocate (in bytes)
-	 * @return			A pointer to a block of memory with size Size or NULL
+	 * @return			A pointer to a block of memory with size Size or nullptr
 	 */
 	void* operator new( size_t Size )
 	{
@@ -53,7 +51,7 @@ public:
 	 * Overloaded array new operator using the system allocator.
 	 *
 	 * @param	Size	Amount of memory to allocate (in bytes)
-	 * @return			A pointer to a block of memory with size Size or NULL
+	 * @return			A pointer to a block of memory with size Size or nullptr
 	 */
 	void* operator new[]( size_t Size )
 	{
@@ -196,5 +194,5 @@ protected:
 	static uint32 TotalReallocCalls;
 };
 
-#endif
 
+#endif
diff --git a/Engine/Source/Runtime/Core/Public/HAL/OutputDevices.h b/Engine/Source/Runtime/Core/Public/HAL/OutputDevices.h
--- a/Engine/Source/Runtime/Core/Public/HAL/OutputDevices.h
+++ b/Engine/Source/Runtime/Core/Public/HAL/OutputDevices.h
@@ -1,18 +1,14 @@
 // Copyright 1998-2014 Epic Games, Inc. All Rights Reserved.
 
-/*=============================================================================
-	OutputDevices.h: Collection of FOutputDevice subclasses
-=============================================================================*/
-
 #pragma once
-#if !PLATFORM_DESKTOP
-
-// don't support colorized text on consoles
-#define SET_WARN_COLOR(Color)
-#define SET_WARN_COLOR_AND_BACKGROUND(Color, Bkgrnd)
-#define CLEAR_WARN_COLOR() 
 
+#if !PLATFORM_DESKTOP
+	// don't support colorized text on consoles
+	#define SET_WARN_COLOR(Color)
+	#define SET_WARN_COLOR_AND_BACKGROUND(Color, Bkgrnd)
+	#define CLEAR_WARN_COLOR() 
 #else
+
 /*-----------------------------------------------------------------------------
 	Colorized text.
 
@@ -36,33 +32,35 @@
 	
 	An empty string reverts to the normal gray on black.
 -----------------------------------------------------------------------------*/
+
 // putting them in a namespace to protect against future name conflicts
 namespace OutputDeviceColor
 {
-const TCHAR* const COLOR_BLACK			= TEXT("0000");
-
-const TCHAR* const COLOR_DARK_RED		= TEXT("1000");
-const TCHAR* const COLOR_DARK_GREEN		= TEXT("0100");
-const TCHAR* const COLOR_DARK_BLUE		= TEXT("0010");
-const TCHAR* const COLOR_DARK_YELLOW	= TEXT("1100");
-const TCHAR* const COLOR_DARK_CYAN		= TEXT("0110");
-const TCHAR* const COLOR_DARK_PURPLE	= TEXT("1010");
-const TCHAR* const COLOR_DARK_WHITE		= TEXT("1110");
-const TCHAR* const COLOR_GRAY			= COLOR_DARK_WHITE;
-
-const TCHAR* const COLOR_RED			= TEXT("1001");
-const TCHAR* const COLOR_GREEN			= TEXT("0101");
-const TCHAR* const COLOR_BLUE			= TEXT("0011");
-const TCHAR* const COLOR_YELLOW			= TEXT("1101");
-const TCHAR* const COLOR_CYAN			= TEXT("0111");
-const TCHAR* const COLOR_PURPLE			= TEXT("1011");
-const TCHAR* const COLOR_WHITE			= TEXT("1111");
-
-const TCHAR* const COLOR_NONE			= TEXT("");
-
+	const TCHAR* const COLOR_BLACK			= TEXT("0000");
+
+	const TCHAR* const COLOR_DARK_RED		= TEXT("1000");
+	const TCHAR* const COLOR_DARK_GREEN		= TEXT("0100");
+	const TCHAR* const COLOR_DARK_BLUE		= TEXT("0010");
+	const TCHAR* const COLOR_DARK_YELLOW	= TEXT("1100");
+	const TCHAR* const COLOR_DARK_CYAN		= TEXT("0110");
+	const TCHAR* const COLOR_DARK_PURPLE	= TEXT("1010");
+	const TCHAR* const COLOR_DARK_WHITE		= TEXT("1110");
+	const TCHAR* const COLOR_GRAY			= COLOR_DARK_WHITE;
+
+	const TCHAR* const COLOR_RED			= TEXT("1001");
+	const TCHAR* const COLOR_GREEN			= TEXT("0101");
+	const TCHAR* const COLOR_BLUE			= TEXT("0011");
+	const TCHAR* const COLOR_YELLOW			= TEXT("1101");
+	const TCHAR* const COLOR_CYAN			= TEXT("0111");
+	const TCHAR* const COLOR_PURPLE			= TEXT("1011");
+	const TCHAR* const COLOR_WHITE			= TEXT("1111");
+
+	const TCHAR* const COLOR_NONE			= TEXT("");
 }
+
 using namespace OutputDeviceColor;
 
+
 // let a console or (UE_BUILD_SHIPPING || UE_BUILD_TEST) define it to nothing
 #ifndef SET_WARN_COLOR
 
@@ -383,10 +381,10 @@ public:
 	/** 
 	 * Constructor, initializing member variables.
 	 *
-	 * @param InFilename	Filename to use, can be NULL
+	 * @param InFilename	Filename to use, can be nullptr
 	 * @param bDisableBackup If true, existing files will not be backed up
 	 */
-	FOutputDeviceFile( const TCHAR* InFilename = NULL, bool bDisableBackup=false );
+	FOutputDeviceFile( const TCHAR* InFilename = nullptr, bool bDisableBackup=false );
 
 	/** Sets the filename that the output device writes to.  If the output device was already writing to a file, closes that file. */
 	void SetFilename(const TCHAR* InFilename);
@@ -526,4 +524,3 @@ CORE_API DECLARE_LOG_CATEGORY_EXTERN(LogMemory, Log, All);
 
 // Temporary log category, generally you should not check things in that use this
 CORE_API DECLARE_LOG_CATEGORY_EXTERN(LogTemp, Log, All);
-
diff --git a/Engine/Source/Runtime/Core/Public/HAL/PThreadEvent.h b/Engine/Source/Runtime/Core/Public/HAL/PThreadEvent.h
--- a/Engine/Source/Runtime/Core/Public/HAL/PThreadEvent.h
+++ b/Engine/Source/Runtime/Core/Public/HAL/PThreadEvent.h
@@ -2,10 +2,12 @@
 
 #pragma once
 
+
 /**
  * This is the PThreads version of FEvent.
  */
-class FPThreadEvent : public FEvent
+class FPThreadEvent
+	: public FEvent
 {
 	// This is a little complicated, in an attempt to match Win32 Event semantics...
     typedef enum
@@ -34,7 +36,7 @@ class FPThreadEvent : public FEvent
 		check(rc == 0);
 	}
 	
-	static inline void SubtractTimevals(const struct timeval *FromThis, struct timeval *SubThis, struct timeval *Difference)
+	static inline void SubtractTimevals( const struct timeval *FromThis, struct timeval *SubThis, struct timeval *Difference )
 	{
 		if (FromThis->tv_usec < SubThis->tv_usec)
 		{
@@ -91,16 +93,16 @@ public:
 		}
 	}
 
-	virtual bool Create(bool _bIsManualReset = false) override
+	virtual bool Create( bool _bIsManualReset = false ) override
 	{
 		check(!bInitialized);
 		bool RetVal = false;
 		Triggered = TRIGGERED_NONE;
 		bIsManualReset = _bIsManualReset;
 
-		if (pthread_mutex_init(&Mutex, NULL) == 0)
+		if (pthread_mutex_init(&Mutex, nullptr) == 0)
 		{
-			if (pthread_cond_init(&Condition, NULL) == 0)
+			if (pthread_cond_init(&Condition, nullptr) == 0)
 			{
 				bInitialized = true;
 				RetVal = true;
@@ -146,7 +148,5 @@ public:
 		UnlockEventMutex();
 	}
 
-	virtual bool Wait(uint32 WaitTime = (uint32)-1, const bool bIgnoreThreadIdleStats = false) override;
-
+	virtual bool Wait( uint32 WaitTime = (uint32)-1, const bool bIgnoreThreadIdleStats = false ) override;
 };
-
diff --git a/Engine/Source/Runtime/Core/Public/HAL/Platform.h b/Engine/Source/Runtime/Core/Public/HAL/Platform.h
--- a/Engine/Source/Runtime/Core/Public/HAL/Platform.h
+++ b/Engine/Source/Runtime/Core/Public/HAL/Platform.h
@@ -1,10 +1,8 @@
 ﻿// Copyright 1998-2014 Epic Games, Inc. All Rights Reserved.
 
-/*================================================================================
-	Platform.h: Setup the includes for the platform
-==================================================================================*/
 #pragma once
 
+
 // define all other platforms to be zero
 //@port Define the platform here to be zero when compiling for other platforms
 #if !defined(PLATFORM_WINDOWS)
diff --git a/Engine/Source/Runtime/Core/Public/HAL/PlatformFilemanager.h b/Engine/Source/Runtime/Core/Public/HAL/PlatformFilemanager.h
--- a/Engine/Source/Runtime/Core/Public/HAL/PlatformFilemanager.h
+++ b/Engine/Source/Runtime/Core/Public/HAL/PlatformFilemanager.h
@@ -1,11 +1,8 @@
 // Copyright 1998-2014 Epic Games, Inc. All Rights Reserved.
 
-/*=============================================================================================
-	PlatformFileManager.h: Generic platform file interfaces
-==============================================================================================*/
-
 #pragma once
 
+
 /**
 * Platform File chain manager.
 **/
@@ -17,40 +14,40 @@ class CORE_API FPlatformFileManager
 public:
 
 	/** Constructor. */
-	FPlatformFileManager();
+	FPlatformFileManager( );
 
 	/**
 	 * Gets the currently used platform file.
 	 *
 	 * @return Reference to the currently used platform file.
 	 */
-	IPlatformFile& GetPlatformFile();
+	IPlatformFile& GetPlatformFile( );
 
 	/**
 	 * Sets the current platform file.
 	 *
 	 * @param NewTopmostPlatformFile Platform file to be used.
 	 */
-	void SetPlatformFile(IPlatformFile& NewTopmostPlatformFile);
+	void SetPlatformFile( IPlatformFile& NewTopmostPlatformFile );
 
 	/**
 	 * Finds a platform file in the chain of active platform files.
 	 *
 	 * @param Name of the platform file.
-	 * @return Pointer to the active platform file or NULL if the platform file was not found.
+	 * @return Pointer to the active platform file or nullptr if the platform file was not found.
 	 */
-	IPlatformFile* FindPlatformFile(const TCHAR* Name);
+	IPlatformFile* FindPlatformFile( const TCHAR* Name );
 
 	/**
 	 * Creates a new platform file instance.
 	 *
 	 * @param Name of the platform file to create.
-	 * @return Platform file instance of the platform file type was found, NULL otherwise.
+	 * @return Platform file instance of the platform file type was found, nullptr otherwise.
 	 */
-	IPlatformFile* GetPlatformFile(const TCHAR* Name);
+	IPlatformFile* GetPlatformFile( const TCHAR* Name );
 
 	/**
 	 * Gets FPlatformFileManager Singleton.
 	 */
-	static FPlatformFileManager& Get();
-};
\ No newline at end of file
+	static FPlatformFileManager& Get( );
+};
diff --git a/Engine/Source/Runtime/Core/Public/HAL/PlatformIncludes.h b/Engine/Source/Runtime/Core/Public/HAL/PlatformIncludes.h
--- a/Engine/Source/Runtime/Core/Public/HAL/PlatformIncludes.h
+++ b/Engine/Source/Runtime/Core/Public/HAL/PlatformIncludes.h
@@ -1,8 +1,5 @@
 // Copyright 1998-2014 Epic Games, Inc. All Rights Reserved.
 
-/*================================================================================
-	PlatformIncludes.h: Setup the main includes for the platform
-==================================================================================*/
 #pragma once
 
 #include <new>
@@ -59,4 +56,3 @@
 #if !defined(TEXT) && !UE_BUILD_DOCS
 		#define TEXT(s) L##s
 #endif
-
