Addec AreCapsLocked() to the modifier state keys for the generic application api.

diff --git a/Engine/Source/Runtime/Core/Private/Linux/LinuxApplication.cpp b/Engine/Source/Runtime/Core/Private/Linux/LinuxApplication.cpp
--- a/Engine/Source/Runtime/Core/Private/Linux/LinuxApplication.cpp
+++ b/Engine/Source/Runtime/Core/Private/Linux/LinuxApplication.cpp
@@ -862,8 +862,9 @@ FModifierKeysState FLinuxApplication::GetModifierKeys() const
 	const bool bIsRightControlDown	= (modifiers & KMOD_RCTRL) != 0;
 	const bool bIsLeftAltDown		= (modifiers & KMOD_LALT) != 0;
 	const bool bIsRightAltDown		= (modifiers & KMOD_RALT) != 0;
+	const bool bAreCapsLocked		= (modifiers & KMOD_CAPS) != 0;
 
-	return FModifierKeysState( bIsLeftShiftDown, bIsRightShiftDown, bIsLeftControlDown, bIsRightControlDown, bIsLeftAltDown, bIsRightAltDown, false, false );
+	return FModifierKeysState( bIsLeftShiftDown, bIsRightShiftDown, bIsLeftControlDown, bIsRightControlDown, bIsLeftAltDown, bIsRightAltDown, false, false, bAreCapsLocked );
 }
 
 
diff --git a/Engine/Source/Runtime/Core/Private/Mac/MacApplication.cpp b/Engine/Source/Runtime/Core/Private/Mac/MacApplication.cpp
--- a/Engine/Source/Runtime/Core/Private/Mac/MacApplication.cpp
+++ b/Engine/Source/Runtime/Core/Private/Mac/MacApplication.cpp
@@ -831,8 +831,9 @@ FModifierKeysState FMacApplication::GetModifierKeys() const
 	const bool bIsRightAltDown			= ( CurrentFlags & ( 1 << 5 ) ) != 0;
 	const bool bIsLeftCommandDown		= ( CurrentFlags & ( 1 << 2 ) ) != 0; // Mac pretends the Control key is Command
 	const bool bIsRightCommandDown		= ( CurrentFlags & ( 1 << 3 ) ) != 0; // Mac pretends the Control key is Command
+	const bool bAreCapsLocked           = ( CurrentFlags & ( 1 << 8 ) ) != 0; 
 
-	return FModifierKeysState( bIsLeftShiftDown, bIsRightShiftDown, bIsLeftControlDown, bIsRightControlDown, bIsLeftAltDown, bIsRightAltDown, bIsLeftCommandDown, bIsRightCommandDown );
+	return FModifierKeysState(bIsLeftShiftDown, bIsRightShiftDown, bIsLeftControlDown, bIsRightControlDown, bIsLeftAltDown, bIsRightAltDown, bIsLeftCommandDown, bIsRightCommandDown, bAreCapsLocked);
 }
 
 FPlatformRect FMacApplication::GetWorkArea( const FPlatformRect& CurrentWindow ) const
diff --git a/Engine/Source/Runtime/Core/Private/Windows/WindowsApplication.cpp b/Engine/Source/Runtime/Core/Private/Windows/WindowsApplication.cpp
--- a/Engine/Source/Runtime/Core/Private/Windows/WindowsApplication.cpp
+++ b/Engine/Source/Runtime/Core/Private/Windows/WindowsApplication.cpp
@@ -261,9 +261,10 @@ FModifierKeysState FWindowsApplication::GetModifierKeys() const
 	const bool bIsLeftControlDown = ( ::GetAsyncKeyState( VK_LCONTROL ) & 0x8000 ) != 0;
 	const bool bIsRightControlDown = ( ::GetAsyncKeyState( VK_RCONTROL ) & 0x8000 ) != 0;
 	const bool bIsLeftAltDown = ( ::GetAsyncKeyState( VK_LMENU ) & 0x8000 ) != 0;
-	const bool bIsRightAltDown = ( ::GetAsyncKeyState( VK_RMENU ) & 0x8000 ) != 0;
+	const bool bIsRightAltDown = ( ::GetAsyncKeyState (VK_RMENU ) & 0x8000 ) != 0;
+	const bool bAreCapsLocked = ( ::GetKeyState( VK_CAPITAL ) & 0x0001 ) != 0;
 
-	return FModifierKeysState( bIsLeftShiftDown, bIsRightShiftDown, bIsLeftControlDown, bIsRightControlDown, bIsLeftAltDown, bIsRightAltDown, false, false ); // Win key is ignored
+	return FModifierKeysState(bIsLeftShiftDown, bIsRightShiftDown, bIsLeftControlDown, bIsRightControlDown, bIsLeftAltDown, bIsRightAltDown, false, false, bAreCapsLocked); // Win key is ignored
 }
 
 void FWindowsApplication::SetCapture( const TSharedPtr< FGenericWindow >& InWindow )
diff --git a/Engine/Source/Runtime/Core/Public/GenericPlatform/GenericApplication.h b/Engine/Source/Runtime/Core/Public/GenericPlatform/GenericApplication.h
--- a/Engine/Source/Runtime/Core/Public/GenericPlatform/GenericApplication.h
+++ b/Engine/Source/Runtime/Core/Public/GenericPlatform/GenericApplication.h
@@ -43,7 +43,8 @@ public:
 					    const bool bInIsLeftAltDown,
 						const bool bInIsRightAltDown,
 						const bool bInIsLeftCommandDown,
-						const bool bInIsRightCommandDown )
+						const bool bInIsRightCommandDown,
+						const bool bAreCapsLocked)
 		: bIsLeftShiftDown( bInIsLeftShiftDown ),
 		  bIsRightShiftDown( bInIsRightShiftDown ),
 		  bIsLeftControlDown( bInIsLeftControlDown ),
@@ -51,10 +52,23 @@ public:
 		  bIsLeftAltDown( bInIsLeftAltDown ),
 		  bIsRightAltDown( bInIsRightAltDown ),
 		  bIsLeftCommandDown( bInIsLeftCommandDown ),
-		  bIsRightCommandDown( bInIsRightCommandDown )
+		  bIsRightCommandDown( bInIsRightCommandDown ),
+		  bAreCapsLocked( bAreCapsLocked )
 	{
 	}
 
+	FModifierKeysState()
+		: bIsLeftShiftDown(false),
+		  bIsRightShiftDown(false),
+		  bIsLeftControlDown(false),
+		  bIsRightControlDown(false),
+		  bIsLeftAltDown(false),
+		  bIsRightAltDown(false),
+		  bIsLeftCommandDown(false),
+		  bIsRightCommandDown(false),
+		  bAreCapsLocked(false)
+	{
+	}
 	
 	/**
 	 * Returns true if either shift key was down when this event occurred
@@ -177,6 +191,14 @@ public:
 	}
 
 
+	/**
+	 * @return  true if the Caps Lock key has been toggled to the enabled state.
+	 */
+	bool AreCapsLocked() const
+	{
+		return bAreCapsLocked;
+	}
+
 private:
 
 	/** True if the left shift key was down when this event occurred. */
@@ -202,6 +224,9 @@ private:
 	
 	/** True if the right command key was down when this event occurred. */
 	bool bIsRightCommandDown;
+
+	/** True if the Caps Lock key has been toggled to the enabled state. */
+	bool bAreCapsLocked;
 };
 
 struct FPlatformRect
@@ -306,7 +331,7 @@ public:
 
 	virtual void* GetCapture( void ) const { return NULL; }
 
-	virtual FModifierKeysState GetModifierKeys() const  { return FModifierKeysState( false, false, false, false, false, false, false, false ); }
+	virtual FModifierKeysState GetModifierKeys() const  { return FModifierKeysState(); }
 
 	virtual void SetHighPrecisionMouseMode( const bool Enable, const TSharedPtr< FGenericWindow >& InWindow ) { };
 
