Enabled SupportsQuit for android.

diff --git a/Engine/Source/Runtime/Core/Public/Android/AndroidProperties.h b/Engine/Source/Runtime/Core/Public/Android/AndroidProperties.h
--- a/Engine/Source/Runtime/Core/Public/Android/AndroidProperties.h
+++ b/Engine/Source/Runtime/Core/Public/Android/AndroidProperties.h
@@ -88,6 +88,11 @@ struct FAndroidPlatformProperties
 		return false;
 	}
 
+	static FORCEINLINE bool SupportsQuit() 
+	{
+		return true;
+	}
+
 	static FORCEINLINE bool HasFixedResolution()
 	{
 		return true;
