Fix for FIOSystem not being considered during shutdown if never started.

- PR #329. Found by amigo, debugged and fixed on the IRC (caused crashes on SCW exit under Linux).

diff --git a/Engine/Source/Runtime/Core/Private/Serialization/AsyncIOSystemBase.cpp b/Engine/Source/Runtime/Core/Private/Serialization/AsyncIOSystemBase.cpp
--- a/Engine/Source/Runtime/Core/Private/Serialization/AsyncIOSystemBase.cpp
+++ b/Engine/Source/Runtime/Core/Private/Serialization/AsyncIOSystemBase.cpp
@@ -793,5 +793,5 @@ void FIOSystem::Shutdown()
 
 bool FIOSystem::HasShutdown()
 {
-	return AsyncIOThread == (FRunnableThread*)-1;
-}
\ No newline at end of file
+	return AsyncIOThread == nullptr || AsyncIOThread == (FRunnableThread*)-1;
+}
