Linux: misc/minor changes from PR #306.

- Linux splash (changed to match our coding style better).  WARNING: currently cannot be cross-compiled.
- Platform (SDL) message box.
- Other minor additions elsewhere.

diff --git a/Engine/Source/Runtime/Core/Private/Linux/LinuxPlatformFile.cpp b/Engine/Source/Runtime/Core/Private/Linux/LinuxPlatformFile.cpp
--- a/Engine/Source/Runtime/Core/Private/Linux/LinuxPlatformFile.cpp
+++ b/Engine/Source/Runtime/Core/Private/Linux/LinuxPlatformFile.cpp
@@ -2,6 +2,7 @@
 
 #include "Core.h"
 #include <sys/file.h>	// flock()
+#include <sys/stat.h>   // mkdirp()
 
 DEFINE_LOG_CATEGORY_STATIC(LogLinuxPlatformFile, Log, All);
 
@@ -259,6 +260,13 @@ IFileHandle* FLinuxPlatformFile::OpenWrite(const TCHAR* Filename, bool bAppend,
 	{
 		Flags |= O_WRONLY;
 	}
+
+	// create directories if needed.
+	if (!CreateDirectoriesFromPath(Filename))
+	{
+		return NULL;
+	}
+
 	int32 Handle = open(TCHAR_TO_UTF8(*NormalizeFilename(Filename)), Flags, S_IRUSR | S_IWUSR);
 	if (Handle != -1)
 	{
@@ -329,6 +337,53 @@ bool FLinuxPlatformFile::IterateDirectory(const TCHAR* Directory, FDirectoryVisi
 	return Result;
 }
 
+bool FLinuxPlatformFile::CreateDirectoriesFromPath(const TCHAR* Path)
+{
+	// if the file already exists, then directories exist.
+	struct stat FileInfo;
+	if (stat(TCHAR_TO_UTF8(*NormalizeFilename(Path)), &FileInfo) != -1)
+	{
+		return true;
+	}
+
+	// convert path to native char string.
+	int32 Len = strlen(TCHAR_TO_UTF8(*NormalizeFilename(Path)));
+	char *DirPath = reinterpret_cast<char *>(FMemory::Malloc((Len+2) * sizeof(char)));
+	char *SubPath = reinterpret_cast<char *>(FMemory::Malloc((Len+2) * sizeof(char)));
+	strcpy(DirPath, TCHAR_TO_UTF8(*NormalizeFilename(Path)));
+
+	for (int32 i=0; i<Len; ++i)
+	{
+		struct stat FileInfo;
+
+		SubPath[i] = DirPath[i];
+
+		if (SubPath[i] == '/')
+		{
+			SubPath[i+1] = 0;
+
+			// directory exists?
+			if (stat(SubPath, &FileInfo) == -1)
+			{
+				// nope. create it.
+				if (mkdir(SubPath, 0755) == -1)
+				{
+					int ErrNo = errno;
+					UE_LOG(LogLinuxPlatformFile, Warning, TEXT( "create dir('%s') failed: errno=%d (%s)" ),
+																DirPath, ErrNo, ANSI_TO_TCHAR(strerror(ErrNo)));
+					FMemory::Free(DirPath);
+					FMemory::Free(SubPath);
+					return false;
+				}
+			}
+		}
+	}
+
+	FMemory::Free(DirPath);
+	FMemory::Free(SubPath);
+	return true;
+}
+
 IPlatformFile& IPlatformFile::GetPlatformPhysical()
 {
 	static FLinuxPlatformFile Singleton;
diff --git a/Engine/Source/Runtime/Core/Private/Linux/LinuxPlatformMisc.cpp b/Engine/Source/Runtime/Core/Private/Linux/LinuxPlatformMisc.cpp
--- a/Engine/Source/Runtime/Core/Private/Linux/LinuxPlatformMisc.cpp
+++ b/Engine/Source/Runtime/Core/Private/Linux/LinuxPlatformMisc.cpp
@@ -230,6 +230,155 @@ uint32 FLinuxPlatformMisc::GetKeyMap( uint16* KeyCodes, FString* KeyNames, uint3
 	return NumMappings;
 }
 
+EAppReturnType::Type FLinuxPlatformMisc::MessageBoxExt(EAppMsgType::Type MsgType, const TCHAR* Text, const TCHAR* Caption)
+{
+	int NumberOfButtons = 0;
+
+	if (SDL_WasInit(SDL_INIT_VIDEO) == 0)
+	{
+		SDL_InitSubSystem(SDL_INIT_VIDEO);
+	}
+
+	SDL_MessageBoxButtonData *Buttons = nullptr;
+
+	switch (MsgType)
+	{
+		case EAppMsgType::Ok:
+			NumberOfButtons = 1;
+			Buttons = new SDL_MessageBoxButtonData[NumberOfButtons];
+			Buttons[0].flags = SDL_MESSAGEBOX_BUTTON_RETURNKEY_DEFAULT;
+			Buttons[0].text = "Ok";
+			Buttons[0].buttonid = EAppReturnType::Ok;
+			break;
+
+		case EAppMsgType::YesNo:
+			NumberOfButtons = 2;
+			Buttons = new SDL_MessageBoxButtonData[NumberOfButtons];
+			Buttons[0].flags = SDL_MESSAGEBOX_BUTTON_ESCAPEKEY_DEFAULT;
+			Buttons[0].text = "Yes";
+			Buttons[0].buttonid = EAppReturnType::Yes;
+			Buttons[1].flags = SDL_MESSAGEBOX_BUTTON_ESCAPEKEY_DEFAULT;
+			Buttons[1].text = "No";
+			Buttons[1].buttonid = EAppReturnType::No;
+			break;
+
+		case EAppMsgType::OkCancel:
+			NumberOfButtons = 2;
+			Buttons = new SDL_MessageBoxButtonData[NumberOfButtons];
+			Buttons[0].flags = SDL_MESSAGEBOX_BUTTON_ESCAPEKEY_DEFAULT;
+			Buttons[0].text = "Ok";
+			Buttons[0].buttonid = EAppReturnType::Ok;
+			Buttons[1].flags = SDL_MESSAGEBOX_BUTTON_ESCAPEKEY_DEFAULT;
+			Buttons[1].text = "Cancel";
+			Buttons[1].buttonid = EAppReturnType::Cancel;
+			break;
+
+		case EAppMsgType::YesNoCancel:
+			NumberOfButtons = 3;
+			Buttons = new SDL_MessageBoxButtonData[NumberOfButtons];
+			Buttons[0].flags = SDL_MESSAGEBOX_BUTTON_ESCAPEKEY_DEFAULT;
+			Buttons[0].text = "Yes";
+			Buttons[0].buttonid = EAppReturnType::Yes;
+			Buttons[1].flags = SDL_MESSAGEBOX_BUTTON_ESCAPEKEY_DEFAULT;
+			Buttons[1].text = "No";
+			Buttons[1].buttonid = EAppReturnType::No;
+			Buttons[2].flags = SDL_MESSAGEBOX_BUTTON_ESCAPEKEY_DEFAULT;
+			Buttons[2].text = "Cancel";
+			Buttons[2].buttonid = EAppReturnType::Cancel;
+			break;
+
+		case EAppMsgType::CancelRetryContinue:
+			NumberOfButtons = 3;
+			Buttons = new SDL_MessageBoxButtonData[NumberOfButtons];
+			Buttons[0].flags = SDL_MESSAGEBOX_BUTTON_ESCAPEKEY_DEFAULT;
+			Buttons[0].text = "Continue";
+			Buttons[0].buttonid = EAppReturnType::Continue;
+			Buttons[1].flags = SDL_MESSAGEBOX_BUTTON_ESCAPEKEY_DEFAULT;
+			Buttons[1].text = "Retry";
+			Buttons[1].buttonid = EAppReturnType::Retry;
+			Buttons[2].flags = SDL_MESSAGEBOX_BUTTON_ESCAPEKEY_DEFAULT;
+			Buttons[2].text = "Cancel";
+			Buttons[2].buttonid = EAppReturnType::Cancel;
+			break;
+
+		case EAppMsgType::YesNoYesAllNoAll:
+			NumberOfButtons = 4;
+			Buttons = new SDL_MessageBoxButtonData[NumberOfButtons];
+			Buttons[0].flags = SDL_MESSAGEBOX_BUTTON_ESCAPEKEY_DEFAULT;
+			Buttons[0].text = "Yes";
+			Buttons[0].buttonid = EAppReturnType::Yes;
+			Buttons[1].flags = SDL_MESSAGEBOX_BUTTON_ESCAPEKEY_DEFAULT;
+			Buttons[1].text = "No";
+			Buttons[1].buttonid = EAppReturnType::No;
+			Buttons[2].flags = SDL_MESSAGEBOX_BUTTON_ESCAPEKEY_DEFAULT;
+			Buttons[2].text = "Yes to all";
+			Buttons[2].buttonid = EAppReturnType::YesAll;
+			Buttons[3].flags = SDL_MESSAGEBOX_BUTTON_ESCAPEKEY_DEFAULT;
+			Buttons[3].text = "No to all";
+			Buttons[3].buttonid = EAppReturnType::NoAll;
+			break;
+
+		case EAppMsgType::YesNoYesAllNoAllCancel:
+			NumberOfButtons = 5;
+			Buttons = new SDL_MessageBoxButtonData[NumberOfButtons];
+			Buttons[0].flags = SDL_MESSAGEBOX_BUTTON_ESCAPEKEY_DEFAULT;
+			Buttons[0].text = "Yes";
+			Buttons[0].buttonid = EAppReturnType::Yes;
+			Buttons[1].flags = SDL_MESSAGEBOX_BUTTON_ESCAPEKEY_DEFAULT;
+			Buttons[1].text = "No";
+			Buttons[1].buttonid = EAppReturnType::No;
+			Buttons[2].flags = SDL_MESSAGEBOX_BUTTON_ESCAPEKEY_DEFAULT;
+			Buttons[2].text = "Yes to all";
+			Buttons[2].buttonid = EAppReturnType::YesAll;
+			Buttons[3].flags = SDL_MESSAGEBOX_BUTTON_ESCAPEKEY_DEFAULT;
+			Buttons[3].text = "No to all";
+			Buttons[3].buttonid = EAppReturnType::NoAll;
+			Buttons[4].flags = SDL_MESSAGEBOX_BUTTON_ESCAPEKEY_DEFAULT;
+			Buttons[4].text = "Cancel";
+			Buttons[4].buttonid = EAppReturnType::Cancel;
+			break;
+
+		case EAppMsgType::YesNoYesAll:
+			NumberOfButtons = 3;
+			Buttons = new SDL_MessageBoxButtonData[NumberOfButtons];
+			Buttons[0].flags = SDL_MESSAGEBOX_BUTTON_ESCAPEKEY_DEFAULT;
+			Buttons[0].text = "Yes";
+			Buttons[0].buttonid = EAppReturnType::Yes;
+			Buttons[1].flags = SDL_MESSAGEBOX_BUTTON_ESCAPEKEY_DEFAULT;
+			Buttons[1].text = "No";
+			Buttons[1].buttonid = EAppReturnType::No;
+			Buttons[2].flags = SDL_MESSAGEBOX_BUTTON_ESCAPEKEY_DEFAULT;
+			Buttons[2].text = "Yes to all";
+			Buttons[2].buttonid = EAppReturnType::YesAll;
+			break;
+	}
+
+	SDL_MessageBoxData MessageBoxData = 
+	{
+		SDL_MESSAGEBOX_INFORMATION,
+		NULL, // No parent window
+		TCHAR_TO_UTF8(Caption),
+		TCHAR_TO_UTF8(Text),
+		NumberOfButtons,
+		Buttons,
+		NULL // Default color scheme
+	};
+
+	int ButtonPressed = -1;
+	if (SDL_ShowMessageBox(&MessageBoxData, &ButtonPressed) == -1) 
+	{
+		UE_LOG(LogInit, Fatal, TEXT("Error Presenting MessageBox: %s\n"), ANSI_TO_TCHAR(SDL_GetError()));
+		// unreachable
+		SDL_Quit();
+		return EAppReturnType::Cancel;
+	}
+
+	delete[] Buttons;
+	SDL_Quit();
+
+	return ButtonPressed == -1 ? EAppReturnType::Cancel : static_cast<EAppReturnType::Type>(ButtonPressed);
+}
+
 int32 FLinuxPlatformMisc::NumberOfCores()
 {
 	cpu_set_t AvailableCpusMask;
@@ -369,16 +518,16 @@ FString DescribeSignal(int32 Signal, siginfo_t* Info)
 		break;
 
 		HANDLE_CASE(SIGINT, "program interrupted")
-			HANDLE_CASE(SIGQUIT, "user-requested crash")
-			HANDLE_CASE(SIGILL, "illegal instruction")
-			HANDLE_CASE(SIGTRAP, "trace trap")
-			HANDLE_CASE(SIGABRT, "abort() called")
-			HANDLE_CASE(SIGFPE, "floating-point exception")
-			HANDLE_CASE(SIGKILL, "program killed")
-			HANDLE_CASE(SIGSYS, "non-existent system call invoked")
-			HANDLE_CASE(SIGPIPE, "write on a pipe with no reader")
-			HANDLE_CASE(SIGTERM, "software termination signal")
-			HANDLE_CASE(SIGSTOP, "stop")
+		HANDLE_CASE(SIGQUIT, "user-requested crash")
+		HANDLE_CASE(SIGILL, "illegal instruction")
+		HANDLE_CASE(SIGTRAP, "trace trap")
+		HANDLE_CASE(SIGABRT, "abort() called")
+		HANDLE_CASE(SIGFPE, "floating-point exception")
+		HANDLE_CASE(SIGKILL, "program killed")
+		HANDLE_CASE(SIGSYS, "non-existent system call invoked")
+		HANDLE_CASE(SIGPIPE, "write on a pipe with no reader")
+		HANDLE_CASE(SIGTERM, "software termination signal")
+		HANDLE_CASE(SIGSTOP, "stop")
 
 	default:
 		ErrorString += FString::Printf(TEXT("Signal %d (unknown)"), Signal);
diff --git a/Engine/Source/Runtime/Core/Private/Linux/LinuxPlatformSplash.cpp b/Engine/Source/Runtime/Core/Private/Linux/LinuxPlatformSplash.cpp
new file mode 100644
--- /dev/null
+++ b/Engine/Source/Runtime/Core/Private/Linux/LinuxPlatformSplash.cpp
@@ -0,0 +1,688 @@
+// Copyright 1998-2014 Epic Games, Inc. All Rights Reserved.
+
+/*================================================================================
+	SplashScreen.cpp: Splash screen for game/editor startup
+================================================================================*/
+
+
+#include "Core.h"
+#include "EngineVersion.h"
+#include "EngineBuildSettings.h"
+
+#if WITH_EDITOR
+
+#include "SDL_thread.h"
+#include "SDL_timer.h"
+#include "SDL_opengl.h"
+
+#include "ft2build.h"
+#include FT_FREETYPE_H
+
+/**
+ * Splash screen functions and static globals
+ */
+
+struct Rect
+{
+	int32 Top;
+	int32 Left;
+	int32 Right;
+	int32 Bottom;
+};
+
+
+struct SplashFont
+{
+	enum State
+	{
+		NotLoaded = 0,
+		Loaded
+	};
+
+	FT_Face Face;
+	int32 Status;
+};
+
+FT_Library FontLibrary;
+
+SplashFont FontSmall;
+SplashFont FontNormal;
+SplashFont FontLarge;
+
+static SDL_Thread *GSplashThread = nullptr;
+static SDL_Window *GSplashWindow = nullptr;
+static SDL_Renderer *GSplashRenderer = nullptr;
+static SDL_Texture *GSplashTexture = nullptr;
+static SDL_Surface *GSplashScreenImage = nullptr;
+static SDL_Surface *GSplashIconImage = nullptr;
+static SDL_Surface *message = nullptr;
+static FText GSplashScreenAppName;
+static FText GSplashScreenText[ SplashTextType::NumTextTypes ];
+static Rect GSplashScreenTextRects[ SplashTextType::NumTextTypes ];
+static FString GSplashPath;
+static FString GIconPath;
+
+
+static int32 SplashWidth = 0, SplashHeight = 0;
+static unsigned char *ScratchSpace = nullptr;
+static int32 ThreadState = 0;
+static int32 SplashBPP = 0;
+
+//---------------------------------------------------------
+static int32 OpenFonts ( )
+{
+	FontSmall.Status = SplashFont::NotLoaded;
+	FontNormal.Status = SplashFont::NotLoaded;
+	FontLarge.Status = SplashFont::NotLoaded;
+	
+	if (FT_Init_FreeType(&FontLibrary))
+	{
+		UE_LOG(LogHAL, Error, TEXT("*** Unable to initialize trutype font library."));		
+		return -1;
+	}
+	
+	// small font face
+	FString FontPath = FPaths::EngineContentDir() / TEXT("Slate/Fonts/Roboto-Light.ttf");
+	
+	if (FT_New_Face(FontLibrary, TCHAR_TO_ANSI(*FontPath), 0, &FontSmall.Face))
+	{
+		UE_LOG(LogHAL, Error, TEXT("*** Unable to open small font face for splash screen."));		
+	}
+	else
+	{
+		FT_Set_Pixel_Sizes(FontSmall.Face, 0, 10);
+		FontSmall.Status = SplashFont::Loaded;
+	}
+
+
+	// normal font face
+	FontPath = FPaths::EngineContentDir() / TEXT("Slate/Fonts/Roboto-Regular.ttf");
+	
+	if (FT_New_Face(FontLibrary, TCHAR_TO_ANSI(*FontPath), 0, &FontNormal.Face))
+	{
+		UE_LOG(LogHAL, Error, TEXT("*** Unable to open normal font face for splash screen."));		
+	}
+	else
+	{
+		FT_Set_Pixel_Sizes(FontNormal.Face, 0, 12);
+		FontNormal.Status = SplashFont::Loaded;
+	}	
+	
+	// large font face
+	FontPath = FPaths::EngineContentDir() / TEXT("Slate/Fonts/Roboto-Bold.ttf");
+	
+	if (FT_New_Face(FontLibrary, TCHAR_TO_ANSI(*FontPath), 0, &FontLarge.Face))
+	{
+		UE_LOG(LogHAL, Error, TEXT("*** Unable to open large font face for splash screen."));		
+	}
+	else
+	{
+		FT_Set_Pixel_Sizes (FontLarge.Face, 0, 40);
+		FontLarge.Status = SplashFont::Loaded;
+	}
+	
+	return 0;
+}
+
+
+//---------------------------------------------------------
+static void DrawCharacter(int32 penx, int32 peny, FT_GlyphSlot Glyph, int32 CurTypeIndex, float red, float green, float blue)
+{
+
+	// drawing boundaries
+	int32 xmin = GSplashScreenTextRects[CurTypeIndex].Left;
+	int32 xmax = GSplashScreenTextRects[CurTypeIndex].Right;
+	int32 ymax = GSplashScreenTextRects[CurTypeIndex].Bottom;
+	int32 ymin = GSplashScreenTextRects[CurTypeIndex].Top;
+	
+	// glyph dimensions
+	int32 cwidth = Glyph->bitmap.width;
+	int32 cheight = Glyph->bitmap.rows;
+	int32 cpitch = Glyph->bitmap.pitch;
+		
+	unsigned char *pixels = Glyph->bitmap.buffer;
+	
+	// draw glyph raster to texture
+	for (int y=0; y<cheight; y++)
+	{
+		for (int x=0; x<cwidth; x++)
+		{
+			// find pixel position in splash image
+			int xpos = penx + x;
+			int ypos = peny + y - (Glyph->metrics.horiBearingY >> 6);
+			
+			// make sure pixel is in drawing rectangle
+			if (xpos < xmin || xpos >= xmax || ypos < ymin || ypos >= ymax)
+				continue;
+
+			// get index of pixel in glyph bitmap			
+			int32 src_idx = (y * cpitch) + x;
+			int32 dst_idx = (ypos * SplashWidth + xpos) * SplashBPP;
+			
+			// write pixel
+			float alpha = pixels[src_idx] / 255.0f;
+			
+			ScratchSpace[dst_idx] = ScratchSpace[dst_idx]*(1.0 - alpha) + alpha*red;
+			dst_idx++;
+			ScratchSpace[dst_idx] = ScratchSpace[dst_idx]*(1.0 - alpha) + alpha*green;
+			dst_idx++;
+			ScratchSpace[dst_idx] = ScratchSpace[dst_idx]*(1.0 - alpha) + alpha*blue;
+		}
+	}
+}
+
+
+//---------------------------------------------------------
+static int RenderString (GLuint tex_idx)
+{
+	FT_UInt last_glyph = 0;
+	FT_Vector kern;	
+	bool bRightJustify = false;
+	float red, blue, green; // font color
+
+	// clear the rendering scratch pad.
+	FMemory::Memcpy(ScratchSpace, GSplashScreenImage->pixels, SplashWidth*SplashHeight*SplashBPP);
+
+	// draw each type of string
+	for (int CurTypeIndex=0; CurTypeIndex<SplashTextType::NumTextTypes; CurTypeIndex++)
+	{
+		SplashFont *Font;
+		
+		// initial pen position
+		uint32 penx = GSplashScreenTextRects[ CurTypeIndex].Left;
+		uint32 peny = GSplashScreenTextRects[ CurTypeIndex].Bottom;
+		kern.x = 0;		
+		
+		// Set font color and text position
+		if (CurTypeIndex == SplashTextType::StartupProgress)
+		{
+			red = green = blue = 200.0f;
+			Font = &FontSmall;
+		}
+		else if (CurTypeIndex == SplashTextType::VersionInfo1)
+		{
+			red = green = blue = 240.0f;
+			Font = &FontNormal;
+		}
+		else if (CurTypeIndex == SplashTextType::GameName)
+		{
+			red = green = blue = 240.0f;
+			Font = &FontLarge;
+			penx = GSplashScreenTextRects[ CurTypeIndex].Right;
+			bRightJustify = true;
+		}
+		else
+		{
+			red = green = blue = 160.0f;
+			Font = &FontSmall;
+		}
+		
+		// sanity check: make sure we have a font loaded.
+		if (Font->Status == SplashFont::NotLoaded)
+			continue;
+
+		// adjust verticle pos to allow for descenders
+		peny += Font->Face->descender >> 6;
+
+		// convert strings to glyphs and place them in bitmap.
+		FString Text = GSplashScreenText[CurTypeIndex].ToString();
+			
+		for (int i=0; i<Text.Len(); i++)
+		{
+			FT_ULong charcode;
+			
+			// fetch next glyph
+			if (bRightJustify)
+			{
+				charcode = (Uint32)(Text[Text.Len() - i - 1]);
+			}
+			else
+			{
+				charcode = (Uint32)(Text[i]);				
+			}
+						
+			FT_UInt glyph_idx = FT_Get_Char_Index(Font->Face, charcode);
+			FT_Load_Glyph(Font->Face, glyph_idx, FT_LOAD_DEFAULT);
+
+			
+			if (Font->Face->glyph->format != FT_GLYPH_FORMAT_BITMAP)
+			{
+				FT_Render_Glyph(Font->Face->glyph, FT_RENDER_MODE_NORMAL);
+			}
+
+
+			// pen advance and kerning
+			if (bRightJustify)
+			{
+				if (last_glyph != 0)
+				{
+					FT_Get_Kerning(Font->Face, glyph_idx, last_glyph, FT_KERNING_DEFAULT, &kern);
+				}
+				
+				penx -= (Font->Face->glyph->metrics.horiAdvance - kern.x) >> 6;	
+			}
+			else
+			{
+				if (last_glyph != 0)
+				{
+					FT_Get_Kerning(Font->Face, last_glyph, glyph_idx, FT_KERNING_DEFAULT, &kern);
+				}	
+			}
+
+			last_glyph = glyph_idx;
+
+			// draw character
+			DrawCharacter(penx, peny, Font->Face->glyph,CurTypeIndex, red, green, blue);
+						
+			if (!bRightJustify)
+			{
+				penx += (Font->Face->glyph->metrics.horiAdvance - kern.x) >> 6;
+			}
+		}
+				
+		// store rendered text as texture
+		glBindTexture(GL_TEXTURE_2D, tex_idx);
+		
+		if (SplashBPP == 3)
+		{
+			glTexImage2D(GL_TEXTURE_2D, 0, GL_RGB,
+				GSplashScreenImage->w, GSplashScreenImage->h,
+				0, GL_RGB, GL_UNSIGNED_BYTE, ScratchSpace);
+		}
+		else
+		{
+			glTexImage2D(GL_TEXTURE_2D, 0, GL_RGBA,
+				GSplashScreenImage->w, GSplashScreenImage->h,
+				0, GL_RGBA, GL_UNSIGNED_BYTE, ScratchSpace);			
+		}
+	}
+
+	return 0;
+}
+
+
+/**
+ * Thread function that actually creates the window and
+ * renders its contents.
+ */
+static int StartSplashScreenThread(void *ptr)
+{	
+	//	init the sdl here
+	if (SDL_WasInit( SDL_INIT_VIDEO ) == 0)
+	{
+		SDL_InitSubSystem(SDL_INIT_VIDEO);
+	}
+	
+	// create splash window
+	GSplashWindow = SDL_CreateWindow(NULL,
+									SDL_WINDOWPOS_CENTERED,
+									SDL_WINDOWPOS_CENTERED,
+									GSplashScreenImage->w,
+									GSplashScreenImage->h,
+									SDL_WINDOW_BORDERLESS | SDL_WINDOW_OPENGL | SDL_WINDOW_SHOWN
+									);
+	if(GSplashWindow == nullptr)
+	{
+		UE_LOG(LogHAL, Error, TEXT("Splash screen window could not be created! SDL_Error: %s"), UTF8_TO_TCHAR(SDL_GetError()));
+		return -1;
+	}
+
+	SDL_SetWindowTitle( GSplashWindow, TCHAR_TO_UTF8(*GSplashScreenAppName.ToString()) );
+
+	GSplashIconImage = SDL_LoadBMP(TCHAR_TO_UTF8(*(FString(FPlatformProcess::BaseDir()) / GIconPath)));
+
+	if (GSplashIconImage != nullptr)
+	{
+		SDL_SetWindowIcon(GSplashWindow, GSplashIconImage);
+	}
+	else
+	{
+		UE_LOG(LogHAL, Error, TEXT("Splash icon could not be created! SDL_Error: %s"), UTF8_TO_TCHAR(SDL_GetError()));
+	}
+
+	SDL_GL_SetAttribute( SDL_GL_CONTEXT_MAJOR_VERSION, 3 );
+	SDL_GL_SetAttribute( SDL_GL_CONTEXT_MINOR_VERSION, 2 );
+	SDL_GL_SetAttribute( SDL_GL_CONTEXT_FLAGS, SDL_GL_CONTEXT_FORWARD_COMPATIBLE_FLAG);
+	SDL_GL_SetAttribute( SDL_GL_CONTEXT_PROFILE_MASK, SDL_GL_CONTEXT_PROFILE_COMPATIBILITY );
+
+	SDL_GL_SetAttribute( SDL_GL_SHARE_WITH_CURRENT_CONTEXT, 0 );
+
+	SDL_GLContext Context = SDL_GL_CreateContext( GSplashWindow );
+
+	if (Context == NULL)
+	{
+		UE_LOG(LogHAL, Error, TEXT("Splash screen SDL_GL_CreateContext failed: %s"), UTF8_TO_TCHAR(SDL_GetError()));
+		return -1;		
+	}
+
+	// set up viewport
+	glMatrixMode(GL_MODELVIEW);
+	glLoadIdentity();
+	glMatrixMode(GL_PROJECTION);
+	glLoadIdentity();
+
+	glDisable(GL_DEPTH_TEST);
+	glEnable(GL_TEXTURE_2D);
+
+	// create texture slot
+	GLuint texture;
+	glGenTextures(1, &texture);
+	glBindTexture(GL_TEXTURE_2D, texture);
+
+	// set filtering
+	glTexParameteri(GL_TEXTURE_2D, GL_TEXTURE_MIN_FILTER, GL_LINEAR);
+	glTexParameteri(GL_TEXTURE_2D, GL_TEXTURE_MAG_FILTER, GL_LINEAR);
+
+	// load splash image as texture in opengl
+	if (GSplashScreenImage->format->BytesPerPixel == 3)
+	{
+		glTexImage2D(GL_TEXTURE_2D, 0, GL_RGB,
+			GSplashScreenImage->w, GSplashScreenImage->h,
+			0, GL_RGB, GL_UNSIGNED_BYTE, GSplashScreenImage->pixels);
+	}
+	else if (GSplashScreenImage->format->BytesPerPixel == 4)
+	{
+		glTexImage2D(GL_TEXTURE_2D, 0, GL_RGBA,
+			GSplashScreenImage->w, GSplashScreenImage->h,
+			0, GL_RGBA, GL_UNSIGNED_BYTE, GSplashScreenImage->pixels);		
+	}
+	else
+	{
+		UE_LOG(LogHAL, Warning, TEXT("Splash BMP image has unsupported format"));
+		return -1;
+	}
+
+	// remember bmp stats
+	SplashWidth = GSplashScreenImage->w;
+	SplashHeight =  GSplashScreenImage->h;
+	SplashBPP = GSplashScreenImage->format->BytesPerPixel;
+
+	if (SplashWidth <= 0 || SplashHeight <= 0)
+	{
+		UE_LOG(LogHAL, Warning, TEXT("Invalid splash image dimensions."));		
+		return -1;
+	}
+
+	// allocate scratch space for rendering text
+	ScratchSpace = reinterpret_cast<unsigned char *>(FMemory::Malloc(SplashHeight*SplashWidth*SplashBPP));
+
+	// Setup bounds for game name
+	GSplashScreenTextRects[ SplashTextType::GameName ].Top = 0;
+	GSplashScreenTextRects[ SplashTextType::GameName ].Bottom = 50;
+	GSplashScreenTextRects[ SplashTextType::GameName ].Left = 12;
+	GSplashScreenTextRects[ SplashTextType::GameName ].Right = SplashWidth - 12;
+
+	// Setup bounds for version info text 1
+	GSplashScreenTextRects[ SplashTextType::VersionInfo1 ].Top = SplashHeight - 60;
+	GSplashScreenTextRects[ SplashTextType::VersionInfo1 ].Bottom = SplashHeight - 40;
+	GSplashScreenTextRects[ SplashTextType::VersionInfo1 ].Left = 10;
+	GSplashScreenTextRects[ SplashTextType::VersionInfo1 ].Right = SplashWidth - 10;
+
+	// Setup bounds for copyright info text
+	if (GIsEditor)
+	{
+		GSplashScreenTextRects[ SplashTextType::CopyrightInfo ].Top = SplashHeight - 44;
+		GSplashScreenTextRects[ SplashTextType::CopyrightInfo ].Bottom = SplashHeight - 24;
+	}
+	else
+	{
+		GSplashScreenTextRects[ SplashTextType::CopyrightInfo ].Top = SplashHeight - 16;
+		GSplashScreenTextRects[ SplashTextType::CopyrightInfo ].Bottom = SplashHeight - 6;
+	}
+
+	GSplashScreenTextRects[ SplashTextType::CopyrightInfo ].Left = 10;
+	GSplashScreenTextRects[ SplashTextType::CopyrightInfo ].Right = SplashWidth - 20;
+
+	// Setup bounds for startup progress text
+	GSplashScreenTextRects[ SplashTextType::StartupProgress ].Top = SplashHeight - 20;
+	GSplashScreenTextRects[ SplashTextType::StartupProgress ].Bottom = SplashHeight;
+	GSplashScreenTextRects[ SplashTextType::StartupProgress ].Left = 10;
+	GSplashScreenTextRects[ SplashTextType::StartupProgress ].Right = SplashWidth - 20;
+
+	OpenFonts();
+
+	// drawing loop
+	while (ThreadState >= 0)
+	{
+		SDL_Delay(100);
+		
+		if (ThreadState > 0)
+		{
+			RenderString(texture);
+			ThreadState--;
+		}
+		
+		if (SDL_GL_MakeCurrent( GSplashWindow, Context ) == 0)
+		{
+			// draw splash image			
+			glBindTexture(GL_TEXTURE_2D, texture);
+			
+			glBegin(GL_QUADS);
+			glTexCoord2i(0,1); glVertex2i(-1, -1);
+			glTexCoord2i(1,1); glVertex2i(1, -1);
+			glTexCoord2i(1,0); glVertex2i(1, 1);
+			glTexCoord2i(0,0); glVertex2i(-1, 1);
+			glEnd();
+			
+			SDL_GL_SwapWindow( GSplashWindow );
+		}
+	}
+
+	// clean up
+	if (ScratchSpace)
+	{
+		FMemory::Free(ScratchSpace);
+		ScratchSpace = NULL;
+	}
+
+	if (FontSmall.Status == SplashFont::Loaded)
+	{
+		FT_Done_Face(FontSmall.Face);
+	}
+
+	if (FontNormal.Status == SplashFont::Loaded)
+	{
+		FT_Done_Face(FontNormal.Face);
+	}
+
+	if (FontLarge.Status == SplashFont::Loaded)
+	{
+		FT_Done_Face(FontLarge.Face);
+	}
+
+	FT_Done_FreeType(FontLibrary);
+
+	SDL_GL_DeleteContext(Context);
+	SDL_DestroyWindow(GSplashWindow);
+	FMemory::Free(ScratchSpace);
+
+	return 0;
+}
+
+/**
+ * Finds a usable splash pathname for the given filename
+ *
+ * @param SplashFilename Name of the desired splash name ("Splash.bmp")
+ * @param OutPath String containing the path to the file, if this function returns true
+ *
+ * @return true if a splash screen was found
+ */
+static bool GetSplashPath(const TCHAR* SplashFilename, const TCHAR* IconFilename, bool& OutIsCustom)
+{
+	// first look in game's splash directory
+	GSplashPath = FPaths::GameContentDir() + TEXT("Splash/") + SplashFilename;
+	GIconPath = FPaths::GameContentDir() + TEXT("Splash/") + IconFilename;
+	OutIsCustom = true;
+
+	// if this was found, then we're done
+	if (IFileManager::Get().FileSize(*GSplashPath) != -1)
+	{
+		return true;
+	}
+
+	// next look in Engine/Splash
+	GSplashPath = FPaths::ConvertRelativePathToFull(FPaths::EngineContentDir() + TEXT("Splash/") + SplashFilename);
+	GIconPath = FPaths::EngineContentDir() + TEXT("Splash/") + IconFilename;
+	OutIsCustom = false;
+
+	// if this was found, then we're done
+	if (IFileManager::Get().FileSize(*GSplashPath) != -1)
+	{
+		return true;
+	}
+
+	// if not found yet, then return failure
+	return false;
+}
+#endif //WITH_EDITOR
+
+/**
+ * Sets the text displayed on the splash screen (for startup/loading progress)
+ *
+ * @param	InType		Type of text to change
+ * @param	InText		Text to display
+ */
+static void StartSetSplashText( const SplashTextType::Type InType, const FText& InText )
+{
+#if WITH_EDITOR
+	// Update splash text
+	GSplashScreenText[ InType ] = InText;
+#endif
+}
+
+
+/**
+ * Open a splash screen if there's not one already and if it's not disabled.
+ *
+ */
+void FLinuxPlatformSplash::Show( )
+{
+#if WITH_EDITOR
+	// need to do a splash screen?
+	if(GSplashThread || FParse::Param(FCommandLine::Get(),TEXT("NOSPLASH")) == true)
+	{
+		return;
+	}
+
+	// decide on which splash screen to show
+	const FText GameName = FText::FromString(FApp::GetGameName());
+
+	const TCHAR* SplashImage = GIsEditor ? ( GameName.IsEmpty() ? TEXT("EdSplashDefault.bmp") : TEXT("EdSplash.bmp") ) : ( GameName.IsEmpty() ? TEXT("SplashDefault.bmp") : TEXT("Splash.bmp") );
+	// TODO: add an icon
+	//const TCHAR* IconImage = GIsEditor ? ( GameName.IsEmpty() ? TEXT("EdIconDefault.bmp") : TEXT("EdIcon.bmp") ) : ( GameName.IsEmpty() ? TEXT("IconDefault.bmp") : TEXT("Icon.bmp") );
+	const TCHAR* IconImage = GIsEditor ? ( GameName.IsEmpty() ? TEXT("EdSplashDefault.bmp") : TEXT("EdSplash.bmp") ) : ( GameName.IsEmpty() ? TEXT("SplashDefault.bmp") : TEXT("Splash.bmp") );
+
+	// make sure a splash was found
+	bool IsCustom;
+	
+	if ( GetSplashPath(SplashImage, IconImage, IsCustom) == false )
+	{
+		UE_LOG(LogHAL, Warning, TEXT("Splash screen image not found."));
+	}
+
+	
+	// Don't set the game name if the splash screen is custom.
+	if ( !IsCustom )
+	{
+		StartSetSplashText( SplashTextType::GameName, GameName );
+	}
+
+	// In the editor, we'll display loading info
+	if( GIsEditor )
+	{
+		// Set initial startup progress info
+		{
+			StartSetSplashText( SplashTextType::StartupProgress,
+				NSLOCTEXT("UnrealEd", "SplashScreen_InitialStartupProgress", "Loading..." ) );
+		}
+
+		// Set version info
+		{
+			const FText GameName = FText::FromString( FApp::GetGameName() );
+			const FText Version = FText::FromString( GEngineVersion.ToString( FEngineBuildSettings::IsPerforceBuild() ? EVersionComponent::Branch : EVersionComponent::Patch ) );
+
+			FText VersionInfo = FText::Format( NSLOCTEXT("UnrealEd", "UnrealEdTitleWithVersion_F", "Unreal Editor - {0} [Version {1}]" ), GameName, Version );
+			FText AppName =     FText::Format( NSLOCTEXT("UnrealEd", "UnrealEdTitle_F",            "Unreal Editor - {0}" ), GameName );
+
+			StartSetSplashText( SplashTextType::VersionInfo1, VersionInfo );
+
+			// Change the window text (which will be displayed in the taskbar)
+			GSplashScreenAppName = AppName;
+		}
+
+		// Display copyright information in editor splash screen
+		{
+			const FText CopyrightInfo = NSLOCTEXT( "UnrealEd", "SplashScreen_CopyrightInfo", "Copyright \x00a9 1998-2014   Epic Games, Inc.   All rights reserved." );
+			StartSetSplashText( SplashTextType::CopyrightInfo, CopyrightInfo );
+		}
+	}
+
+	// load splash .bmp image
+	GSplashScreenImage = SDL_LoadBMP(TCHAR_TO_UTF8(*GSplashPath));
+	
+	if (GSplashScreenImage == nullptr)
+	{
+		UE_LOG(LogHAL, Warning, TEXT("Could not load splash BMP! SDL_Error: %s"), UTF8_TO_TCHAR(SDL_GetError()));
+		return;
+	}
+
+	// Create rendering thread to display splash
+	GSplashThread = SDL_CreateThread(StartSplashScreenThread, "StartSplashScreenThread", (void *)NULL);
+
+	if (GSplashThread == nullptr)
+	{
+		UE_LOG(LogHAL, Error, TEXT("Splash screen SDL_CreateThread failed: %s"), UTF8_TO_TCHAR(SDL_GetError()));
+	}
+	else
+	{
+		SDL_DetachThread(GSplashThread);
+		ThreadState = 1;
+	}
+#endif //WITH_EDITOR
+}
+
+
+/**
+* Done with splash screen. Close it and clean up.
+*/
+void FLinuxPlatformSplash::Hide()
+{
+#if WITH_EDITOR
+	// signal thread it's time to quit
+	GSplashThread = nullptr;
+	ThreadState = -99;
+#endif
+}
+
+
+/**
+* Sets the text displayed on the splash screen (for startup/loading progress)
+*
+* @param	InType		Type of text to change
+* @param	InText		Text to display
+*/
+void FLinuxPlatformSplash::SetSplashText( const SplashTextType::Type InType, const TCHAR* InText )
+{
+#if WITH_EDITOR
+	// We only want to bother drawing startup progress in the editor, since this information is
+	// not interesting to an end-user (also, it's not usually localized properly.)	
+
+	if (InType == SplashTextType::CopyrightInfo || GIsEditor)
+	{
+		bool bWasUpdated = false;
+		{
+			// Update splash text
+			if (FCString::Strcmp( InText, *GSplashScreenText[ InType ].ToString() ) != 0)
+			{
+				GSplashScreenText[ InType ] = FText::FromString( InText );
+				bWasUpdated = true;
+			}
+		}
+
+		if (bWasUpdated)
+		{
+			ThreadState++;
+		}
+	}
+#endif //WITH_EDITOR
+}
diff --git a/Engine/Source/Runtime/Core/Private/Misc/MessageLog.cpp b/Engine/Source/Runtime/Core/Private/Misc/MessageLog.cpp
--- a/Engine/Source/Runtime/Core/Private/Misc/MessageLog.cpp
+++ b/Engine/Source/Runtime/Core/Private/Misc/MessageLog.cpp
@@ -137,14 +137,18 @@ TSharedRef<FTokenizedMessage> FMessageLog::Error( const FText& InMessage )
 TSharedRef<FTokenizedMessage> FMessageLog::PerformanceWarning( const FText& InMessage )
 {
 	TSharedRef<FTokenizedMessage> Message = FTokenizedMessage::Create(EMessageSeverity::PerformanceWarning, InMessage);
+#if !PLATFORM_LINUX // @todo: these are too spammy for now on Linux
 	Messages.Add(Message);
+#endif // !PLATFORM_LINUX
 	return Message;
 }
 
 TSharedRef<FTokenizedMessage> FMessageLog::Warning( const FText& InMessage )
 {
 	TSharedRef<FTokenizedMessage> Message = FTokenizedMessage::Create(EMessageSeverity::Warning, InMessage);
-	Messages.Add(Message);
+#if !PLATFORM_LINUX // @todo: these are too spammy for now on Linux
+	Messages.Add(Message); // TODO These are too spammy for now
+#endif // !PLATFORM_LINUX
 	return Message;
 }
 
diff --git a/Engine/Source/Runtime/Core/Public/Linux/LinuxPlatformFile.h b/Engine/Source/Runtime/Core/Public/Linux/LinuxPlatformFile.h
--- a/Engine/Source/Runtime/Core/Public/Linux/LinuxPlatformFile.h
+++ b/Engine/Source/Runtime/Core/Public/Linux/LinuxPlatformFile.h
@@ -34,5 +34,6 @@ public:
 	virtual bool DirectoryExists(const TCHAR* Directory) override;
 	virtual bool CreateDirectory(const TCHAR* Directory) override;
 	virtual bool DeleteDirectory(const TCHAR* Directory) override;
+	bool CreateDirectoriesFromPath(const TCHAR* Path);
 	bool IterateDirectory(const TCHAR* Directory, FDirectoryVisitor& Visitor);
 };
diff --git a/Engine/Source/Runtime/Core/Public/Linux/LinuxPlatformMisc.h b/Engine/Source/Runtime/Core/Public/Linux/LinuxPlatformMisc.h
--- a/Engine/Source/Runtime/Core/Public/Linux/LinuxPlatformMisc.h
+++ b/Engine/Source/Runtime/Core/Public/Linux/LinuxPlatformMisc.h
@@ -113,6 +113,8 @@ struct CORE_API FLinuxPlatformMisc : public FGenericPlatformMisc
 		return TEXT(":");
 	}
 
+	static EAppReturnType::Type MessageBoxExt(EAppMsgType::Type MsgType, const TCHAR* Text, const TCHAR* Caption);
+
 	FORCEINLINE static void MemoryBarrier()
 	{
 		__sync_synchronize();
diff --git a/Engine/Source/Runtime/Core/Public/Linux/LinuxPlatformSplash.h b/Engine/Source/Runtime/Core/Public/Linux/LinuxPlatformSplash.h
--- a/Engine/Source/Runtime/Core/Public/Linux/LinuxPlatformSplash.h
+++ b/Engine/Source/Runtime/Core/Public/Linux/LinuxPlatformSplash.h
@@ -1,7 +1,7 @@
 // Copyright 1998-2014 Epic Games, Inc. All Rights Reserved.
 
 /*=============================================================================================
-	LinuxSplash.h: Linux platform splash screen...
+	LinuxSplash.h: Linux platform splash screen.
 ==============================================================================================*/
 
 #pragma once
@@ -9,10 +9,24 @@
 /**
  * Linux splash implementation
  */
-struct CORE_API FLinuxSplash : public FGenericPlatformSplash
+struct CORE_API FLinuxPlatformSplash : public FGenericPlatformSplash
 {
-	// default implementation for now
-};
+	/**
+	 * Show the splash screen
+	 */
+	static void Show();
+	/**
+	 * Hide the splash screen
+	 */
+	static void Hide();
 
+	/**
+	 * Sets the text displayed on the splash screen (for startup/loading progress)
+	 *
+	 * @param	InType		Type of text to change
+	 * @param	InText		Text to display
+	 */
+	static void SetSplashText(const SplashTextType::Type InType, const TCHAR* InText);
+};
 
-typedef FLinuxSplash FPlatformSplash;
+typedef FLinuxPlatformSplash FPlatformSplash;
diff --git a/Engine/Source/Runtime/Engine/Engine.Build.cs b/Engine/Source/Runtime/Engine/Engine.Build.cs
--- a/Engine/Source/Runtime/Engine/Engine.Build.cs
+++ b/Engine/Source/Runtime/Engine/Engine.Build.cs
@@ -81,12 +81,6 @@ public class Engine : ModuleRules
 
 		if (UEBuildConfiguration.bBuildDeveloperTools)
 		{
-			PrivateDependencyModuleNames.AddRange(
-				new string[] {
-					"ImageCore",
-					"RawMesh"
-				});
-
 			// Add "BlankModule" so that it gets compiled as an example and will be maintained and tested.  This can be removed
 			// at any time if needed.  The module isn't actually loaded by the engine so there is no runtime cost.
 			DynamicallyLoadedModuleNames.Add("BlankModule");
@@ -95,6 +89,12 @@ public class Engine : ModuleRules
 			{
 				PrivateIncludePathModuleNames.Add("MeshUtilities");
 				DynamicallyLoadedModuleNames.Add("MeshUtilities");
+
+				PrivateDependencyModuleNames.AddRange(
+					new string[] {
+						"ImageCore",
+						"RawMesh"
+					});
 			}
 
 			if (Target.Configuration != UnrealTargetConfiguration.Shipping && Target.Configuration != UnrealTargetConfiguration.Test && Target.Type != TargetRules.TargetType.Server)
@@ -128,6 +128,16 @@ public class Engine : ModuleRules
 					}
 				);
 			}
+			else if (Target.Platform == UnrealTargetPlatform.Linux)
+			{
+				DynamicallyLoadedModuleNames.AddRange(
+					new string[] {
+						"LinuxTargetPlatform",
+						"LinuxNoEditorTargetPlatform",
+						"LinuxServerTargetPlatform",
+					}
+				);
+			}
 		}
 
 		DynamicallyLoadedModuleNames.AddRange(
diff --git a/Engine/Source/Runtime/Engine/Private/ShaderCompiler/ShaderCompiler.cpp b/Engine/Source/Runtime/Engine/Private/ShaderCompiler/ShaderCompiler.cpp
--- a/Engine/Source/Runtime/Engine/Private/ShaderCompiler/ShaderCompiler.cpp
+++ b/Engine/Source/Runtime/Engine/Private/ShaderCompiler/ShaderCompiler.cpp
@@ -630,7 +630,7 @@ void FShaderCompileThreadRunnable::WriteNewTasks()
 
 			const FString WorkingDirectory = Manager->AbsoluteShaderBaseWorkingDirectory + FString::FromInt(WorkerIndex);
 
-#if PLATFORM_MAC			
+#if PLATFORM_MAC || PLATFORM_LINUX
 			// To make sure that the process waiting for input file won't try to read it until it's ready
 			// we use a temp file name during writing.
 			FString TransferFileName;
@@ -678,7 +678,7 @@ void FShaderCompileThreadRunnable::WriteNewTasks()
 				DoWriteTasks(CurrentWorkerInfo.QueuedJobs, *TransferFile);
 				delete TransferFile;
 
-#if PLATFORM_MAC			
+#if PLATFORM_MAC || PLATFORM_LINUX
 				// Change the transfer file name to proper one
 				FString ProperTransferFileName = WorkingDirectory / TEXT("WorkerInputOnly.in");
 				IFileManager::Get().Move(*ProperTransferFileName, *TransferFileName);
@@ -988,6 +988,8 @@ FShaderCompilingManager::FShaderCompilingManager() :
 	NumOutstandingJobs(0),
 #if PLATFORM_MAC
 	ShaderCompileWorkerName(TEXT("../../../Engine/Binaries/Mac/ShaderCompileWorker"))
+#elif PLATFORM_LINUX
+	ShaderCompileWorkerName(TEXT("../../../Engine/Binaries/Linux/ShaderCompileWorker"))
 #else
 	ShaderCompileWorkerName(TEXT("../../../Engine/Binaries/Win64/ShaderCompileWorker.exe"))
 #endif
@@ -1159,13 +1161,15 @@ uint32 FShaderCompilingManager::LaunchWorker(const FString& WorkingDirectory, ui
 	FPlatformMisc::LowLevelOutputDebugStringf(TEXT("Launching shader compile worker w/ WorkerParameters\n\t%s\n"), WorkerParametersText);
 	return 17;
 #else
-//FPlatformMisc::LowLevelOutputDebugStringf(TEXT("Launching shader compile worker:\n\t%s\n"), *WorkerParameters);
+#if UE_BUILD_DEBUG && PLATFORM_LINUX
+	FPlatformMisc::LowLevelOutputDebugStringf(TEXT("Launching shader compile worker:\n\t%s\n"), *WorkerParameters);
+#endif
 	uint32 WorkerId = 0;
 	FProcHandle WorkerHandle = FPlatformProcess::CreateProc(*ShaderCompileWorkerName, *WorkerParameters, true, false, false, &WorkerId, PriorityModifier, NULL, NULL);
 	if( !WorkerHandle.IsValid() )
 	{
 		// If this doesn't error, the app will hang waiting for jobs that can never be completed
-		UE_LOG(LogShaderCompilers, Fatal, TEXT( "Couldn't launch %s! Make sure the exe is in your binaries folder." ), *ShaderCompileWorkerName );
+		UE_LOG(LogShaderCompilers, Fatal, TEXT( "Couldn't launch %s! Make sure the file is in your binaries folder." ), *ShaderCompileWorkerName );
 	}
 
 	// Close the handle as we no longer need it.
diff --git a/Engine/Source/Runtime/RHI/RHI.Build.cs b/Engine/Source/Runtime/RHI/RHI.Build.cs
--- a/Engine/Source/Runtime/RHI/RHI.Build.cs
+++ b/Engine/Source/Runtime/RHI/RHI.Build.cs
@@ -8,7 +8,7 @@ public class RHI : ModuleRules
 	{
 		PrivateDependencyModuleNames.Add("Core");
 
-		if (UEBuildConfiguration.bCompileAgainstEngine)
+		if (UEBuildConfiguration.bCompileAgainstEngine && !UEBuildConfiguration.bBuildDedicatedServer)
 		{
             DynamicallyLoadedModuleNames.Add("NullDrv");
 
@@ -19,7 +19,8 @@ public class RHI : ModuleRules
 
 			if ((Target.Platform == UnrealTargetPlatform.Win32) ||
 				(Target.Platform == UnrealTargetPlatform.Win64) ||
-				(Target.Platform == UnrealTargetPlatform.Mac)   || 
+				(Target.Platform == UnrealTargetPlatform.Mac)   ||
+				(Target.Platform == UnrealTargetPlatform.Linux) ||
                 (Target.Platform == UnrealTargetPlatform.HTML5))
 			{
 				DynamicallyLoadedModuleNames.Add("OpenGLDrv");
diff --git a/Engine/Source/UE4Game.Target.cs b/Engine/Source/UE4Game.Target.cs
--- a/Engine/Source/UE4Game.Target.cs
+++ b/Engine/Source/UE4Game.Target.cs
@@ -44,7 +44,7 @@ public class UE4GameTarget : TargetRules
 				}
 				OutExtraModuleNames.Add("OnlineSubsystemFacebook");
 			}
-			else if (Target.Platform == UnrealTargetPlatform.Mac)
+			else if (Target.Platform == UnrealTargetPlatform.Mac || Target.Platform == UnrealTargetPlatform.Linux)
 			{
 				OutExtraModuleNames.Add("OnlineSubsystemNull");
 				if (UEBuildConfiguration.bCompileSteamOSS == true)
