Initial work to fix fullscreen mode on OS X properly - these changes make it possible to use 'proper' fullscreen mode which respects the chosen resolution for games where the mouse cursor is hidden. The default remains as WindowedFullcreen as these changes don't address the difference in mouse & display coordinate systems that results from using OS X's scaled fullscreen mode.

diff --git a/Engine/Source/Runtime/Core/Private/Mac/MacApplication.cpp b/Engine/Source/Runtime/Core/Private/Mac/MacApplication.cpp
--- a/Engine/Source/Runtime/Core/Private/Mac/MacApplication.cpp
+++ b/Engine/Source/Runtime/Core/Private/Mac/MacApplication.cpp
@@ -1041,7 +1041,7 @@ void FMacApplication::OnWindowDidResize( FSlateCocoaWindow* Window )
 		uint32 Width = [Window openGLFrame].size.width;
 		uint32 Height = [Window openGLFrame].size.height;
 		
-		if([Window windowMode] != EWindowMode::Windowed)
+		if([Window windowMode] == EWindowMode::WindowedFullscreen)
 		{
 			// Grab current monitor data for sizing
 			Width = FMath::TruncToInt([[Window screen] frame].size.width);
diff --git a/Engine/Source/Runtime/Core/Private/Mac/MacWindow.cpp b/Engine/Source/Runtime/Core/Private/Mac/MacWindow.cpp
--- a/Engine/Source/Runtime/Core/Private/Mac/MacWindow.cpp
+++ b/Engine/Source/Runtime/Core/Private/Mac/MacWindow.cpp
@@ -12,6 +12,8 @@ TArray< FSlateCocoaWindow* > GRunningModalWindows;
 @implementation FSlateCocoaWindow
 
 @synthesize bForwardEvents;
+@synthesize TargetWindowMode;
+@synthesize PreFullScreenRect;
 
 - (id)initWithContentRect:(NSRect)ContentRect styleMask:(NSUInteger)Style backing:(NSBackingStoreType)BufferingType defer:(BOOL)Flag
 {
@@ -27,8 +29,10 @@ TArray< FSlateCocoaWindow* > GRunningModalWindows;
 		bRenderInitialised = false;
 		bZoomed = [super isZoomed];
 		self.bForwardEvents = true;
+		self.TargetWindowMode = EWindowMode::Windowed;
 		[super setAlphaValue:DeferOpacity];
 		DeferFrame = [super frame];
+		self.PreFullScreenRect = DeferFrame;
 		bDeferSetFrame = false;
 		bDeferSetOrigin = false;
 	}
@@ -37,7 +41,11 @@ TArray< FSlateCocoaWindow* > GRunningModalWindows;
 
 - (NSRect)openGLFrame
 {
-	if([self styleMask] & (NSTexturedBackgroundWindowMask))
+	if(self.TargetWindowMode == EWindowMode::Fullscreen || WindowMode == EWindowMode::Fullscreen)
+	{
+		return self.PreFullScreenRect;
+	}
+	else if([self styleMask] & (NSTexturedBackgroundWindowMask))
 	{
 		return (!bDeferSetFrame ? [self frame] : DeferFrame);
 	}
@@ -300,9 +308,19 @@ TArray< FSlateCocoaWindow* > GRunningModalWindows;
 	}
 }
 
+- (void)windowWillEnterFullScreen:(NSNotification *)notification
+{
+	// Handle clicking on the titlebar fullscreen item
+	if(self.TargetWindowMode == EWindowMode::Windowed)
+	{
+		// @todo: Fix fullscreen mode mouse coordinate handling - for now default to windowed fullscreen
+		self.TargetWindowMode = EWindowMode::WindowedFullscreen;
+	}
+}
+
 - (void)windowDidEnterFullScreen:(NSNotification *)notification
 {
-	WindowMode = EWindowMode::WindowedFullscreen;
+	WindowMode = self.TargetWindowMode;
 	if(self.bForwardEvents)
 	{
 		MacApplication->OnWindowDidResize( self );
@@ -312,6 +330,7 @@ TArray< FSlateCocoaWindow* > GRunningModalWindows;
 - (void)windowDidExitFullScreen:(NSNotification *)notification
 {
 	WindowMode = EWindowMode::Windowed;
+	self.TargetWindowMode = EWindowMode::Windowed;
 	if(self.bForwardEvents)
 	{
 		MacApplication->OnWindowDidResize( self );
@@ -789,7 +808,7 @@ FMacWindow::FMacWindow()
 	, WindowMode(EWindowMode::Windowed)
 	, bIsVisible(false)
 {
-	PreFullscreenWindowRect.left = PreFullscreenWindowRect.top = PreFullscreenWindowRect.right = PreFullscreenWindowRect.bottom = 0;
+	PreFullscreenWindowRect.origin.x = PreFullscreenWindowRect.origin.y = PreFullscreenWindowRect.size.width = PreFullscreenWindowRect.size.height = 0.0f;
 }
 
 FSlateCocoaWindow* FMacWindow::GetWindowHandle() const
@@ -805,35 +824,38 @@ void FMacWindow::ReshapeWindow( int32 X, int32 Y, int32 Width, int32 Height )
 	const TSharedRef<FGenericApplicationMessageHandler> MessageHandler = OwningApplication->MessageHandler;
 	MessageHandler->BeginReshapingWindow( SharedThis( this ) );
 
-	BOOL DisplayIfNeeded = (WindowMode == EWindowMode::Windowed);
-	
-	const int32 InvertedY = FPlatformMisc::ConvertSlateYPositionToCocoa(Y) - Height + 1;
-	if (Definition->HasOSWindowBorder)
-	{
-		[WindowHandle setFrame: [WindowHandle frameRectForContentRect: NSMakeRect(X, InvertedY, FMath::Max(Width, 1), FMath::Max(Height, 1))] display:DisplayIfNeeded];
-	}
-	else
-	{
-		[WindowHandle setFrame: NSMakeRect(X, InvertedY, FMath::Max(Width, 1), FMath::Max(Height, 1)) display:DisplayIfNeeded];
-	}
-	
-	// Force resize back to screen size in fullscreen - not ideally pretty but means we don't
-	// have to subvert the OS X or UE fullscreen handling events elsewhere.
-	if(WindowMode != EWindowMode::Windowed)
+	if(WindowMode == EWindowMode::Windowed || WindowMode == EWindowMode::WindowedFullscreen)
 	{
-		[WindowHandle setFrame: [WindowHandle screen].frame display:YES];
+		BOOL DisplayIfNeeded = (WindowMode == EWindowMode::Windowed);
+		
+		const int32 InvertedY = FPlatformMisc::ConvertSlateYPositionToCocoa(Y) - Height + 1;
+		if (Definition->HasOSWindowBorder)
+		{
+			[WindowHandle setFrame: [WindowHandle frameRectForContentRect: NSMakeRect(X, InvertedY, FMath::Max(Width, 1), FMath::Max(Height, 1))] display:DisplayIfNeeded];
+		}
+		else
+		{
+			[WindowHandle setFrame: NSMakeRect(X, InvertedY, FMath::Max(Width, 1), FMath::Max(Height, 1)) display:DisplayIfNeeded];
+		}
+		
+		// Force resize back to screen size in fullscreen - not ideally pretty but means we don't
+		// have to subvert the OS X or UE fullscreen handling events elsewhere.
+		if(WindowMode != EWindowMode::Windowed)
+		{
+			[WindowHandle setFrame: [WindowHandle screen].frame display:YES];
+		}
+		
+		WindowHandle->bZoomed = [WindowHandle isZoomed];
 	}
 	
-	WindowHandle->bZoomed = [WindowHandle isZoomed];
-	
 	MessageHandler->FinishedReshapingWindow( SharedThis( this ) );
 }
 
 bool FMacWindow::GetFullScreenInfo( int32& X, int32& Y, int32& Width, int32& Height ) const
 {
 	SCOPED_AUTORELEASE_POOL;
-
-	const NSRect Frame = [WindowHandle screen].frame;
+	bool const bIsFullscreen = (WindowMode == EWindowMode::Fullscreen);
+	const NSRect Frame = (!bIsFullscreen) ? [WindowHandle screen].frame : PreFullscreenWindowRect;
 	X = Frame.origin.x;
 	Y = Frame.origin.y;
 	Width = Frame.size.width;
@@ -994,6 +1016,13 @@ void FMacWindow::SetWindowMode( EWindowMode::Type NewWindowMode )
 		
 		[WindowHandle setCollectionBehavior: Behaviour];
 		
+		if(!bIsFullscreen)
+		{
+			PreFullscreenWindowRect = [WindowHandle openGLFrame];
+			WindowHandle.PreFullScreenRect = PreFullscreenWindowRect;
+		}
+		
+		WindowHandle.TargetWindowMode = NewWindowMode;
 		[WindowHandle toggleFullScreen:nil];
 		// Ensure that the window has transitioned BEFORE leaving this function
 		// this prevents problems with failure to correctly update mouse locks
diff --git a/Engine/Source/Runtime/Core/Public/Mac/MacWindow.h b/Engine/Source/Runtime/Core/Public/Mac/MacWindow.h
--- a/Engine/Source/Runtime/Core/Public/Mac/MacWindow.h
+++ b/Engine/Source/Runtime/Core/Public/Mac/MacWindow.h
@@ -27,6 +27,8 @@
 	bool bZoomed;
 }
 
+@property (assign) NSRect PreFullScreenRect;
+@property (assign) EWindowMode::Type TargetWindowMode;
 @property (assign) bool bForwardEvents;
 
 /** Get the frame filled by a child OpenGL view, which may cover the window or fill the content view depending upon the window style.
@@ -193,7 +195,7 @@ private:
 	/** The mode that the window is in (windowed, fullscreen, windowedfullscreen ) */
 	EWindowMode::Type WindowMode;
 
-	RECT PreFullscreenWindowRect;
+	NSRect PreFullscreenWindowRect;
 
 	bool bIsVisible : 1;
 };
