Fixed Indian Hindi currency pattern for legacy number formatting. Should resolve Soul automation test failures.

diff --git a/Engine/Source/Runtime/Core/Private/Internationalization/Cultures/IndiaHindi.h b/Engine/Source/Runtime/Core/Private/Internationalization/Cultures/IndiaHindi.h
--- a/Engine/Source/Runtime/Core/Private/Internationalization/Cultures/IndiaHindi.h
+++ b/Engine/Source/Runtime/Core/Private/Internationalization/Cultures/IndiaHindi.h
@@ -44,8 +44,8 @@ public:
 			, "%"					//const FString PercentSymbol
 			, "‰"					//const FString PerMilleSymbol
 			, ""					//const FString PositiveSign
-			, ECurrencyNegativePattern::NegativeSymbolCurrencySymbolSpaceX
-			, ECurrencyPositivePattern::CurrencySymbolSpaceX
+			, ECurrencyNegativePattern::NegativeSymbolCurrencySymbolX
+			, ECurrencyPositivePattern::CurrencySymbolX
 			, EPercentNegativePattern::NegativeSymbolXPct
 			, EPercentPositivePattern::XPct
 			, ENumberNegativePattern::NegativeSymbolX
