A bunch of Linux changes.

- Mostly from PR #306, but sanitized/reworked.
- LinuxApplication now defers processing events (should make it differ less from Windows).
- Process spawning changes in #306 partially left out until explained by amigo/#ue4linux folks.

diff --git a/Engine/Source/Runtime/Core/Private/Linux/LinuxApplication.cpp b/Engine/Source/Runtime/Core/Private/Linux/LinuxApplication.cpp
--- a/Engine/Source/Runtime/Core/Private/Linux/LinuxApplication.cpp
+++ b/Engine/Source/Runtime/Core/Private/Linux/LinuxApplication.cpp
@@ -74,14 +74,25 @@ FLinuxApplication::FLinuxApplication() : GenericApplication( MakeShareable( new
 #endif // STEAM_CONTROLLER_SUPPORT
 {
 	bUsingHighPrecisionMouseInput = false;
-	bAllowedToDeferMessageProcessing = false;
+	bAllowedToDeferMessageProcessing = true;
 	MouseCaptureWindow = NULL;
 	ControllerStates = new SDLControllerState[SDL_NumJoysticks()];
 	memset( ControllerStates, 0, sizeof(SDLControllerState) * SDL_NumJoysticks() );
+
+	fMouseWheelScrollAccel = 1.0f;
+	if (GConfig)
+	{
+		GConfig->GetFloat(TEXT("X11.Tweaks"), TEXT( "MouseWheelScrollAcceleration" ), fMouseWheelScrollAccel, GEngineIni);
+	}
 }
 
 FLinuxApplication::~FLinuxApplication()
 {
+	if ( GConfig )
+	{
+		GConfig->GetFloat(TEXT("X11.Tweaks"), TEXT("MouseWheelScrollAcceleration"), fMouseWheelScrollAccel, GEngineIni);
+		GConfig->Flush(false, GEngineIni);
+	}
 	delete [] ControllerStates;
 }
 
@@ -95,10 +106,9 @@ void FLinuxApplication::DestroyApplication()
 	}
 }
 
-
-TSharedRef< FGenericWindow > FLinuxApplication::MakeWindow() 
-{ 
-	return FLinuxWindow::Make(); 
+TSharedRef< FGenericWindow > FLinuxApplication::MakeWindow()
+{
+	return FLinuxWindow::Make();
 }
 
 void FLinuxApplication::InitializeWindow(	const TSharedRef< FGenericWindow >& InWindow,
@@ -121,18 +131,26 @@ void FLinuxApplication::SetMessageHandler( const TSharedRef< FGenericApplication
 #endif // STEAM_CONTROLLER_SUPPORT
 }
 
-static TSharedPtr< FLinuxWindow > FindWindowBySDLWindow( const TArray< TSharedRef< FLinuxWindow > >& WindowsToSearch, SDL_HWindow const WindowHandle )
+namespace
 {
-	for (int32 WindowIndex=0; WindowIndex < WindowsToSearch.Num(); ++WindowIndex)
+	TSharedPtr< FLinuxWindow > FindWindowBySDLWindow(const TArray< TSharedRef< FLinuxWindow > >& WindowsToSearch, SDL_HWindow const WindowHandle)
 	{
-		TSharedRef< FLinuxWindow > Window = WindowsToSearch[ WindowIndex ];
-		if ( Window->GetHWnd() == WindowHandle )
+		for (int32 WindowIndex=0; WindowIndex < WindowsToSearch.Num(); ++WindowIndex)
 		{
-			return Window;
+			TSharedRef< FLinuxWindow > Window = WindowsToSearch[WindowIndex];
+			if (Window->GetHWnd() == WindowHandle)
+			{
+				return Window;
+			}
 		}
+
+		return TSharedPtr< FLinuxWindow >(nullptr);
 	}
+}
 
-	return TSharedPtr< FLinuxWindow >( NULL );
+TSharedPtr< FLinuxWindow > FLinuxApplication::FindWindowBySDLWindow(SDL_Window *win)
+{
+	return ::FindWindowBySDLWindow(Windows, win);
 }
 
 void FLinuxApplication::PumpMessages( const float TimeDelta )
@@ -167,24 +185,22 @@ bool FLinuxApplication::GeneratesKeyCharMessage(const SDL_KeyboardEvent & KeyDow
 void FLinuxApplication::ProcessDeferredMessage( SDL_Event Event )
 {
 	// This function can be reentered when entering a modal tick loop.
-	// We need to make a copy of the events that need to be processed or we may end up processing the same messages twice 
+	// We need to make a copy of the events that need to be processed or we may end up processing the same messages twice
 	SDL_HWindow NativeWindow = NULL;
 
-	TSharedPtr< FLinuxWindow > CurrentEventWindow = FindEventWindow( &Event );
-	if( !CurrentEventWindow.IsValid() && LastEventWindow.IsValid() )
-	{
-		CurrentEventWindow = LastEventWindow;
-	}
-	if( CurrentEventWindow.IsValid() )
+	// get pointer to window that received this event
+	TSharedPtr< FLinuxWindow > CurrentEventWindow = FindEventWindow(&Event);
+
+	if (CurrentEventWindow.IsValid())
 	{
 		LastEventWindow = CurrentEventWindow;
 		NativeWindow = CurrentEventWindow->GetHWnd();
 	}
-	if( !NativeWindow )
+	if (!NativeWindow)
 	{
 		return;
 	}
-	switch( Event.type )
+	switch(Event.type)
 	{
 	case SDL_KEYDOWN:
 		{
@@ -217,7 +233,7 @@ void FLinuxApplication::ProcessDeferredMessage( SDL_Event Event )
 			SDL_MouseMotionEvent motionEvent = Event.motion;
 			FLinuxCursor *LinuxCursor = (FLinuxCursor*)Cursor.Get();
 
-			if(LinuxCursor->IsHidden())
+			if(SDL_ShowCursor(-1) == 0)
 			{
 				int width, height;
 				SDL_GetWindowSize( NativeWindow, &width, &height );
@@ -266,7 +282,7 @@ void FLinuxApplication::ProcessDeferredMessage( SDL_Event Event )
 	case SDL_MOUSEBUTTONUP:
 		{
 			SDL_MouseButtonEvent buttonEvent = Event.button;
-				
+
 			EMouseButtons::Type button;
 			switch(buttonEvent.button)
 			{
@@ -289,28 +305,30 @@ void FLinuxApplication::ProcessDeferredMessage( SDL_Event Event )
 				button = EMouseButtons::Invalid;
 				break;
 			}
-			if(buttonEvent.type == SDL_MOUSEBUTTONUP)
+			
+			if (buttonEvent.type == SDL_MOUSEBUTTONUP)
 			{
 				MessageHandler->OnMouseUp(button);
 			}
-			// SDL 2.0.2+
-			//else if(buttonEvent.clicks > 1)
-			//{
-			//	MessageHandler->OnMouseDoubleClick( NativeWindow, button );
-			//}
 			else
 			{
-				MessageHandler->OnMouseDown( CurrentEventWindow, button );
+				if (buttonEvent.clicks == 2)
+				{
+					MessageHandler->OnMouseDoubleClick(CurrentEventWindow, button);
+				}
+				else
+				{
+					MessageHandler->OnMouseDown(CurrentEventWindow, button);
+				}
 			}
 		}
 		break;
 	case SDL_MOUSEWHEEL:
 		{
-			SDL_MouseWheelEvent wheelEvent = Event.wheel;
-			const float SpinFactor = 1 / 120.0f;
-			const short WheelDelta = wheelEvent.y;
+			SDL_MouseWheelEvent *WheelEvent = &Event.wheel;
+			float Amount = WheelEvent->y * fMouseWheelScrollAccel;
 
-			MessageHandler->OnMouseWheel( static_cast<float>( WheelDelta ) * SpinFactor );
+			MessageHandler->OnMouseWheel(Amount);
 		}
 		break;
 	case SDL_CONTROLLERAXISMOTION:
@@ -544,7 +562,7 @@ void FLinuxApplication::ProcessDeferredMessage( SDL_Event Event )
 			}
 		}
 		break;
-		
+
 	case SDL_WINDOWEVENT:
 		{
 			SDL_WindowEvent windowEvent = Event.window;
@@ -553,137 +571,134 @@ void FLinuxApplication::ProcessDeferredMessage( SDL_Event Event )
 			{
 				case SDL_WINDOWEVENT_SIZE_CHANGED:
 					{
-					//	printf( "Ariel - SDL_WINDOWEVENT_SIZE_CHANGED has been send.\n" );
-					
 						int NewWidth  = windowEvent.data1;
 						int NewHeight = windowEvent.data2;
 
-						MessageHandler->OnSizeChanged( CurrentEventWindow.ToSharedRef(), NewWidth, NewHeight, 
-																		//	bWasMinimized
-																			false
-																			);
+						MessageHandler->OnSizeChanged(
+							CurrentEventWindow.ToSharedRef(),
+							NewWidth,
+							NewHeight,
+							//	bWasMinimized
+							false
+						);
 					}
 					break;
 
 				case SDL_WINDOWEVENT_RESIZED:
 					{
-					//	printf( "Ariel - SDL_WINDOWEVENT_RESIZED has been send.\n" );
-
 						int NewWidth  = windowEvent.data1;
 						int NewHeight = windowEvent.data2;
 
-						MessageHandler->OnSizeChanged( CurrentEventWindow.ToSharedRef(), NewWidth, NewHeight, 
-																		//	bWasMinimized
-																			false
-																			);
+						MessageHandler->OnSizeChanged(
+							CurrentEventWindow.ToSharedRef(),
+							NewWidth,
+							NewHeight,
+							//	bWasMinimized
+							false
+						);
 
-					//	MessageHandler->OnResizingWindow( CurrentEventWindow.ToSharedRef() );
+						MessageHandler->OnResizingWindow( CurrentEventWindow.ToSharedRef() );
 					}
+					break;
 
 				case SDL_WINDOWEVENT_CLOSE:
 					{
-					//	printf( "Ariel - SDL_WINDOWEVENT_CLOSE has been send: %d %d.\n", windowEvent.data1, windowEvent.data2 );
-						if(windowEvent.data1 == 0 && windowEvent.data2 == 0)
-						{
-							MessageHandler->OnWindowClose( CurrentEventWindow.ToSharedRef() );
-						}
+						MessageHandler->OnWindowClose( CurrentEventWindow.ToSharedRef() );
 					}
 					break;
 
 				case SDL_WINDOWEVENT_SHOWN:
 					{
-					//	printf( "Ariel - SDL_WINDOWEVENT_SHOWN has been send.\n" );
-					}
-					break;
-
-				case SDL_WINDOWEVENT_HIDDEN:
-					{
-					//	printf( "Ariel - SDL_WINDOWEVENT_HIDDEN has been send.\n" );
-					}
-					break;
-
-				case SDL_WINDOWEVENT_EXPOSED:
-					{
-					//	printf( "Ariel - SDL_WINDOWEVENT_EXPOSED has been send.\n" );
+						int Width, Height;
+
+						SDL_GetWindowSize(NativeWindow, &Width, &Height);
+						
+						MessageHandler->OnSizeChanged(
+							CurrentEventWindow.ToSharedRef(),
+							Width,
+							Height,
+							false
+						);
+					
 					}
 					break;
 
 				case SDL_WINDOWEVENT_MOVED:
 					{
-					//	printf( "Ariel - SDL_WINDOWEVENT_MOVED has been send.\n" );
-						MessageHandler->OnMovedWindow( CurrentEventWindow.ToSharedRef(), windowEvent.data1, windowEvent.data2 );
-					}
-					break;
-
-				case SDL_WINDOWEVENT_MINIMIZED:
-					{
-					//	printf( "Ariel - SDL_WINDOWEVENT_MINIMIZED has been send.\n" );
+						MessageHandler->OnMovedWindow(CurrentEventWindow.ToSharedRef(), windowEvent.data1, windowEvent.data2);
 					}
 					break;
 
 				case SDL_WINDOWEVENT_MAXIMIZED:
 					{
-					//	printf( "Ariel - SDL_WINDOWEVENT_MAXIMIZED has been send.\n" );
-						MessageHandler->OnWindowAction( CurrentEventWindow.ToSharedRef(), EWindowAction::Maximize );
+						MessageHandler->OnWindowAction(CurrentEventWindow.ToSharedRef(), EWindowAction::Maximize);
 					}
 					break;
 
 				case SDL_WINDOWEVENT_RESTORED:
 					{
-					//	printf( "Ariel - SDL_WINDOWEVENT_RESTORED has been send.\n" );
-						MessageHandler->OnWindowAction( CurrentEventWindow.ToSharedRef(), EWindowAction::Restore );
+						MessageHandler->OnWindowAction(CurrentEventWindow.ToSharedRef(), EWindowAction::Restore);
 					}
 					break;
 
 				case SDL_WINDOWEVENT_ENTER:
 					{
-						if ( CurrentEventWindow.IsValid() )
+						if (CurrentEventWindow.IsValid())
 						{
 							MessageHandler->OnCursorSet();
-							//MessageHandler->OnWindowActivationChanged( CurrentEventWindow.ToSharedRef(), EWindowActivation::ActivateByMouse );
+							MessageHandler->OnWindowActivationChanged(CurrentEventWindow.ToSharedRef(), EWindowActivation::ActivateByMouse);
 						}
 					}
 					break;
 
 				case SDL_WINDOWEVENT_LEAVE:
 					{
-						if( CurrentEventWindow.IsValid() && GetCapture() != NULL)
+						if (CurrentEventWindow.IsValid() && GetCapture() != NULL)
 						{
 							UpdateMouseCaptureWindow((SDL_HWindow)GetCapture());
+							MessageHandler->OnWindowActivationChanged(CurrentEventWindow.ToSharedRef(), EWindowActivation::Deactivate);
 						}
 					}
 					break;
 
 				case SDL_WINDOWEVENT_FOCUS_GAINED:
 					{
-						if ( CurrentEventWindow.IsValid() )
+						if (CurrentEventWindow.IsValid())
 						{
-							MessageHandler->OnWindowActivationChanged( CurrentEventWindow.ToSharedRef(), EWindowActivation::Activate );
+							MessageHandler->OnWindowActivationChanged(CurrentEventWindow.ToSharedRef(), EWindowActivation::Activate);
 						}
 					}
 					break;
 
 				case SDL_WINDOWEVENT_FOCUS_LOST:
 					{
-					//	printf( "Ariel - SDL_WINDOWEVENT_FOCUS_LOST has been send.\n" );
-						if ( CurrentEventWindow.IsValid() )
+						if (CurrentEventWindow.IsValid())
 						{
-							MessageHandler->OnWindowActivationChanged( CurrentEventWindow.ToSharedRef(), EWindowActivation::Deactivate );
+							MessageHandler->OnWindowActivationChanged(CurrentEventWindow.ToSharedRef(), EWindowActivation::Deactivate);
 						}
 					}
 				break;
+
+				case SDL_WINDOWEVENT_HIDDEN:		// intended fall-through
+				case SDL_WINDOWEVENT_EXPOSED:		// intended fall-through
+				case SDL_WINDOWEVENT_MINIMIZED:		// intended fall-through
+				default:
+					break;
 			}
 		}
 		break;
 	}
-
 }
 
+EWindowZone::Type FLinuxApplication::WindowHitTest(const TSharedPtr< FLinuxWindow > &Window, int x, int y)
+{
+	return MessageHandler->GetWindowZoneForPoint(Window.ToSharedRef(), x, y);
+}
 
 void FLinuxApplication::ProcessDeferredEvents( const float TimeDelta )
 {
 	// This function can be reentered when entering a modal tick loop.
-	// We need to make a copy of the events that need to be processed or we may end up processing the same messages twice 
+	// We need to make a copy of the events that need to be processed or we may end up processing the same messages twice
 	SDL_HWindow NativeWindow = NULL;
 
 	TArray< SDL_Event > Events( PendingEvents );
@@ -810,46 +825,56 @@ TCHAR FLinuxApplication::ConvertChar( SDL_Keysym Keysym )
 
 TSharedPtr< FLinuxWindow > FLinuxApplication::FindEventWindow( SDL_Event* Event )
 {
-	uint16 windowID;
+	uint16 WindowID = 0;
 	switch (Event->type)
 	{
 		case SDL_KEYDOWN:
 		case SDL_KEYUP:
-			windowID = Event->key.windowID;
+			WindowID = Event->key.windowID;
 			break;
 		case SDL_MOUSEMOTION:
-			windowID = Event->motion.windowID;
+			WindowID = Event->motion.windowID;
 			break;
 		case SDL_MOUSEBUTTONDOWN:
 		case SDL_MOUSEBUTTONUP:
-			windowID = Event->button.windowID;
+			WindowID = Event->button.windowID;
 			break;
 		case SDL_MOUSEWHEEL:
-			windowID = Event->wheel.windowID;
+			WindowID = Event->wheel.windowID;
 			break;
 
 		case SDL_WINDOWEVENT:
-			windowID = Event->window.windowID;
+			WindowID = Event->window.windowID;
 			break;
 		default:
-			return TSharedPtr< FLinuxWindow >( NULL );
+			return TSharedPtr< FLinuxWindow >(nullptr);
 	}
 
 	for (int32 WindowIndex=0; WindowIndex < Windows.Num(); ++WindowIndex)
 	{
-		TSharedRef< FLinuxWindow > Window = Windows[ WindowIndex ];
-		if ( SDL_GetWindowID(Window->GetHWnd()) == windowID )
+		TSharedRef< FLinuxWindow > Window = Windows[WindowIndex];
+		
+		if (SDL_GetWindowID(Window->GetHWnd()) == WindowID)
 		{
 			return Window;
 		}
 	}
 
-	if(Windows.Num() > 0)
+	return TSharedPtr< FLinuxWindow >(nullptr);
+}
+
+void FLinuxApplication::RemoveEventWindow(SDL_HWindow HWnd)
+{
+	for (int32 WindowIndex=0; WindowIndex < Windows.Num(); ++WindowIndex)
 	{
-		return TSharedPtr< FLinuxWindow >( Windows[0] );
+		TSharedRef< FLinuxWindow > Window = Windows[ WindowIndex ];
+		
+		if ( Window->GetHWnd() == HWnd )
+		{
+			Windows.RemoveAt(WindowIndex);
+			return;
+		}
 	}
-
-	return TSharedPtr< FLinuxWindow >( NULL );
 }
 
 FModifierKeysState FLinuxApplication::GetModifierKeys() const
@@ -880,7 +905,7 @@ void* FLinuxApplication::GetCapture( void ) const
 	return ( bIsMouseCaptureEnabled && MouseCaptureWindow ) ? MouseCaptureWindow : NULL;
 }
 
-void FLinuxApplication::UpdateMouseCaptureWindow( SDL_HWindow TargetWindow )
+void FLinuxApplication::UpdateMouseCaptureWindow(SDL_HWindow TargetWindow)
 {
 	const bool bEnable = bIsMouseCaptureEnabled || bIsMouseCursorLocked;
 	FLinuxCursor *LinuxCursor = static_cast<FLinuxCursor*>(Cursor.Get());
@@ -907,7 +932,7 @@ void FLinuxApplication::UpdateMouseCaptureWindow( SDL_HWindow TargetWindow )
 			{
 				SDL_CaptureMouse(SDL_FALSE);
 			}
-			MouseCaptureWindow = NULL;
+			MouseCaptureWindow = nullptr;
 		}
 	}
 }
@@ -921,37 +946,15 @@ void FLinuxApplication::SetHighPrecisionMouseMode( const bool Enable, const TSha
 
 FPlatformRect FLinuxApplication::GetWorkArea( const FPlatformRect& CurrentWindow ) const
 {
-	RECT WindowsWindowDim;
-	WindowsWindowDim.left	= CurrentWindow.Left;
-	WindowsWindowDim.top	= CurrentWindow.Top;
-	WindowsWindowDim.right	= CurrentWindow.Right;
-	WindowsWindowDim.bottom	= CurrentWindow.Bottom;
-
-	//	asdd
-#if 0
-	// ... figure out the best monitor for that window.
-	HMONITOR hBestMonitor = MonitorFromRect( &WindowsWindowDim, MONITOR_DEFAULTTONEAREST );
-
-	// Get information about that monitor...
-	MONITORINFO MonitorInfo;
-	MonitorInfo.cbSize = sizeof(MonitorInfo);
-	GetMonitorInfo( hBestMonitor, &MonitorInfo);
-
-	// ... so that we can figure out the work area (are not covered by taskbar)
-	MonitorInfo.rcWork;
-
-	FPlatformRect WorkArea;
-	WorkArea.Left = MonitorInfo.rcWork.left;
-	WorkArea.Top = MonitorInfo.rcWork.top;
-	WorkArea.Right = MonitorInfo.rcWork.right;
-	WorkArea.Bottom = MonitorInfo.rcWork.bottom;
-#endif
-
+	SDL_Rect sdl_rect;
+	SDL_GetDisplayBounds(0, &sdl_rect);
+	
 	FPlatformRect WorkArea;
-	WorkArea.Left	= WindowsWindowDim.left;
-	WorkArea.Top	= WindowsWindowDim.top;
-	WorkArea.Right	= WindowsWindowDim.right;
-	WorkArea.Bottom	= WindowsWindowDim.bottom;
+	
+	WorkArea.Top = sdl_rect.y;
+	WorkArea.Bottom = sdl_rect.y + sdl_rect.h;
+	WorkArea.Left = sdl_rect.x;
+	WorkArea.Right = sdl_rect.x + sdl_rect.w;
 
 	return WorkArea;
 }
diff --git a/Engine/Source/Runtime/Core/Private/Linux/LinuxCursor.cpp b/Engine/Source/Runtime/Core/Private/Linux/LinuxCursor.cpp
--- a/Engine/Source/Runtime/Core/Private/Linux/LinuxCursor.cpp
+++ b/Engine/Source/Runtime/Core/Private/Linux/LinuxCursor.cpp
@@ -1,5 +1,7 @@
+// Copyright 1998-2014 Epic Games, Inc. All Rights Reserved.
 
 #include "Core.h"
+#include "LinuxApplication.h"
 #include "LinuxCursor.h"
 #include "LinuxWindow.h"
 
@@ -91,7 +93,19 @@ FLinuxCursor::FLinuxCursor() : bHidden(false)
 			break;
 
 		case EMouseCursor::EyeDropper:
-			//CursorHandle = LoadCursorFromFile((LPCTSTR)*(FString( FPlatformProcess::BaseDir() ) / FString::Printf( TEXT("%sEditor/Slate/Icons/eyedropper.cur"), *FPaths::EngineContentDir() )));
+			{
+				Uint8 mask[32]={0x00, 0x07, 0x00, 0x0f, 0x00, 0x5f, 0x00, 0xfe,
+								0x01, 0xfc, 0x00, 0xf8, 0x01, 0xfc, 0x02, 0xf8,
+								0x07, 0xd0, 0x0f, 0x80, 0x1f, 0x00, 0x3e, 0x00,
+								0x7c, 0x00, 0x78, 0x00, 0xf0, 0x00, 0x40, 0x00};
+
+				Uint8 data[32]={0x00, 0x07, 0x00, 0x0b, 0x00, 0x53, 0x00, 0xa6,
+								0x01, 0x0c, 0x00, 0xf8, 0x01, 0x7c, 0x02, 0x38,
+								0x04, 0x50, 0x08, 0x80, 0x11, 0x00, 0x22, 0x00,
+								0x44, 0x00, 0x48, 0x00, 0xb0, 0x00, 0x40, 0x00};
+
+				CursorHandle = SDL_CreateCursor(data, mask, 16, 16, 0, 15);
+			}
 			break;
 
 			// NOTE: For custom app cursors, use:
@@ -119,27 +133,29 @@ FLinuxCursor::~FLinuxCursor()
 	{
 		switch( CurCursorIndex )
 		{
-		case EMouseCursor::None:
-		case EMouseCursor::Default:
-		case EMouseCursor::TextEditBeam:
-		case EMouseCursor::ResizeLeftRight:
-		case EMouseCursor::ResizeUpDown:
-		case EMouseCursor::ResizeSouthEast:
-		case EMouseCursor::ResizeSouthWest:
-		case EMouseCursor::CardinalCross:
-		case EMouseCursor::Crosshairs:
-		case EMouseCursor::Hand:
-		case EMouseCursor::GrabHand:
-		case EMouseCursor::GrabHandClosed:
-		case EMouseCursor::SlashedCircle:
-		case EMouseCursor::EyeDropper:
-			// Standard shared cursors don't need to be destroyed
-			break;
-
-		default:
-			// Unrecognized cursor type!
-			check( 0 );
-			break;
+			case EMouseCursor::None:
+			case EMouseCursor::Default:
+			case EMouseCursor::TextEditBeam:
+			case EMouseCursor::ResizeLeftRight:
+			case EMouseCursor::ResizeUpDown:
+			case EMouseCursor::ResizeSouthEast:
+			case EMouseCursor::ResizeSouthWest:
+			case EMouseCursor::CardinalCross:
+			case EMouseCursor::Crosshairs:
+			case EMouseCursor::Hand:
+			case EMouseCursor::GrabHand:
+			case EMouseCursor::GrabHandClosed:
+			case EMouseCursor::SlashedCircle:
+				// Standard shared cursors don't need to be destroyed
+				break;
+			case EMouseCursor::EyeDropper:
+				SDL_FreeCursor(CursorHandles[CurCursorIndex]);
+				break;
+
+			default:
+				// Unrecognized cursor type!
+				check( 0 );
+				break;
 		}
 	}
 }
@@ -169,13 +185,14 @@ void FLinuxCursor::SetType( const EMouseCursor::Type InNewCursor )
 	if(InNewCursor == EMouseCursor::None)
 	{
 		bHidden = true;
-		SDL_ShowCursor(0);
+		SDL_ShowCursor(SDL_DISABLE);
+		SDL_SetCursor(CursorHandles[0]);
 	}
 	else
 	{
 		bHidden = false;
-		SDL_ShowCursor(1);
-		SDL_SetCursor( CursorHandles[ InNewCursor ] );
+		SDL_ShowCursor(SDL_ENABLE);
+		SDL_SetCursor(CursorHandles[InNewCursor]);
 	}
 }
 
@@ -190,12 +207,14 @@ void FLinuxCursor::Show( bool bShow )
 	if( bShow )
 	{
 		// Show mouse cursor.
-		SDL_ShowCursor(1);
+		bHidden = false;
+		SDL_ShowCursor(SDL_ENABLE);
 	}
 	else
 	{
 		// Disable the cursor.
-		SDL_ShowCursor(0);
+		bHidden = true;
+		SDL_ShowCursor(SDL_DISABLE);
 	}
 }
 
diff --git a/Engine/Source/Runtime/Core/Private/Linux/LinuxPlatformProcess.cpp b/Engine/Source/Runtime/Core/Private/Linux/LinuxPlatformProcess.cpp
--- a/Engine/Source/Runtime/Core/Private/Linux/LinuxPlatformProcess.cpp
+++ b/Engine/Source/Runtime/Core/Private/Linux/LinuxPlatformProcess.cpp
@@ -10,13 +10,14 @@
 #include <spawn.h>
 #include <sys/wait.h>
 #include <sys/resource.h>
-#include <sys/ioctl.h>	// ioctl
-#include <asm/ioctls.h>	// FIONREAD
+#include <sys/ioctl.h> // ioctl
+#include <asm/ioctls.h> // FIONREAD
 
 void* FLinuxPlatformProcess::GetDllHandle( const TCHAR* Filename )
 {
 	check( Filename );
-	void *Handle = dlopen( TCHAR_TO_ANSI(Filename), RTLD_LAZY | RTLD_LOCAL );
+	FString AbsolutePath = FPaths::ConvertRelativePathToFull(Filename);
+	void *Handle = dlopen( TCHAR_TO_ANSI(*AbsolutePath), RTLD_LAZY | RTLD_LOCAL );
 	if (!Handle)
 	{
 		UE_LOG(LogLinux, Warning, TEXT("dlopen failed: %s"), ANSI_TO_TCHAR(dlerror()) );
@@ -40,7 +41,7 @@ void* FLinuxPlatformProcess::GetDllExport( void* DllHandle, const TCHAR* ProcNam
 int32 FLinuxPlatformProcess::GetDllApiVersion( const TCHAR* Filename )
 {
 	check(Filename);
-	return MODULE_API_VERSION;
+	return ENGINE_VERSION;
 }
 
 const TCHAR* FLinuxPlatformProcess::GetModulePrefix()
@@ -88,6 +89,32 @@ const TCHAR* FLinuxPlatformProcess::ComputerName()
 	return CachedResult;
 }
 
+void FLinuxPlatformProcess::CleanFileCache()
+{
+	bool bShouldCleanShaderWorkingDirectory = true;
+#if !(UE_BUILD_SHIPPING && WITH_EDITOR)
+	// Only clean the shader working directory if we are the first instance, to avoid deleting files in use by other instances
+	// @todo - check if any other instances are running right now
+	bShouldCleanShaderWorkingDirectory = GIsFirstInstance;
+#endif
+
+	if (bShouldCleanShaderWorkingDirectory && !FParse::Param( FCommandLine::Get(), TEXT("Multiprocess")))
+	{
+		// get shader path, and convert it to the userdirectory
+		FString ShaderDir = FString(FPlatformProcess::BaseDir()) / FPlatformProcess::ShaderDir();
+		FString UserShaderDir = IFileManager::Get().ConvertToAbsolutePathForExternalAppForWrite(*ShaderDir);
+		FPaths::CollapseRelativeDirectories(ShaderDir);
+
+		// make sure we don't delete from the source directory
+		if (ShaderDir != UserShaderDir)
+		{
+			IFileManager::Get().DeleteDirectory(*UserShaderDir, false, true);
+		}
+
+		FPlatformProcess::CleanShaderWorkingDir();
+	}
+}
+
 const TCHAR* FLinuxPlatformProcess::BaseDir()
 {
 	static bool bHaveResult = false;
@@ -96,7 +123,14 @@ const TCHAR* FLinuxPlatformProcess::BaseDir()
 	if (!bHaveResult)
 	{
 		char SelfPath[ PlatformProcessLimits::MaxBaseDirLength ] = {0};
-		readlink( "/proc/self/exe", SelfPath, ARRAY_COUNT(SelfPath) - 1);
+		if (readlink( "/proc/self/exe", SelfPath, ARRAY_COUNT(SelfPath) - 1) == -1)
+		{
+			int ErrNo = errno;
+			UE_LOG(LogHAL, Fatal, TEXT("readlink() failed with errno = %d (%s)"), ErrNo,
+				StringCast< TCHAR >(strerror(ErrNo)).Get());
+			// unreachable
+			return CachedResult;
+		}
 		SelfPath[ARRAY_COUNT(SelfPath) - 1] = 0;
 
 		FCString::Strcpy(CachedResult, ARRAY_COUNT(CachedResult) - 1, ANSI_TO_TCHAR(dirname(SelfPath)));
@@ -107,6 +141,56 @@ const TCHAR* FLinuxPlatformProcess::BaseDir()
 	return CachedResult;
 }
 
+const TCHAR* FLinuxPlatformProcess::UserDir()
+{
+	// The UserDir is where user visible files (such as game projects) live.
+	// On Linux (just like on Mac) this corresponds to $HOME/Documents.
+	static TCHAR Result[MAX_PATH] = TEXT("");
+
+	if (!Result[0])
+	{
+		char* Home = secure_getenv("HOME");
+		if (!Home)
+		{
+			UE_LOG(LogHAL, Warning, TEXT("Unable to read the $HOME environment variable"));
+		}
+		else
+		{
+			FCString::Strcpy(Result, ANSI_TO_TCHAR(Home));
+			FCString::Strcat(Result, TEXT("/Documents/"));
+		}
+	}
+	return Result;
+}
+
+const TCHAR* FLinuxPlatformProcess::UserSettingsDir()
+{
+	// Like on Mac we use the same folder for UserSettingsDir and ApplicationSettingsDir
+	// $HOME/.config/Epic/
+	return ApplicationSettingsDir();
+}
+
+const TCHAR* FLinuxPlatformProcess::ApplicationSettingsDir()
+{
+	// The ApplicationSettingsDir is where the engine stores settings and configuration
+	// data.  On linux this corresponds to $HOME/.config/Epic
+	static TCHAR Result[MAX_PATH] = TEXT("");
+	if (!Result[0])
+	{
+		char* Home = secure_getenv("HOME");
+		if (!Home)
+		{
+			UE_LOG(LogHAL, Warning, TEXT("Unable to read the $HOME environment variable"));
+		}
+		else
+		{
+			FCString::Strcpy(Result, ANSI_TO_TCHAR(Home));
+			FCString::Strcat(Result, TEXT("/.config/Epic/"));
+		}
+	}
+	return Result;
+}
+
 bool FLinuxPlatformProcess::SetProcessLimits(EProcessResource::Type Resource, uint64 Limit)
 {
 	rlimit NativeLimit;
@@ -122,7 +206,7 @@ bool FLinuxPlatformProcess::SetProcessLimits(EProcessResource::Type Resource, ui
 			break;
 
 		default:
-			UE_LOG(LogHAL, Warning, TEXT("Unkown resource type %d"), Resource);
+			UE_LOG(LogHAL, Warning, TEXT("Unknown resource type %d"), Resource);
 			return false;
 	}
 
@@ -142,7 +226,13 @@ const TCHAR* FLinuxPlatformProcess::ExecutableName(bool bRemoveExtension)
 	if (!bHaveResult)
 	{
 		char SelfPath[ PlatformProcessLimits::MaxBaseDirLength ] = {0};
-		readlink( "/proc/self/exe", SelfPath, ARRAY_COUNT(SelfPath) - 1);
+		if (readlink( "/proc/self/exe", SelfPath, ARRAY_COUNT(SelfPath) - 1) == -1)
+		{
+			int ErrNo = errno;
+			UE_LOG(LogHAL, Fatal, TEXT("readlink() failed with errno = %d (%s)"), ErrNo,
+				StringCast< TCHAR >(strerror(ErrNo)).Get());
+			return CachedResult;
+		}
 		SelfPath[ARRAY_COUNT(SelfPath) - 1] = 0;
 
 		FCString::Strcpy(CachedResult, ARRAY_COUNT(CachedResult) - 1, ANSI_TO_TCHAR(basename(SelfPath)));
@@ -291,14 +381,31 @@ FRunnableThread* FLinuxPlatformProcess::CreateRunnableThread()
 
 void FLinuxPlatformProcess::LaunchURL(const TCHAR* URL, const TCHAR* Parms, FString* Error)
 {
-	// stub implementation for now
-	// TODO: consider looking for gnome-open, xdg-open, sensible-browser or just hardcoded names...
+	// @todo This ignores params and error; mostly a stub
+	pid_t pid = fork();
+	UE_LOG(LogHAL, Verbose, TEXT("FLinuxPlatformProcess::LaunchURL: '%s'"), TCHAR_TO_ANSI(URL));
+	if (pid == 0)
+	{
+		exit(execl("/usr/bin/xdg-open", "xdg-open", TCHAR_TO_ANSI(URL), (char *)0));
+	}
 }
 
 FProcHandle FLinuxPlatformProcess::CreateProc(const TCHAR* URL, const TCHAR* Parms, bool bLaunchDetached, bool bLaunchHidden, bool bLaunchReallyHidden, uint32* OutProcessID, int32 PriorityModifier, const TCHAR* OptionalWorkingDirectory, void* PipeWrite)
 {
-	FString AbsolutePath = FPaths::ConvertRelativePathToFull(URL);
-	FString Commandline = AbsolutePath;
+	// @TODO bLaunchHidden bLaunchReallyHidden are not handled
+	// When using OptionalWorkingDirectory, we need an absolute path to executable
+	FString ProcessPath = URL;
+	if (*URL != '/' && OptionalWorkingDirectory)
+	{
+		ProcessPath = FString(BaseDir()) + ProcessPath;
+	}
+
+	if (!FPaths::FileExists(ProcessPath))
+	{
+		return FProcHandle();
+	}
+
+	FString Commandline = ProcessPath;
 	Commandline += TEXT(" ");
 	Commandline += Parms;
 
@@ -306,12 +413,12 @@ FProcHandle FLinuxPlatformProcess::CreateProc(const TCHAR* URL, const TCHAR* Par
 
 	TArray<FString> ArgvArray;
 	int Argc = Commandline.ParseIntoArray(&ArgvArray, TEXT(" "), true);
-	char * Argv[PlatformProcessLimits::MaxArgvParameters + 1] = { NULL };	// last argument is NULL, hence +1
+	char* Argv[PlatformProcessLimits::MaxArgvParameters + 1] = { NULL };	// last argument is NULL, hence +1
 	struct CleanupArgvOnExit
 	{
 		int Argc;
-		char ** Argv;	// relying on it being long enough to hold Argc elements
-		
+		char** Argv;	// relying on it being long enough to hold Argc elements
+
 		CleanupArgvOnExit( int InArgc, char *InArgv[] )
 			:	Argc(InArgc)
 			,	Argv(InArgv)
@@ -326,6 +433,65 @@ FProcHandle FLinuxPlatformProcess::CreateProc(const TCHAR* URL, const TCHAR* Par
 		}
 	} CleanupGuard(Argc, Argv);
 
+	// make sure we do not lose arguments with spaces in them due to Commandline.ParseIntoArray breaking them apart above
+	// @todo this code might need to be optimized somehow and integrated with main argument parser below it
+	TArray<FString> NewArgvArray;
+	if (Argc > 0)
+	{
+		if (Argc > PlatformProcessLimits::MaxArgvParameters)
+		{
+			UE_LOG(LogHAL, Warning, TEXT("FLinuxPlatformProcess::CreateProc: too many (%d) commandline arguments passed, will only pass %d"),
+				Argc, PlatformProcessLimits::MaxArgvParameters);
+			Argc = PlatformProcessLimits::MaxArgvParameters;
+		}
+
+		FString MultiPartArg;
+		for (int32 Index = 0; Index < Argc; Index++)
+		{
+			if (MultiPartArg.IsEmpty())
+			{
+				if ((ArgvArray[Index].StartsWith(TEXT("\"")) && !ArgvArray[Index].EndsWith(TEXT("\""))) // check for a starting quote but no ending quote, excludes quoted single arguments
+					|| (ArgvArray[Index].Contains(TEXT("=\"")) && !ArgvArray[Index].EndsWith(TEXT("\""))) // check for quote after =, but no ending quote, this gets arguments of the type -blah="string string string"
+					|| ArgvArray[Index].EndsWith(TEXT("=\""))) // check for ending quote after =, this gets arguments of the type -blah=" string string string "
+				{
+					MultiPartArg = ArgvArray[Index];
+				}
+				else
+				{
+					if (ArgvArray[Index].Contains(TEXT("=\"")))
+					{
+						FString SingleArg = ArgvArray[Index];
+						SingleArg = SingleArg.Replace(TEXT("=\""), TEXT("="));
+						NewArgvArray.Add(SingleArg.TrimQuotes(NULL));
+					}
+					else
+					{
+						NewArgvArray.Add(ArgvArray[Index].TrimQuotes(NULL));
+					}
+				}
+			}
+			else
+			{
+				MultiPartArg += TEXT(" ");
+				MultiPartArg += ArgvArray[Index];
+				if (ArgvArray[Index].EndsWith(TEXT("\"")))
+				{
+					if (MultiPartArg.StartsWith(TEXT("\"")))
+					{
+						NewArgvArray.Add(MultiPartArg.TrimQuotes(NULL));
+					}
+					else
+					{
+						NewArgvArray.Add(MultiPartArg);
+					}
+					MultiPartArg.Empty();
+				}
+			}
+		}
+	}
+	// update Argc with the new argument count
+	Argc = NewArgvArray.Num();
+
 	if (Argc > 0)	// almost always, unless there's no program name
 	{
 		if (Argc > PlatformProcessLimits::MaxArgvParameters)
@@ -337,8 +503,8 @@ FProcHandle FLinuxPlatformProcess::CreateProc(const TCHAR* URL, const TCHAR* Par
 
 		for (int Idx = 0; Idx < Argc; ++Idx)
 		{
-			auto AnsiBuffer = StringCast<ANSICHAR>(*ArgvArray[Idx]);
-			const char * Ansi = AnsiBuffer.Get();
+			auto AnsiBuffer = StringCast<ANSICHAR>(*NewArgvArray[Idx]);
+			const char* Ansi = AnsiBuffer.Get();
 			size_t AnsiSize = FCStringAnsi::Strlen(Ansi) + 1;
 			check(AnsiSize);
 
@@ -357,7 +523,7 @@ FProcHandle FLinuxPlatformProcess::CreateProc(const TCHAR* URL, const TCHAR* Par
 	pid_t ChildPid = -1;
 
 	posix_spawn_file_actions_t FileActions;
-	
+
 	posix_spawn_file_actions_init(&FileActions);
 	if (PipeWrite)
 	{
@@ -365,7 +531,7 @@ FProcHandle FLinuxPlatformProcess::CreateProc(const TCHAR* URL, const TCHAR* Par
 		posix_spawn_file_actions_adddup2(&FileActions, PipeWriteHandle->GetHandle(), STDOUT_FILENO);
 	}
 
-	int ErrNo = posix_spawn(&ChildPid, TCHAR_TO_ANSI(*AbsolutePath), &FileActions, NULL, Argv, environ);
+	int ErrNo = posix_spawn(&ChildPid, TCHAR_TO_ANSI(*ProcessPath), &FileActions, NULL, Argv, environ);
 	posix_spawn_file_actions_destroy(&FileActions);
 	if (ErrNo != 0)
 	{
@@ -377,6 +543,13 @@ FProcHandle FLinuxPlatformProcess::CreateProc(const TCHAR* URL, const TCHAR* Par
 		UE_LOG(LogHAL, Log, TEXT("FLinuxPlatformProcess::CreateProc: spawned child %d"), ChildPid);
 	}
 
+	if (OutProcessID)
+	{
+		*OutProcessID = ChildPid;
+	}
+
+	posix_spawn_file_actions_destroy(&FileActions);
+
 	return FProcHandle( ChildPid );
 }
 
@@ -430,7 +603,7 @@ bool FProcHandle::GetReturnCode(int32* ReturnCodePtr)
 		Wait();
 	}
 
-	if (ReturnCode < 0)
+	if (ReturnCode != -1)
 	{
 		if (ReturnCodePtr != NULL)
 		{
@@ -488,14 +661,19 @@ void FLinuxPlatformProcess::TerminateProc( FProcHandle & ProcessHandle, bool Kil
 	if (KillTree)
 	{
 		// TODO: enumerate the children
-		check(!"Killing a subtree is not implemented yet");
+		STUBBED("FLinuxPlatformProcess::TerminateProc() : Killing a subtree is not implemented yet");
 	}
 
 	int KillResult = kill(ProcessHandle.Get(), SIGTERM);	// graceful
 	check(KillResult != -1 || errno != EINVAL);
 }
 
-bool FLinuxPlatformProcess::GetProcReturnCode( FProcHandle & ProcHandle, int32* ReturnCode )
+uint32 FLinuxPlatformProcess::GetCurrentProcessId()
+{
+	return getpid();
+}
+
+bool FLinuxPlatformProcess::GetProcReturnCode( FProcHandle& ProcHandle, int32* ReturnCode )
 {
 	if (IsProcRunning(ProcHandle))
 	{
@@ -517,3 +695,155 @@ bool FLinuxPlatformProcess::Daemonize()
 
 	return true;
 }
+
+bool FLinuxPlatformProcess::IsApplicationRunning( uint32 ProcessId )
+{
+	errno = 0;
+	getpriority(PRIO_PROCESS, ProcessId);
+	return errno == 0;
+}
+
+bool FLinuxPlatformProcess::IsThisApplicationForeground()
+{
+	STUBBED("FLinuxPlatformProcess::IsThisApplicationForeground");
+	return true;
+}
+
+bool FLinuxPlatformProcess::IsApplicationRunning( const TCHAR* ProcName )
+{
+	FString Commandline = TEXT("pidof '");
+	Commandline += ProcName;
+	Commandline += TEXT("'  > /dev/null");
+	return !system(TCHAR_TO_ANSI(*Commandline));
+}
+
+bool FLinuxPlatformProcess::ExecProcess( const TCHAR* URL, const TCHAR* Params, int32* OutReturnCode, FString* OutStdOut, FString* OutStdErr )
+{
+	TArray<FString> ArgsArray;
+	FString(Params).ParseIntoArray(&ArgsArray, TEXT(" "), true);
+	char *args[ArgsArray.Num()];
+
+	for(int i = 0; i < ArgsArray.Num(); i++)
+	{
+		args[i] = TCHAR_TO_ANSI(*ArgsArray[i]);
+	}
+	pid_t pid;
+	int status;
+	int fd_stdout[2], fd_stderr[2];
+
+	if (pipe(fd_stdout) == -1) 
+	{
+		int ErrNo = errno;
+		UE_LOG(LogHAL, Fatal, TEXT("Creating fd_stdout pipe failed with errno = %d (%s)"), ErrNo,
+			StringCast< TCHAR >(strerror(ErrNo)).Get());
+		return false;
+	}
+
+	if (pipe(fd_stderr) == -1) 
+	{
+		int ErrNo = errno;
+		UE_LOG(LogHAL, Fatal, TEXT("Creating fd_stderr pipe failed with errno = %d (%s)"), ErrNo,
+			StringCast< TCHAR >(strerror(ErrNo)).Get());
+		return false;
+	}
+
+	pid = fork();
+	if(pid < 0)
+	{
+		int ErrNo = errno;
+		UE_LOG(LogHAL, Fatal, TEXT("fork() failed with errno = %d (%s)"), ErrNo,
+			StringCast< TCHAR >(strerror(ErrNo)).Get());
+		return false;
+	}
+
+	if(pid == 0)
+	{
+		close(fd_stdout[0]);
+		close(fd_stderr[0]);
+		dup2(fd_stdout[1], 1);
+		dup2(fd_stderr[1], 2);
+		close(fd_stdout[1]);
+		close(fd_stderr[1]);
+
+		// TODO Not sure if we need to look up URL in PATH
+		exit(execv(TCHAR_TO_ANSI(URL), args));
+	}
+	else
+	{
+		// TODO This might deadlock, should use select. Rewrite me. Also doesn't handle all errors correctly.
+		close(fd_stdout[1]);
+		close(fd_stderr[1]);
+		do
+		{
+			pid_t wpid = waitpid(pid, &status, 0);
+			if (wpid == -1)
+			{
+				int ErrNo = errno;
+				UE_LOG(LogHAL, Fatal, TEXT("waitpid() failed with errno = %d (%s)"), ErrNo,
+				StringCast< TCHAR >(strerror(ErrNo)).Get());
+				return false;
+			}
+
+			if (WIFEXITED(status))
+			{
+				*OutReturnCode = WEXITSTATUS(status);
+			}
+			else if (WIFSIGNALED(status))
+			{
+				*OutReturnCode = WTERMSIG(status);
+			}
+
+		} while (!WIFEXITED(status) && !WIFSIGNALED(status));
+
+		while(1)
+		{
+			char buf[100];
+			int size = read(fd_stdout[0], buf, 100);
+			if(size)
+			{
+				if(OutStdErr)
+				{
+					*OutStdErr += FString(buf);
+				}
+			}
+			else
+			{
+				break;
+			}
+		}
+
+		while(1)
+		{
+			char buf[100];
+			int size = read(fd_stderr[0], buf, 100);
+			if(size)
+			{
+				if(OutStdOut)
+				{
+					*OutStdOut += FString(buf);
+				}
+			}
+			else
+			{
+				break;
+			}
+		}
+		return true;
+	}
+}
+
+void FLinuxPlatformProcess::LaunchFileInDefaultExternalApplication( const TCHAR* FileName, const TCHAR* Parms, ELaunchVerb::Type Verb )
+{
+	// TODO This ignores parms and verb
+	pid_t pid = fork();
+	if (pid == 0)
+	{
+		exit(execl("/usr/bin/xdg-open", "xdg-open", TCHAR_TO_ANSI(FileName), (char *)0));
+	}
+}
+
+void FLinuxPlatformProcess::ExploreFolder( const TCHAR* FilePath )
+{
+	// TODO This is broken, not an explore action but should be fine if called on a directory
+	FLinuxPlatformProcess::LaunchFileInDefaultExternalApplication(FilePath, NULL, ELaunchVerb::Edit);
+}
diff --git a/Engine/Source/Runtime/Core/Public/Linux/LinuxApplication.h b/Engine/Source/Runtime/Core/Public/Linux/LinuxApplication.h
--- a/Engine/Source/Runtime/Core/Public/Linux/LinuxApplication.h
+++ b/Engine/Source/Runtime/Core/Public/Linux/LinuxApplication.h
@@ -56,6 +56,11 @@ public:
 	void AddPendingEvent( SDL_Event event );
 
 	void OnMouseCursorLock( bool bLockEnabled );
+	
+	void RemoveEventWindow(SDL_HWindow Window);
+
+	EWindowZone::Type WindowHitTest( const TSharedPtr< FLinuxWindow > &window, int x, int y );
+	TSharedPtr< FLinuxWindow > FindWindowBySDLWindow( SDL_Window *win );
 
 private:
 	FLinuxApplication();
@@ -106,6 +111,8 @@ private:
 
 	SDLControllerState *ControllerStates;
 
+	float fMouseWheelScrollAccel;
+
 #if STEAM_CONTROLLER_SUPPORT
 	TSharedPtr< class SteamControllerInterface > SteamInput;
 #endif // STEAM_CONTROLLER_SUPPORT
diff --git a/Engine/Source/Runtime/Core/Public/Linux/LinuxPlatformProcess.h b/Engine/Source/Runtime/Core/Public/Linux/LinuxPlatformProcess.h
--- a/Engine/Source/Runtime/Core/Public/Linux/LinuxPlatformProcess.h
+++ b/Engine/Source/Runtime/Core/Public/Linux/LinuxPlatformProcess.h
@@ -136,7 +136,11 @@ struct CORE_API FLinuxPlatformProcess : public FGenericPlatformProcess
 	static void* GetDllExport( void* DllHandle, const TCHAR* ProcName );
 	static int32 GetDllApiVersion( const TCHAR* Filename );
 	static const TCHAR* ComputerName();
+	static void CleanFileCache();
 	static const TCHAR* BaseDir();
+	static const TCHAR* UserDir();
+	static const TCHAR* UserSettingsDir();
+	static const TCHAR* ApplicationSettingsDir();
 	static FString GetApplicationName( uint32 ProcessId );
 	static bool SetProcessLimits(EProcessResource::Type Resource, uint64 Limit);
 	static const TCHAR* ExecutableName(bool bRemoveExtension = true);
@@ -153,8 +157,15 @@ struct CORE_API FLinuxPlatformProcess : public FGenericPlatformProcess
 	static bool IsProcRunning( FProcHandle & ProcessHandle );
 	static void WaitForProc( FProcHandle & ProcessHandle );
 	static void TerminateProc( FProcHandle & ProcessHandle, bool KillTree = false );
+	static uint32 GetCurrentProcessId();
 	static bool GetProcReturnCode( FProcHandle & ProcHandle, int32* ReturnCode );
 	static bool Daemonize();
+	static bool IsApplicationRunning( uint32 ProcessId );
+	static bool IsApplicationRunning( const TCHAR* ProcName );
+	static bool IsThisApplicationForeground();
+	static bool ExecProcess( const TCHAR* URL, const TCHAR* Params, int32* OutReturnCode, FString* OutStdOut, FString* OutStdErr );
+	static void ExploreFolder( const TCHAR* FilePath );
+	static void LaunchFileInDefaultExternalApplication( const TCHAR* FileName, const TCHAR* Parms = NULL, ELaunchVerb::Type Verb = ELaunchVerb::Open );
 };
 
 typedef FLinuxPlatformProcess FPlatformProcess;
