Checks removed from TDereferenceWrapper, so we can sort containers with nullptrs in them (albeit by dereferencing null references).

diff --git a/Engine/Source/Runtime/Core/Public/Templates/Sorting.h b/Engine/Source/Runtime/Core/Public/Templates/Sorting.h
--- a/Engine/Source/Runtime/Core/Public/Templates/Sorting.h
+++ b/Engine/Source/Runtime/Core/Public/Templates/Sorting.h
@@ -55,8 +55,6 @@ struct TDereferenceWrapper<T*, PREDICATE_CLASS>
 	/** Dereferennce pointers */
 	FORCEINLINE bool operator()( T* A, T* B ) const 
 	{
-		check( A != NULL );
-		check( B != NULL );
 		return Predicate( *A, *B ); 
 	} 
 };
