Added a method to Internationalization Cultures to get the "parent" name of a culture from it's name string. Culture fallback logic improved and made consistent - any locale can be used if its internationalization data is present and some parent of it has localization data. Renamed all localization data and code using the erroneous "zh-CN" locale instead of the correct "zh-Hans-CN" locale for China.

diff --git a/Engine/Config/Localization/Editor.ini b/Engine/Config/Localization/Editor.ini
--- a/Engine/Config/Localization/Editor.ini
+++ b/Engine/Config/Localization/Editor.ini
@@ -8,7 +8,7 @@ SourceCulture=en
 CulturesToGenerate=en
 CulturesToGenerate=ja
 CulturesToGenerate=ko
-CulturesToGenerate=zh-CN
+CulturesToGenerate=zh-Hans-CN
 CulturesToGenerate=es
 
 ;Gather text from source code
diff --git a/Engine/Config/Localization/Engine.ini b/Engine/Config/Localization/Engine.ini
--- a/Engine/Config/Localization/Engine.ini
+++ b/Engine/Config/Localization/Engine.ini
@@ -7,7 +7,7 @@ SourceCulture=en
 CulturesToGenerate=en
 CulturesToGenerate=ja
 CulturesToGenerate=ko
-CulturesToGenerate=zh-CN
+CulturesToGenerate=zh-Hans-CN
 CulturesToGenerate=es
 
 ;Gather text from source code
diff --git a/Engine/Config/Localization/PortableObjectExport.ini b/Engine/Config/Localization/PortableObjectExport.ini
--- a/Engine/Config/Localization/PortableObjectExport.ini
+++ b/Engine/Config/Localization/PortableObjectExport.ini
@@ -3,7 +3,7 @@ CulturesToGenerate=en
 CulturesToGenerate=es
 CulturesToGenerate=ja
 CulturesToGenerate=ko
-CulturesToGenerate=zh-CN
+CulturesToGenerate=zh-Hans-CN
 
 [GatherTextStep0]
 CommandletClass=InternationalizationExport
diff --git a/Engine/Config/Localization/PortableObjectImport.ini b/Engine/Config/Localization/PortableObjectImport.ini
--- a/Engine/Config/Localization/PortableObjectImport.ini
+++ b/Engine/Config/Localization/PortableObjectImport.ini
@@ -3,7 +3,7 @@ CulturesToGenerate=en
 CulturesToGenerate=es
 CulturesToGenerate=ja
 CulturesToGenerate=ko
-CulturesToGenerate=zh-CN
+CulturesToGenerate=zh-Hans-CN
 
 [GatherTextStep0]
 CommandletClass=InternationalizationExport
diff --git a/Engine/Config/Localization/PropertyNames.ini b/Engine/Config/Localization/PropertyNames.ini
--- a/Engine/Config/Localization/PropertyNames.ini
+++ b/Engine/Config/Localization/PropertyNames.ini
@@ -8,7 +8,7 @@ CulturesToGenerate=es
 CulturesToGenerate=en
 CulturesToGenerate=ja
 CulturesToGenerate=ko
-CulturesToGenerate=zh-CN
+CulturesToGenerate=zh-Hans-CN
 
 ;Gather text from metadata
 [GatherTextStep0]
diff --git a/Engine/Config/Localization/ResourceFileGen.ini b/Engine/Config/Localization/ResourceFileGen.ini
--- a/Engine/Config/Localization/ResourceFileGen.ini
+++ b/Engine/Config/Localization/ResourceFileGen.ini
@@ -8,7 +8,7 @@ CulturesToGenerate=en
 CulturesToGenerate=es
 CulturesToGenerate=ja
 CulturesToGenerate=ko
-CulturesToGenerate=zh-CN
+CulturesToGenerate=zh-Hans-CN
 ResourceName=Engine.locres
 
 ;Write Localized Text Resource for Editor
@@ -21,7 +21,7 @@ CulturesToGenerate=en
 CulturesToGenerate=es
 CulturesToGenerate=ja
 CulturesToGenerate=ko
-CulturesToGenerate=zh-CN
+CulturesToGenerate=zh-Hans-CN
 ResourceName=Editor.locres
 
 ;Write Localized Text Resource for Editor Tool Tips
@@ -34,7 +34,7 @@ CulturesToGenerate=en
 CulturesToGenerate=es
 CulturesToGenerate=ja
 CulturesToGenerate=ko
-CulturesToGenerate=zh-CN
+CulturesToGenerate=zh-Hans-CN
 ResourceName=ToolTips.locres
 
 ;Write Localized Text Resource for Editor Property Names
@@ -47,5 +47,5 @@ CulturesToGenerate=en
 CulturesToGenerate=es
 CulturesToGenerate=ja
 CulturesToGenerate=ko
-CulturesToGenerate=zh-CN
+CulturesToGenerate=zh-Hans-CN
 ResourceName=PropertyNames.locres
\ No newline at end of file
diff --git a/Engine/Config/Localization/ToolTips.ini b/Engine/Config/Localization/ToolTips.ini
--- a/Engine/Config/Localization/ToolTips.ini
+++ b/Engine/Config/Localization/ToolTips.ini
@@ -8,7 +8,7 @@ CulturesToGenerate=es
 CulturesToGenerate=en
 CulturesToGenerate=ja
 CulturesToGenerate=ko
-CulturesToGenerate=zh-CN
+CulturesToGenerate=zh-Hans-CN
 
 ;Gather text from metadata
 [GatherTextStep0]
diff --git a/Engine/Config/Localization/WordCount.ini b/Engine/Config/Localization/WordCount.ini
--- a/Engine/Config/Localization/WordCount.ini
+++ b/Engine/Config/Localization/WordCount.ini
@@ -10,7 +10,7 @@ ManifestName=Engine.manifest
 CulturesToGenerate=es
 CulturesToGenerate=ja
 CulturesToGenerate=ko
-CulturesToGenerate=zh-CN
+CulturesToGenerate=zh-Hans-CN
 WordCountReportName=Engine.csv
 
 ;Write Text Localization Report for Editor
@@ -22,7 +22,7 @@ ManifestName=Editor.manifest
 CulturesToGenerate=es
 CulturesToGenerate=ja
 CulturesToGenerate=ko
-CulturesToGenerate=zh-CN
+CulturesToGenerate=zh-Hans-CN
 WordCountReportName=Editor.csv
 
 ;Write Text Localization Report for Editor Tool Tips
@@ -34,7 +34,7 @@ ManifestName=ToolTips.manifest
 CulturesToGenerate=es
 CulturesToGenerate=ja
 CulturesToGenerate=ko
-CulturesToGenerate=zh-CN
+CulturesToGenerate=zh-Hans-CN
 WordCountReportName=ToolTips.csv
 
 ;Write Text Localization Report for Editor Property Names
@@ -46,5 +46,5 @@ ManifestName=PropertyNames.manifest
 CulturesToGenerate=es
 CulturesToGenerate=ja
 CulturesToGenerate=ko
-CulturesToGenerate=zh-CN
+CulturesToGenerate=zh-Hans-CN
 WordCountReportName=PropertyNames.csv
\ No newline at end of file
diff --git a/Engine/Source/Runtime/Core/Private/Internationalization/Culture.cpp b/Engine/Source/Runtime/Core/Private/Internationalization/Culture.cpp
--- a/Engine/Source/Runtime/Core/Private/Internationalization/Culture.cpp
+++ b/Engine/Source/Runtime/Core/Private/Internationalization/Culture.cpp
@@ -49,6 +49,11 @@ FString FCulture::GetName() const
 	return Implementation->GetName();
 }
 
+FString FCulture::GetParentName(const FString& CultureName)
+{
+	return FImplementation::GetParentName(CultureName);
+}
+
 FString FCulture::GetNativeName() const
 {
 	return Implementation->GetNativeName();
diff --git a/Engine/Source/Runtime/Core/Private/Internationalization/ICUCulture.cpp b/Engine/Source/Runtime/Core/Private/Internationalization/ICUCulture.cpp
--- a/Engine/Source/Runtime/Core/Private/Internationalization/ICUCulture.cpp
+++ b/Engine/Source/Runtime/Core/Private/Internationalization/ICUCulture.cpp
@@ -107,6 +107,14 @@ FString FCulture::FICUCultureImplementation::GetName() const
 	return ICULocale.getName();
 }
 
+FString FCulture::FICUCultureImplementation::GetParentName(const FString& CultureName)
+{
+	char Parent[32];
+	UErrorCode ICUStatus = U_ZERO_ERROR;
+	int32_t ParentLength = uloc_getParent(TCHAR_TO_ANSI(*CultureName), Parent, 32, &ICUStatus);
+	return FString(Parent);
+}
+
 FString FCulture::FICUCultureImplementation::GetNativeName() const
 {
 	icu::UnicodeString ICUResult;
diff --git a/Engine/Source/Runtime/Core/Private/Internationalization/ICUCulture.h b/Engine/Source/Runtime/Core/Private/Internationalization/ICUCulture.h
--- a/Engine/Source/Runtime/Core/Private/Internationalization/ICUCulture.h
+++ b/Engine/Source/Runtime/Core/Private/Internationalization/ICUCulture.h
@@ -116,6 +116,8 @@ class FCulture::FICUCultureImplementation
 
 	FString GetName() const;
 
+	static FString GetParentName(const FString& CultureName);
+
 	FString GetNativeName() const;
 
 	FString GetNativeLanguage() const;
diff --git a/Engine/Source/Runtime/Core/Private/Internationalization/Internationalization.cpp b/Engine/Source/Runtime/Core/Private/Internationalization/Internationalization.cpp
--- a/Engine/Source/Runtime/Core/Private/Internationalization/Internationalization.cpp
+++ b/Engine/Source/Runtime/Core/Private/Internationalization/Internationalization.cpp
@@ -566,7 +566,7 @@ void FInternationalization::GetCultureNames(TArray<FString>& CultureNames) const
 	}
 }
 
-void FInternationalization::GetCulturesWithAvailableLocalization(const TArray<FString>& InLocalizationPaths, TArray< TSharedPtr<FCulture, ESPMode::ThreadSafe> >& OutAvailableCultures, const bool bAllowFallback) const
+void FInternationalization::GetCulturesWithAvailableLocalization(const TArray<FString>& InLocalizationPaths, TArray< TSharedPtr<FCulture, ESPMode::ThreadSafe> >& OutAvailableCultures, const bool bIncludeDerivedCultures) const
 {
 	OutAvailableCultures.Reset();
 
@@ -604,13 +604,28 @@ void FInternationalization::GetCulturesWithAvailableLocalization(const TArray<FS
 		FileManager.IterateDirectory(*LocalizationPath, CultureEnumeratorVistor);
 	}
 
-	// Find any cultures that are a complete or partial match for the languages we have translations for
-	for(const auto& Culture : AllCultures)
+	// Find any cultures that are a partial match for those we have translations for.
+	if(bIncludeDerivedCultures)
+	{
+		for(const auto& Culture : AllCultures)
+		{
+			for(FString ParentCultureName = Culture->GetName(); !ParentCultureName.IsEmpty(); ParentCultureName = FCulture::GetParentName(ParentCultureName))
+			{
+				if(AllLocalizationFolders.Contains(ParentCultureName))
+				{
+					OutAvailableCultures.AddUnique(Culture);
+					break;
+				}
+			}
+		}
+	}
+	// Find any cultures that are a complete match for those we have translations for.
+	else
 	{
 		for(const FString& LocalizationFolder : AllLocalizationFolders)
 		{
-			// Check the full name, but if it doesn't match, see if the base language is present
-			if(LocalizationFolder == Culture->GetName() || (bAllowFallback && LocalizationFolder == Culture->GetTwoLetterISOLanguageName()))
+			TSharedPtr<FCulture, ESPMode::ThreadSafe> Culture = GetCulture(LocalizationFolder);
+			if(Culture.IsValid() && AllCultures.Contains(Culture.ToSharedRef()))
 			{
 				OutAvailableCultures.AddUnique(Culture);
 			}
diff --git a/Engine/Source/Runtime/Core/Private/Internationalization/LegacyCulture.cpp b/Engine/Source/Runtime/Core/Private/Internationalization/LegacyCulture.cpp
--- a/Engine/Source/Runtime/Core/Private/Internationalization/LegacyCulture.cpp
+++ b/Engine/Source/Runtime/Core/Private/Internationalization/LegacyCulture.cpp
@@ -43,6 +43,19 @@ FString FCulture::FLegacyCultureImplementation::GetName() const
 	return Name;
 }
 
+FString FCulture::FLegacyCultureImplementation::GetParentName(const FString& CultureName)
+{
+	FString Result;
+
+	int32 UnderscoreIndex = INDEX_NONE;
+	if(Name.FindLastChar(TEXT('_'), UnderscoreIndex))
+	{
+		Result = FString(CultureName).RemoveAt(UnderscoreIndex);
+	}
+
+	return Result;
+}
+
 FString FCulture::FLegacyCultureImplementation::GetNativeName() const
 {
 	return NativeName;
diff --git a/Engine/Source/Runtime/Core/Private/Internationalization/LegacyCulture.h b/Engine/Source/Runtime/Core/Private/Internationalization/LegacyCulture.h
--- a/Engine/Source/Runtime/Core/Private/Internationalization/LegacyCulture.h
+++ b/Engine/Source/Runtime/Core/Private/Internationalization/LegacyCulture.h
@@ -28,6 +28,8 @@ class FCulture::FLegacyCultureImplementation
 
 	FString GetName() const;
 
+	static FString GetParentName(const FString& CultureName);
+
 	FString GetNativeName() const;
 
 	FString GetNativeLanguage() const;
diff --git a/Engine/Source/Runtime/Core/Private/Internationalization/TextLocalizationManager.cpp b/Engine/Source/Runtime/Core/Private/Internationalization/TextLocalizationManager.cpp
--- a/Engine/Source/Runtime/Core/Private/Internationalization/TextLocalizationManager.cpp
+++ b/Engine/Source/Runtime/Core/Private/Internationalization/TextLocalizationManager.cpp
@@ -36,19 +36,19 @@ void EndInitTextLocalization()
 	else
 #endif
 	{
-		FString TargetCultureName;
-		if (FParse::Value(FCommandLine::Get(), TEXT("CULTUREFORCOOKING="), TargetCultureName))
+		FString RequestedCultureName;
+		if (FParse::Value(FCommandLine::Get(), TEXT("CULTUREFORCOOKING="), RequestedCultureName))
 		{
 			// Write the culture passed in if first install...
 			if (FParse::Param(FCommandLine::Get(), TEXT("firstinstall")))
 			{
-				GConfig->SetString(TEXT("Internationalization"), TEXT("Culture"), *TargetCultureName, GEngineIni);
+				GConfig->SetString(TEXT("Internationalization"), TEXT("Culture"), *RequestedCultureName, GEngineIni);
 			}
 		}
 		else
 #if !UE_BUILD_SHIPPING
 		// Use culture override specified on commandline.
-		if (FParse::Value(FCommandLine::Get(), TEXT("CULTURE="), TargetCultureName))
+		if (FParse::Value(FCommandLine::Get(), TEXT("CULTURE="), RequestedCultureName))
 		{
 			//UE_LOG(LogInit, Log, TEXT("Overriding culture %s w/ command-line option".), *CultureName);
 		}
@@ -56,66 +56,65 @@ void EndInitTextLocalization()
 #endif // !UE_BUILD_SHIPPING
 #if WITH_EDITOR
 		// See if we've been provided a culture override in the editor
-		if(GConfig->GetString( TEXT("Internationalization"), TEXT("Culture"), TargetCultureName, GEditorGameAgnosticIni ))
+		if(GConfig->GetString( TEXT("Internationalization"), TEXT("Culture"), RequestedCultureName, GEditorGameAgnosticIni ))
 		{
 			//UE_LOG(LogInit, Log, TEXT("Overriding culture %s w/ editor configuration."), *CultureName);
 		}
 		else
 #endif // WITH_EDITOR
 		// Use culture specified in engine configuration.
-		if(GConfig->GetString( TEXT("Internationalization"), TEXT("Culture"), TargetCultureName, GEngineIni ))
+		if(GConfig->GetString( TEXT("Internationalization"), TEXT("Culture"), RequestedCultureName, GEngineIni ))
 		{
 			//UE_LOG(LogInit, Log, TEXT("Overriding culture %s w/ engine configuration."), *CultureName);
 		}
 		else
 		{
-			TargetCultureName = I18N.GetDefaultCulture()->GetName();
+			RequestedCultureName = I18N.GetDefaultCulture()->GetName();
 		}
 
+		FString TargetCultureName = RequestedCultureName;
 		{
 			TArray<FString> LocalizationPaths;
-
 			if(ShouldLoadEditor)
 			{
 				LocalizationPaths += FPaths::GetEditorLocalizationPaths();
 			}
-
 			if(ShouldLoadGame)
 			{
 				LocalizationPaths += FPaths::GetGameLocalizationPaths();
 			}
-
 			LocalizationPaths += FPaths::GetEngineLocalizationPaths();
 
 			TArray< TSharedPtr<FCulture, ESPMode::ThreadSafe> > AvailableCultures;
-			I18N.GetCulturesWithAvailableLocalization(LocalizationPaths, AvailableCultures);
+			I18N.GetCulturesWithAvailableLocalization(LocalizationPaths, AvailableCultures, false);
 
-			const FString RequestedCultureName = TargetCultureName;
-			TSharedPtr<FCulture, ESPMode::ThreadSafe> RequestedCulture = I18N.GetCulture(TargetCultureName);
-			if(RequestedCulture.IsValid())
+			// Query for parent culture name based on request culture name.
+			FString ParentCultureName = RequestedCultureName;
+				
+			// If we do not have localization data, try the parent culture.
+			for(TSharedPtr<FCulture, ESPMode::ThreadSafe> ParentCulture = I18N.GetCulture(ParentCultureName); !AvailableCultures.Contains(ParentCulture); ParentCulture = I18N.GetCulture(ParentCultureName))
 			{
-				if(!AvailableCultures.Contains(RequestedCulture))
+				ParentCultureName = FCulture::GetParentName(ParentCultureName);
+
+				// No parent culture
+				if(ParentCultureName.IsEmpty())
 				{
-					// Try base language.
-					TargetCultureName = RequestedCulture->GetTwoLetterISOLanguageName();
-					TSharedPtr<FCulture, ESPMode::ThreadSafe> TargetCulture = I18N.GetCulture(TargetCultureName);
-					if(AvailableCultures.Contains(TargetCulture))
-					{
-						// Fallback to base language.
-						UE_LOG(LogTextLocalizationManager, Log, TEXT("The selected culture '%s' has no localization data; falling back to base language '%s'"), *RequestedCultureName, *TargetCultureName);
-					}
-					else
-					{
-						// Fallback to English.
-						UE_LOG(LogTextLocalizationManager, Log, TEXT("The selected culture '%s' has no localization data; falling back to 'en'"), *RequestedCultureName);
-						TargetCultureName = "en";
-					}
+					break;
+				}
+			}
+
+			if(!ParentCultureName.IsEmpty())
+			{
+				if(RequestedCultureName != ParentCultureName)
+				{
+					// Make the user aware that the localization data belongs to a parent culture.
+					UE_LOG(LogTextLocalizationManager, Log, TEXT("The requested culture ('%s') has no localization data; parent culture's ('%s') localization data will be used."), *RequestedCultureName, *ParentCultureName);
 				}
 			}
-			// Fallback to English.
 			else
 			{
-				UE_LOG(LogTextLocalizationManager, Warning, TEXT("The selected culture '%s' is not valid; falling back to 'en'"), *RequestedCultureName);
+				// Fallback to English.
+				UE_LOG(LogTextLocalizationManager, Log, TEXT("The requested culture ('%s') has no localization data; falling back to 'en' for localization and internationalization data."), *RequestedCultureName);
 				TargetCultureName = "en";
 			}
 		}
diff --git a/Engine/Source/Runtime/Core/Public/Internationalization/Culture.h b/Engine/Source/Runtime/Core/Public/Internationalization/Culture.h
--- a/Engine/Source/Runtime/Core/Public/Internationalization/Culture.h
+++ b/Engine/Source/Runtime/Core/Public/Internationalization/Culture.h
@@ -41,6 +41,8 @@ public:
 
 	FString GetName() const;
 
+	static FString GetParentName(const FString& CultureName);
+	
 	FString GetNativeName() const;
 
 	FString GetNativeLanguage() const;
@@ -60,10 +62,12 @@ public:
 #if UE_ENABLE_ICU
 public:
 	class FICUCultureImplementation;
+	typedef FICUCultureImplementation FImplementation;
 	TSharedRef<FICUCultureImplementation> Implementation;
 #else
 public:
 	class FLegacyCultureImplementation;
+	typedef FLegacyCultureImplementation FImplementation;
 	TSharedRef<FLegacyCultureImplementation> Implementation;
 
 	// The date formating rule for this culture
diff --git a/Engine/Source/Runtime/Core/Public/Internationalization/Internationalization.h b/Engine/Source/Runtime/Core/Public/Internationalization/Internationalization.h
--- a/Engine/Source/Runtime/Core/Public/Internationalization/Internationalization.h
+++ b/Engine/Source/Runtime/Core/Public/Internationalization/Internationalization.h
@@ -68,8 +68,8 @@ public:
 
 	CORE_API void GetCultureNames(TArray<FString>& CultureNames) const;
 
-	// Given some paths to look at, populate a list of cultures that we have available localization information for. If bAllowFallback, include cultures for which we have non-region-specific data for.
-	CORE_API void GetCulturesWithAvailableLocalization(const TArray<FString>& InLocalizationPaths, TArray< TSharedPtr<FCulture, ESPMode::ThreadSafe> >& OutAvailableCultures, const bool bAllowFallback = false) const;
+	// Given some paths to look at, populate a list of cultures that we have available localization information for. If bIncludeDerivedCultures, include cultures that are derived from those we have localization data for.
+	CORE_API void GetCulturesWithAvailableLocalization(const TArray<FString>& InLocalizationPaths, TArray< TSharedPtr<FCulture, ESPMode::ThreadSafe> >& OutAvailableCultures, const bool bIncludeDerivedCultures) const;
 
 private:
 	FInternationalization();
