Fixed compilation error in LegacyCulture.cpp regarding using a member variable in a static method instead of the argument passed in.

diff --git a/Engine/Source/Runtime/Core/Private/Internationalization/LegacyCulture.cpp b/Engine/Source/Runtime/Core/Private/Internationalization/LegacyCulture.cpp
--- a/Engine/Source/Runtime/Core/Private/Internationalization/LegacyCulture.cpp
+++ b/Engine/Source/Runtime/Core/Private/Internationalization/LegacyCulture.cpp
@@ -48,7 +48,7 @@ FString FCulture::FLegacyCultureImplementation::GetParentName(const FString& Cul
 	FString Result;
 
 	int32 UnderscoreIndex = INDEX_NONE;
-	if(Name.FindLastChar(TEXT('_'), UnderscoreIndex))
+	if(CultureName.FindLastChar(TEXT('_'), UnderscoreIndex))
 	{
 		Result = FString(CultureName).RemoveAt(UnderscoreIndex);
 	}
