Fixed improper assignment in LegacyCulture.cpp in GetParentName.

diff --git a/Engine/Source/Runtime/Core/Private/Internationalization/LegacyCulture.cpp b/Engine/Source/Runtime/Core/Private/Internationalization/LegacyCulture.cpp
--- a/Engine/Source/Runtime/Core/Private/Internationalization/LegacyCulture.cpp
+++ b/Engine/Source/Runtime/Core/Private/Internationalization/LegacyCulture.cpp
@@ -50,7 +50,8 @@ FString FCulture::FLegacyCultureImplementation::GetParentName(const FString& Cul
 	int32 UnderscoreIndex = INDEX_NONE;
 	if(CultureName.FindLastChar(TEXT('_'), UnderscoreIndex))
 	{
-		Result = FString(CultureName).RemoveAt(UnderscoreIndex);
+		Result = CultureName;
+		Result.RemoveAt(UnderscoreIndex);
 	}
 
 	return Result;
