Made regular modal windows on Mac use NSNormalWindowLevel (to fix issues with tooltips) and not disappear when the app is in background.

diff --git a/Engine/Source/Runtime/Core/Private/Mac/MacWindow.cpp b/Engine/Source/Runtime/Core/Private/Mac/MacWindow.cpp
--- a/Engine/Source/Runtime/Core/Private/Mac/MacWindow.cpp
+++ b/Engine/Source/Runtime/Core/Private/Mac/MacWindow.cpp
@@ -341,11 +341,6 @@ TArray< FSlateCocoaWindow* > GRunningModalWindows;
 {
 	if([NSApp isHidden] == NO)
 	{
-		bool bModal = GRunningModalWindows.Num() != 0 && GRunningModalWindows.Last() == self;
-		if(bModal)
-		{
-			[self setLevel:NSModalPanelWindowLevel];
-		}
 		if(GRunningModalWindows.Num() == 0 || GRunningModalWindows.Last() == self || [self styleMask] == NSBorderlessWindowMask)
 		{
 			[self orderFrontAndMakeMain:false andKey:false];
@@ -367,11 +362,6 @@ TArray< FSlateCocoaWindow* > GRunningModalWindows;
 	[self setMovable: YES];
 	[self setMovableByWindowBackground: NO];
 	
-	if([self level] == NSModalPanelWindowLevel)
-	{
-		[self setLevel:NSNormalWindowLevel];
-	}
-	
 	if(self.bForwardEvents)
 	{
 		MacApplication->OnWindowDidResignKey( self );
@@ -740,7 +730,7 @@ void FMacWindow::Initialize( FMacApplication* const Application, const TSharedRe
 	[WindowHandle setDisplayReconfiguring: false];
 	[WindowHandle setAcceptsMouseMovedEvents: YES];
 	[WindowHandle setDelegate: WindowHandle];
-	if(Definition->IsRegularWindow && !InParent.IsValid())
+	if(Definition->IsRegularWindow && (!InParent.IsValid() || Definition->IsModalWindow))
 	{
 		[WindowHandle setLevel: NSNormalWindowLevel];
 	}
