Fixed negative currency formatting pattern in legacy culture number formatting.

diff --git a/Engine/Source/Runtime/Core/Private/Internationalization/Cultures/AmericanEnglish.h b/Engine/Source/Runtime/Core/Private/Internationalization/Cultures/AmericanEnglish.h
--- a/Engine/Source/Runtime/Core/Private/Internationalization/Cultures/AmericanEnglish.h
+++ b/Engine/Source/Runtime/Core/Private/Internationalization/Cultures/AmericanEnglish.h
@@ -43,7 +43,7 @@ public:
 			, "%"					//const FString PercentSymbol
 			, "‰"					//const FString PerMilleSymbol
 			, ""					//const FString PositiveSign
-			, ECurrencyNegativePattern::BracketCurrencySymbolXBracket
+			, ECurrencyNegativePattern::NegativeSymbolCurrencySymbolX
 			, ECurrencyPositivePattern::CurrencySymbolX
 			, EPercentNegativePattern::NegativeSymbolXPct
 			, EPercentPositivePattern::XPct
