Giving the http request class an accessor for the response class

diff --git a/Engine/Source/Runtime/Core/Private/GenericPlatform/GenericPlatformHttp.cpp b/Engine/Source/Runtime/Core/Private/GenericPlatform/GenericPlatformHttp.cpp
--- a/Engine/Source/Runtime/Core/Private/GenericPlatform/GenericPlatformHttp.cpp
+++ b/Engine/Source/Runtime/Core/Private/GenericPlatform/GenericPlatformHttp.cpp
@@ -33,6 +33,7 @@ public:
 	virtual FHttpRequestProgressDelegate& OnRequestProgress() override { static FHttpRequestProgressDelegate RequestProgressDelegate; return RequestProgressDelegate; }
 	virtual void CancelRequest() override {}
 	virtual EHttpRequestStatus::Type GetStatus() override { return EHttpRequestStatus::NotStarted; }
+	virtual const FHttpResponsePtr GetResponse() const override { return nullptr; }
 	virtual void Tick(float DeltaSeconds) override {}
 };
 
diff --git a/Engine/Source/Runtime/Online/HTTP/Private/Curl/CurlHttp.cpp b/Engine/Source/Runtime/Online/HTTP/Private/Curl/CurlHttp.cpp
--- a/Engine/Source/Runtime/Online/HTTP/Private/Curl/CurlHttp.cpp
+++ b/Engine/Source/Runtime/Online/HTTP/Private/Curl/CurlHttp.cpp
@@ -620,6 +620,11 @@ EHttpRequestStatus::Type FCurlHttpRequest::GetStatus()
 	return CompletionStatus;
 }
 
+const FHttpResponsePtr FCurlHttpRequest::GetResponse() const
+{
+	return Response;
+}
+
 void FCurlHttpRequest::Tick(float DeltaSeconds)
 {
 	// check for true completion/cancellation
diff --git a/Engine/Source/Runtime/Online/HTTP/Private/Curl/CurlHttp.h b/Engine/Source/Runtime/Online/HTTP/Private/Curl/CurlHttp.h
--- a/Engine/Source/Runtime/Online/HTTP/Private/Curl/CurlHttp.h
+++ b/Engine/Source/Runtime/Online/HTTP/Private/Curl/CurlHttp.h
@@ -128,6 +128,7 @@ public:
 	virtual FHttpRequestProgressDelegate& OnRequestProgress() override;
 	virtual void CancelRequest() override;
 	virtual EHttpRequestStatus::Type GetStatus() override;
+	virtual const FHttpResponsePtr GetResponse() const override;
 	virtual void Tick(float DeltaSeconds) override;
 	// End IHttpRequest interface
 
diff --git a/Engine/Source/Runtime/Online/HTTP/Private/IOS/IOSHTTP.cpp b/Engine/Source/Runtime/Online/HTTP/Private/IOS/IOSHTTP.cpp
--- a/Engine/Source/Runtime/Online/HTTP/Private/IOS/IOSHTTP.cpp
+++ b/Engine/Source/Runtime/Online/HTTP/Private/IOS/IOSHTTP.cpp
@@ -301,6 +301,12 @@ EHttpRequestStatus::Type FIOSHttpRequest::GetStatus()
 }
 
 
+const FHttpResponsePtr FIOSHttpRequest::GetResponse() const
+{
+	return Response;
+}
+
+
 void FIOSHttpRequest::Tick(float DeltaSeconds) 
 {
 	if(CompletionStatus == EHttpRequestStatus::Processing)
diff --git a/Engine/Source/Runtime/Online/HTTP/Private/IOS/IOSHTTP.h b/Engine/Source/Runtime/Online/HTTP/Private/IOS/IOSHTTP.h
--- a/Engine/Source/Runtime/Online/HTTP/Private/IOS/IOSHTTP.h
+++ b/Engine/Source/Runtime/Online/HTTP/Private/IOS/IOSHTTP.h
@@ -43,6 +43,7 @@ public:
 	virtual FHttpRequestProgressDelegate& OnRequestProgress() override;
 	virtual void CancelRequest() override;
 	virtual EHttpRequestStatus::Type GetStatus() override;
+	virtual const FHttpResponsePtr GetResponse() const override;
 	virtual void Tick(float DeltaSeconds) override;
 	// End IHttpRequest interface
 
diff --git a/Engine/Source/Runtime/Online/HTTP/Private/Mac/MacHTTP.cpp b/Engine/Source/Runtime/Online/HTTP/Private/Mac/MacHTTP.cpp
--- a/Engine/Source/Runtime/Online/HTTP/Private/Mac/MacHTTP.cpp
+++ b/Engine/Source/Runtime/Online/HTTP/Private/Mac/MacHTTP.cpp
@@ -312,6 +312,12 @@ EHttpRequestStatus::Type FMacHttpRequest::GetStatus()
 }
 
 
+const FHttpResponsePtr FMacHttpRequest::GetResponse() const
+{
+	return Response;
+}
+
+
 void FMacHttpRequest::Tick(float DeltaSeconds)
 {
 	if( CompletionStatus == EHttpRequestStatus::Processing || Response->HadError() )
diff --git a/Engine/Source/Runtime/Online/HTTP/Private/Mac/MacHTTP.h b/Engine/Source/Runtime/Online/HTTP/Private/Mac/MacHTTP.h
--- a/Engine/Source/Runtime/Online/HTTP/Private/Mac/MacHTTP.h
+++ b/Engine/Source/Runtime/Online/HTTP/Private/Mac/MacHTTP.h
@@ -43,6 +43,7 @@ public:
 	virtual FHttpRequestProgressDelegate& OnRequestProgress() override;
 	virtual void CancelRequest() override;
 	virtual EHttpRequestStatus::Type GetStatus() override;
+	virtual const FHttpResponsePtr GetResponse() const override;
 	virtual void Tick(float DeltaSeconds) override;
 	// End IHttpRequest interface
 
diff --git a/Engine/Source/Runtime/Online/HTTP/Private/Windows/HttpWinInet.cpp b/Engine/Source/Runtime/Online/HTTP/Private/Windows/HttpWinInet.cpp
--- a/Engine/Source/Runtime/Online/HTTP/Private/Windows/HttpWinInet.cpp
+++ b/Engine/Source/Runtime/Online/HTTP/Private/Windows/HttpWinInet.cpp
@@ -724,6 +724,11 @@ EHttpRequestStatus::Type FHttpRequestWinInet::GetStatus()
 	return CompletionStatus;
 }
 
+const FHttpResponsePtr FHttpRequestWinInet::GetResponse() const
+{
+	return Response;
+}
+
 void FHttpRequestWinInet::Tick(float DeltaSeconds)
 {
 	// keep track of elapsed milliseconds
diff --git a/Engine/Source/Runtime/Online/HTTP/Private/Windows/HttpWinInet.h b/Engine/Source/Runtime/Online/HTTP/Private/Windows/HttpWinInet.h
--- a/Engine/Source/Runtime/Online/HTTP/Private/Windows/HttpWinInet.h
+++ b/Engine/Source/Runtime/Online/HTTP/Private/Windows/HttpWinInet.h
@@ -91,6 +91,7 @@ public:
 	virtual FHttpRequestProgressDelegate& OnRequestProgress() override;
 	virtual void CancelRequest() override;
 	virtual EHttpRequestStatus::Type GetStatus() override;
+	virtual const FHttpResponsePtr GetResponse() const override;
 	virtual void Tick(float DeltaSeconds) override;
 
 	// FHttpRequestWinInet
diff --git a/Engine/Source/Runtime/Online/HTTP/Public/Interfaces/IHttpRequest.h b/Engine/Source/Runtime/Online/HTTP/Public/Interfaces/IHttpRequest.h
--- a/Engine/Source/Runtime/Online/HTTP/Public/Interfaces/IHttpRequest.h
+++ b/Engine/Source/Runtime/Online/HTTP/Public/Interfaces/IHttpRequest.h
@@ -156,6 +156,13 @@ public:
 	 */
 	virtual EHttpRequestStatus::Type GetStatus() = 0;
 
+	/**
+	 * Get the associated Response
+	 *
+	 * @return the response
+	 */
+	virtual const FHttpResponsePtr GetResponse() const = 0;
+
 	/**
 	 * Used to tick the request
 	 *
