Added the option to specify a currency code when formatting a number as currency.

diff --git a/Engine/Source/Runtime/Core/Private/Internationalization/ICUCulture.cpp b/Engine/Source/Runtime/Core/Private/Internationalization/ICUCulture.cpp
--- a/Engine/Source/Runtime/Core/Private/Internationalization/ICUCulture.cpp
+++ b/Engine/Source/Runtime/Core/Private/Internationalization/ICUCulture.cpp
@@ -239,10 +239,11 @@ TSharedRef<const icu::DecimalFormat> FCulture::FICUCultureImplementation::GetDec
 	}
 }
 
-TSharedRef<const icu::DecimalFormat> FCulture::FICUCultureImplementation::GetCurrencyFormatter(const FNumberFormattingOptions* const Options) const
+TSharedRef<const icu::DecimalFormat> FCulture::FICUCultureImplementation::GetCurrencyFormatter(const FString& CurrencyCode, const FNumberFormattingOptions* const Options) const
 {
-	const bool bIsDefault = Options == NULL;
+	const bool bIsDefault = Options == NULL && CurrencyCode.IsEmpty();
 	const TSharedRef<const icu::DecimalFormat> DefaultFormatter( ICUCurrencyFormat );
+
 	if(bIsDefault)
 	{
 		return DefaultFormatter;
@@ -250,6 +251,14 @@ TSharedRef<const icu::DecimalFormat> FCulture::FICUCultureImplementation::GetCur
 	else
 	{
 		const TSharedRef<icu::DecimalFormat> Formatter( static_cast<icu::DecimalFormat*>(DefaultFormatter->clone()) );
+		
+		if (!CurrencyCode.IsEmpty())
+		{
+			icu::UnicodeString ICUCurrencyCode;
+			ICUUtilities::ConvertString(CurrencyCode, ICUCurrencyCode);
+			Formatter->setCurrency(ICUCurrencyCode.getBuffer());
+		}
+
 		if(Options)
 		{
 			Formatter->setGroupingUsed(Options->UseGrouping);
@@ -259,6 +268,7 @@ TSharedRef<const icu::DecimalFormat> FCulture::FICUCultureImplementation::GetCur
 			Formatter->setMinimumFractionDigits(Options->MinimumFractionalDigits);
 			Formatter->setMaximumFractionDigits(Options->MaximumFractionalDigits);
 		}
+
 		return Formatter;
 	}
 }
diff --git a/Engine/Source/Runtime/Core/Private/Internationalization/ICUCulture.h b/Engine/Source/Runtime/Core/Private/Internationalization/ICUCulture.h
--- a/Engine/Source/Runtime/Core/Private/Internationalization/ICUCulture.h
+++ b/Engine/Source/Runtime/Core/Private/Internationalization/ICUCulture.h
@@ -136,7 +136,7 @@ class FCulture::FICUCultureImplementation
 
 	TSharedRef<const icu::Collator, ESPMode::ThreadSafe> GetCollator(const ETextComparisonLevel::Type ComparisonLevel) const;
 	TSharedRef<const icu::DecimalFormat> GetDecimalFormatter(const FNumberFormattingOptions* const Options = NULL) const;
-	TSharedRef<const icu::DecimalFormat> GetCurrencyFormatter(const FNumberFormattingOptions* const Options = NULL) const;
+	TSharedRef<const icu::DecimalFormat> GetCurrencyFormatter(const FString& CurrencyCode = FString(), const FNumberFormattingOptions* const Options = NULL) const;
 	TSharedRef<const icu::DecimalFormat> GetPercentFormatter(const FNumberFormattingOptions* const Options = NULL) const;
 	TSharedRef<const icu::DateFormat> GetDateFormatter(const EDateTimeStyle::Type DateStyle) const;
 	TSharedRef<const icu::DateFormat> GetTimeFormatter(const EDateTimeStyle::Type TimeStyle, const FString& TimeZone) const;
diff --git a/Engine/Source/Runtime/Core/Private/Internationalization/ICUText.h b/Engine/Source/Runtime/Core/Private/Internationalization/ICUText.h
--- a/Engine/Source/Runtime/Core/Private/Internationalization/ICUText.h
+++ b/Engine/Source/Runtime/Core/Private/Internationalization/ICUText.h
@@ -31,13 +31,13 @@ FText FText::AsNumberTemplate(T1 Val, const FNumberFormattingOptions* const Opti
 }
 
 template<typename T1, typename T2>
-FText FText::AsCurrencyTemplate(T1 Val, const FNumberFormattingOptions* const Options, const TSharedPtr<FCulture, ESPMode::ThreadSafe>& TargetCulture)
+FText FText::AsCurrencyTemplate(T1 Val, const FString& CurrencyCode, const FNumberFormattingOptions* const Options, const TSharedPtr<FCulture, ESPMode::ThreadSafe>& TargetCulture)
 {
 	FInternationalization& I18N = FInternationalization::Get();
 	checkf(I18N.IsInitialized() == true, TEXT("FInternationalization is not initialized. An FText formatting method was likely used in static object initialization - this is not supported."));
 	const TSharedRef<FCulture, ESPMode::ThreadSafe> Culture = TargetCulture.IsValid() ? TargetCulture.ToSharedRef() : I18N.GetCurrentCulture();
 	UErrorCode ICUStatus = U_ZERO_ERROR;
-	const TSharedRef<const icu::DecimalFormat> ICUDecimalFormat( Culture->Implementation->GetCurrencyFormatter(Options) );
+	const TSharedRef<const icu::DecimalFormat> ICUDecimalFormat(Culture->Implementation->GetCurrencyFormatter(CurrencyCode, Options));
 	icu::Formattable FormattableVal(static_cast<T2>(Val));
 	icu::UnicodeString FormattedString;
 	ICUDecimalFormat->format(FormattableVal, FormattedString, ICUStatus);
@@ -46,7 +46,7 @@ FText FText::AsCurrencyTemplate(T1 Val, const FNumberFormattingOptions* const Op
 	ICUUtilities::ConvertString(FormattedString, NativeString);
 
 	FText ReturnText = FText::CreateNumericalText(NativeString);
-	ReturnText.History = MakeShareable(new FTextHistory_AsCurrency(Val, Options, TargetCulture));
+	ReturnText.History = MakeShareable(new FTextHistory_AsCurrency(Val, CurrencyCode, Options, TargetCulture));
 	return ReturnText;
 }
 
diff --git a/Engine/Source/Runtime/Core/Private/Internationalization/LegacyText.h b/Engine/Source/Runtime/Core/Private/Internationalization/LegacyText.h
--- a/Engine/Source/Runtime/Core/Private/Internationalization/LegacyText.h
+++ b/Engine/Source/Runtime/Core/Private/Internationalization/LegacyText.h
@@ -12,7 +12,7 @@ FText FText::AsNumberTemplate(T1 Val, const FNumberFormattingOptions* const Opti
 }
 
 template<typename T1, typename T2>
-FText FText::AsCurrencyTemplate(T1 Val, const FNumberFormattingOptions* const Options, const TSharedPtr<FCulture, ESPMode::ThreadSafe>& TargetCulture)
+FText FText::AsCurrencyTemplate(T1 Val, const FString& CurrencyCode, const FNumberFormattingOptions* const Options, const TSharedPtr<FCulture, ESPMode::ThreadSafe>& TargetCulture)
 {
 	FInternationalization& I18N = FInternationalization::Get();
 	checkf(I18N.IsInitialized() == true, TEXT("FInternationalization is not initialized. An FText formatting method was likely used in static object initialization - this is not supported."));
diff --git a/Engine/Source/Runtime/Core/Private/Internationalization/Text.cpp b/Engine/Source/Runtime/Core/Private/Internationalization/Text.cpp
--- a/Engine/Source/Runtime/Core/Private/Internationalization/Text.cpp
+++ b/Engine/Source/Runtime/Core/Private/Internationalization/Text.cpp
@@ -462,12 +462,12 @@ DEF_ASNUMBER(uint16)
 DEF_ASNUMBER_CAST(uint32, int64_t)
 DEF_ASNUMBER_CAST(uint64, int64_t)
 #undef DEF_ASNUMBER
+#undef DEF_ASNUMBER_CAST
 
 /**
-* Generate an FText that represents the passed number as currency in the current culture
-*/
-
-#define DEF_ASCURRENCY_CAST(T1, T2) FText FText::AsCurrency(T1 Val, const FNumberFormattingOptions* const Options, const TSharedPtr<FCulture, ESPMode::ThreadSafe>& TargetCulture) { return FText::AsCurrencyTemplate<T1, T2>(Val, Options, TargetCulture); }
+ * Generate an FText that represents the passed number as currency in the current culture
+ */
+#define DEF_ASCURRENCY_CAST(T1, T2) FText FText::AsCurrency(T1 Val, const FString& CurrencyCode, const FNumberFormattingOptions* const Options, const TSharedPtr<FCulture, ESPMode::ThreadSafe>& TargetCulture) { return FText::AsCurrencyTemplate<T1, T2>(Val, CurrencyCode, Options, TargetCulture); }
 #define DEF_ASCURRENCY(T) DEF_ASCURRENCY_CAST(T, T)
 DEF_ASCURRENCY(float)
 	DEF_ASCURRENCY(double)
@@ -480,6 +480,7 @@ DEF_ASCURRENCY(float)
 	DEF_ASCURRENCY_CAST(uint32, int64_t)
 	DEF_ASCURRENCY_CAST(uint64, int64_t)
 #undef DEF_ASCURRENCY
+#undef DEF_ASCURRENCY_CAST
 
 /**
 * Generate an FText that represents the passed number as a percentage in the current culture
@@ -490,6 +491,7 @@ DEF_ASCURRENCY(float)
 DEF_ASPERCENT(double)
 DEF_ASPERCENT(float)
 #undef DEF_ASPERCENT
+#undef DEF_ASPERCENT_CAST
 
 bool FText::FindText( const FString& Namespace, const FString& Key, FText& OutText, const FString* const SourceString )
 {
diff --git a/Engine/Source/Runtime/Core/Private/Internationalization/TextHistory.cpp b/Engine/Source/Runtime/Core/Private/Internationalization/TextHistory.cpp
--- a/Engine/Source/Runtime/Core/Private/Internationalization/TextHistory.cpp
+++ b/Engine/Source/Runtime/Core/Private/Internationalization/TextHistory.cpp
@@ -405,8 +405,9 @@ void FTextHistory_AsPercent::Serialize( FArchive& Ar )
 ///////////////////////////////////////
 // FTextHistory_AsCurrency
 
-FTextHistory_AsCurrency::FTextHistory_AsCurrency(const FFormatArgumentValue& InSourceValue, const FNumberFormattingOptions* const InFormatOptions, const TSharedPtr<FCulture, ESPMode::ThreadSafe> InTargetCulture)
-	: FTextHistory_FormatNumber(InSourceValue, InFormatOptions, InTargetCulture)
+FTextHistory_AsCurrency::FTextHistory_AsCurrency(const FFormatArgumentValue& InSourceValue, const FString& CurrencyCode, const FNumberFormattingOptions* const InFormatOptions, const TSharedPtr<FCulture, ESPMode::ThreadSafe> InTargetCulture)
+: FTextHistory_FormatNumber(InSourceValue, InFormatOptions, InTargetCulture)
+, CurrencyCode(CurrencyCode)
 {
 }
 
@@ -418,19 +419,19 @@ FText FTextHistory_AsCurrency::ToText(bool bInAsSource) const
 	{
 	case EFormatArgumentType::UInt:
 		{
-			return FText::AsCurrency(SourceValue.UIntValue, FormatOptions, CurrentCulture);
+			return FText::AsCurrency(SourceValue.UIntValue, CurrencyCode, FormatOptions, CurrentCulture);
 		}
 	case EFormatArgumentType::Int:
 		{
-			return FText::AsCurrency(SourceValue.IntValue, FormatOptions, CurrentCulture);
+			return FText::AsCurrency(SourceValue.IntValue, CurrencyCode, FormatOptions, CurrentCulture);
 		}
 	case EFormatArgumentType::Float:
 		{
-			return FText::AsCurrency(SourceValue.FloatValue, FormatOptions, CurrentCulture);
+			return FText::AsCurrency(SourceValue.FloatValue, CurrencyCode, FormatOptions, CurrentCulture);
 		}
 	case EFormatArgumentType::Double:
 		{
-			return FText::AsCurrency(SourceValue.DoubleValue, FormatOptions, CurrentCulture);
+			return FText::AsCurrency(SourceValue.DoubleValue, CurrencyCode, FormatOptions, CurrentCulture);
 		}
 	default:
 		{
@@ -449,6 +450,11 @@ void FTextHistory_AsCurrency::Serialize( FArchive& Ar )
 		Ar << HistoryType;
 	}
 
+	if (Ar.UE4Ver() >= VER_UE4_ADDED_CURRENCY_CODE_TO_FTEXT)
+	{
+		Ar << CurrencyCode;
+	}
+
 	FTextHistory_FormatNumber::Serialize(Ar);
 }
 
diff --git a/Engine/Source/Runtime/Core/Private/Internationalization/TextHistory.h b/Engine/Source/Runtime/Core/Private/Internationalization/TextHistory.h
--- a/Engine/Source/Runtime/Core/Private/Internationalization/TextHistory.h
+++ b/Engine/Source/Runtime/Core/Private/Internationalization/TextHistory.h
@@ -179,12 +179,16 @@ class CORE_API FTextHistory_AsCurrency : public FTextHistory_FormatNumber
 {
 public:
 	FTextHistory_AsCurrency() {};
-	FTextHistory_AsCurrency(const FFormatArgumentValue& InSourceValue, const FNumberFormattingOptions* const InFormatOptions, const TSharedPtr<FCulture, ESPMode::ThreadSafe> InTargetCulture);
+	FTextHistory_AsCurrency(const FFormatArgumentValue& InSourceValue, const FString& CurrencyCode, const FNumberFormattingOptions* const InFormatOptions, const TSharedPtr<FCulture, ESPMode::ThreadSafe> InTargetCulture);
 
 	// Begin FTextHistory interface
 	virtual FText ToText(bool bInAsSource) const override;
 	virtual void Serialize(FArchive& Ar) override;
 	// End FTextHistory interface
+
+private:
+	/** The currency used to format the number. */
+	FString CurrencyCode;
 };
 
 /**  Handles history for formatting using AsDate */
diff --git a/Engine/Source/Runtime/Core/Public/Internationalization/Text.h b/Engine/Source/Runtime/Core/Public/Internationalization/Text.h
--- a/Engine/Source/Runtime/Core/Public/Internationalization/Text.h
+++ b/Engine/Source/Runtime/Core/Public/Internationalization/Text.h
@@ -138,17 +138,17 @@ public:
 	/**
 	 * Generate an FText that represents the passed number as currency in the current culture
 	 */
-	static FText AsCurrency(float Val,	const FNumberFormattingOptions* const Options = NULL, const TSharedPtr<FCulture, ESPMode::ThreadSafe>& TargetCulture = NULL);
-	static FText AsCurrency(double Val,	const FNumberFormattingOptions* const Options = NULL, const TSharedPtr<FCulture, ESPMode::ThreadSafe>& TargetCulture = NULL);
-	static FText AsCurrency(int8 Val,	const FNumberFormattingOptions* const Options = NULL, const TSharedPtr<FCulture, ESPMode::ThreadSafe>& TargetCulture = NULL);
-	static FText AsCurrency(int16 Val,	const FNumberFormattingOptions* const Options = NULL, const TSharedPtr<FCulture, ESPMode::ThreadSafe>& TargetCulture = NULL);
-	static FText AsCurrency(int32 Val,	const FNumberFormattingOptions* const Options = NULL, const TSharedPtr<FCulture, ESPMode::ThreadSafe>& TargetCulture = NULL);
-	static FText AsCurrency(int64 Val,	const FNumberFormattingOptions* const Options = NULL, const TSharedPtr<FCulture, ESPMode::ThreadSafe>& TargetCulture = NULL);
-	static FText AsCurrency(uint8 Val,	const FNumberFormattingOptions* const Options = NULL, const TSharedPtr<FCulture, ESPMode::ThreadSafe>& TargetCulture = NULL);
-	static FText AsCurrency(uint16 Val,	const FNumberFormattingOptions* const Options = NULL, const TSharedPtr<FCulture, ESPMode::ThreadSafe>& TargetCulture = NULL);
-	static FText AsCurrency(uint32 Val,	const FNumberFormattingOptions* const Options = NULL, const TSharedPtr<FCulture, ESPMode::ThreadSafe>& TargetCulture = NULL);
-	static FText AsCurrency(uint64 Val,	const FNumberFormattingOptions* const Options = NULL, const TSharedPtr<FCulture, ESPMode::ThreadSafe>& TargetCulture = NULL);
-	static FText AsCurrency(long Val,	const FNumberFormattingOptions* const Options = NULL, const TSharedPtr<FCulture, ESPMode::ThreadSafe>& TargetCulture = NULL);
+	static FText AsCurrency(float Val,  const FString& CurrencyCode = FString(), const FNumberFormattingOptions* const Options = NULL, const TSharedPtr<FCulture, ESPMode::ThreadSafe>& TargetCulture = NULL);
+	static FText AsCurrency(double Val, const FString& CurrencyCode = FString(), const FNumberFormattingOptions* const Options = NULL, const TSharedPtr<FCulture, ESPMode::ThreadSafe>& TargetCulture = NULL);
+	static FText AsCurrency(int8 Val,   const FString& CurrencyCode = FString(), const FNumberFormattingOptions* const Options = NULL, const TSharedPtr<FCulture, ESPMode::ThreadSafe>& TargetCulture = NULL);
+	static FText AsCurrency(int16 Val,  const FString& CurrencyCode = FString(), const FNumberFormattingOptions* const Options = NULL, const TSharedPtr<FCulture, ESPMode::ThreadSafe>& TargetCulture = NULL);
+	static FText AsCurrency(int32 Val,  const FString& CurrencyCode = FString(), const FNumberFormattingOptions* const Options = NULL, const TSharedPtr<FCulture, ESPMode::ThreadSafe>& TargetCulture = NULL);
+	static FText AsCurrency(int64 Val,  const FString& CurrencyCode = FString(), const FNumberFormattingOptions* const Options = NULL, const TSharedPtr<FCulture, ESPMode::ThreadSafe>& TargetCulture = NULL);
+	static FText AsCurrency(uint8 Val,  const FString& CurrencyCode = FString(), const FNumberFormattingOptions* const Options = NULL, const TSharedPtr<FCulture, ESPMode::ThreadSafe>& TargetCulture = NULL);
+	static FText AsCurrency(uint16 Val, const FString& CurrencyCode = FString(), const FNumberFormattingOptions* const Options = NULL, const TSharedPtr<FCulture, ESPMode::ThreadSafe>& TargetCulture = NULL);
+	static FText AsCurrency(uint32 Val, const FString& CurrencyCode = FString(), const FNumberFormattingOptions* const Options = NULL, const TSharedPtr<FCulture, ESPMode::ThreadSafe>& TargetCulture = NULL);
+	static FText AsCurrency(uint64 Val, const FString& CurrencyCode = FString(), const FNumberFormattingOptions* const Options = NULL, const TSharedPtr<FCulture, ESPMode::ThreadSafe>& TargetCulture = NULL);
+	static FText AsCurrency(long Val,   const FString& CurrencyCode = FString(), const FNumberFormattingOptions* const Options = NULL, const TSharedPtr<FCulture, ESPMode::ThreadSafe>& TargetCulture = NULL);
 
 	/**
 	 * Generate an FText that represents the passed number as a percentage in the current culture
@@ -319,7 +319,7 @@ private:
 	template<typename T1, typename T2>
 	static FText AsNumberTemplate(T1 Val, const FNumberFormattingOptions* const Options = NULL, const TSharedPtr<FCulture, ESPMode::ThreadSafe>& TargetCulture = NULL);
 	template<typename T1, typename T2>
-	static FText AsCurrencyTemplate(T1 Val, const FNumberFormattingOptions* const Options = NULL, const TSharedPtr<FCulture, ESPMode::ThreadSafe>& TargetCulture = NULL);
+	static FText AsCurrencyTemplate(T1 Val, const FString& CurrencyCode, const FNumberFormattingOptions* const Options = NULL, const TSharedPtr<FCulture, ESPMode::ThreadSafe>& TargetCulture = NULL);
 	template<typename T1, typename T2>
 	static FText AsPercentTemplate(T1 Val, const FNumberFormattingOptions* const Options = NULL, const TSharedPtr<FCulture, ESPMode::ThreadSafe>& TargetCulture = NULL);
 
diff --git a/Engine/Source/Runtime/Core/Public/UObject/ObjectVersion.h b/Engine/Source/Runtime/Core/Public/UObject/ObjectVersion.h
--- a/Engine/Source/Runtime/Core/Public/UObject/ObjectVersion.h
+++ b/Engine/Source/Runtime/Core/Public/UObject/ObjectVersion.h
@@ -590,6 +590,8 @@ enum EUnrealEngineObjectUE4Version
 	VER_UE4_CLEAR_NOTIFY_TRIGGERS,
 	// Convert old curve names stored in anim assets into skeleton smartnames
 	VER_UE4_SKELETON_ADD_SMARTNAMES,
+	// Added the currency code field to FTextHistory_AsCurrency
+	VER_UE4_ADDED_CURRENCY_CODE_TO_FTEXT,
 
 	// -----<new versions can be added before this line>-------------------------------------------------
 	// - this needs to be the last line (see note below)
diff --git a/Engine/Source/Runtime/Engine/Classes/Kismet/KismetTextLibrary.h b/Engine/Source/Runtime/Engine/Classes/Kismet/KismetTextLibrary.h
--- a/Engine/Source/Runtime/Engine/Classes/Kismet/KismetTextLibrary.h
+++ b/Engine/Source/Runtime/Engine/Classes/Kismet/KismetTextLibrary.h
@@ -118,12 +118,12 @@ class UKismetTextLibrary : public UBlueprintFunctionLibrary
 	// Default values are duplicated from FNumberFormattingOptions and should be replicated in all functions and in the struct when changed!
 	/* Converts a passed in integer to a text formatted as a currency */
 	UFUNCTION(BlueprintPure, meta=(FriendlyName = "AsCurrency (int)", AdvancedDisplay = "1"), Category="Utilities|Text")
-	static FText AsCurrency_Integer(int32 Value, TEnumAsByte<ERoundingMode> RoundingMode, bool bUseGrouping = true, int32 MinimumIntegralDigits = 1, int32 MaximumIntegralDigits = 324, int32 MinimumFractionalDigits = 0, int32 MaximumFractionalDigits = 3);
+	static FText AsCurrency_Integer(int32 Value, TEnumAsByte<ERoundingMode> RoundingMode, bool bUseGrouping = true, int32 MinimumIntegralDigits = 1, int32 MaximumIntegralDigits = 324, int32 MinimumFractionalDigits = 0, int32 MaximumFractionalDigits = 3, const FString& CurrencyCode = TEXT(""));
 
 	// Default values are duplicated from FNumberFormattingOptions and should be replicated in all functions and in the struct when changed!
 	/* Converts a passed in float to a text formatted as a currency */
 	UFUNCTION(BlueprintPure, meta=(FriendlyName = "AsCurrency (float)", AdvancedDisplay = "1"), Category="Utilities|Text")
-	static FText AsCurrency_Float(float Value, TEnumAsByte<ERoundingMode> RoundingMode, bool bUseGrouping = true, int32 MinimumIntegralDigits = 1, int32 MaximumIntegralDigits = 324, int32 MinimumFractionalDigits = 0, int32 MaximumFractionalDigits = 3);
+	static FText AsCurrency_Float(float Value, TEnumAsByte<ERoundingMode> RoundingMode, bool bUseGrouping = true, int32 MinimumIntegralDigits = 1, int32 MaximumIntegralDigits = 324, int32 MinimumFractionalDigits = 0, int32 MaximumFractionalDigits = 3, const FString& CurrencyCode = TEXT(""));
 
 	// Default values are duplicated from FNumberFormattingOptions and should be replicated in all functions and in the struct when changed!
 	/* Converts a passed in float to a text, formatted as a percent */
diff --git a/Engine/Source/Runtime/Engine/Private/KismetTextLibrary.cpp b/Engine/Source/Runtime/Engine/Private/KismetTextLibrary.cpp
--- a/Engine/Source/Runtime/Engine/Private/KismetTextLibrary.cpp
+++ b/Engine/Source/Runtime/Engine/Private/KismetTextLibrary.cpp
@@ -111,7 +111,7 @@ FText UKismetTextLibrary::Conv_FloatToText(float Value, TEnumAsByte<ERoundingMod
 	return FText::AsNumber(Value, &NumberFormatOptions);
 }
 
-FText UKismetTextLibrary::AsCurrency_Integer(int32 Value, TEnumAsByte<ERoundingMode> RoundingMode, bool bUseGrouping/* = true*/, int32 MinimumIntegralDigits/* = 1*/, int32 MaximumIntegralDigits/* = 324*/, int32 MinimumFractionalDigits/* = 0*/, int32 MaximumFractionalDigits/* = 3*/)
+FText UKismetTextLibrary::AsCurrency_Integer(int32 Value, TEnumAsByte<ERoundingMode> RoundingMode, bool bUseGrouping/* = true*/, int32 MinimumIntegralDigits/* = 1*/, int32 MaximumIntegralDigits/* = 324*/, int32 MinimumFractionalDigits/* = 0*/, int32 MaximumFractionalDigits/* = 3*/, const FString& CurrencyCode)
 {
 	FNumberFormattingOptions NumberFormatOptions;
 	NumberFormatOptions.UseGrouping = bUseGrouping;
@@ -121,10 +121,10 @@ FText UKismetTextLibrary::AsCurrency_Integer(int32 Value, TEnumAsByte<ERoundingM
 	NumberFormatOptions.MinimumFractionalDigits = MinimumFractionalDigits;
 	NumberFormatOptions.MaximumFractionalDigits = MaximumFractionalDigits;
 
-	return FText::AsCurrency(Value, &NumberFormatOptions);
+	return FText::AsCurrency(Value, CurrencyCode, &NumberFormatOptions);
 }
 
-FText UKismetTextLibrary::AsCurrency_Float(float Value, TEnumAsByte<ERoundingMode> RoundingMode, bool bUseGrouping/* = true*/, int32 MinimumIntegralDigits/* = 1*/, int32 MaximumIntegralDigits/* = 324*/, int32 MinimumFractionalDigits/* = 0*/, int32 MaximumFractionalDigits/* = 3*/)
+FText UKismetTextLibrary::AsCurrency_Float(float Value, TEnumAsByte<ERoundingMode> RoundingMode, bool bUseGrouping/* = true*/, int32 MinimumIntegralDigits/* = 1*/, int32 MaximumIntegralDigits/* = 324*/, int32 MinimumFractionalDigits/* = 0*/, int32 MaximumFractionalDigits/* = 3*/, const FString& CurrencyCode)
 {
 	FNumberFormattingOptions NumberFormatOptions;
 	NumberFormatOptions.UseGrouping = bUseGrouping;
@@ -134,7 +134,7 @@ FText UKismetTextLibrary::AsCurrency_Float(float Value, TEnumAsByte<ERoundingMod
 	NumberFormatOptions.MinimumFractionalDigits = MinimumFractionalDigits;
 	NumberFormatOptions.MaximumFractionalDigits = MaximumFractionalDigits;
 
-	return FText::AsCurrency(Value, &NumberFormatOptions);
+	return FText::AsCurrency(Value, CurrencyCode, &NumberFormatOptions);
 }
 
 FText UKismetTextLibrary::AsPercent_Float(float Value, TEnumAsByte<ERoundingMode> RoundingMode, bool bUseGrouping/* = true*/, int32 MinimumIntegralDigits/* = 1*/, int32 MaximumIntegralDigits/* = 324*/, int32 MinimumFractionalDigits/* = 0*/, int32 MaximumFractionalDigits/* = 3*/)
