Fixed a problem with certain windows (for example Add Code) treated as panels (floating window level, hide on deactivate) on Mac (TTP #343432)

diff --git a/Engine/Source/Runtime/Core/Private/Mac/MacWindow.cpp b/Engine/Source/Runtime/Core/Private/Mac/MacWindow.cpp
--- a/Engine/Source/Runtime/Core/Private/Mac/MacWindow.cpp
+++ b/Engine/Source/Runtime/Core/Private/Mac/MacWindow.cpp
@@ -721,7 +721,7 @@ void FMacWindow::Initialize( FMacApplication* const Application, const TSharedRe
 		check(0);
 		return;
 	}
-	
+
 	// Disable automatic release on close - explicit retain/release makes the Destroy() logic much easier to follow
 	[WindowHandle setReleasedWhenClosed: NO];
 	
@@ -730,7 +730,9 @@ void FMacWindow::Initialize( FMacApplication* const Application, const TSharedRe
 	[WindowHandle setDisplayReconfiguring: false];
 	[WindowHandle setAcceptsMouseMovedEvents: YES];
 	[WindowHandle setDelegate: WindowHandle];
-	if(Definition->IsRegularWindow && (!InParent.IsValid() || Definition->IsModalWindow))
+
+	// @todo: We really need a window type in tab manager to know what window level to use and whether or not a window should hide on deactivate
+	if(Definition->IsRegularWindow && (!InParent.IsValid() || Definition->IsModalWindow || (!Definition->SupportsMaximize && !Definition->SupportsMinimize)))
 	{
 		[WindowHandle setLevel: NSNormalWindowLevel];
 	}
