#ue4 - FRandomStream.GetUnitVector() generates new sample if vector length is very small (improving precision and avoiding potential NaN).

diff --git a/Engine/Source/Runtime/Core/Public/Math/RandomStream.h b/Engine/Source/Runtime/Core/Public/Math/RandomStream.h
--- a/Engine/Source/Runtime/Core/Public/Math/RandomStream.h
+++ b/Engine/Source/Runtime/Core/Public/Math/RandomStream.h
@@ -109,15 +109,17 @@ public:
 	FVector GetUnitVector( ) const
 	{
 		FVector Result;
+		float L;
 
 		do
 		{
 			// Check random vectors in the unit sphere so result is statistically uniform.
-			Result.X = GetFraction() * 2 - 1;
-			Result.Y = GetFraction() * 2 - 1;
-			Result.Z = GetFraction() * 2 - 1;
+			Result.X = GetFraction() * 2.f - 1.f;
+			Result.Y = GetFraction() * 2.f - 1.f;
+			Result.Z = GetFraction() * 2.f - 1.f;
+			L = Result.SizeSquared();
 		}
-		while(Result.SizeSquared() > 1.f);
+		while(L > 1.f || L < KINDA_SMALL_NUMBER);
 
 		return Result.UnsafeNormal();
 	}
diff --git a/Engine/Source/Runtime/Core/Public/Math/UnrealMath.h b/Engine/Source/Runtime/Core/Public/Math/UnrealMath.h
--- a/Engine/Source/Runtime/Core/Public/Math/UnrealMath.h
+++ b/Engine/Source/Runtime/Core/Public/Math/UnrealMath.h
@@ -351,12 +351,12 @@ inline FVector FMath::VRand()
 	do
 	{
 		// Check random vectors in the unit sphere so result is statistically uniform.
-		Result.X = FRand() * 2 - 1;
-		Result.Y = FRand() * 2 - 1;
-		Result.Z = FRand() * 2 - 1;
-
+		Result.X = FRand() * 2.f - 1.f;
+		Result.Y = FRand() * 2.f - 1.f;
+		Result.Z = FRand() * 2.f - 1.f;
 		L = Result.SizeSquared();
-	} while(L > 1.0f || L < 0.001f);
+	}
+	while(L > 1.0f || L < KINDA_SMALL_NUMBER);
 
 	return Result * (1.0f / Sqrt(L));
 }
