Fix for TMap/TSet move semantics.

Allocator traits for FDefaultBitArrayAllocator fixed, which properly enables TSparseArray, TSet, TMap etc. move semantics.
Other missing allocator traits added.
TPair made move-aware.

diff --git a/Engine/Source/Runtime/Core/Public/Containers/ContainerAllocationPolicies.h b/Engine/Source/Runtime/Core/Public/Containers/ContainerAllocationPolicies.h
--- a/Engine/Source/Runtime/Core/Public/Containers/ContainerAllocationPolicies.h
+++ b/Engine/Source/Runtime/Core/Public/Containers/ContainerAllocationPolicies.h
@@ -721,6 +721,7 @@ class FDefaultSetAllocator         : public TSetAllocator<>         { public: ty
 class FDefaultBitArrayAllocator    : public TInlineAllocator<4>     { public: typedef TInlineAllocator<4>     Typedef; };
 class FDefaultSparseArrayAllocator : public TSparseArrayAllocator<> { public: typedef TSparseArrayAllocator<> Typedef; };
 
-template <> struct TAllocatorTraits<FDefaultAllocator>         : TAllocatorTraits    <typename FDefaultAllocator        ::Typedef> {};
-template <> struct TAllocatorTraits<FDefaultBitArrayAllocator> : TAllocatorTraitsBase<typename FDefaultBitArrayAllocator::Typedef> {};
-
+template <> struct TAllocatorTraits<FDefaultAllocator>            : TAllocatorTraits<typename FDefaultAllocator           ::Typedef> {};
+template <> struct TAllocatorTraits<FDefaultSetAllocator>         : TAllocatorTraits<typename FDefaultSetAllocator        ::Typedef> {};
+template <> struct TAllocatorTraits<FDefaultBitArrayAllocator>    : TAllocatorTraits<typename FDefaultBitArrayAllocator   ::Typedef> {};
+template <> struct TAllocatorTraits<FDefaultSparseArrayAllocator> : TAllocatorTraits<typename FDefaultSparseArrayAllocator::Typedef> {};
diff --git a/Engine/Source/Runtime/Core/Public/Containers/Map.h b/Engine/Source/Runtime/Core/Public/Containers/Map.h
--- a/Engine/Source/Runtime/Core/Public/Containers/Map.h
+++ b/Engine/Source/Runtime/Core/Public/Containers/Map.h
@@ -50,10 +50,6 @@ public:
 	KeyType   Key;
 	ValueType Value;
 
-	/** Default constructor. */
-	FORCEINLINE TPair()
-	{}
-
 	/** Initialization constructor. */
 	template <typename InitKeyType, typename InitValueType>
 	FORCEINLINE TPair(const TPairInitializer<InitKeyType, InitValueType>& InInitializer)
@@ -74,21 +70,49 @@ public:
 		// the initializer is itself an rvalue reference.
 	}
 
-	// Explicit copy constructor to make the error messages for copying a pair containing a noncopyable type show the caller's code and not only
-	// "This diagnostic occurred in the compiler generated function 'TPair<KeyType,ValueType>::TPair(const TPair<KeyType,ValueType> &)'" (VC++)
-	FORCEINLINE TPair(const TPair& InInitializer)
-		: Key(InInitializer.Key)
-		, Value(InInitializer.Value)
-	{
-	}
+	#if PLATFORM_COMPILER_HAS_DEFAULTED_FUNCTIONS
 
-	FORCEINLINE TPair& operator=(const TPair& Other)
-	{
-		Key = Other.Key;
-		Value = Other.Value;
+		FORCEINLINE TPair() = default;
+		FORCEINLINE TPair(TPair&&) = default;
+		FORCEINLINE TPair(const TPair&) = default;
+		FORCEINLINE TPair& operator=(TPair&&) = default;
+		FORCEINLINE TPair& operator=(const TPair&) = default;
 
-		return *this;
-	}
+	#else
+
+		FORCEINLINE TPair()
+		{
+		}
+
+		FORCEINLINE TPair(TPair&& InInitializer)
+			: Key  (MoveTemp(InInitializer.Key))
+			, Value(MoveTemp(InInitializer.Value))
+		{
+		}
+
+		FORCEINLINE TPair(const TPair& InInitializer)
+			: Key  (InInitializer.Key)
+			, Value(InInitializer.Value)
+		{
+		}
+
+		FORCEINLINE TPair& operator=(TPair&& Other)
+		{
+			Key   = MoveTemp(Other.Key);
+			Value = MoveTemp(Other.Value);
+
+			return *this;
+		}
+
+		FORCEINLINE TPair& operator=(const TPair& Other)
+		{
+			Key   = Other.Key;
+			Value = Other.Value;
+
+			return *this;
+		}
+
+	#endif
 
 	/** Serializer. */
 	FORCEINLINE friend FArchive& operator<<(FArchive& Ar,TPair& Pair)
@@ -142,6 +166,8 @@ struct TDefaultMapKeyFuncs : BaseKeyFuncs<TPair<KeyType,ValueType>,KeyType,bInAl
 template<typename KeyType,typename ValueType,bool bInAllowDuplicateKeys,typename SetAllocator = FDefaultSetAllocator,typename KeyFuncs = TDefaultMapKeyFuncs<KeyType,ValueType,bInAllowDuplicateKeys> >
 class TMapBase
 {
+	friend struct TContainerTraits<TMapBase>;
+
 public:
 	typedef typename TTypeTraits<KeyType  >::ConstPointerType KeyConstPointerType;
 	typedef typename TTypeTraits<KeyType  >::ConstInitType    KeyInitType;
@@ -728,6 +754,8 @@ private:
 template<typename KeyType,typename ValueType,bool bInAllowDuplicateKeys,typename SetAllocator,typename KeyFuncs>
 class TSortableMapBase : public TMapBase<KeyType,ValueType,bInAllowDuplicateKeys,SetAllocator,KeyFuncs>
 {
+	friend struct TContainerTraits<TSortableMapBase>;
+
 public:
 	typedef TMapBase<KeyType,ValueType,bInAllowDuplicateKeys,SetAllocator,KeyFuncs> Super;
 
@@ -812,8 +840,9 @@ private:
 template<typename KeyType,typename ValueType,typename SetAllocator /*= FDefaultSetAllocator*/,typename KeyFuncs /*= TDefaultMapKeyFuncs<KeyType,ValueType,false>*/>
 class TMap : public TSortableMapBase<KeyType,ValueType,false,SetAllocator,KeyFuncs>
 {
-public:
+	friend struct TContainerTraits<TMap>;
 
+public:
 	typedef TSortableMapBase<KeyType,ValueType,false,SetAllocator,KeyFuncs> Super;
 	typedef typename Super::KeyInitType KeyInitType;
 	typedef typename Super::KeyConstPointerType KeyConstPointerType;
@@ -902,8 +931,9 @@ public:
 template<typename KeyType,typename ValueType,typename SetAllocator /* = FDefaultSetAllocator */,typename KeyFuncs /*= TDefaultMapKeyFuncs<KeyType,ValueType,true>*/>
 class TMultiMap : public TSortableMapBase<KeyType,ValueType,true,SetAllocator,KeyFuncs>
 {
-public:
+	friend struct TContainerTraits<TMultiMap>;
 
+public:
 	typedef TSortableMapBase<KeyType,ValueType,true,SetAllocator,KeyFuncs> Super;
 	typedef typename Super::KeyConstPointerType KeyConstPointerType;
 	typedef typename Super::KeyInitType KeyInitType;
@@ -1175,3 +1205,27 @@ public:
 		return Super::Num();
 	}
 };
+
+template <typename KeyType, typename ValueType, bool bInAllowDuplicateKeys, typename SetAllocator, typename KeyFuncs>
+struct TContainerTraits<TMapBase<KeyType, ValueType, bInAllowDuplicateKeys, SetAllocator, KeyFuncs>> : public TContainerTraitsBase<TMapBase<KeyType, ValueType, bInAllowDuplicateKeys, SetAllocator, KeyFuncs>>
+{
+	enum { MoveWillEmptyContainer = TContainerTraits<typename TMapBase<KeyType, ValueType, bInAllowDuplicateKeys, SetAllocator, KeyFuncs>::PairSetType>::MoveWillEmptyContainer };
+};
+
+template <typename KeyType, typename ValueType, bool bInAllowDuplicateKeys, typename SetAllocator, typename KeyFuncs>
+struct TContainerTraits<TSortableMapBase<KeyType, ValueType, bInAllowDuplicateKeys, SetAllocator, KeyFuncs>> : public TContainerTraitsBase<TSortableMapBase<KeyType, ValueType, bInAllowDuplicateKeys, SetAllocator, KeyFuncs>>
+{
+	enum { MoveWillEmptyContainer = TContainerTraits<typename TSortableMapBase<KeyType, ValueType, bInAllowDuplicateKeys, SetAllocator, KeyFuncs>::Super>::MoveWillEmptyContainer };
+};
+
+template <typename KeyType, typename ValueType, typename SetAllocator,typename KeyFuncs>
+struct TContainerTraits<TMap<KeyType, ValueType, SetAllocator, KeyFuncs>> : public TContainerTraitsBase<TMap<KeyType, ValueType, SetAllocator, KeyFuncs>>
+{
+	enum { MoveWillEmptyContainer = TContainerTraits<typename TMap<KeyType, ValueType, SetAllocator, KeyFuncs>::Super>::MoveWillEmptyContainer };
+};
+
+template <typename KeyType, typename ValueType, typename SetAllocator,typename KeyFuncs>
+struct TContainerTraits<TMultiMap<KeyType, ValueType, SetAllocator, KeyFuncs>> : public TContainerTraitsBase<TMultiMap<KeyType, ValueType, SetAllocator, KeyFuncs>>
+{
+	enum { MoveWillEmptyContainer = TContainerTraits<typename TMultiMap<KeyType, ValueType, SetAllocator, KeyFuncs>::Super>::MoveWillEmptyContainer };
+};
diff --git a/Engine/Source/Runtime/Core/Public/Containers/Set.h b/Engine/Source/Runtime/Core/Public/Containers/Set.h
--- a/Engine/Source/Runtime/Core/Public/Containers/Set.h
+++ b/Engine/Source/Runtime/Core/Public/Containers/Set.h
@@ -240,8 +240,6 @@ public:
 		return *this;
 	}
 
-#if PLATFORM_COMPILER_HAS_RVALUE_REFERENCES
-
 private:
 	template <typename SetType>
 	static FORCEINLINE typename TEnableIf<TContainerTraits<SetType>::MoveWillEmptyContainer>::Type MoveOrCopy(SetType& ToSet, SetType& FromSet)
@@ -279,8 +277,6 @@ public:
 		return *this;
 	}
 
-#endif
-
 	/**
 	 * Removes all elements from the set, potentially leaving space allocated for an expected number of elements about to be added.
 	 * @param ExpectedNumElements - The number of elements about to be added to the set.
@@ -1126,7 +1122,6 @@ template<typename ElementType, typename KeyFuncs, typename Allocator>
 struct TContainerTraits<TSet<ElementType, KeyFuncs, Allocator> > : public TContainerTraitsBase<TSet<ElementType, KeyFuncs, Allocator> >
 {
 	enum { MoveWillEmptyContainer =
-		PLATFORM_COMPILER_HAS_RVALUE_REFERENCES &&
 		TContainerTraits<typename TSet<ElementType, KeyFuncs, Allocator>::ElementArrayType>::MoveWillEmptyContainer &&
 		TAllocatorTraits<typename Allocator::HashAllocator>::SupportsMove };
 };
