Added a CamelCase break iterator for use by the text layout

TTP#322500 - UX: Add wrapping of Asset Names based on CamelCase and Underscores, not just spaces

- Refactored the existing break iterators to share the same interface so that they could be swapped out at runtime.
- Made the iterators recyclable (so you can change the string rather than have to allocate a new iterator).
- Removed the public interfaces for the different break iterator types - they're all now accessed via IBreakIterator and are created via the relevant function in FBreakIterator.
- Removed a load of duplication in the ICU break iterators.
- Made FTextLayout accept a custom line break iterator (as used by its wrapping cache).
- Fixed a crash where the legacy line or word break iterators could underflow/overflow the string range.

diff --git a/Engine/Source/Editor/UnrealEd/Private/Commandlets/GenerateTextLocalizationReportCommandlet.cpp b/Engine/Source/Editor/UnrealEd/Private/Commandlets/GenerateTextLocalizationReportCommandlet.cpp
--- a/Engine/Source/Editor/UnrealEd/Private/Commandlets/GenerateTextLocalizationReportCommandlet.cpp
+++ b/Engine/Source/Editor/UnrealEd/Private/Commandlets/GenerateTextLocalizationReportCommandlet.cpp
@@ -6,6 +6,7 @@
 #include "InternationalizationManifestJsonSerializer.h"
 #include "InternationalizationArchive.h"
 #include "InternationalizationArchiveJsonSerializer.h"
+#include "BreakIterator.h"
 
 DEFINE_LOG_CATEGORY_STATIC(LogGenerateTextLocalizationReportCommandlet, Log, All);
 
@@ -19,12 +20,13 @@ int32 UGenerateTextLocalizationReportCommandlet::CountWords( const FString& Text
 	int32 NumOfWords = 0;
 	// Calculate # of words
 
-	FLineBreakIterator LineBreakIterator( Text );
+	TSharedRef<IBreakIterator> LineBreakIterator = FBreakIterator::CreateLineBreakIterator();
+	LineBreakIterator->SetString( Text );
 
 	int32 PreviousBreak = 0;
 	int32 CurrentBreak = 0;
 
-	while( ( CurrentBreak = LineBreakIterator.MoveToNext() ) != INDEX_NONE )
+	while( ( CurrentBreak = LineBreakIterator->MoveToNext() ) != INDEX_NONE )
 	{
 		if ( CurrentBreak > PreviousBreak )
 		{
diff --git a/Engine/Source/Runtime/Core/Private/Internationalization/CamelCaseBreakIterator.cpp b/Engine/Source/Runtime/Core/Private/Internationalization/CamelCaseBreakIterator.cpp
new file mode 100644
--- /dev/null
+++ b/Engine/Source/Runtime/Core/Private/Internationalization/CamelCaseBreakIterator.cpp
@@ -0,0 +1,151 @@
+// Copyright 1998-2014 Epic Games, Inc. All Rights Reserved.
+
+#include "Core.h"
+#include "CamelCaseBreakIterator.h"
+
+FCamelCaseBreakIterator::FCamelCaseBreakIterator()
+	: String()
+	, CurrentPosition(0)
+{
+	BreakPoints.Add(0);
+}
+
+void FCamelCaseBreakIterator::SetString(const FText& InText)
+{
+	String = InText.ToString();
+	UpdateBreakPointsArray();
+	ResetToBeginning();
+}
+
+void FCamelCaseBreakIterator::SetString(const FString& InString)
+{
+	String = InString;
+	UpdateBreakPointsArray();
+	ResetToBeginning();
+}
+
+void FCamelCaseBreakIterator::SetString(const TCHAR* const InString, const int32 InStringLength) 
+{
+	String = FString(InString, InStringLength);
+	UpdateBreakPointsArray();
+	ResetToBeginning();
+}
+
+void FCamelCaseBreakIterator::ClearString()
+{
+	String = FString();
+	BreakPoints.Empty(1);
+	BreakPoints.Add(0);
+	ResetToBeginning();
+}
+
+int32 FCamelCaseBreakIterator::GetCurrentPosition() const
+{
+	return CurrentPosition;
+}
+
+int32 FCamelCaseBreakIterator::ResetToBeginning()
+{
+	return CurrentPosition = 0;
+}
+
+int32 FCamelCaseBreakIterator::ResetToEnd()
+{
+	return CurrentPosition = String.Len();
+}
+
+int32 FCamelCaseBreakIterator::MoveToPrevious()
+{
+	return MoveToCandidateBefore(CurrentPosition);
+}
+
+int32 FCamelCaseBreakIterator::MoveToNext()
+{
+	return MoveToCandidateAfter(CurrentPosition);
+}
+
+int32 FCamelCaseBreakIterator::MoveToCandidateBefore(const int32 InIndex)
+{
+	CurrentPosition = InIndex;
+
+	for(int32 BreakPointIndex = BreakPoints.Num() - 1; BreakPointIndex >= 0; --BreakPointIndex)
+	{
+		const int32 BreakPoint = BreakPoints[BreakPointIndex];
+		if(BreakPoint < InIndex)
+		{
+			CurrentPosition = BreakPoint;
+			break;
+		}
+	}
+
+	return (CurrentPosition >= InIndex) ? INDEX_NONE : CurrentPosition;
+}
+
+int32 FCamelCaseBreakIterator::MoveToCandidateAfter(const int32 InIndex)
+{
+	CurrentPosition = InIndex;
+
+	for(int32 BreakPointIndex = 0; BreakPointIndex < BreakPoints.Num(); ++BreakPointIndex)
+	{
+		const int32 BreakPoint = BreakPoints[BreakPointIndex];
+		if(BreakPoint > InIndex)
+		{
+			CurrentPosition = BreakPoint;
+			break;
+		}
+	}
+
+	return (CurrentPosition <= InIndex) ? INDEX_NONE : CurrentPosition;
+}
+
+void FCamelCaseBreakIterator::UpdateBreakPointsArray()
+{
+	TArray<FToken> Tokens;
+	TokenizeString(Tokens);
+	PopulateBreakPointsArray(Tokens);
+}
+
+void FCamelCaseBreakIterator::PopulateBreakPointsArray(const TArray<FToken>& InTokens)
+{
+	BreakPoints.Empty(InTokens.Num());
+
+	// Process the tokens so that input like "ICUBreakIterator1234_Ext" would produce the following break points:
+	// ICU|Break|Iterator|1234|_|Ext|
+
+	BreakPoints.Add(0); // start of the string
+
+	ETokenType TokenRunType = ETokenType::Other;
+	for(int32 TokenIndex = 0; TokenIndex < InTokens.Num(); ++TokenIndex)
+	{
+		const FToken& Token = InTokens[TokenIndex];
+
+		// End of string?
+		if(Token.TokenType == ETokenType::Null)
+		{
+			BreakPoints.AddUnique(Token.StrIndex);
+			break;
+		}
+
+		// Have we found the end of some kind of run of tokens?
+		if(TokenRunType != Token.TokenType)
+		{
+			// If we've moved from a run of upper-case tokens, to a lower-case token, then we need to try and make the previous upper-case token part of the next run
+			const int32 BreakTokenIndex = TokenIndex - ((TokenRunType == ETokenType::Uppercase && Token.TokenType == ETokenType::Lowercase) ? 1 : 0);
+			if(BreakTokenIndex > 0)
+			{
+				BreakPoints.AddUnique(InTokens[BreakTokenIndex].StrIndex);
+			}
+		}
+
+		// Always add "other" tokens as break points
+		if(Token.TokenType == ETokenType::Other)
+		{
+			BreakPoints.AddUnique(Token.StrIndex);
+		}
+
+		TokenRunType = Token.TokenType;
+	}
+
+	// There should always be at least one entry for the end of the string
+	check(BreakPoints.Num());
+}
diff --git a/Engine/Source/Runtime/Core/Private/Internationalization/CamelCaseBreakIterator.h b/Engine/Source/Runtime/Core/Private/Internationalization/CamelCaseBreakIterator.h
new file mode 100644
--- /dev/null
+++ b/Engine/Source/Runtime/Core/Private/Internationalization/CamelCaseBreakIterator.h
@@ -0,0 +1,62 @@
+// Copyright 1998-2014 Epic Games, Inc. All Rights Reserved.
+
+#pragma once
+
+#include "IBreakIterator.h"
+
+/**
+ * Base class for camel case break iterators
+ * A derived type must provide a way to tokenize the string for processing
+ */
+class FCamelCaseBreakIterator : public IBreakIterator
+{
+public:
+	FCamelCaseBreakIterator();
+
+	virtual void SetString(const FText& InText) override;
+	virtual void SetString(const FString& InString) override;
+	virtual void SetString(const TCHAR* const InString, const int32 InStringLength) override;
+	virtual void ClearString() override;
+
+	virtual int32 GetCurrentPosition() const override;
+
+	virtual int32 ResetToBeginning() override;
+	virtual int32 ResetToEnd() override;
+
+	virtual int32 MoveToPrevious() override;
+	virtual int32 MoveToNext() override;
+	virtual int32 MoveToCandidateBefore(const int32 InIndex) override;
+	virtual int32 MoveToCandidateAfter(const int32 InIndex) override;
+
+protected:
+	enum class ETokenType : uint8
+	{
+		Uppercase,
+		Lowercase,
+		Digit,
+		Null,
+		Other,
+	};
+
+	struct FToken
+	{
+		FToken(const ETokenType InTokenType, const int32 InStrIndex)
+			: TokenType(InTokenType)
+			, StrIndex(InStrIndex)
+		{
+		}
+
+		ETokenType TokenType;
+		int32 StrIndex;
+	};
+
+	virtual void TokenizeString(TArray<FToken>& OutTokens) = 0;
+
+protected:
+	void UpdateBreakPointsArray();
+	void PopulateBreakPointsArray(const TArray<FToken>& InTokens);
+
+	FString String;
+	int32 CurrentPosition;
+	TArray<int32> BreakPoints;
+};
diff --git a/Engine/Source/Runtime/Core/Private/Internationalization/ICUBreakIterator.cpp b/Engine/Source/Runtime/Core/Private/Internationalization/ICUBreakIterator.cpp
new file mode 100644
--- /dev/null
+++ b/Engine/Source/Runtime/Core/Private/Internationalization/ICUBreakIterator.cpp
@@ -0,0 +1,73 @@
+// Copyright 1998-2014 Epic Games, Inc. All Rights Reserved.
+
+#include "Core.h"
+#include "ICUBreakIterator.h"
+
+#if UE_ENABLE_ICU
+#include "ICUTextCharacterIterator.h"
+
+FICUBreakIterator::FICUBreakIterator(TSharedRef<icu::BreakIterator>&& InICUBreakIterator)
+	: ICUBreakIterator(InICUBreakIterator)
+{
+}
+
+void FICUBreakIterator::SetString(const FText& InText)
+{
+	ICUBreakIterator->adoptText(new FICUTextCharacterIterator(InText)); // ICUBreakIterator takes ownership of this instance
+	ResetToBeginning();
+}
+
+void FICUBreakIterator::SetString(const FString& InString)
+{
+	ICUBreakIterator->adoptText(new FICUTextCharacterIterator(InString)); // ICUBreakIterator takes ownership of this instance
+	ResetToBeginning();
+}
+
+void FICUBreakIterator::SetString(const TCHAR* const InString, const int32 InStringLength) 
+{
+	ICUBreakIterator->adoptText(new FICUTextCharacterIterator(InString, InStringLength)); // ICUBreakIterator takes ownership of this instance
+	ResetToBeginning();
+}
+
+void FICUBreakIterator::ClearString()
+{
+	ICUBreakIterator->setText(icu::UnicodeString());
+	ResetToBeginning();
+}
+
+int32 FICUBreakIterator::GetCurrentPosition() const
+{
+	return ICUBreakIterator->current();
+}
+
+int32 FICUBreakIterator::ResetToBeginning()
+{
+	return ICUBreakIterator->first();
+}
+
+int32 FICUBreakIterator::ResetToEnd()
+{
+	return ICUBreakIterator->last();
+}
+
+int32 FICUBreakIterator::MoveToPrevious()
+{
+	return ICUBreakIterator->previous();
+}
+
+int32 FICUBreakIterator::MoveToNext()
+{
+	return ICUBreakIterator->next();
+}
+
+int32 FICUBreakIterator::MoveToCandidateBefore(const int32 InIndex)
+{
+	return ICUBreakIterator->preceding(InIndex);
+}
+
+int32 FICUBreakIterator::MoveToCandidateAfter(const int32 InIndex)
+{
+	return ICUBreakIterator->following(InIndex);
+}
+
+#endif
diff --git a/Engine/Source/Runtime/Core/Private/Internationalization/ICUBreakIterator.h b/Engine/Source/Runtime/Core/Private/Internationalization/ICUBreakIterator.h
new file mode 100644
--- /dev/null
+++ b/Engine/Source/Runtime/Core/Private/Internationalization/ICUBreakIterator.h
@@ -0,0 +1,38 @@
+// Copyright 1998-2014 Epic Games, Inc. All Rights Reserved.
+
+#pragma once
+
+#include "IBreakIterator.h"
+
+#if UE_ENABLE_ICU
+
+#include <unicode/brkiter.h>
+
+/**
+ * Wraps an ICU break iterator instance inside our own break iterator API
+ */
+class FICUBreakIterator : public IBreakIterator
+{
+public:
+	FICUBreakIterator(TSharedRef<icu::BreakIterator>&& InICUBreakIterator);
+
+	virtual void SetString(const FText& InText) override;
+	virtual void SetString(const FString& InString) override;
+	virtual void SetString(const TCHAR* const InString, const int32 InStringLength) override;
+	virtual void ClearString() override;
+
+	virtual int32 GetCurrentPosition() const override;
+
+	virtual int32 ResetToBeginning() override;
+	virtual int32 ResetToEnd() override;
+
+	virtual int32 MoveToPrevious() override;
+	virtual int32 MoveToNext() override;
+	virtual int32 MoveToCandidateBefore(const int32 InIndex) override;
+	virtual int32 MoveToCandidateAfter(const int32 InIndex) override;
+
+protected:
+	TSharedRef<icu::BreakIterator> ICUBreakIterator;
+};
+
+#endif
diff --git a/Engine/Source/Runtime/Core/Private/Internationalization/ICUCamelCaseBreakIterator.cpp b/Engine/Source/Runtime/Core/Private/Internationalization/ICUCamelCaseBreakIterator.cpp
new file mode 100644
--- /dev/null
+++ b/Engine/Source/Runtime/Core/Private/Internationalization/ICUCamelCaseBreakIterator.cpp
@@ -0,0 +1,54 @@
+// Copyright 1998-2014 Epic Games, Inc. All Rights Reserved.
+
+#include "Core.h"
+#include "BreakIterator.h"
+#include "CamelCaseBreakIterator.h"
+
+#if UE_ENABLE_ICU
+#include "ICUTextCharacterIterator.h"
+#include <unicode/uchar.h>
+
+class FICUCamelCaseBreakIterator : public FCamelCaseBreakIterator
+{
+protected:
+	virtual void TokenizeString(TArray<FToken>& OutTokens) override;
+};
+
+void FICUCamelCaseBreakIterator::TokenizeString(TArray<FToken>& OutTokens)
+{
+	OutTokens.Empty(String.Len());
+
+	FICUTextCharacterIterator CharIter(String);
+	for(CharIter.setToStart(); CharIter.current32() != FICUTextCharacterIterator::DONE; CharIter.next32PostInc())
+	{
+		const UChar32 CurrentChar = CharIter.current32();
+
+		ETokenType TokenType = ETokenType::Other;
+		if(u_isULowercase(CurrentChar))
+		{
+			TokenType = ETokenType::Lowercase;
+		}
+		else if(u_isUUppercase(CurrentChar))
+		{
+			TokenType = ETokenType::Uppercase;
+		}
+		else if(u_isdigit(CurrentChar))
+		{
+			TokenType = ETokenType::Digit;
+		}
+
+		OutTokens.Emplace(FToken(TokenType, CharIter.getIndex()));
+	}
+
+	OutTokens.Emplace(FToken(ETokenType::Null, String.Len()));
+
+	// There should always be at least one token for the end of the string
+	check(OutTokens.Num());
+}
+
+TSharedRef<IBreakIterator> FBreakIterator::CreateCamelCaseBreakIterator()
+{
+	return MakeShareable(new FICUCamelCaseBreakIterator());
+}
+
+#endif
diff --git a/Engine/Source/Runtime/Core/Private/Internationalization/ICUCharacterBoundaryIterator.cpp b/Engine/Source/Runtime/Core/Private/Internationalization/ICUCharacterBoundaryIterator.cpp
--- a/Engine/Source/Runtime/Core/Private/Internationalization/ICUCharacterBoundaryIterator.cpp
+++ b/Engine/Source/Runtime/Core/Private/Internationalization/ICUCharacterBoundaryIterator.cpp
@@ -1,117 +1,34 @@
 // Copyright 1998-2014 Epic Games, Inc. All Rights Reserved.
 
 #include "Core.h"
-#include "CharacterBoundaryIterator.h"
+#include "BreakIterator.h"
 
 #if UE_ENABLE_ICU
-#include "ICUTextCharacterIterator.h"
-#include <unicode/brkiter.h>
+#include "ICUBreakIterator.h"
 
-namespace
-{
-	TSharedRef<icu::BreakIterator> CreateGraphemeBreakIterator()
-	{
-		UErrorCode ICUStatus = U_ZERO_ERROR;
-		return MakeShareable( icu::BreakIterator::createCharacterInstance( icu::Locale::getDefault(), ICUStatus ) );
-	}
-}
-
-class FCharacterBoundaryIterator::FImplementation
+class FICUCharacterBoundaryIterator : public FICUBreakIterator
 {
 public:
-	FImplementation(const TCHAR* const String, const int32 StringLength)
-		: ICUBreakIterator( CreateGraphemeBreakIterator() )
-	{
-		ICUBreakIterator->adoptText(new FICUTextCharacterIterator(String, StringLength)); // ICUBreakIterator takes ownership of this instance
-	}
-
-	int32 GetCurrentPosition() const
-	{
-		return ICUBreakIterator->current();
-	}
-
-	int32 ResetToBeginning()
-	{
-		return ICUBreakIterator->first();
-	}
-
-	int32 ResetToEnd()
-	{
-		return ICUBreakIterator->last();
-	}
-
-	int32 MoveToPrevious()
-	{
-		return ICUBreakIterator->previous();
-	}
-
-	int32 MoveToNext()
-	{
-		return ICUBreakIterator->next();
-	}
-
-	int32 MoveToCandidateBefore(const int32 Index)
-	{
-		return ICUBreakIterator->preceding(Index);
-	}
-
-	int32 MoveToCandidateAfter(const int32 Index)
-	{
-		return ICUBreakIterator->following(Index);
-	}
+	FICUCharacterBoundaryIterator();
 
 private:
-	TSharedRef<icu::BreakIterator> ICUBreakIterator;
+	static TSharedRef<icu::BreakIterator> CreateInternalGraphemeBreakIterator();
 };
 
-FCharacterBoundaryIterator::FCharacterBoundaryIterator(const FText& Text)
-	: Implementation( new FImplementation( *Text.ToString(), Text.ToString().Len() ) )
-{
-}
-
-FCharacterBoundaryIterator::FCharacterBoundaryIterator(const FString& String)
-	: Implementation( new FImplementation( *String, String.Len() ) )
-{
-}
-
-FCharacterBoundaryIterator::FCharacterBoundaryIterator(const TCHAR* const String, const int32 StringLength)
-	: Implementation( new FImplementation( String, StringLength ) )
-{
-}
-
-int32 FCharacterBoundaryIterator::GetCurrentPosition() const
-{
-	return Implementation->GetCurrentPosition();
-}
-
-int32 FCharacterBoundaryIterator::ResetToBeginning()
-{
-	return Implementation->ResetToBeginning();
-}
-
-int32 FCharacterBoundaryIterator::ResetToEnd()
-{
-	return Implementation->ResetToEnd();
-}
-
-int32 FCharacterBoundaryIterator::MoveToPrevious()
-{
-	return Implementation->MoveToPrevious();
-}
-
-int32 FCharacterBoundaryIterator::MoveToNext()
+FICUCharacterBoundaryIterator::FICUCharacterBoundaryIterator()
+	: FICUBreakIterator(CreateInternalGraphemeBreakIterator())
 {
-	return Implementation->MoveToNext();
 }
 
-int32 FCharacterBoundaryIterator::MoveToCandidateBefore(const int32 Index)
+TSharedRef<icu::BreakIterator> FICUCharacterBoundaryIterator::CreateInternalGraphemeBreakIterator()
 {
-	return Implementation->MoveToCandidateBefore(Index);
+	UErrorCode ICUStatus = U_ZERO_ERROR;
+	return MakeShareable(icu::BreakIterator::createCharacterInstance(icu::Locale::getDefault(), ICUStatus));
 }
 
-int32 FCharacterBoundaryIterator::MoveToCandidateAfter(const int32 Index)
+TSharedRef<IBreakIterator> FBreakIterator::CreateCharacterBoundaryIterator()
 {
-	return Implementation->MoveToCandidateAfter(Index);
+	return MakeShareable(new FICUCharacterBoundaryIterator());
 }
 
-#endif
\ No newline at end of file
+#endif
diff --git a/Engine/Source/Runtime/Core/Private/Internationalization/ICULineBreakIterator.cpp b/Engine/Source/Runtime/Core/Private/Internationalization/ICULineBreakIterator.cpp
--- a/Engine/Source/Runtime/Core/Private/Internationalization/ICULineBreakIterator.cpp
+++ b/Engine/Source/Runtime/Core/Private/Internationalization/ICULineBreakIterator.cpp
@@ -1,117 +1,34 @@
 // Copyright 1998-2014 Epic Games, Inc. All Rights Reserved.
 
 #include "Core.h"
-#include "LineBreakIterator.h"
+#include "BreakIterator.h"
 
 #if UE_ENABLE_ICU
-#include "ICUTextCharacterIterator.h"
-#include <unicode/brkiter.h>
+#include "ICUBreakIterator.h"
 
-namespace
-{
-	TSharedRef<icu::BreakIterator> CreateLineBreakIterator()
-	{
-		UErrorCode ICUStatus = U_ZERO_ERROR;
-		return MakeShareable( icu::BreakIterator::createLineInstance( icu::Locale::getDefault(), ICUStatus ) );
-	}
-}
-
-class FLineBreakIterator::FImplementation
+class FICULineBreakIterator : public FICUBreakIterator
 {
 public:
-	FImplementation(const TCHAR* const String, const int32 StringLength)
-		: ICUBreakIterator( CreateLineBreakIterator() )
-	{
-		ICUBreakIterator->adoptText(new FICUTextCharacterIterator(String, StringLength)); // ICUBreakIterator takes ownership of this instance
-	}
-
-	int32 GetCurrentPosition() const
-	{
-		return ICUBreakIterator->current();
-	}
-
-	int32 ResetToBeginning()
-	{
-		return ICUBreakIterator->first();
-	}
-
-	int32 ResetToEnd()
-	{
-		return ICUBreakIterator->last();
-	}
-
-	int32 MoveToPrevious()
-	{
-		return ICUBreakIterator->previous();
-	}
-
-	int32 MoveToNext()
-	{
-		return ICUBreakIterator->next();
-	}
-
-	int32 MoveToCandidateBefore(const int32 Index)
-	{
-		return ICUBreakIterator->preceding(Index);
-	}
-
-	int32 MoveToCandidateAfter(const int32 Index)
-	{
-		return ICUBreakIterator->following(Index);
-	}
+	FICULineBreakIterator();
 
 private:
-	TSharedRef<icu::BreakIterator> ICUBreakIterator;
+	static TSharedRef<icu::BreakIterator> CreateInternalLineBreakIterator();
 };
 
-FLineBreakIterator::FLineBreakIterator(const FText& Text)
-	: Implementation( new FImplementation( *Text.ToString(), Text.ToString().Len() ) )
-{
-}
-
-FLineBreakIterator::FLineBreakIterator(const FString& String)
-	: Implementation( new FImplementation( *String, String.Len() ) )
-{
-}
-
-FLineBreakIterator::FLineBreakIterator(const TCHAR* const String, const int32 StringLength)
-	: Implementation( new FImplementation( String, StringLength ) )
-{
-}
-
-int32 FLineBreakIterator::GetCurrentPosition() const
-{
-	return Implementation->GetCurrentPosition();
-}
-
-int32 FLineBreakIterator::ResetToBeginning()
-{
-	return Implementation->ResetToBeginning();
-}
-
-int32 FLineBreakIterator::ResetToEnd()
-{
-	return Implementation->ResetToEnd();
-}
-
-int32 FLineBreakIterator::MoveToPrevious()
-{
-	return Implementation->MoveToPrevious();
-}
-
-int32 FLineBreakIterator::MoveToNext()
+FICULineBreakIterator::FICULineBreakIterator()
+	: FICUBreakIterator(CreateInternalLineBreakIterator())
 {
-	return Implementation->MoveToNext();
 }
 
-int32 FLineBreakIterator::MoveToCandidateBefore(const int32 Index)
+TSharedRef<icu::BreakIterator> FICULineBreakIterator::CreateInternalLineBreakIterator()
 {
-	return Implementation->MoveToCandidateBefore(Index);
+	UErrorCode ICUStatus = U_ZERO_ERROR;
+	return MakeShareable( icu::BreakIterator::createLineInstance(icu::Locale::getDefault(), ICUStatus));
 }
 
-int32 FLineBreakIterator::MoveToCandidateAfter(const int32 Index)
+TSharedRef<IBreakIterator> FBreakIterator::CreateLineBreakIterator()
 {
-	return Implementation->MoveToCandidateAfter(Index);
+	return MakeShareable(new FICULineBreakIterator());
 }
 
-#endif
\ No newline at end of file
+#endif
diff --git a/Engine/Source/Runtime/Core/Private/Internationalization/ICUWordBreakIterator.cpp b/Engine/Source/Runtime/Core/Private/Internationalization/ICUWordBreakIterator.cpp
--- a/Engine/Source/Runtime/Core/Private/Internationalization/ICUWordBreakIterator.cpp
+++ b/Engine/Source/Runtime/Core/Private/Internationalization/ICUWordBreakIterator.cpp
@@ -1,117 +1,34 @@
 // Copyright 1998-2014 Epic Games, Inc. All Rights Reserved.
 
 #include "Core.h"
-#include "WordBreakIterator.h"
+#include "BreakIterator.h"
 
 #if UE_ENABLE_ICU
-#include "ICUTextCharacterIterator.h"
-#include <unicode/brkiter.h>
+#include "ICUBreakIterator.h"
 
-namespace
-{
-	TSharedRef<icu::BreakIterator> CreateWordBreakIterator()
-	{
-		UErrorCode ICUStatus = U_ZERO_ERROR;
-		return MakeShareable( icu::BreakIterator::createWordInstance( icu::Locale::getDefault(), ICUStatus ) );
-	}
-}
-
-class FWordBreakIterator::FImplementation
+class FICUWordBreakIterator : public FICUBreakIterator
 {
 public:
-	FImplementation(const TCHAR* const String, const int32 StringLength)
-		: ICUBreakIterator(CreateWordBreakIterator())
-	{
-		ICUBreakIterator->adoptText(new FICUTextCharacterIterator(String, StringLength)); // ICUBreakIterator takes ownership of this instance
-	}
-
-	int32 GetCurrentPosition() const
-	{
-		return ICUBreakIterator->current();
-	}
-
-	int32 ResetToBeginning()
-	{
-		return ICUBreakIterator->first();
-	}
-
-	int32 ResetToEnd()
-	{
-		return ICUBreakIterator->last();
-	}
-
-	int32 MoveToPrevious()
-	{
-		return ICUBreakIterator->previous();
-	}
-
-	int32 MoveToNext()
-	{
-		return ICUBreakIterator->next();
-	}
-
-	int32 MoveToCandidateBefore(const int32 Index)
-	{
-		return ICUBreakIterator->preceding(Index);
-	}
-
-	int32 MoveToCandidateAfter(const int32 Index)
-	{
-		return ICUBreakIterator->following(Index);
-	}
+	FICUWordBreakIterator();
 
 private:
-	TSharedRef<icu::BreakIterator> ICUBreakIterator;
+	static TSharedRef<icu::BreakIterator> CreateInternalWordBreakIterator();
 };
 
-FWordBreakIterator::FWordBreakIterator(const FText& Text)
-	: Implementation( new FImplementation( *Text.ToString(), Text.ToString().Len() ) )
-{
-}
-
-FWordBreakIterator::FWordBreakIterator(const FString& String)
-	: Implementation( new FImplementation( *String, String.Len() ) )
-{
-}
-
-FWordBreakIterator::FWordBreakIterator(const TCHAR* const String, const int32 StringLength)
-	: Implementation( new FImplementation( String, StringLength ) )
-{
-}
-
-int32 FWordBreakIterator::GetCurrentPosition() const
-{
-	return Implementation->GetCurrentPosition();
-}
-
-int32 FWordBreakIterator::ResetToBeginning()
-{
-	return Implementation->ResetToBeginning();
-}
-
-int32 FWordBreakIterator::ResetToEnd()
-{
-	return Implementation->ResetToEnd();
-}
-
-int32 FWordBreakIterator::MoveToPrevious()
-{
-	return Implementation->MoveToPrevious();
-}
-
-int32 FWordBreakIterator::MoveToNext()
+FICUWordBreakIterator::FICUWordBreakIterator()
+	: FICUBreakIterator(CreateInternalWordBreakIterator())
 {
-	return Implementation->MoveToNext();
 }
 
-int32 FWordBreakIterator::MoveToCandidateBefore(const int32 Index)
+TSharedRef<icu::BreakIterator> FICUWordBreakIterator::CreateInternalWordBreakIterator()
 {
-	return Implementation->MoveToCandidateBefore(Index);
+	UErrorCode ICUStatus = U_ZERO_ERROR;
+	return MakeShareable(icu::BreakIterator::createWordInstance(icu::Locale::getDefault(), ICUStatus));
 }
 
-int32 FWordBreakIterator::MoveToCandidateAfter(const int32 Index)
+TSharedRef<IBreakIterator> FBreakIterator::CreateWordBreakIterator()
 {
-	return Implementation->MoveToCandidateAfter(Index);
+	return MakeShareable(new FICUWordBreakIterator());
 }
 
-#endif
\ No newline at end of file
+#endif
diff --git a/Engine/Source/Runtime/Core/Private/Internationalization/LegacyCamelCaseBreakIterator.cpp b/Engine/Source/Runtime/Core/Private/Internationalization/LegacyCamelCaseBreakIterator.cpp
new file mode 100644
--- /dev/null
+++ b/Engine/Source/Runtime/Core/Private/Internationalization/LegacyCamelCaseBreakIterator.cpp
@@ -0,0 +1,51 @@
+// Copyright 1998-2014 Epic Games, Inc. All Rights Reserved.
+
+#include "Core.h"
+#include "BreakIterator.h"
+#include "CamelCaseBreakIterator.h"
+
+#if !UE_ENABLE_ICU
+
+class FLegacyCamelCaseBreakIterator : public FCamelCaseBreakIterator
+{
+protected:
+	virtual void TokenizeString(TArray<FToken>& OutTokens) override;
+};
+
+void FLegacyCamelCaseBreakIterator::TokenizeString(TArray<FToken>& OutTokens)
+{
+	OutTokens.Empty(String.Len());
+
+	for(int32 CurrentCharIndex = 0; CurrentCharIndex < String.Len(); ++CurrentCharIndex)
+	{
+		const TCHAR CurrentChar = String[CurrentCharIndex];
+
+		ETokenType TokenType = ETokenType::Other;
+		if(FChar::IsLower(CurrentChar))
+		{
+			TokenType = ETokenType::Lowercase;
+		}
+		else if(FChar::IsUpper(CurrentChar))
+		{
+			TokenType = ETokenType::Uppercase;
+		}
+		else if(FChar::IsDigit(CurrentChar))
+		{
+			TokenType = ETokenType::Digit;
+		}
+
+		OutTokens.Emplace(FToken(TokenType, CurrentCharIndex));
+	}
+
+	OutTokens.Emplace(FToken(ETokenType::Null, String.Len()));
+
+	// There should always be at least one token for the end of the string
+	check(OutTokens.Num());
+}
+
+TSharedRef<IBreakIterator> FBreakIterator::CreateCamelCaseBreakIterator()
+{
+	return MakeShareable(new FLegacyCamelCaseBreakIterator());
+}
+
+#endif
diff --git a/Engine/Source/Runtime/Core/Private/Internationalization/LegacyCharacterBoundaryIterator.cpp b/Engine/Source/Runtime/Core/Private/Internationalization/LegacyCharacterBoundaryIterator.cpp
--- a/Engine/Source/Runtime/Core/Private/Internationalization/LegacyCharacterBoundaryIterator.cpp
+++ b/Engine/Source/Runtime/Core/Private/Internationalization/LegacyCharacterBoundaryIterator.cpp
@@ -1,60 +1,105 @@
 // Copyright 1998-2014 Epic Games, Inc. All Rights Reserved.
 
 #include "Core.h"
-#include "CharacterBoundaryIterator.h"
+#include "BreakIterator.h"
 
 #if !UE_ENABLE_ICU
 
-FCharacterBoundaryIterator::FCharacterBoundaryIterator(const FText& Text)
-	: String( Text.ToString() ), CurrentPosition(0)
+class FLegacyCharacterBoundaryIterator : public IBreakIterator
 {
+public:
+	FLegacyCharacterBoundaryIterator();
+
+	virtual void SetString(const FText& InText) override;
+	virtual void SetString(const FString& InString) override;
+	virtual void SetString(const TCHAR* const InString, const int32 InStringLength) override;
+	virtual void ClearString() override;
+
+	virtual int32 GetCurrentPosition() const override;
+
+	virtual int32 ResetToBeginning() override;
+	virtual int32 ResetToEnd() override;
+
+	virtual int32 MoveToPrevious() override;
+	virtual int32 MoveToNext() override;
+	virtual int32 MoveToCandidateBefore(const int32 InIndex) override;
+	virtual int32 MoveToCandidateAfter(const int32 InIndex) override;
+
+private:
+	FString String;
+	int32 CurrentPosition;
+};
+
+FLegacyCharacterBoundaryIterator::FLegacyCharacterBoundaryIterator()
+	: String()
+	, CurrentPosition(0)
+{
+}
+
+void FLegacyCharacterBoundaryIterator::SetString(const FText& InText)
+{
+	String = InText.ToString();
+	ResetToBeginning();
+}
+
+void FLegacyCharacterBoundaryIterator::SetString(const FString& InString)
+{
+	String = InString;
+	ResetToBeginning();
 }
 
-FCharacterBoundaryIterator::FCharacterBoundaryIterator(const FString& InString)
-	: String( InString ), CurrentPosition(0)
+void FLegacyCharacterBoundaryIterator::SetString(const TCHAR* const InString, const int32 InStringLength) 
 {
+	String = FString(InString, InStringLength);
+	ResetToBeginning();
 }
 
-FCharacterBoundaryIterator::FCharacterBoundaryIterator(const TCHAR* const InString, const int32 StringLength)
-	: String( StringLength, InString ), CurrentPosition(0)
+void FLegacyCharacterBoundaryIterator::ClearString()
 {
+	String = FString();
+	ResetToBeginning();
 }
 
-int32 FCharacterBoundaryIterator::GetCurrentPosition() const
+int32 FLegacyCharacterBoundaryIterator::GetCurrentPosition() const
 {
 	return CurrentPosition;
 }
 
-int32 FCharacterBoundaryIterator::ResetToBeginning()
+int32 FLegacyCharacterBoundaryIterator::ResetToBeginning()
 {
 	return CurrentPosition = 0;
 }
 
-int32 FCharacterBoundaryIterator::ResetToEnd()
+int32 FLegacyCharacterBoundaryIterator::ResetToEnd()
 {
 	return CurrentPosition = String.Len();
 }
 
-int32 FCharacterBoundaryIterator::MoveToPrevious()
+int32 FLegacyCharacterBoundaryIterator::MoveToPrevious()
 {
 	return MoveToCandidateBefore(CurrentPosition);
 }
 
-int32 FCharacterBoundaryIterator::MoveToNext()
+int32 FLegacyCharacterBoundaryIterator::MoveToNext()
 {
 	return MoveToCandidateAfter(CurrentPosition);
 }
 
-int32 FCharacterBoundaryIterator::MoveToCandidateBefore(const int32 Index)
+int32 FLegacyCharacterBoundaryIterator::MoveToCandidateBefore(const int32 InIndex)
+{
+	CurrentPosition = FMath::Clamp( InIndex - 1, 0, String.Len() );
+	return CurrentPosition >= InIndex ? INDEX_NONE : CurrentPosition;
+}
+
+int32 FLegacyCharacterBoundaryIterator::MoveToCandidateAfter(const int32 InIndex)
 {
-	CurrentPosition = FMath::Clamp( Index - 1, 0, String.Len() );
-	return CurrentPosition >= Index ? -1 : CurrentPosition;
+	CurrentPosition = FMath::Clamp( InIndex + 1, 0, String.Len() );
+	return CurrentPosition <= InIndex ? INDEX_NONE : CurrentPosition;
 }
 
-int32 FCharacterBoundaryIterator::MoveToCandidateAfter(const int32 Index)
+TSharedRef<IBreakIterator> FBreakIterator::CreateCharacterBoundaryIterator()
 {
-	CurrentPosition = FMath::Clamp( Index + 1, 0, String.Len() );
-	return CurrentPosition <= Index ? -1 : CurrentPosition;
+	return MakeShareable(new FLegacyCharacterBoundaryIterator());
 }
 
-#endif
\ No newline at end of file
+#endif
diff --git a/Engine/Source/Runtime/Core/Private/Internationalization/LegacyLineBreakIterator.cpp b/Engine/Source/Runtime/Core/Private/Internationalization/LegacyLineBreakIterator.cpp
--- a/Engine/Source/Runtime/Core/Private/Internationalization/LegacyLineBreakIterator.cpp
+++ b/Engine/Source/Runtime/Core/Private/Internationalization/LegacyLineBreakIterator.cpp
@@ -1,54 +1,94 @@
 // Copyright 1998-2014 Epic Games, Inc. All Rights Reserved.
 
 #include "Core.h"
-#include "LineBreakIterator.h"
+#include "BreakIterator.h"
 
 #if !UE_ENABLE_ICU
 
-FLineBreakIterator::FLineBreakIterator(const FText& Text)
-	: String( Text.ToString() ), CurrentPosition(0)
+class FLegacyLineBreakIterator : public IBreakIterator
 {
+public:
+	FLegacyLineBreakIterator();
+
+	virtual void SetString(const FText& InText) override;
+	virtual void SetString(const FString& InString) override;
+	virtual void SetString(const TCHAR* const InString, const int32 InStringLength) override;
+	virtual void ClearString() override;
+
+	virtual int32 GetCurrentPosition() const override;
+
+	virtual int32 ResetToBeginning() override;
+	virtual int32 ResetToEnd() override;
+
+	virtual int32 MoveToPrevious() override;
+	virtual int32 MoveToNext() override;
+	virtual int32 MoveToCandidateBefore(const int32 InIndex) override;
+	virtual int32 MoveToCandidateAfter(const int32 InIndex) override;
+
+private:
+	FString String;
+	int32 CurrentPosition;
+};
+
+FLegacyLineBreakIterator::FLegacyLineBreakIterator()
+	: String()
+	, CurrentPosition(0)
+{
+}
+
+void FLegacyLineBreakIterator::SetString(const FText& InText)
+{
+	String = InText.ToString();
+	ResetToBeginning();
+}
+
+void FLegacyLineBreakIterator::SetString(const FString& InString)
+{
+	String = InString;
+	ResetToBeginning();
 }
 
-FLineBreakIterator::FLineBreakIterator(const FString& InString)
-	: String( InString ), CurrentPosition(0)
+void FLegacyLineBreakIterator::SetString(const TCHAR* const InString, const int32 InStringLength) 
 {
+	String = FString(InString, InStringLength);
+	ResetToBeginning();
 }
 
-FLineBreakIterator::FLineBreakIterator(const TCHAR* const InString, const int32 StringLength)
-	: String( StringLength, InString ), CurrentPosition(0)
+void FLegacyLineBreakIterator::ClearString()
 {
+	String = FString();
+	ResetToBeginning();
 }
 
-int32 FLineBreakIterator::GetCurrentPosition() const
+int32 FLegacyLineBreakIterator::GetCurrentPosition() const
 {
 	return CurrentPosition;
 }
 
-int32 FLineBreakIterator::ResetToBeginning()
+int32 FLegacyLineBreakIterator::ResetToBeginning()
 {
 	return CurrentPosition = 0;
 }
 
-int32 FLineBreakIterator::ResetToEnd()
+int32 FLegacyLineBreakIterator::ResetToEnd()
 {
 	return CurrentPosition = String.Len();
 }
 
-int32 FLineBreakIterator::MoveToPrevious()
+int32 FLegacyLineBreakIterator::MoveToPrevious()
 {
 	return MoveToCandidateBefore(CurrentPosition);
 }
 
-int32 FLineBreakIterator::MoveToNext()
+int32 FLegacyLineBreakIterator::MoveToNext()
 {
 	return MoveToCandidateAfter(CurrentPosition);
 }
 
-int32 FLineBreakIterator::MoveToCandidateBefore(const int32 Index)
+int32 FLegacyLineBreakIterator::MoveToCandidateBefore(const int32 InIndex)
 {
 	// Can break between char and whitespace.
-	for(CurrentPosition = FMath::Clamp( Index, 0, String.Len() ) - 1; CurrentPosition >= 0 + 1; --CurrentPosition)
+	for(CurrentPosition = FMath::Clamp( InIndex - 1, 0, String.Len() ); CurrentPosition >= 1; --CurrentPosition)
 	{
 		bool bPreviousCharWasWhitespace = FChar::IsWhitespace(String[CurrentPosition - 1]);
 		bool bCurrentCharIsWhiteSpace = FChar::IsWhitespace(String[CurrentPosition]);
@@ -58,13 +98,13 @@ int32 FLineBreakIterator::MoveToCandidateBefore(const int32 Index)
 		}
 	}
 
-	return CurrentPosition >= Index ? -1 : CurrentPosition;
+	return CurrentPosition >= InIndex ? INDEX_NONE : CurrentPosition;
 }
 
-int32 FLineBreakIterator::MoveToCandidateAfter(const int32 Index)
+int32 FLegacyLineBreakIterator::MoveToCandidateAfter(const int32 InIndex)
 {
 	// Can break between char and whitespace.
-	for(CurrentPosition = FMath::Clamp( Index, 0, String.Len() ) + 1; CurrentPosition <= String.Len() - 1; ++CurrentPosition)
+	for(CurrentPosition = FMath::Clamp( InIndex + 1, 0, String.Len() ); CurrentPosition < String.Len(); ++CurrentPosition)
 	{
 		bool bPreviousCharWasWhitespace = FChar::IsWhitespace(String[CurrentPosition - 1]);
 		bool bCurrentCharIsWhiteSpace = FChar::IsWhitespace(String[CurrentPosition]);
@@ -74,7 +114,12 @@ int32 FLineBreakIterator::MoveToCandidateAfter(const int32 Index)
 		}
 	}
 
-	return CurrentPosition <= Index ? -1 : CurrentPosition;
+	return (CurrentPosition <= InIndex) ? INDEX_NONE : CurrentPosition;
+}
+
+TSharedRef<IBreakIterator> FBreakIterator::CreateLineBreakIterator()
+{
+	return MakeShareable(new FLegacyLineBreakIterator());
 }
 
-#endif
\ No newline at end of file
+#endif
diff --git a/Engine/Source/Runtime/Core/Private/Internationalization/LegacyWordBreakIterator.cpp b/Engine/Source/Runtime/Core/Private/Internationalization/LegacyWordBreakIterator.cpp
--- a/Engine/Source/Runtime/Core/Private/Internationalization/LegacyWordBreakIterator.cpp
+++ b/Engine/Source/Runtime/Core/Private/Internationalization/LegacyWordBreakIterator.cpp
@@ -1,54 +1,94 @@
 // Copyright 1998-2014 Epic Games, Inc. All Rights Reserved.
 
 #include "Core.h"
-#include "WordBreakIterator.h"
+#include "BreakIterator.h"
 
 #if !UE_ENABLE_ICU
 
-FWordBreakIterator::FWordBreakIterator(const FText& Text)
-	: String( Text.ToString() ), CurrentPosition(0)
+class FLegacyWordBreakIterator : public IBreakIterator
 {
+public:
+	FLegacyWordBreakIterator();
+
+	virtual void SetString(const FText& InText) override;
+	virtual void SetString(const FString& InString) override;
+	virtual void SetString(const TCHAR* const InString, const int32 InStringLength) override;
+	virtual void ClearString() override;
+
+	virtual int32 GetCurrentPosition() const override;
+
+	virtual int32 ResetToBeginning() override;
+	virtual int32 ResetToEnd() override;
+
+	virtual int32 MoveToPrevious() override;
+	virtual int32 MoveToNext() override;
+	virtual int32 MoveToCandidateBefore(const int32 InIndex) override;
+	virtual int32 MoveToCandidateAfter(const int32 InIndex) override;
+
+private:
+	FString String;
+	int32 CurrentPosition;
+};
+
+FLegacyWordBreakIterator::FLegacyWordBreakIterator()
+	: String()
+	, CurrentPosition(0)
+{
+}
+
+void FLegacyWordBreakIterator::SetString(const FText& InText)
+{
+	String = InText.ToString();
+	ResetToBeginning();
+}
+
+void FLegacyWordBreakIterator::SetString(const FString& InString)
+{
+	String = InString;
+	ResetToBeginning();
 }
 
-FWordBreakIterator::FWordBreakIterator(const FString& InString)
-	: String( InString ), CurrentPosition(0)
+void FLegacyWordBreakIterator::SetString(const TCHAR* const InString, const int32 InStringLength) 
 {
+	String = FString(InString, InStringLength);
+	ResetToBeginning();
 }
 
-FWordBreakIterator::FWordBreakIterator(const TCHAR* const InString, const int32 StringLength)
-	: String( StringLength, InString ), CurrentPosition(0)
+void FLegacyWordBreakIterator::ClearString()
 {
+	String = FString();
+	ResetToBeginning();
 }
 
-int32 FWordBreakIterator::GetCurrentPosition() const
+int32 FLegacyWordBreakIterator::GetCurrentPosition() const
 {
 	return CurrentPosition;
 }
 
-int32 FWordBreakIterator::ResetToBeginning()
+int32 FLegacyWordBreakIterator::ResetToBeginning()
 {
 	return CurrentPosition = 0;
 }
 
-int32 FWordBreakIterator::ResetToEnd()
+int32 FLegacyWordBreakIterator::ResetToEnd()
 {
 	return CurrentPosition = String.Len();
 }
 
-int32 FWordBreakIterator::MoveToPrevious()
+int32 FLegacyWordBreakIterator::MoveToPrevious()
 {
 	return MoveToCandidateBefore(CurrentPosition);
 }
 
-int32 FWordBreakIterator::MoveToNext()
+int32 FLegacyWordBreakIterator::MoveToNext()
 {
 	return MoveToCandidateAfter(CurrentPosition);
 }
 
-int32 FWordBreakIterator::MoveToCandidateBefore(const int32 Index)
+int32 FLegacyWordBreakIterator::MoveToCandidateBefore(const int32 InIndex)
 {
 	// Can break between char and whitespace.
-	for(CurrentPosition = FMath::Clamp( Index, 0, String.Len() ) - 1; CurrentPosition >= 0 + 1; --CurrentPosition)
+	for(CurrentPosition = FMath::Clamp( InIndex - 1, 0, String.Len() ); CurrentPosition >= 1; --CurrentPosition)
 	{
 		bool bPreviousCharWasWhitespace = FChar::IsWhitespace(String[CurrentPosition - 1]);
 		bool bCurrentCharIsWhiteSpace = FChar::IsWhitespace(String[CurrentPosition]);
@@ -58,13 +98,13 @@ int32 FWordBreakIterator::MoveToCandidateBefore(const int32 Index)
 		}
 	}
 
-	return CurrentPosition >= Index ? -1 : CurrentPosition;
+	return CurrentPosition >= InIndex ? INDEX_NONE : CurrentPosition;
 }
 
-int32 FWordBreakIterator::MoveToCandidateAfter(const int32 Index)
+int32 FLegacyWordBreakIterator::MoveToCandidateAfter(const int32 InIndex)
 {
 	// Can break between char and whitespace.
-	for(CurrentPosition = FMath::Clamp( Index, 0, String.Len() ) + 1; CurrentPosition <= String.Len() - 1; ++CurrentPosition)
+	for(CurrentPosition = FMath::Clamp( InIndex + 1, 0, String.Len() ); CurrentPosition < String.Len(); ++CurrentPosition)
 	{
 		bool bPreviousCharWasWhitespace = FChar::IsWhitespace(String[CurrentPosition - 1]);
 		bool bCurrentCharIsWhiteSpace = FChar::IsWhitespace(String[CurrentPosition]);
@@ -74,7 +114,12 @@ int32 FWordBreakIterator::MoveToCandidateAfter(const int32 Index)
 		}
 	}
 
-	return CurrentPosition <= Index ? -1 : CurrentPosition;
+	return CurrentPosition <= InIndex ? INDEX_NONE : CurrentPosition;
+}
+
+TSharedRef<IBreakIterator> FBreakIterator::CreateWordBreakIterator()
+{
+	return MakeShareable(new FLegacyWordBreakIterator());
 }
 
 #endif
\ No newline at end of file
diff --git a/Engine/Source/Runtime/Core/Private/Internationalization/WordWrapper.cpp b/Engine/Source/Runtime/Core/Private/Internationalization/WordWrapper.cpp
--- a/Engine/Source/Runtime/Core/Private/Internationalization/WordWrapper.cpp
+++ b/Engine/Source/Runtime/Core/Private/Internationalization/WordWrapper.cpp
@@ -2,15 +2,19 @@
 
 #include "Core.h"
 #include "WordWrapper.h"
+#include "BreakIterator.h"
 
 FWordWrapper::FWordWrapper(const TCHAR* const InString, const int32 InStringLength, FWrappedLineData* const OutWrappedLineData)
 	: String(InString)
 	, StringLength(InStringLength)
-	, GraphemeBreakIterator(InString, InStringLength)
-	, LineBreakIterator(InString, InStringLength)
+	, GraphemeBreakIterator(FBreakIterator::CreateCharacterBoundaryIterator())
+	, LineBreakIterator(FBreakIterator::CreateLineBreakIterator())
 	, StartIndex(0)
 	, WrappedLineData(OutWrappedLineData)
 {
+	GraphemeBreakIterator->SetString(InString, InStringLength);
+	LineBreakIterator->SetString(InString, InStringLength);
+
 	if(WrappedLineData)
 	{
 		WrappedLineData->Empty();
@@ -139,7 +143,7 @@ int32 FWordWrapper::FindFirstMandatoryBreakBetween(const int32 InStartIndex, con
 
 int32 FWordWrapper::FindLastBreakCandidateBetween(const int32 InStartIndex, const int32 WrapIndex)
 {
-	int32 BreakIndex = LineBreakIterator.MoveToCandidateBefore(WrapIndex + 1);
+	int32 BreakIndex = LineBreakIterator->MoveToCandidateBefore(WrapIndex + 1);
 	if(BreakIndex < InStartIndex)
 	{
 		BreakIndex = INDEX_NONE;
@@ -149,7 +153,7 @@ int32 FWordWrapper::FindLastBreakCandidateBetween(const int32 InStartIndex, cons
 
 int32 FWordWrapper::FindEndOfLastWholeGraphemeCluster(const int32 InStartIndex, const int32 WrapIndex)
 {
-	int32 BreakIndex = GraphemeBreakIterator.MoveToCandidateBefore(WrapIndex + 1);
+	int32 BreakIndex = GraphemeBreakIterator->MoveToCandidateBefore(WrapIndex + 1);
 	if(BreakIndex < InStartIndex)
 	{
 		BreakIndex = INDEX_NONE;
diff --git a/Engine/Source/Runtime/Core/Public/Internationalization/BreakIterator.h b/Engine/Source/Runtime/Core/Public/Internationalization/BreakIterator.h
new file mode 100644
--- /dev/null
+++ b/Engine/Source/Runtime/Core/Public/Internationalization/BreakIterator.h
@@ -0,0 +1,29 @@
+// Copyright 1998-2014 Epic Games, Inc. All Rights Reserved.
+#pragma once
+
+#include "IBreakIterator.h"
+
+struct CORE_API FBreakIterator
+{
+	/**
+	 * Create a new instance of a break iterator designed to locate character boundary points within a string
+	 */
+	static TSharedRef<IBreakIterator> CreateCharacterBoundaryIterator();
+
+	/**
+	 * Create a new instance of a break iterator designed to locate word boundary points within a string
+	 */
+	static TSharedRef<IBreakIterator> CreateWordBreakIterator();
+
+	/**
+	 * Create a new instance of a break iterator designed to locate appropriate soft-wrapping points within a string
+	 * Note that these aren't newlines as such, but rather places where the line could be soft-wrapped when displaying the text
+	 */
+	static TSharedRef<IBreakIterator> CreateLineBreakIterator();
+
+	/**
+	 * Create a new instance of a break iterator which is designed to locate appropriate soft-wrapping points within a string by searching for CamelCase boundaries
+	 * This may produce odd results when used with languages which don't have a concept of cases for characters, however should be fine for parsing variable or asset names
+	 */
+	static TSharedRef<IBreakIterator> CreateCamelCaseBreakIterator();
+};
diff --git a/Engine/Source/Runtime/Core/Public/Internationalization/CharacterBoundaryIterator.h b/Engine/Source/Runtime/Core/Public/Internationalization/CharacterBoundaryIterator.h
deleted file mode 100644
diff --git a/Engine/Source/Runtime/Core/Public/Internationalization/IBreakIterator.h b/Engine/Source/Runtime/Core/Public/Internationalization/IBreakIterator.h
new file mode 100644
--- /dev/null
+++ b/Engine/Source/Runtime/Core/Public/Internationalization/IBreakIterator.h
@@ -0,0 +1,23 @@
+// Copyright 1998-2014 Epic Games, Inc. All Rights Reserved.
+#pragma once
+
+class IBreakIterator
+{
+public:
+	virtual ~IBreakIterator() {}
+
+	virtual void SetString(const FText& InText) = 0;
+	virtual void SetString(const FString& InString) = 0;
+	virtual void SetString(const TCHAR* const InString, const int32 InStringLength) = 0;
+	virtual void ClearString() = 0;
+
+	virtual int32 GetCurrentPosition() const = 0;
+
+	virtual int32 ResetToBeginning() = 0;
+	virtual int32 ResetToEnd() = 0;
+
+	virtual int32 MoveToPrevious() = 0;
+	virtual int32 MoveToNext() = 0;
+	virtual int32 MoveToCandidateBefore(const int32 InIndex) = 0;
+	virtual int32 MoveToCandidateAfter(const int32 InIndex) = 0;
+};
diff --git a/Engine/Source/Runtime/Core/Public/Internationalization/LineBreakIterator.h b/Engine/Source/Runtime/Core/Public/Internationalization/LineBreakIterator.h
deleted file mode 100644
diff --git a/Engine/Source/Runtime/Core/Public/Internationalization/WordBreakIterator.h b/Engine/Source/Runtime/Core/Public/Internationalization/WordBreakIterator.h
deleted file mode 100644
diff --git a/Engine/Source/Runtime/Core/Public/Internationalization/WordWrapper.h b/Engine/Source/Runtime/Core/Public/Internationalization/WordWrapper.h
--- a/Engine/Source/Runtime/Core/Public/Internationalization/WordWrapper.h
+++ b/Engine/Source/Runtime/Core/Public/Internationalization/WordWrapper.h
@@ -1,8 +1,7 @@
 // Copyright 1998-2014 Epic Games, Inc. All Rights Reserved.
 #pragma once
 
-#include "CharacterBoundaryIterator.h"
-#include "LineBreakIterator.h"
+#include "IBreakIterator.h"
 
 // Abstract Base Class for word wrapping functionality.
 class CORE_API FWordWrapper
@@ -55,8 +54,8 @@ protected:
 	const TCHAR* const String;
 	const int32 StringLength;
 private:
-	FCharacterBoundaryIterator GraphemeBreakIterator;
-	FLineBreakIterator LineBreakIterator;
+	TSharedPtr<IBreakIterator> GraphemeBreakIterator;
+	TSharedPtr<IBreakIterator> LineBreakIterator;
 	int32 StartIndex;
 	FWrappedLineData* WrappedLineData;
 };
\ No newline at end of file
diff --git a/Engine/Source/Runtime/Slate/Private/Framework/Text/SyntaxTokenizer.cpp b/Engine/Source/Runtime/Slate/Private/Framework/Text/SyntaxTokenizer.cpp
--- a/Engine/Source/Runtime/Slate/Private/Framework/Text/SyntaxTokenizer.cpp
+++ b/Engine/Source/Runtime/Slate/Private/Framework/Text/SyntaxTokenizer.cpp
@@ -2,7 +2,7 @@
 
 #include "SlatePrivatePCH.h"
 #include "SyntaxTokenizer.h"
-#include "WordBreakIterator.h"
+#include "BreakIterator.h"
 
 TSharedRef< FSyntaxTokenizer > FSyntaxTokenizer::Create(TArray<FRule> InRules)
 {
@@ -34,8 +34,10 @@ FSyntaxTokenizer::FSyntaxTokenizer(TArray<FRule> InRules)
 
 void FSyntaxTokenizer::TokenizeLineRanges(const FString& Input, const TArray<FTextRange>& LineRanges, TArray<FTokenizedLine>& OutTokenizedLines)
 {
+	TSharedRef<IBreakIterator> WBI = FBreakIterator::CreateWordBreakIterator();
+	WBI->SetString(Input);
+
 	// Tokenize line ranges
-	FWordBreakIterator WBI(Input);
 	for(const FTextRange& LineRange : LineRanges)
 	{
 		FTokenizedLine TokenizedLine;
@@ -73,7 +75,7 @@ void FSyntaxTokenizer::TokenizeLineRanges(const FString& Input, const TArray<FTe
 				}
 
 				// If none matched, consume the character(s) as text
-				const int32 NextWordBoundary = WBI.MoveToCandidateAfter(CurrentOffset);
+				const int32 NextWordBoundary = WBI->MoveToCandidateAfter(CurrentOffset);
 				const int32 TextTokenEnd = (NextWordBoundary == INDEX_NONE) ? LineRange.EndIndex : FMath::Min(NextWordBoundary, LineRange.EndIndex);
 				TokenizedLine.Tokens.Emplace(FToken(ETokenType::Literal, FTextRange(CurrentOffset, TextTokenEnd)));
 				CurrentOffset = TextTokenEnd;
diff --git a/Engine/Source/Runtime/Slate/Private/Framework/Text/TextLayout.cpp b/Engine/Source/Runtime/Slate/Private/Framework/Text/TextLayout.cpp
--- a/Engine/Source/Runtime/Slate/Private/Framework/Text/TextLayout.cpp
+++ b/Engine/Source/Runtime/Slate/Private/Framework/Text/TextLayout.cpp
@@ -1,8 +1,7 @@
 // Copyright 1998-2014 Epic Games, Inc. All Rights Reserved.
 
 #include "SlatePrivatePCH.h"
-#include "LineBreakIterator.h"
-#include "WordBreakIterator.h"
+#include "BreakIterator.h"
 
 FTextLayout::FBreakCandidate FTextLayout::CreateBreakCandidate( int32& OutRunIndex, FLineModel& Line, int32 PreviousBreak, int32 CurrentBreak )
 {
@@ -518,6 +517,12 @@ void FTextLayout::ClearView()
 
 void FTextLayout::CreateWrappingCache()
 {
+	// If we've not yet been provided with a custom line break iterator, then just use the default one
+	if (!LineBreakIterator.IsValid())
+	{
+		LineBreakIterator = FBreakIterator::CreateLineBreakIterator();
+	}
+
 	for (int LineModelIndex = 0; LineModelIndex < LineModels.Num(); LineModelIndex++)
 	{
 		FLineModel& LineModel = LineModels[ LineModelIndex ];
@@ -532,19 +537,21 @@ void FTextLayout::CreateWrappingCache()
 				LineModel.Runs[ RunIndex ].ClearCache();
 			}
 
-			FLineBreakIterator LineBreakIterator( **LineModel.Text );
+			LineBreakIterator->SetString( **LineModel.Text );
 
 			int32 PreviousBreak = 0;
 			int32 CurrentBreak = 0;
 			int32 CurrentRunIndex = 0;
 
-			while( ( CurrentBreak = LineBreakIterator.MoveToNext() ) != INDEX_NONE )
+			while( ( CurrentBreak = LineBreakIterator->MoveToNext() ) != INDEX_NONE )
 			{
 				LineModel.BreakCandidates.Add( CreateBreakCandidate(/*OUT*/CurrentRunIndex, LineModel, PreviousBreak, CurrentBreak) );
 				PreviousBreak = CurrentBreak;
 			}
 		}
 	}
+
+	LineBreakIterator->ClearString();
 }
 
 void FTextLayout::ClearWrappingCache()
@@ -556,7 +563,8 @@ void FTextLayout::ClearWrappingCache()
 	}
 }
 
-FTextLayout::FTextLayout() : LineModels()
+FTextLayout::FTextLayout() 
+	: LineModels()
 	, LineViews()
 	, DirtyFlags( EDirtyState::None )
 	, Scale( 1.0f )
@@ -564,6 +572,8 @@ FTextLayout::FTextLayout() : LineModels()
 	, Margin()
 	, LineHeightPercentage( 1.0f )
 	, DrawSize( ForceInitToZero )
+	, LineBreakIterator() // Initialized in FTextLayout::CreateWrappingCache if no custom iterator is provided
+	, WordBreakIterator(FBreakIterator::CreateWordBreakIterator())
 {
 
 }
@@ -1387,12 +1397,12 @@ FTextSelection FTextLayout::GetWordAt(const FTextLocation& Location) const
 
 	const FLineModel& LineModel = LineModels[LineIndex];
 
-	FWordBreakIterator WordBreakIterator(**LineModel.Text);
+	WordBreakIterator->SetString(**LineModel.Text);
 
-	int32 PreviousBreak = WordBreakIterator.MoveToCandidateAfter(Offset);
+	int32 PreviousBreak = WordBreakIterator->MoveToCandidateAfter(Offset);
 	int32 CurrentBreak = 0;
 
-	while ((CurrentBreak = WordBreakIterator.MoveToPrevious()) != INDEX_NONE)
+	while ((CurrentBreak = WordBreakIterator->MoveToPrevious()) != INDEX_NONE)
 	{
 		bool HasLetter = false;
 		for (int32 Index = CurrentBreak; Index < PreviousBreak; Index++)
@@ -1412,6 +1422,8 @@ FTextSelection FTextLayout::GetWordAt(const FTextLocation& Location) const
 		PreviousBreak = CurrentBreak;
 	}
 
+	WordBreakIterator->ClearString();
+
 	if (PreviousBreak == CurrentBreak || CurrentBreak == INDEX_NONE)
 	{
 		return FTextSelection();
@@ -1420,6 +1432,12 @@ FTextSelection FTextLayout::GetWordAt(const FTextLocation& Location) const
 	return FTextSelection(FTextLocation(LineIndex, CurrentBreak), FTextLocation(LineIndex, PreviousBreak));
 }
 
+void FTextLayout::SetLineBreakIterator( TSharedPtr<IBreakIterator> InLineBreakIterator )
+{
+	LineBreakIterator = InLineBreakIterator;
+	DirtyFlags |= EDirtyState::Layout;
+}
+
 void FTextLayout::SetMargin( const FMargin& InMargin )
 {
 	if ( Margin == InMargin )
diff --git a/Engine/Source/Runtime/Slate/Public/Framework/Text/TextLayout.h b/Engine/Source/Runtime/Slate/Public/Framework/Text/TextLayout.h
--- a/Engine/Source/Runtime/Slate/Public/Framework/Text/TextLayout.h
+++ b/Engine/Source/Runtime/Slate/Public/Framework/Text/TextLayout.h
@@ -1,6 +1,7 @@
 // Copyright 1998-2014 Epic Games, Inc. All Rights Reserved.
 #pragma once
 
+#include "IBreakIterator.h"
 #include "TextLayout.generated.h"
 
 #define TEXT_LAYOUT_DEBUG 0
@@ -291,6 +292,9 @@ public:
 	FMargin GetMargin() const;
 	void SetMargin( const FMargin& InMargin );
 
+	/** Set the iterator to use to detect appropriate soft-wrapping points for lines (or null to go back to using the default) */
+	void SetLineBreakIterator( TSharedPtr<IBreakIterator> InLineBreakIterator );
+
 	void ClearLines();
 
 	void AddLine( const TSharedRef< FString >& Text, const TArray< TSharedRef< IRun > >& Runs );
@@ -458,4 +462,10 @@ protected:
 
 	/** The final size of the text layout on screen. */
 	FVector2D DrawSize;
+
+	/** The iterator to use to detect appropriate soft-wrapping points for lines */
+	TSharedPtr<IBreakIterator> LineBreakIterator;
+
+	/** The iterator to use to detect word boundaries */
+	TSharedPtr<IBreakIterator> WordBreakIterator;
 };
\ No newline at end of file
