Added an IsEmptyOrWhitespace function to FText

diff --git a/Engine/Source/Runtime/Core/Private/Internationalization/Text.cpp b/Engine/Source/Runtime/Core/Private/Internationalization/Text.cpp
--- a/Engine/Source/Runtime/Core/Private/Internationalization/Text.cpp
+++ b/Engine/Source/Runtime/Core/Private/Internationalization/Text.cpp
@@ -291,6 +291,25 @@ FText::FText( FString InSourceString, FString InNamespace, FString InKey, int32
 	History = MakeShareable(new FTextHistory_Base(InSourceString));
 }
 
+
+bool FText::IsEmptyOrWhitespace() const
+{
+	if (DisplayString.Get().IsEmpty())
+	{
+		return true;
+	}
+
+	for( const TCHAR Character : DisplayString.Get() )
+	{
+		if (!IsWhitespace(Character))
+		{
+			return false;
+		}
+	}
+
+	return true;
+}
+
 FText FText::TrimPreceding( const FText& InText )
 {
 	FString TrimmedString = InText.ToString();
diff --git a/Engine/Source/Runtime/Core/Public/Internationalization/Text.h b/Engine/Source/Runtime/Core/Public/Internationalization/Text.h
--- a/Engine/Source/Runtime/Core/Public/Internationalization/Text.h
+++ b/Engine/Source/Runtime/Core/Public/Internationalization/Text.h
@@ -238,6 +238,8 @@ public:
 		return DisplayString.Get().IsEmpty();
 	}
 
+	bool IsEmptyOrWhitespace() const;
+
 	/**
 	 * Removes whitespace characters from the front of the string.
 	 */
