Support more output devices on Linux.

- Partially based on PR #306 submission, corrected for cross-compiling.
- Each platform is now able to set up output devices for logging according to its own needs.

diff --git a/Engine/Source/Programs/UnrealBuildTool/Linux/LinuxToolChain.cs b/Engine/Source/Programs/UnrealBuildTool/Linux/LinuxToolChain.cs
--- a/Engine/Source/Programs/UnrealBuildTool/Linux/LinuxToolChain.cs
+++ b/Engine/Source/Programs/UnrealBuildTool/Linux/LinuxToolChain.cs
@@ -90,9 +90,14 @@ namespace UnrealBuildTool
 
             if (CrossCompiling())
             {
-                // There are exceptions used in the code base (e.g. UnrealHeadTool).
-                // So this flag cannot be used, at least not for native linux builds.
+                // There are exceptions used in the code base (e.g. UnrealHeadTool).  @todo: weed out exceptions
+                // So this flag cannot be used, at least not for native Linux builds.
                 Result += " -fno-exceptions";               // no exceptions
+                Result += " -DPLATFORM_EXCEPTIONS_DISABLED=1";
+            }
+            else
+            {
+                Result += " -DPLATFORM_EXCEPTIONS_DISABLED=0";
             }
 
             Result += " -Wall -Werror";
diff --git a/Engine/Source/Runtime/Core/Private/GenericPlatform/GenericPlatformOutputDevices.cpp b/Engine/Source/Runtime/Core/Private/GenericPlatform/GenericPlatformOutputDevices.cpp
--- a/Engine/Source/Runtime/Core/Private/GenericPlatform/GenericPlatformOutputDevices.cpp
+++ b/Engine/Source/Runtime/Core/Private/GenericPlatform/GenericPlatformOutputDevices.cpp
@@ -3,6 +3,27 @@
 #include "Core.h"
 #include "FeedbackContextAnsi.h"
 
+void							FGenericPlatformOutputDevices::SetupOutputDevices()
+{
+	check(GLog);
+	check(GLogConsole);
+
+	GLog->AddOutputDevice(FPlatformOutputDevices::GetLog());
+
+	if (!FParse::Param(FCommandLine::Get(), TEXT("NOCONSOLE")))
+	{
+		GLog->AddOutputDevice(GLogConsole);
+	}
+
+	// Only create debug output device if a debugger is attached or we're running on a console or build machine
+	// A shipping build with logging explicitly enabled will fail the IsDebuggerPresent() check, but we still need to add the debug output device for logging purposes
+	if (!FPlatformProperties::SupportsWindowedMode() || FPlatformMisc::IsDebuggerPresent() || (UE_BUILD_SHIPPING && !NO_LOGGING) || GIsBuildMachine)
+	{
+		GLog->AddOutputDevice(new FOutputDeviceDebug());
+	}
+
+	GLog->AddOutputDevice(FPlatformOutputDevices::GetEventLog());
+};
 
 FString								FGenericPlatformOutputDevices::GetAbsoluteLogFilename()
 {
diff --git a/Engine/Source/Runtime/Core/Private/Linux/LinuxPlatformFeedbackContextPrivate.h b/Engine/Source/Runtime/Core/Private/Linux/LinuxPlatformFeedbackContextPrivate.h
new file mode 100644
--- /dev/null
+++ b/Engine/Source/Runtime/Core/Private/Linux/LinuxPlatformFeedbackContextPrivate.h
@@ -0,0 +1,112 @@
+// Copyright 1998-2014 Epic Games, Inc. All Rights Reserved.
+
+#include <syslog.h>
+
+#pragma once
+
+/**
+ * Feedback context implementation for Linux.
+ */
+class FFeedbackContextLinux : public FFeedbackContext
+{
+	/** Context information for warning and error messages */
+	FContextSupplier*	Context;
+
+public:
+	// Variables.
+	int32					SlowTaskCount;
+
+	// Constructor.
+	FFeedbackContextLinux()
+	: FFeedbackContext()
+	, Context( NULL )
+	, SlowTaskCount( 0 )
+	{}
+
+	void Serialize( const TCHAR* V, ELogVerbosity::Type Verbosity, const class FName& Category )
+	{
+		// if we set the color for warnings or errors, then reset at the end of the function
+		// note, we have to set the colors directly without using the standard SET_WARN_COLOR macro
+		if (Verbosity==ELogVerbosity::Error || Verbosity==ELogVerbosity::Warning)
+		{
+			if (TreatWarningsAsErrors && Verbosity==ELogVerbosity::Warning)
+			{
+				Verbosity = ELogVerbosity::Error;
+			}
+
+			FString Prefix;
+			if (Context)
+			{
+				Prefix = Context->GetContext() + TEXT(" : ");
+			}
+			FString Format = Prefix + FOutputDevice::FormatLogLine(Verbosity, Category, V);
+
+			if (Verbosity == ELogVerbosity::Error)
+			{
+				// Only store off the message if running a commandlet.
+				if (IsRunningCommandlet())
+				{
+					Errors.Add(Format);
+				}
+			}
+			else
+			{
+				// Only store off the message if running a commandlet.
+				if ( IsRunningCommandlet() )
+				{
+					Warnings.Add(Format);
+				}
+			}
+
+			// send to syslog too (for zabbix etc)
+			int BasePriority = (Verbosity == ELogVerbosity::Error) ? LOG_ERR : LOG_WARNING;
+			syslog(BasePriority | LOG_USER, "%s", StringCast< char >(*Format).Get());
+		}
+
+		if( GLogConsole && IsRunningCommandlet() )
+		{
+			GLogConsole->Serialize( V, Verbosity, Category );
+		}
+		if( !GLog->IsRedirectingTo( this ) )
+		{
+			GLog->Serialize( V, Verbosity, Category );
+		}
+	}
+	VARARG_BODY( bool, YesNof, const TCHAR*, VARARG_NONE )
+	{
+		TCHAR TempStr[4096];
+		GET_VARARGS( TempStr, ARRAY_COUNT(TempStr), ARRAY_COUNT(TempStr)-1, Fmt, Fmt );
+		if( ( GIsClient || GIsEditor ) && ( ( GIsSilent != true ) && ( FApp::IsUnattended() != true ) ) )
+		{
+			//return( ::MessageBox( NULL, TempStr, *NSLOCTEXT("Core", "Question", "Question").ToString(), MB_YESNO|MB_TASKMODAL ) == IDYES);
+			STUBBED("+++++++++++++++ LINUXPLATFORMFEEDBACKCONTEXTPRIVATE.H DIALOG BOX PROMPT +++++++++++++++++++");
+			return true;
+		}
+		else
+		{
+			return false;
+		}
+	}
+
+	void BeginSlowTask( const FText& Task, bool ShowProgressDialog, bool bShowCancelButton=false )
+	{
+		GIsSlowTask = ++SlowTaskCount>0;
+	}
+	void EndSlowTask()
+	{
+		check(SlowTaskCount>0);
+		GIsSlowTask = --SlowTaskCount>0;
+	}
+	virtual bool StatusUpdate( int32 Numerator, int32 Denominator, const FText& StatusText )
+	{
+		return true;
+	}
+	FContextSupplier* GetContext() const
+	{
+		return Context;
+	}
+	void SetContext( FContextSupplier* InSupplier )
+	{
+		Context = InSupplier;
+	}
+};
diff --git a/Engine/Source/Runtime/Core/Private/Linux/LinuxPlatformMisc.cpp b/Engine/Source/Runtime/Core/Private/Linux/LinuxPlatformMisc.cpp
--- a/Engine/Source/Runtime/Core/Private/Linux/LinuxPlatformMisc.cpp
+++ b/Engine/Source/Runtime/Core/Private/Linux/LinuxPlatformMisc.cpp
@@ -174,6 +174,12 @@ uint32 FLinuxPlatformMisc::GetCharKeyMap(uint16* KeyCodes, FString* KeyNames, ui
 	return FGenericPlatformMisc::GetStandardPrintableKeyMap(KeyCodes, KeyNames, MaxMappings, false, true);
 }
 
+void FLinuxPlatformMisc::LowLevelOutputDebugString(const TCHAR *Message)
+{
+	static_assert(PLATFORM_USE_LS_SPEC_FOR_WIDECHAR, "Check printf format");
+	fprintf(stderr, "%ls", Message);	// there's no good way to implement that really
+}
+
 uint32 FLinuxPlatformMisc::GetKeyMap( uint16* KeyCodes, FString* KeyNames, uint32 MaxMappings )
 {
 #define ADDKEYMAP(KeyCode, KeyName)		if (NumMappings<MaxMappings) { KeyCodes[NumMappings]=KeyCode; KeyNames[NumMappings]=KeyName; ++NumMappings; };
diff --git a/Engine/Source/Runtime/Core/Private/Linux/LinuxPlatformOutputDevices.cpp b/Engine/Source/Runtime/Core/Private/Linux/LinuxPlatformOutputDevices.cpp
--- a/Engine/Source/Runtime/Core/Private/Linux/LinuxPlatformOutputDevices.cpp
+++ b/Engine/Source/Runtime/Core/Private/Linux/LinuxPlatformOutputDevices.cpp
@@ -8,28 +8,228 @@
 #include "Core.h"
 #include "LinuxPlatformOutputDevices.h"
 #include "FeedbackContextAnsi.h"
-#include <syslog.h>
+#include "LinuxApplication.h"
+#include "LinuxPlatformOutputDevicesPrivate.h"
+#include "LinuxPlatformFeedbackContextPrivate.h"
 
-class FFeedbackContextLinux : public FFeedbackContextAnsi
+#define COLOR_RED     "\x1b[31m"
+#define COLOR_GREEN   "\x1b[32m"
+#define COLOR_YELLOW  "\x1b[33m"
+#define COLOR_BLUE    "\x1b[34m"
+#define COLOR_NONE   "\x1b[0m"
+
+//////////////////////////////////
+// FLinuxOutputDevices
+//////////////////////////////////
+
+void								FLinuxOutputDevices::SetupOutputDevices()
 {
-public:
+	check(GLog);
+	check(GLogConsole);
 
-	void Serialize(const TCHAR* V, ELogVerbosity::Type Verbosity, const class FName& Category)
-	{
-		if (Verbosity == ELogVerbosity::Error || Verbosity == ELogVerbosity::Warning)
-		{
-			FString Format = FOutputDevice::FormatLogLine(Verbosity, Category, V);
+	// add file log
+	GLog->AddOutputDevice(FPlatformOutputDevices::GetLog());
 
-			int BasePriority = (Verbosity == ELogVerbosity::Error) ? LOG_ERR : LOG_WARNING;
-			syslog(BasePriority | LOG_USER, "%s", StringCast< char >(*Format).Get());
-		}
+	// @todo: set to false for minor utils?
+	bool bLogToConsole = !NO_LOGGING && !FParse::Param(FCommandLine::Get(), TEXT("NOCONSOLE"));
 
-		FFeedbackContextAnsi::Serialize(V, Verbosity, Category);
+	if (bLogToConsole)
+	{
+		GLog->AddOutputDevice(GLogConsole);
 	}
-};
+
+	// debug and event logging is not really supported on Linux. 
+}
+
+class FOutputDeviceError*			FLinuxOutputDevices::GetError()
+{
+	static FOutputDeviceLinuxError Singleton;
+	return &Singleton;
+}
+
+class FOutputDevice*				FLinuxOutputDevices::GetEventLog()
+{
+	return NULL; // @TODO No event logging
+}
+
+class FOutputDeviceConsole*			FLinuxOutputDevices::GetLogConsole()
+{
+	// this is a slightly different kind of singleton that gives ownership to the caller and should not be called more than once
+	return new FOutputDeviceConsoleLinux();
+}
 
 class FFeedbackContext*				FLinuxOutputDevices::GetWarn()
 {
 	static FFeedbackContextLinux Singleton;
 	return &Singleton;
 }
+
+//////////////////////////////////
+// FOutputDeviceLinuxError
+//////////////////////////////////
+
+FOutputDeviceLinuxError::FOutputDeviceLinuxError()
+:	ErrorPos(0)
+{}
+
+void FOutputDeviceLinuxError::Serialize(const TCHAR* Msg, ELogVerbosity::Type Verbosity, const class FName& Category)
+{
+	FPlatformMisc::DebugBreak();
+
+	if (!GIsCriticalError)
+	{
+		// First appError.
+		GIsCriticalError = 1;
+		TCHAR ErrorBuffer[1024];
+		ErrorBuffer[0] = 0;
+		// pop up a crash window if we are not in unattended mode; 
+		if (WITH_EDITOR && FApp::IsUnattended() == false)
+		{
+			// @TODO Not implemented
+			UE_LOG(LogLinux, Error, TEXT("appError called: %s"), Msg );
+		}
+		// log the warnings to the log
+		else
+		{
+			UE_LOG(LogLinux, Error, TEXT("appError called: %s"), Msg );
+		}
+		FCString::Strncpy( GErrorHist, Msg, ARRAY_COUNT(GErrorHist) - 5 );
+		FCString::Strncat( GErrorHist, TEXT("\r\n\r\n"), ARRAY_COUNT(GErrorHist) - 1 );
+		ErrorPos = FCString::Strlen(GErrorHist);
+	}
+	else
+	{
+		UE_LOG(LogLinux, Error, TEXT("Error reentered: %s"), Msg );
+	}
+
+	if( GIsGuarded )
+	{
+		// Propagate error so structured exception handler can perform necessary work.
+#if PLATFORM_EXCEPTIONS_DISABLED
+		FPlatformMisc::DebugBreak();
+#endif
+		FPlatformMisc::RaiseException(1);
+	}
+	else
+	{
+		// We crashed outside the guarded code (e.g. appExit).
+		HandleError();
+		FPlatformMisc::RequestExit(true);
+	}
+}
+
+void FOutputDeviceLinuxError::HandleError()
+{
+	// make sure we don't report errors twice
+	static int32 CallCount = 0;
+	int32 NewCallCount = FPlatformAtomics::InterlockedIncrement(&CallCount);
+	if (NewCallCount != 1)
+	{
+		UE_LOG(LogLinux, Error, TEXT("HandleError re-entered.") );
+		return;
+	}
+
+	// Trigger the OnSystemFailure hook if it exists
+	FCoreDelegates::OnHandleSystemError.Broadcast();
+
+#if !PLATFORM_EXCEPTIONS_DISABLED
+	try
+	{
+#endif // !PLATFORM_EXCEPTIONS_DISABLED
+		GIsGuarded = 0;
+		GIsRunning = 0;
+		GIsCriticalError = 1;
+		GLogConsole = NULL;
+		GErrorHist[ARRAY_COUNT(GErrorHist)-1] = 0;
+
+		// Dump the error and flush the log.
+		UE_LOG(LogLinux, Log, TEXT("=== Critical error: ===") LINE_TERMINATOR TEXT("%s") LINE_TERMINATOR, GErrorExceptionDescription);
+		UE_LOG(LogLinux, Log, GErrorHist);
+
+		GLog->Flush();
+
+		FPlatformMisc::ClipboardCopy(GErrorHist);
+		FPlatformMisc::SubmitErrorReport(GErrorHist, EErrorReportMode::Interactive);
+		FCoreDelegates::OnShutdownAfterError.Broadcast();
+#if !PLATFORM_EXCEPTIONS_DISABLED
+	}
+	catch(...)
+	{}
+#endif // !PLATFORM_EXCEPTIONS_DISABLED
+}
+
+////////////////////////////////////////
+// FOutputDeviceConsoleLinux
+////////////////////////////////////////
+
+FOutputDeviceConsoleLinux::FOutputDeviceConsoleLinux()
+	: bOverrideColorSet(false)
+{
+}
+
+FOutputDeviceConsoleLinux::~FOutputDeviceConsoleLinux()
+{
+}
+
+void FOutputDeviceConsoleLinux::Show(bool bShowWindow)
+{
+}
+
+bool FOutputDeviceConsoleLinux::IsShown()
+{
+	return true;
+}
+
+void FOutputDeviceConsoleLinux::Serialize(const TCHAR* Data, ELogVerbosity::Type Verbosity, const class FName& Category)
+{
+	static bool bEntry=false;
+	if (!GIsCriticalError || bEntry)
+	{
+		if (Verbosity == ELogVerbosity::SetColor)
+		{
+			printf("%s", TCHAR_TO_UTF8(Data));
+		}
+		else
+		{
+			bool bNeedToResetColor = false;
+			if (!bOverrideColorSet)
+			{
+				if (Verbosity == ELogVerbosity::Error)
+				{
+					printf(COLOR_RED);
+					bNeedToResetColor = true;
+				}
+				else if (Verbosity == ELogVerbosity::Warning)
+				{
+					printf(COLOR_YELLOW);
+					bNeedToResetColor = true;
+				}
+			}
+			TCHAR OutputString[MAX_SPRINTF]=TEXT(""); //@warning: this is safe as FCString::Sprintf only use 1024 characters max
+			FCString::Sprintf(OutputString,TEXT("%s%s"),*FOutputDevice::FormatLogLine(Verbosity, Category, Data, GPrintLogTimes),LINE_TERMINATOR);
+			printf("%s", TCHAR_TO_UTF8(OutputString));
+
+			if (bNeedToResetColor)
+			{
+				printf(COLOR_NONE);
+			}
+		}
+	}
+	else
+	{
+		bEntry=true;
+#if !PLATFORM_EXCEPTIONS_DISABLED
+		try
+		{
+#endif // !PLATFORM_EXCEPTIONS_DISABLED
+			// Ignore errors to prevent infinite-recursive exception reporting.
+			Serialize(Data, Verbosity, Category);
+#if !PLATFORM_EXCEPTIONS_DISABLED
+		}
+		catch (...)
+		{
+		}
+#endif // !PLATFORM_EXCEPTIONS_DISABLED
+		bEntry = false;
+	}
+}
diff --git a/Engine/Source/Runtime/Core/Private/Linux/LinuxPlatformOutputDevicesPrivate.h b/Engine/Source/Runtime/Core/Private/Linux/LinuxPlatformOutputDevicesPrivate.h
new file mode 100644
--- /dev/null
+++ b/Engine/Source/Runtime/Core/Private/Linux/LinuxPlatformOutputDevicesPrivate.h
@@ -0,0 +1,72 @@
+// Copyright 1998-2014 Epic Games, Inc. All Rights Reserved.
+
+class FOutputDeviceLinuxError : public FOutputDeviceError
+{
+public:
+
+	/** Constructor, initializing member variables */
+	CORE_API FOutputDeviceLinuxError();
+
+	/**
+	 * Serializes the passed in data unless the current event is suppressed.
+	 *
+	 * @param	Data	Text to log
+	 * @param	Event	Event name used for suppression purposes
+	 */
+	virtual void Serialize(const TCHAR* Msg, ELogVerbosity::Type Verbosity, const class FName& Category) override;
+
+	/**
+	 * Error handling function that is being called from within the system wide global
+	 * error handler, e.g. using structured exception handling on the PC.
+	 */
+	void HandleError();
+
+private:
+
+	int32		ErrorPos;
+};
+
+/**
+* Linux implementation of console log window, just a printf to the terminal for now
+*/
+class CORE_API FOutputDeviceConsoleLinux : public FOutputDeviceConsole
+{
+
+private:
+
+	/** true if the color is currently set by the caller */
+	bool bOverrideColorSet;
+
+public:
+
+	/**
+	 * Constructor, setting console control handler.
+	 */
+	FOutputDeviceConsoleLinux();
+	~FOutputDeviceConsoleLinux();
+
+	/**
+	 * Shows or hides the console window.
+	 *
+	 * @param ShowWindow	Whether to show (true) or hide (false) the console window.
+	 */
+	virtual void Show(bool bShowWindow);
+
+	/**
+	 * Returns whether console is currently shown or not
+	 *
+	 * @return true if console is shown, false otherwise
+	 */
+	virtual bool IsShown();
+
+	virtual bool IsAttached() {return false;}
+
+	/**
+	 * Displays text on the console and scrolls if necessary.
+	 *
+	 * @param Data	Text to display
+	 * @param Event	Event type, used for filtering/ suppression
+	 */
+	void Serialize(const TCHAR* Data, ELogVerbosity::Type Verbosity, const class FName& Category);
+
+};
diff --git a/Engine/Source/Runtime/Core/Public/GenericPlatform/GenericPlatformOutputDevices.h b/Engine/Source/Runtime/Core/Public/GenericPlatform/GenericPlatformOutputDevices.h
--- a/Engine/Source/Runtime/Core/Public/GenericPlatform/GenericPlatformOutputDevices.h
+++ b/Engine/Source/Runtime/Core/Public/GenericPlatform/GenericPlatformOutputDevices.h
@@ -12,6 +12,11 @@
 **/
 struct CORE_API FGenericPlatformOutputDevices
 {
+	/**
+	 * Adds output devices which can vary depending on platform, configuration, command line parameters.
+	 */
+	static void							SetupOutputDevices();
+
 	static FString						GetAbsoluteLogFilename();
 
 	static class FOutputDevice*			GetLog();
diff --git a/Engine/Source/Runtime/Core/Public/Linux/LinuxPlatform.h b/Engine/Source/Runtime/Core/Public/Linux/LinuxPlatform.h
--- a/Engine/Source/Runtime/Core/Public/Linux/LinuxPlatform.h
+++ b/Engine/Source/Runtime/Core/Public/Linux/LinuxPlatform.h
@@ -43,7 +43,6 @@ typedef FLinuxPlatformTypes FPlatformTypes;
 #define PLATFORM_HAS_NO_EPROCLIM					1
 #define PLATFORM_HAS_BSD_SOCKET_FEATURE_IOCTL		1
 #define PLATFORM_HAS_BSD_IPV6_SOCKETS				1
-#define PLATFORM_EXCEPTIONS_DISABLED				1
 
 #define PLATFORM_USES_DYNAMIC_RHI					1
 #define PLATFORM_ENABLE_VECTORINTRINSICS			1
diff --git a/Engine/Source/Runtime/Core/Public/Linux/LinuxPlatformMisc.h b/Engine/Source/Runtime/Core/Public/Linux/LinuxPlatformMisc.h
--- a/Engine/Source/Runtime/Core/Public/Linux/LinuxPlatformMisc.h
+++ b/Engine/Source/Runtime/Core/Public/Linux/LinuxPlatformMisc.h
@@ -103,6 +103,7 @@ struct CORE_API FLinuxPlatformMisc : public FGenericPlatformMisc
 	static void PumpMessages(bool bFromMainLoop);
 	static uint32 GetKeyMap( uint16* KeyCodes, FString* KeyNames, uint32 MaxMappings );
 	static uint32 GetCharKeyMap(uint16* KeyCodes, FString* KeyNames, uint32 MaxMappings);
+	static void LowLevelOutputDebugString(const TCHAR *Message);
 	static bool ControlScreensaver(EScreenSaverAction Action);
 
 	static const TCHAR* RootDir();
diff --git a/Engine/Source/Runtime/Core/Public/Linux/LinuxPlatformOutputDevices.h b/Engine/Source/Runtime/Core/Public/Linux/LinuxPlatformOutputDevices.h
--- a/Engine/Source/Runtime/Core/Public/Linux/LinuxPlatformOutputDevices.h
+++ b/Engine/Source/Runtime/Core/Public/Linux/LinuxPlatformOutputDevices.h
@@ -9,6 +9,11 @@
 
 struct CORE_API FLinuxOutputDevices : public FGenericPlatformOutputDevices
 {
+	static void							SetupOutputDevices();
+
+	static class FOutputDevice*			GetEventLog();
+    static class FOutputDeviceConsole*	GetLogConsole();
+    static class FOutputDeviceError*	GetError();
 	static class FFeedbackContext*		GetWarn();
 };
 
diff --git a/Engine/Source/Runtime/Launch/Private/LaunchEngineLoop.cpp b/Engine/Source/Runtime/Launch/Private/LaunchEngineLoop.cpp
--- a/Engine/Source/Runtime/Launch/Private/LaunchEngineLoop.cpp
+++ b/Engine/Source/Runtime/Launch/Private/LaunchEngineLoop.cpp
@@ -2313,21 +2313,7 @@ bool FEngineLoop::AppInit( )
 #endif
 
 	// Init logging to disk
-	GLog->AddOutputDevice(FPlatformOutputDevices::GetLog());
-
-	if (!FParse::Param(FCommandLine::Get(),TEXT("NOCONSOLE")))
-	{
-		GLog->AddOutputDevice(GLogConsole);
-	}
-
-	// Only create debug output device if a debugger is attached or we're running on a console or build machine
-	// A shipping build with logging explicitly enabled will fail the IsDebuggerPresent() check, but we still need to add the debug output device for logging purposes
-	if (!FPlatformProperties::SupportsWindowedMode() || FPlatformMisc::IsDebuggerPresent() || (UE_BUILD_SHIPPING && !NO_LOGGING) || GIsBuildMachine)
-	{
-		GLog->AddOutputDevice( new FOutputDeviceDebug() );
-	}
-
-	GLog->AddOutputDevice(FPlatformOutputDevices::GetEventLog());
+	FPlatformOutputDevices::SetupOutputDevices();
 
 	// init config system
 	FConfigCacheIni::InitializeConfigSystem();
