As a temporary band-aid for command-key suppressing keyup events on Mac OS X, intercept them and handle them directly.

diff --git a/Engine/Source/Runtime/Core/Private/Mac/MacPlatformMisc.cpp b/Engine/Source/Runtime/Core/Private/Mac/MacPlatformMisc.cpp
--- a/Engine/Source/Runtime/Core/Private/Mac/MacPlatformMisc.cpp
+++ b/Engine/Source/Runtime/Core/Private/Mac/MacPlatformMisc.cpp
@@ -473,7 +473,7 @@ void FMacPlatformMisc::PumpMessages( bool bFromMainLoop )
 			const bool bIsMouseClickOrKeyEvent = [Event type] == NSLeftMouseDown || [Event type] == NSLeftMouseUp
 										 || [Event type] == NSRightMouseDown || [Event type] == NSRightMouseUp
 										 || [Event type] == NSOtherMouseDown || [Event type] == NSOtherMouseUp
-										 || [Event type] == NSKeyDown || [Event type] == NSKeyUp;
+										 || [Event type] == NSKeyDown || ([Event type] == NSKeyUp && !([Event modifierFlags] & NSCommandKeyMask));
 
 			if( MacApplication )
 			{
