UE4 - fix lifo task graph behavior. Fix permanent stats "single stepping" when no stats are enabled, prevent draw events from being emitted when parallel rendering

diff --git a/Engine/Source/Runtime/Core/Private/Async/TaskGraph.cpp b/Engine/Source/Runtime/Core/Private/Async/TaskGraph.cpp
--- a/Engine/Source/Runtime/Core/Private/Async/TaskGraph.cpp
+++ b/Engine/Source/Runtime/Core/Private/Async/TaskGraph.cpp
@@ -8,6 +8,7 @@ DEFINE_LOG_CATEGORY_STATIC(LogTaskGraph, Log, All);
 
 DEFINE_STAT(STAT_FReturnGraphTask);
 DEFINE_STAT(STAT_FTriggerEventGraphTask);
+DEFINE_STAT(STAT_UnknownGraphTask);
 
 namespace ENamedThreads
 {
@@ -114,6 +115,9 @@ private:
 
 };
 
+// this is used to signify a task that is just a call to wake up
+// It is generally bad to reuse pointers for non-pointer data, but efficiency is important here
+static FBaseGraphTask* WakeUpBaseGraphTask = (FBaseGraphTask*)0x3;
 
 /** 
  *	FTaskThread
@@ -206,10 +210,12 @@ public:
 		}
 		else if (ThreadId == ENamedThreads::RenderThread)
 		{
-			//StatName = none, we need to let the scope empty so that the render thread submits tasks in a timely manner
-			// The render thread has a custom method to calculate the idle time, so leave it as it is for now.
-			//StallStatId = GET_STATID(STAT_TaskGraph_RenderStalls);
-			//bCountAsStall = true;
+			if (QueueIndex > 0)
+			{
+				StallStatId = GET_STATID(STAT_TaskGraph_RenderStalls);
+				bCountAsStall = true;
+			}
+			// else StatName = none, we need to let the scope empty so that the render thread submits tasks in a timely manner. 
 		}
 		else if (ThreadId != ENamedThreads::StatsThread)
 		{
@@ -320,7 +326,10 @@ public:
 					ProcessingTasks.Start(StatName);
 				}
 #endif
-				Task->Execute(NewTasks, ENamedThreads::Type(ThreadId | (QueueIndex << ENamedThreads::QueueIndexShift)));
+				if (Task != WakeUpBaseGraphTask)
+				{
+					Task->Execute(NewTasks, ENamedThreads::Type(ThreadId | (QueueIndex << ENamedThreads::QueueIndexShift)));
+				}
 			}
 			else
 			{
@@ -741,34 +750,42 @@ public:
 		}
 		if (ThreadToExecuteOn == ENamedThreads::AnyThread)
 		{
-			FTaskThread* TempTarget = StalledUnnamedThreads.Pop(); //@todo it is possible that a thread is in the process of stalling and we just missed it, non-fatal, but we could lose a whole task of potential parallelism.
-			if (TempTarget)
-			{
-				ThreadToExecuteOn = TempTarget->GetThreadId();
-			}
-			else
+			if (FPlatformProcess::SupportsMultithreading())
 			{
-				// Run everything on the game thread if multithreading is disabled.
-				if (FPlatformProcess::SupportsMultithreading())
+				IncomingAnyThreadTasks.Push(Task);
+				FTaskThread* TempTarget = StalledUnnamedThreads.Pop(); //@todo it is possible that a thread is in the process of stalling and we just missed it, non-fatal, but we could lose a whole task of potential parallelism.
+				if (TempTarget)
 				{
-					ThreadToExecuteOn = ENamedThreads::Type((uint32(NextUnnamedThreadForTaskFromUnknownThread.Increment()) % uint32(NextUnnamedThreadMod)) + NumNamedThreads);
+					ThreadToExecuteOn = TempTarget->GetThreadId();
 				}
 				else
 				{
-					ThreadToExecuteOn = ENamedThreads::GameThread;
+					ThreadToExecuteOn = ENamedThreads::Type((uint32(NextUnnamedThreadForTaskFromUnknownThread.Increment()) % uint32(NextUnnamedThreadMod)) + NumNamedThreads);
 				}
+				FTaskThread* Target = &Thread(ThreadToExecuteOn);
+				if (ThreadToExecuteOn != CurrentThreadIfKnown)
+				{
+					Target->EnqueueFromOtherThread(0, WakeUpBaseGraphTask);
+				}
+				return;
+			}
+			else
+			{
+				ThreadToExecuteOn = ENamedThreads::GameThread;
 			}
 		}
-		int32 QueueToExecuteOn = ENamedThreads::GetQueueIndex(ThreadToExecuteOn);
-		ThreadToExecuteOn = ENamedThreads::GetThreadIndex(ThreadToExecuteOn);
-		FTaskThread* Target = &Thread(ThreadToExecuteOn);
-		if (ThreadToExecuteOn == CurrentThreadIfKnown)
-		{
-			Target->EnqueueFromThisThread(QueueToExecuteOn, Task);
-		}
-		else
 		{
-			Target->EnqueueFromOtherThread(QueueToExecuteOn, Task);
+			int32 QueueToExecuteOn = ENamedThreads::GetQueueIndex(ThreadToExecuteOn);
+			ThreadToExecuteOn = ENamedThreads::GetThreadIndex(ThreadToExecuteOn);
+			FTaskThread* Target = &Thread(ThreadToExecuteOn);
+			if (ThreadToExecuteOn == CurrentThreadIfKnown)
+			{
+				Target->EnqueueFromThisThread(QueueToExecuteOn, Task);
+			}
+			else
+			{
+				Target->EnqueueFromOtherThread(QueueToExecuteOn, Task);
+			}
 		}
 	}
 
@@ -899,6 +916,42 @@ public:
 	**/
 	FBaseGraphTask* FindWork(ENamedThreads::Type ThreadInNeed)
 	{
+		{
+			FBaseGraphTask* Task = SortedAnyThreadTasks.Pop();
+			if (Task)
+			{
+				return Task;
+			}
+		}
+		do
+		{
+			FScopeLock ScopeLock(&CriticalSectionForSortingIncomingAnyThreadTasks);
+			if (!IncomingAnyThreadTasks.IsEmpty() && SortedAnyThreadTasks.IsEmpty())
+			{
+				static TArray<FBaseGraphTask*> NewTasks;
+				NewTasks.Reset();
+				IncomingAnyThreadTasks.PopAll(NewTasks);
+				check(NewTasks.Num());
+
+				if (NewTasks.Num() > 1)
+				{
+					TLockFreePointerList<FBaseGraphTask> TempSortedAnyThreadTasks;
+					for (int32 Index = 0 ; Index < NewTasks.Num() - 1; Index++) // we are going to take the last one for ourselves
+					{
+						TempSortedAnyThreadTasks.Push(NewTasks[Index]);
+					}
+					verify(SortedAnyThreadTasks.ReplaceListIfEmpty(TempSortedAnyThreadTasks));
+				}
+				return NewTasks[NewTasks.Num() - 1];
+			}
+			{
+				FBaseGraphTask* Task = SortedAnyThreadTasks.Pop();
+				if (Task)
+				{
+					return Task;
+				}
+			}
+		} while (!IncomingAnyThreadTasks.IsEmpty() || !SortedAnyThreadTasks.IsEmpty());
 		// this can be called before my constructor is finished
 		for (int32 Pass = 0; Pass < 2; Pass++)
 		{
@@ -1003,6 +1056,10 @@ private:
 	uint32				PerThreadIDTLSSlot;
 	/** Thread safe list of stalled thread "Hints". **/
 	TLockFreePointerList<FTaskThread>		StalledUnnamedThreads; 
+
+	TLockFreePointerList<FBaseGraphTask>		IncomingAnyThreadTasks;
+	TLockFreePointerList<FBaseGraphTask>		SortedAnyThreadTasks;
+	FCriticalSection CriticalSectionForSortingIncomingAnyThreadTasks;
 };
 
 
diff --git a/Engine/Source/Runtime/Core/Public/Async/TaskGraphInterfaces.h b/Engine/Source/Runtime/Core/Public/Async/TaskGraphInterfaces.h
--- a/Engine/Source/Runtime/Core/Public/Async/TaskGraphInterfaces.h
+++ b/Engine/Source/Runtime/Core/Public/Async/TaskGraphInterfaces.h
@@ -20,6 +20,7 @@ DECLARE_STATS_GROUP(TEXT("Task Graph Tasks"), STATGROUP_TaskGraphTasks, STATCAT_
 
 DECLARE_CYCLE_STAT_EXTERN(TEXT("FReturnGraphTask"), STAT_FReturnGraphTask, STATGROUP_TaskGraphTasks, CORE_API);
 DECLARE_CYCLE_STAT_EXTERN(TEXT("FTriggerEventGraphTask"), STAT_FTriggerEventGraphTask, STATGROUP_TaskGraphTasks, CORE_API);
+DECLARE_CYCLE_STAT_EXTERN(TEXT("Unknown Graph Task"), STAT_UnknownGraphTask, STATGROUP_TaskGraphTasks, CORE_API);
 
 namespace ENamedThreads
 {
@@ -650,7 +651,7 @@ private:
 		
 		TTask& Task = *(TTask*)&TaskStorage;
 		{
-			FScopeCycleCounter Scope(Task.GetStatId()); 
+			FScopeCycleCounter Scope(Task.GetStatId(), true); 
 			Task.DoTask(CurrentThread, Subsequents);
 #if !(UE_BUILD_SHIPPING || UE_BUILD_TEST)
 			if (TTask::GetSubsequentsMode() == ESubsequentsMode::TrackSubsequents)
diff --git a/Engine/Source/Runtime/Core/Public/Containers/LockFreeListImpl.h b/Engine/Source/Runtime/Core/Public/Containers/LockFreeListImpl.h
--- a/Engine/Source/Runtime/Core/Public/Containers/LockFreeListImpl.h
+++ b/Engine/Source/Runtime/Core/Public/Containers/LockFreeListImpl.h
@@ -44,6 +44,15 @@ public:
 		FLockFreeVoidPointerListGeneric::PopAll<TArray<T *>, T * >(Output);
 	}	
 
+	/**	
+	 *	If the list is empty, replace it with the other list and null the other list.
+	 *	@return true if this call actively closed the list
+	 */
+	bool ReplaceListIfEmpty(TLockFreePointerList<T>& NotThreadSafeTempListToReplaceWith)
+	{
+		return FLockFreeVoidPointerListGeneric::ReplaceListIfEmpty(NotThreadSafeTempListToReplaceWith);
+	}
+
 	/**	
 	 *	Check if the list is empty.
 	 *
diff --git a/Engine/Source/Runtime/Core/Public/Containers/LockFreeVoidPointerListBase.h b/Engine/Source/Runtime/Core/Public/Containers/LockFreeVoidPointerListBase.h
--- a/Engine/Source/Runtime/Core/Public/Containers/LockFreeVoidPointerListBase.h
+++ b/Engine/Source/Runtime/Core/Public/Containers/LockFreeVoidPointerListBase.h
@@ -126,6 +126,20 @@ public:
 		return false;
 	}
 
+	/**	
+	 *	If the list is empty, replace it with the other list and null the other list.
+	 *	@return true if this call actively closed the list
+	 */
+	bool ReplaceListIfEmpty(FLockFreeVoidPointerListBase& NotThreadSafeTempListToReplaceWith)
+	{
+		if (NotThreadSafeTempListToReplaceWith.Head->ReplaceHeadIfHeadEqual(&Head,nullptr))
+		{
+			NotThreadSafeTempListToReplaceWith.Head = nullptr;
+			return true;
+		}
+		return false;
+	}
+
 	/**	
 	 *	Pop all items from the list 
 	 *	@param Output The array to hold the returned items. Must be empty.
diff --git a/Engine/Source/Runtime/Core/Public/Misc/CoreStats.h b/Engine/Source/Runtime/Core/Public/Misc/CoreStats.h
--- a/Engine/Source/Runtime/Core/Public/Misc/CoreStats.h
+++ b/Engine/Source/Runtime/Core/Public/Misc/CoreStats.h
@@ -61,7 +61,7 @@ DECLARE_FLOAT_COUNTER_STAT_EXTERN(TEXT("Game thread additional wait time"),STAT_
 DECLARE_CYCLE_STAT_EXTERN(TEXT("Game TaskGraph Tasks"),STAT_TaskGraph_GameTasks,STATGROUP_Threading, CORE_API);
 DECLARE_CYCLE_STAT_EXTERN(TEXT("Game TaskGraph Stalls"),STAT_TaskGraph_GameStalls,STATGROUP_Threading, CORE_API);
 
-DECLARE_CYCLE_STAT_EXTERN(TEXT("Render TaskGraph Stalls"),STAT_TaskGraph_RenderStalls,STATGROUP_Threading, CORE_API);
+DECLARE_CYCLE_STAT_EXTERN(TEXT("Render Local Queue Stalls"),STAT_TaskGraph_RenderStalls,STATGROUP_Threading, CORE_API);
 
 DECLARE_CYCLE_STAT_EXTERN(TEXT("Other TaskGraph Tasks"),STAT_TaskGraph_OtherTasks,STATGROUP_Threading, CORE_API);
 DECLARE_CYCLE_STAT_EXTERN(TEXT("Other TaskGraph Stalls"),STAT_TaskGraph_OtherStalls,STATGROUP_Threading, CORE_API);
diff --git a/Engine/Source/Runtime/Core/Public/Stats/Stats.h b/Engine/Source/Runtime/Core/Public/Stats/Stats.h
--- a/Engine/Source/Runtime/Core/Public/Stats/Stats.h
+++ b/Engine/Source/Runtime/Core/Public/Stats/Stats.h
@@ -50,9 +50,9 @@ public:
 	 * Pushes the specified stat onto the hierarchy for this thread. Starts
 	 * the timing of the cycles used
 	 */
-	FORCEINLINE_STATS FScopeCycleCounter( TStatId StatId )
+	FORCEINLINE_STATS FScopeCycleCounter( TStatId StatId, bool bAlways = false )
 	{
-		Start( StatId );
+		Start( StatId, bAlways );
 	}
 
 	/**
@@ -83,7 +83,7 @@ typedef FNothing TStatId;
 class FScopeCycleCounter
 {
 public:
-	FORCEINLINE_STATS FScopeCycleCounter(TStatId)
+	FORCEINLINE_STATS FScopeCycleCounter(TStatId, bool bAlways = false)
 	{
 	}
 };
diff --git a/Engine/Source/Runtime/Core/Public/Stats/Stats2.h b/Engine/Source/Runtime/Core/Public/Stats/Stats2.h
--- a/Engine/Source/Runtime/Core/Public/Stats/Stats2.h
+++ b/Engine/Source/Runtime/Core/Public/Stats/Stats2.h
@@ -1197,9 +1197,9 @@ public:
 	 * Pushes the specified stat onto the hierarchy for this thread. Starts
 	 * the timing of the cycles used
 	 */
-	FORCEINLINE_STATS void Start( TStatId InStatId )
+	FORCEINLINE_STATS void Start( TStatId InStatId, bool bAlways = false )
 	{
-		if( FThreadStats::IsCollectingData( InStatId ) )
+		if( (bAlways && InStatId.IsValidStat()) || FThreadStats::IsCollectingData( InStatId ) )
 		{
 			StatId = *InStatId;
 			FThreadStats::AddMessage( *InStatId, EStatOperation::CycleScopeStart );
diff --git a/Engine/Source/Runtime/Engine/Private/SceneUtils.cpp b/Engine/Source/Runtime/Engine/Private/SceneUtils.cpp
--- a/Engine/Source/Runtime/Engine/Private/SceneUtils.cpp
+++ b/Engine/Source/Runtime/Engine/Private/SceneUtils.cpp
@@ -20,14 +20,17 @@ FDrawEvent::~FDrawEvent()
 
 void FDrawEvent::Start(const FColor& Color, const TCHAR* Fmt, ...)
 {
-	check(IsInRenderingThread_Internal());
-	va_list ptr;
-	va_start(ptr, Fmt);
-	TCHAR TempStr[256];
-	// Build the string in the temp buffer
-	FCString::GetVarArgs(TempStr, ARRAY_COUNT(TempStr), ARRAY_COUNT(TempStr) - 1, Fmt, ptr);
-	GDynamicRHI->PushEvent(TempStr);
-	bDrawEventHasBeenEmitted = true;
+	if (GRHICommandList.Bypass()) // todo interleave draw evens into command buffers
+	{
+		check(IsInRenderingThread_Internal());
+		va_list ptr;
+		va_start(ptr, Fmt);
+		TCHAR TempStr[256];
+		// Build the string in the temp buffer
+		FCString::GetVarArgs(TempStr, ARRAY_COUNT(TempStr), ARRAY_COUNT(TempStr) - 1, Fmt, ptr);
+		GDynamicRHI->PushEvent(TempStr);
+		bDrawEventHasBeenEmitted = true;
+	}
 }
 
 #endif // WANTS_DRAW_MESH_EVENTS && PLATFORM_SUPPORTS_DRAW_MESH_EVENTS
\ No newline at end of file
