Expose the default GPU name on OS X.

diff --git a/Engine/Source/Runtime/Core/Private/Mac/MacPlatformMisc.cpp b/Engine/Source/Runtime/Core/Private/Mac/MacPlatformMisc.cpp
--- a/Engine/Source/Runtime/Core/Private/Mac/MacPlatformMisc.cpp
+++ b/Engine/Source/Runtime/Core/Private/Mac/MacPlatformMisc.cpp
@@ -888,6 +888,60 @@ void FMacPlatformMisc::NormalizePath(FString& InPath)
 	}
 }
 
+FString FMacPlatformMisc::GetPrimaryGPUBrand()
+{
+	static FString PrimaryGPU;
+	if(PrimaryGPU.IsEmpty())
+	{
+		PrimaryGPU = FGenericPlatformMisc::GetPrimaryGPUBrand();
+		
+		// Enumerate the GPUs via IOKit to avoid dragging in OpenGL
+		io_iterator_t Iterator;
+		CFMutableDictionaryRef MatchDictionary = IOServiceMatching("IOPCIDevice");
+		if(IOServiceGetMatchingServices(kIOMasterPortDefault, MatchDictionary, &Iterator) == kIOReturnSuccess)
+		{
+			io_registry_entry_t ServiceEntry;
+			while((ServiceEntry = IOIteratorNext(Iterator)))
+			{
+				CFMutableDictionaryRef ServiceInfo;
+				if(IORegistryEntryCreateCFProperties(ServiceEntry, &ServiceInfo, kCFAllocatorDefault, kNilOptions) == kIOReturnSuccess)
+				{
+					// GPUs are class-code 0x30000
+					const CFDataRef ClassCode = (const CFDataRef)CFDictionaryGetValue(ServiceInfo, CFSTR("class-code"));
+					if(ClassCode && CFGetTypeID(ClassCode) == CFDataGetTypeID())
+					{
+						const uint32* Value = reinterpret_cast<const uint32*>(CFDataGetBytePtr(ClassCode));
+						if(Value && *Value == 0x30000)
+						{
+							// If there are StartupDisplay & hda-gfx entries then this is likely the online display
+							const CFDataRef StartupDisplay = (const CFDataRef)CFDictionaryGetValue(ServiceInfo, CFSTR("StartupDisplay"));
+							const CFDataRef HDAGfx = (const CFDataRef)CFDictionaryGetValue(ServiceInfo, CFSTR("hda-gfx"));
+							if(StartupDisplay && HDAGfx)
+							{
+								const CFDataRef Model = (const CFDataRef)CFDictionaryGetValue(ServiceInfo, CFSTR("model"));
+								if(Model && CFGetTypeID(Model) == CFDataGetTypeID())
+								{
+									CFStringRef ModelName = CFStringCreateFromExternalRepresentation(kCFAllocatorDefault, Model, kCFStringEncodingASCII);
+									PrimaryGPU = FString((NSString*)ModelName);
+									CFRelease(ModelName);
+									
+									CFRelease(ServiceInfo);
+									IOObjectRelease(ServiceEntry);
+									break;
+								}
+							}
+						}
+					}
+					CFRelease(ServiceInfo);
+				}
+				IOObjectRelease(ServiceEntry);
+			}
+			IOObjectRelease(Iterator);
+		}
+	}
+	return PrimaryGPU;
+}
+
 void FMacPlatformMisc::GetOSVersions( FString& out_OSVersionLabel, FString& out_OSSubVersionLabel )
 {
 	out_OSVersionLabel = GMacAppInfo.OSVersion;
diff --git a/Engine/Source/Runtime/Core/Public/Mac/MacPlatformMisc.h b/Engine/Source/Runtime/Core/Public/Mac/MacPlatformMisc.h
--- a/Engine/Source/Runtime/Core/Public/Mac/MacPlatformMisc.h
+++ b/Engine/Source/Runtime/Core/Public/Mac/MacPlatformMisc.h
@@ -69,6 +69,7 @@ struct CORE_API FMacPlatformMisc : public FGenericPlatformMisc
 	static int32 NumberOfCoresIncludingHyperthreads();
 	static void LoadPreInitModules();
 	static void NormalizePath(FString& InPath);
+	static FString GetPrimaryGPUBrand();
 	static void GetOSVersions( FString& out_OSVersionLabel, FString& out_OSSubVersionLabel );
 
 
