Fixed a crash when an IME was active when the editor was started

TTP# 340272 - IME crash when typing Korean in the Description TextBlock of SSourceControlSubmitWidget

- FWindowsTextInputMethodSystem::Initialize now checks if an IME is active when it's called, and will update the active IME state appropriately.
- Changed TSF to use AssociateFocus rather than SetFocus to fix some issues where the IME would become active incorrectly when the window focus was changed.
- Some belt-and-braces fixes for the IMM based code path (which was what was showing the crash, even though a TSF based IME was active).

diff --git a/Engine/Source/Runtime/Core/Private/Windows/WindowsTextInputMethodSystem.cpp b/Engine/Source/Runtime/Core/Private/Windows/WindowsTextInputMethodSystem.cpp
--- a/Engine/Source/Runtime/Core/Private/Windows/WindowsTextInputMethodSystem.cpp
+++ b/Engine/Source/Runtime/Core/Private/Windows/WindowsTextInputMethodSystem.cpp
@@ -131,42 +131,43 @@ STDMETHODIMP_(ULONG) FTSFActivationProxy::Release()
 
 STDAPI FTSFActivationProxy::OnActivated(DWORD dwProfileType, LANGID langid, __RPC__in REFCLSID clsid, __RPC__in REFGUID catid, __RPC__in REFGUID guidProfile, HKL hkl, DWORD dwFlags)
 {
-	FString APIString;
-	if(::ImmGetIMEFileName(::GetKeyboardLayout(0), nullptr, 0) > 0)
-	{
-		Owner->CurrentAPI = FWindowsTextInputMethodSystem::EAPI::IMM;
-		APIString = TEXT("IMM");
-	}
-	else
-	{
-		Owner->CurrentAPI = FWindowsTextInputMethodSystem::EAPI::TSF;
-		APIString = TEXT("TSF");
-	}
-	UE_LOG(LogWindowsTextInputMethodSystem, Verbose, TEXT("IME system now using %s."), *APIString);
+	const bool bIsEnabled = !!(dwFlags & TF_IPSINK_FLAG_ACTIVE);
+	Owner->OnIMEActivationStateChanged(bIsEnabled);
 	return S_OK;
 }
 
 STDAPI FTSFActivationProxy::OnActivated(REFCLSID clsid, REFGUID guidProfile, BOOL fActivated)
 {
-	FString APIString;
-	if(::ImmGetIMEFileName(::GetKeyboardLayout(0), nullptr, 0) > 0)
-	{
-		Owner->CurrentAPI = FWindowsTextInputMethodSystem::EAPI::IMM;
-		APIString = TEXT("IMM");
-	}
-	else
-	{
-		Owner->CurrentAPI = FWindowsTextInputMethodSystem::EAPI::TSF;
-		APIString = TEXT("TSF");
-	}
-	UE_LOG(LogWindowsTextInputMethodSystem, Verbose, TEXT("IME system now using %s."), *APIString);
+	Owner->OnIMEActivationStateChanged(fActivated == TRUE);
 	return S_OK;
 }
 
 bool FWindowsTextInputMethodSystem::Initialize()
 {
 	CurrentAPI = EAPI::Unknown;
-	bool Result = InitializeIMM() && InitializeTSF();
+	const bool Result = InitializeIMM() && InitializeTSF();
+
+	if(Result)
+	{
+		const HKL KeyboardLayout = ::GetKeyboardLayout(0);
+
+		// We might already be using an IME if its set as the default language
+		// If so, work out what kind of IME it is
+		TF_INPUTPROCESSORPROFILE TSFProfile;
+		if(SUCCEEDED(TSFInputProcessorProfileManager->GetActiveProfile(GUID_TFCAT_TIP_KEYBOARD, &TSFProfile)) && TSFProfile.hkl && TSFProfile.dwProfileType == TF_PROFILETYPE_INPUTPROCESSOR)
+		{
+			check(TSFProfile.hkl == KeyboardLayout);
+
+			CurrentAPI = EAPI::TSF;
+			UE_LOG(LogWindowsTextInputMethodSystem, Verbose, TEXT("IME system now activated using TSF."));
+		}
+		else if(::ImmGetIMEFileName(KeyboardLayout, nullptr, 0) > 0)
+		{
+			CurrentAPI = EAPI::IMM;
+			UE_LOG(LogWindowsTextInputMethodSystem, Verbose, TEXT("IME system now activated using IMM."));
+		}
+	}
+
 	return Result;
 }
 
@@ -183,8 +184,6 @@ bool FWindowsTextInputMethodSystem::InitializeIMM()
 
 void FWindowsTextInputMethodSystem::UpdateIMMProperty(HKL KeyboardLayoutHandle)
 {
-	WORD langID = LOWORD(KeyboardLayoutHandle);
-
 	IMEProperties = ::ImmGetProperty(KeyboardLayoutHandle, IGP_PROPERTY);
 }
 
@@ -539,13 +538,17 @@ void FWindowsTextInputMethodSystem::ActivateContext(const TSharedRef<ITextInputM
 	check(ContextToInternalContextMap.Contains(Context));
 	FInternalContext& InternalContext = ContextToInternalContextMap[Context];
 
+	const TSharedPtr<FGenericWindow> GenericWindow = Context->GetWindow();
+	const HWND Hwnd = GenericWindow.IsValid() ? reinterpret_cast<HWND>(GenericWindow->GetOSWindowHandle()) : nullptr;
+
 	// IMM Implementation
 	InternalContext.IMMContext.IsComposing = false;
 	InternalContext.IMMContext.IsDeactivating = false;
 
 	// TSF Implementation
 	TComPtr<FTextStoreACP>& TextStore = InternalContext.TSFContext;
-	Result = TSFThreadManager->SetFocus(TextStore->TSFDocumentManager);
+	ITfDocumentMgr* Unused;
+	Result = TSFThreadManager->AssociateFocus(Hwnd, TextStore->TSFDocumentManager, &Unused);
 	if(FAILED(Result))
 	{
 		UE_LOG(LogWindowsTextInputMethodSystem, Error, TEXT("Activating a context failed while setting focus on a TSF document manager."));
@@ -561,6 +564,9 @@ void FWindowsTextInputMethodSystem::DeactivateContext(const TSharedRef<ITextInpu
 
 	FInternalContext& InternalContext = ContextToInternalContextMap[Context];
 
+	const TSharedPtr<FGenericWindow> GenericWindow = Context->GetWindow();
+	const HWND Hwnd = GenericWindow.IsValid() ? reinterpret_cast<HWND>(GenericWindow->GetOSWindowHandle()) : nullptr;
+
 	// TSF Implementation
 	TComPtr<FTextStoreACP> TextStore = InternalContext.TSFContext;
 
@@ -572,7 +578,8 @@ void FWindowsTextInputMethodSystem::DeactivateContext(const TSharedRef<ITextInpu
 	}
 	else if(FocusedDocumentManger == TextStore->TSFDocumentManager)
 	{
-		Result = TSFThreadManager->SetFocus(TSFDisabledDocumentManager);
+		ITfDocumentMgr* Unused;
+		Result = TSFThreadManager->AssociateFocus(Hwnd, TSFDisabledDocumentManager, &Unused);
 		if(FAILED(Result))
 		{
 			UE_LOG(LogWindowsTextInputMethodSystem, Error, TEXT("Deactivating a context failed while setting focus to the disabled TSF document manager."));
@@ -584,8 +591,6 @@ void FWindowsTextInputMethodSystem::DeactivateContext(const TSharedRef<ITextInpu
 	}
 
 	// IMM Implementation
-	const TSharedPtr<FGenericWindow> GenericWindow = Context->GetWindow();
-	const HWND Hwnd = GenericWindow.IsValid() ? reinterpret_cast<HWND>(GenericWindow->GetOSWindowHandle()) : nullptr;
 	HIMC IMMContext = ::ImmGetContext(Hwnd);
 	InternalContext.IMMContext.IsDeactivating = true;
 	// Request the composition is completed to ensure that the composition input UI is closed, and that a WM_IME_ENDCOMPOSITION message is sent
@@ -598,9 +603,40 @@ void FWindowsTextInputMethodSystem::DeactivateContext(const TSharedRef<ITextInpu
 	UE_LOG(LogWindowsTextInputMethodSystem, Verbose, TEXT("Deactivated context %p!"), &(Context.Get()));
 }
 
+void FWindowsTextInputMethodSystem::OnIMEActivationStateChanged(const bool bIsEnabled)
+{
+	if(bIsEnabled)
+	{
+		// It seems that switching away from an IMM based IME doesn't generate a deactivation notification
+		//check(CurrentAPI == EAPI::Unknown);
+
+		FString APIString;
+		const HKL KeyboardLayout = ::GetKeyboardLayout(0);
+		if(::ImmGetIMEFileName(KeyboardLayout, nullptr, 0) > 0)
+		{
+			CurrentAPI = EAPI::IMM;
+			UpdateIMMProperty(KeyboardLayout);
+			APIString = TEXT("IMM");
+		}
+		else
+		{
+			CurrentAPI = EAPI::TSF;
+			APIString = TEXT("TSF");
+		}
+		UE_LOG(LogWindowsTextInputMethodSystem, Verbose, TEXT("IME system now activated using %s."), *APIString);
+	}
+	else
+	{
+		check(CurrentAPI != EAPI::Unknown);
+
+		CurrentAPI = EAPI::Unknown;
+		UE_LOG(LogWindowsTextInputMethodSystem, Verbose, TEXT("IME system now deactivated."));
+	}
+}
+
 int32 FWindowsTextInputMethodSystem::ProcessMessage(HWND hwnd, uint32 msg, WPARAM wParam, LPARAM lParam)
 {
-	if(CurrentAPI == EAPI::TSF)
+	if(CurrentAPI != EAPI::IMM)
 	{
 		return DefWindowProc(hwnd, msg, wParam, lParam);
 	}
@@ -693,7 +729,9 @@ int32 FWindowsTextInputMethodSystem::ProcessMessage(HWND hwnd, uint32 msg, WPARA
 
 					// Clear Composition
 					ActiveContext->SetTextInRange(InternalContext.IMMContext.CompositionBeginIndex, InternalContext.IMMContext.CompositionLength, TEXT(""));
-					
+					InternalContext.IMMContext.CompositionLength = 0;
+					ActiveContext->UpdateCompositionRange(InternalContext.IMMContext.CompositionBeginIndex, 0);
+
 					// Restore any previous selection
 					ActiveContext->SetSelectionRange(SelectionBeginIndex, SelectionLength, SelectionCaretPosition);
 
diff --git a/Engine/Source/Runtime/Core/Public/Windows/WindowsTextInputMethodSystem.h b/Engine/Source/Runtime/Core/Public/Windows/WindowsTextInputMethodSystem.h
--- a/Engine/Source/Runtime/Core/Public/Windows/WindowsTextInputMethodSystem.h
+++ b/Engine/Source/Runtime/Core/Public/Windows/WindowsTextInputMethodSystem.h
@@ -15,8 +15,13 @@ class FWindowsTextInputMethodSystem;
 class FTSFActivationProxy : public ITfInputProcessorProfileActivationSink, public ITfActiveLanguageProfileNotifySink
 {
 public:
-	FTSFActivationProxy(FWindowsTextInputMethodSystem* InOwner) : Owner(InOwner), ReferenceCount(1)
-	{}
+	FTSFActivationProxy(FWindowsTextInputMethodSystem* InOwner) 
+		: Owner(InOwner)
+		, ReferenceCount(1)
+		, TSFProfileCookie(TF_INVALID_COOKIE)
+		, TSFLanguageCookie(TF_INVALID_COOKIE)
+	{
+	}
 
 	// IUnknown Interface Begin
 	STDMETHODIMP QueryInterface(REFIID riid, void **ppvObj) override;
@@ -72,6 +77,8 @@ private:
 	// TSF Implementation
 	bool InitializeTSF();
 
+	void OnIMEActivationStateChanged(const bool bIsEnabled);
+
 private:
 	TSharedPtr<ITextInputMethodContext> ActiveContext;
 	enum class EAPI
