Removing errorneous check that broke iOS.

diff --git a/Engine/Source/Runtime/Core/Private/GenericPlatform/GenericPlatformOutputDevices.cpp b/Engine/Source/Runtime/Core/Private/GenericPlatform/GenericPlatformOutputDevices.cpp
--- a/Engine/Source/Runtime/Core/Private/GenericPlatform/GenericPlatformOutputDevices.cpp
+++ b/Engine/Source/Runtime/Core/Private/GenericPlatform/GenericPlatformOutputDevices.cpp
@@ -6,7 +6,6 @@
 void							FGenericPlatformOutputDevices::SetupOutputDevices()
 {
 	check(GLog);
-	check(GLogConsole);
 
 	GLog->AddOutputDevice(FPlatformOutputDevices::GetLog());
 
