Linux window changes from PR #306.

diff --git a/Engine/Source/Runtime/Core/Private/Linux/LinuxWindow.cpp b/Engine/Source/Runtime/Core/Private/Linux/LinuxWindow.cpp
--- a/Engine/Source/Runtime/Core/Private/Linux/LinuxWindow.cpp
+++ b/Engine/Source/Runtime/Core/Private/Linux/LinuxWindow.cpp
@@ -4,6 +4,7 @@
 #include "LinuxWindow.h"
 #include "LinuxApplication.h"
 
+DEFINE_LOG_CATEGORY( LogLinuxWindow );
 
 FLinuxWindow::~FLinuxWindow()
 {
@@ -71,10 +72,10 @@ void FLinuxWindow::Initialize( FLinuxApplication* const Application, const TShar
 	const float WidthInitial = Definition->WidthDesiredOnScreen;
 	const float HeightInitial = Definition->HeightDesiredOnScreen;
 
-	int32 X = FMath::TruncToInt( XInitialRect );
-	int32 Y = FMath::TruncToInt( YInitialRect );
-	int32 ClientWidth = FMath::TruncToInt( WidthInitial );
-	int32 ClientHeight = FMath::TruncToInt( HeightInitial );
+	int32 X = FMath::TruncToInt( XInitialRect + 0.5f );
+	int32 Y = FMath::TruncToInt( YInitialRect + 0.5f );
+	int32 ClientWidth = FMath::TruncToInt( WidthInitial + 0.5f );
+	int32 ClientHeight = FMath::TruncToInt( HeightInitial + 0.5f );
 	int32 WindowWidth = ClientWidth;
 	int32 WindowHeight = ClientHeight;
 
@@ -93,6 +94,7 @@ void FLinuxWindow::Initialize( FLinuxApplication* const Application, const TShar
 	//	The SDL window doesn't need to be reshaped.
 	//	the size of the window you input is the sizeof the client.
 	HWnd = SDL_CreateWindow( TCHAR_TO_ANSI( *Definition->Title ), X, Y, ClientWidth, ClientHeight, WindowStyle  );
+	SDL_SetWindowHitTest( HWnd, FLinuxWindow::HitTest, this );
 
 	VirtualWidth  = ClientWidth;
 	VirtualHeight = ClientHeight;
@@ -116,6 +118,40 @@ void FLinuxWindow::Initialize( FLinuxApplication* const Application, const TShar
 	}
 }
 
+SDL_HitTestResult FLinuxWindow::HitTest( SDL_Window *SDLwin, const SDL_Point *point, void *data )
+{
+	FLinuxWindow *pself = static_cast<FLinuxWindow *>( data );
+	TSharedPtr< FLinuxWindow > self = pself->OwningApplication->FindWindowBySDLWindow( SDLwin );
+	if ( !self.IsValid() ) 
+	{
+		UE_LOG(LogLinuxWindow, Warning, TEXT("BAD EVENT: SDL window = %p\n"), SDLwin);
+		return SDL_HITTEST_NORMAL;
+	}
+
+	EWindowZone::Type eventZone = pself->OwningApplication->WindowHitTest( self, point->x, point->y );
+
+	static const SDL_HitTestResult Results[] = 
+	{
+		SDL_HITTEST_NORMAL,
+		SDL_HITTEST_RESIZE_TOPLEFT,
+		SDL_HITTEST_RESIZE_TOP,
+		SDL_HITTEST_RESIZE_TOPRIGHT,
+		SDL_HITTEST_RESIZE_LEFT,
+		SDL_HITTEST_NORMAL,
+		SDL_HITTEST_RESIZE_RIGHT,
+		SDL_HITTEST_RESIZE_BOTTOMLEFT,
+		SDL_HITTEST_RESIZE_BOTTOM,
+		SDL_HITTEST_RESIZE_BOTTOMRIGHT,
+		SDL_HITTEST_DRAGGABLE,
+		SDL_HITTEST_NORMAL,
+		SDL_HITTEST_NORMAL,
+		SDL_HITTEST_NORMAL,
+		SDL_HITTEST_NORMAL
+	};
+
+	return Results[eventZone];
+}
+
 /** Native windows should implement MoveWindowTo by relocating the platform-specific window to (X,Y). */
 void FLinuxWindow::MoveWindowTo( int32 X, int32 Y )
 {
@@ -136,6 +172,7 @@ void FLinuxWindow::BringToFront( bool bForce )
 /** Native windows should implement this function by asking the OS to destroy OS-specific resource associated with the window (e.g. Win32 window handle) */
 void FLinuxWindow::Destroy()
 {
+	OwningApplication->RemoveEventWindow( HWnd );
 	SDL_DestroyWindow( HWnd );
 }
 
@@ -365,11 +402,21 @@ void FLinuxWindow::AdjustCachedSize( FVector2D& Size ) const
 	}
 }
 
-
 bool FLinuxWindow::GetFullScreenInfo( int32& X, int32& Y, int32& Width, int32& Height ) const
 {
-	//	todo
-	return true;
+	SDL_Rect rect;
+
+	if (SDL_GetDisplayBounds(0, &rect) == 0)
+	{
+		X = rect.x;
+		Y = rect.y;
+		Width = rect.w;
+		Height = rect.h;
+
+		return true;
+	}
+
+	return false;
 }
 
 /** @return true if the native window is maximized, false otherwise */
@@ -397,7 +444,9 @@ bool FLinuxWindow::IsVisible() const
 /** Returns the size and location of the window when it is restored */
 bool FLinuxWindow::GetRestoredDimensions(int32& X, int32& Y, int32& Width, int32& Height)
 {
-	//	todo
+	SDL_GetWindowPosition(HWnd, &X, &Y);
+	SDL_GetWindowSize(HWnd, &Width, &Height);
+
 	return true;
 }
 
@@ -441,21 +490,22 @@ bool FLinuxWindow::IsPointInWindow( int32 X, int32 Y ) const
 
 int32 FLinuxWindow::GetWindowBorderSize() const
 {
-	// todo
-	return 5;
+	// need to lie and return 0 border size even if we have a border.
+	// reporting anything else causes problems in Slate with menu
+	// positioning.
+	return 0;
 }
 
 
 bool FLinuxWindow::IsForegroundWindow() const
 {
-	// todo
-	return 1;
+	return (HWnd == SDL_GetMouseFocus() || HWnd == SDL_GetKeyboardFocus());
 }
 
 
 void FLinuxWindow::SetText( const TCHAR* const Text )
 {
-	//	todo
+	SDL_SetWindowTitle( HWnd, TCHAR_TO_ANSI(Text));
 }
 
 
diff --git a/Engine/Source/Runtime/Core/Public/Linux/LinuxWindow.h b/Engine/Source/Runtime/Core/Public/Linux/LinuxWindow.h
--- a/Engine/Source/Runtime/Core/Public/Linux/LinuxWindow.h
+++ b/Engine/Source/Runtime/Core/Public/Linux/LinuxWindow.h
@@ -8,6 +8,8 @@
 
 #include "SDL.h"
 
+DECLARE_LOG_CATEGORY_EXTERN( LogLinuxWindow, Log, All );
+
 typedef SDL_Window*		SDL_HWindow;
 
 /**
@@ -42,6 +44,30 @@ public:
 	/**	Sets the window region to specified dimensions */
 	void AdjustWindowRegion( int32 Width, int32 Height );
 
+	bool IsNativeMoving() const;
+	bool BeginNativeMove();
+	void AfterNativeMove();
+	void EndNativeMove();
+
+	enum NativeResizeDirection 
+	{
+		ResizeSouthWest = 0,
+		ResizeSouth,
+		ResizeSouthEast,
+		ResizeEast,
+		ResizeNorthEast,
+		ResizeNorth,
+		ResizeNorthWest,
+		ResizeWest,
+
+		InvalidDirection
+	};
+
+	bool IsNativeResizing() const;
+	bool BeginNativeResize( NativeResizeDirection Direction );
+	void AfterNativeResize();
+	void EndNativeResize();
+
 public:
 	virtual void ReshapeWindow( int32 X, int32 Y, int32 Width, int32 Height ) override;
 
@@ -165,4 +191,6 @@ private:
 
 	bool bIsVisible : 1;
 	bool bWasFullscreen;
+
+	static SDL_HitTestResult HitTest( SDL_Window *SDLwin, const SDL_Point *point, void *data );
 };
