Fixed text test that was incorrectly comparing the number of code units in a string of one type to the number of code points in another type.

diff --git a/Engine/Source/Runtime/Core/Private/Tests/Internationalization/TextTest.cpp b/Engine/Source/Runtime/Core/Private/Tests/Internationalization/TextTest.cpp
--- a/Engine/Source/Runtime/Core/Private/Tests/Internationalization/TextTest.cpp
+++ b/Engine/Source/Runtime/Core/Private/Tests/Internationalization/TextTest.cpp
@@ -633,15 +633,15 @@ bool FICUTextTest::RunTest (const FString& Parameters)
 	// ---------------------------------------------------------------------
 
 	ICUUtilities::ConvertString(SourceString, ICUString);
-	if (SourceString.Len() != ICUString.length())
+	if (SourceString.Len() != ICUString.countChar32())
 	{
-		AddError(FString::Printf(TEXT("icu::UnicodeString is the incorrect length (%d; expected %d)."), ICUString.length(), SourceString.Len()));
+		AddError(FString::Printf(TEXT("icu::UnicodeString is the incorrect length (%d; expected %d)."), ICUString.countChar32(), SourceString.Len()));
 	}
 
 	ICUUtilities::ConvertString(ICUString, ConversionBackStr);
 	if (ICUString.length() != ConversionBackStr.Len())
 	{
-		AddError(FString::Printf(TEXT("FString is the incorrect length (%d; expected %d)."), ConversionBackStr.Len(), ICUString.length()));
+		AddError(FString::Printf(TEXT("FString is the incorrect length (%d; expected %d)."), ConversionBackStr.Len(), ICUString.countChar32()));
 	}
 	if (SourceString != ConversionBackStr)
 	{
@@ -651,15 +651,15 @@ bool FICUTextTest::RunTest (const FString& Parameters)
 	// ---------------------------------------------------------------------
 
 	ICUUtilities::ConvertString(SourceString2, ICUString);
-	if (SourceString2.Len() != ICUString.length())
+	if (SourceString2.Len() != ICUString.countChar32())
 	{
-		AddError(FString::Printf(TEXT("icu::UnicodeString is the incorrect length (%d; expected %d)."), ICUString.length(), SourceString2.Len()));
+		AddError(FString::Printf(TEXT("icu::UnicodeString is the incorrect length (%d; expected %d)."), ICUString.countChar32(), SourceString2.Len()));
 	}
 
 	ICUUtilities::ConvertString(ICUString, ConversionBackStr);
 	if (ICUString.length() != ConversionBackStr.Len())
 	{
-		AddError(FString::Printf(TEXT("FString is the incorrect length (%d; expected %d)."), ConversionBackStr.Len(), ICUString.length()));
+		AddError(FString::Printf(TEXT("FString is the incorrect length (%d; expected %d)."), ConversionBackStr.Len(), ICUString.countChar32()));
 	}
 	if (SourceString2 != ConversionBackStr)
 	{
@@ -669,15 +669,15 @@ bool FICUTextTest::RunTest (const FString& Parameters)
 	// ---------------------------------------------------------------------
 
 	ICUUtilities::ConvertString(SourceString, ICUString);
-	if (SourceString.Len() != ICUString.length())
+	if (SourceString.Len() != ICUString.countChar32())
 	{
-		AddError(FString::Printf(TEXT("icu::UnicodeString is the incorrect length (%d; expected %d)."), ICUString.length(), SourceString.Len()));
+		AddError(FString::Printf(TEXT("icu::UnicodeString is the incorrect length (%d; expected %d)."), ICUString.countChar32(), SourceString.Len()));
 	}
 
 	ICUUtilities::ConvertString(ICUString, ConversionBackStr);
 	if (ICUString.length() != ConversionBackStr.Len())
 	{
-		AddError(FString::Printf(TEXT("FString is the incorrect length (%d; expected %d)."), ConversionBackStr.Len(), ICUString.length()));
+		AddError(FString::Printf(TEXT("FString is the incorrect length (%d; expected %d)."), ConversionBackStr.Len(), ICUString.countChar32()));
 	}
 	if (SourceString != ConversionBackStr)
 	{
