Core - Added a missing method to LFL128

diff --git a/Engine/Source/Runtime/Core/Public/Containers/LockFreeVoidPointerListBase128.h b/Engine/Source/Runtime/Core/Public/Containers/LockFreeVoidPointerListBase128.h
--- a/Engine/Source/Runtime/Core/Public/Containers/LockFreeVoidPointerListBase128.h
+++ b/Engine/Source/Runtime/Core/Public/Containers/LockFreeVoidPointerListBase128.h
@@ -151,13 +151,27 @@ public:
 	 */
 	bool CloseIfEmpty()
 	{
-		if( FLink::ReplaceHeadWithClosedIfHeadIsEmpty( Head ) )
+		if( FLink::ReplaceHeadWithOtherIfHeadIsEmpty( Head, FLinkAllocator::Get().ClosedLink() ) )
 		{
 			return true;
 		}
 		return false;
 	}
 
+	/**	
+	 *	If the list is empty, replace it with the other list and null the other list.
+	 *	@return true if this call actively closed the list
+	 */
+	bool ReplaceListIfEmpty( FLockFreeVoidPointerListBase128& NotThreadSafeTempListToReplaceWith )
+	{
+		if( FLink::ReplaceHeadWithOtherIfHeadIsEmpty( Head, NotThreadSafeTempListToReplaceWith.Head ) )
+		{
+			NotThreadSafeTempListToReplaceWith.Head = FLargePointer();
+			return true;
+		}
+		return false;
+	}
+
 	/**	
 	 *	Pop all items from the list 
 	 *	@param Output The array to hold the returned items. Must be empty.
@@ -417,13 +431,14 @@ private:
 		}
 
 		/**	
-		 *	If the head is empty, then replace it with a closed list. This is the primitive that is used to close lists
+		 *	If the head is empty, then replace it with a head from the second list. This is the primitive that is used to close lists
 		 *	@param HeadPointer; the head of the list
-		 *	@return true if head was replaced with a closed list
+		 *	@param HeadToReplace; the head of the list to be replaced
+		 *	@return true if head was replaced with a new head
 		 */
-		static bool ReplaceHeadWithClosedIfHeadIsEmpty( FLargePointer& HeadPointer )
+		static bool ReplaceHeadWithOtherIfHeadIsEmpty( FLargePointer& HeadPointer, const FLargePointer& HeadToReplace )
 		{
-			const FLargePointer LocalClosedLink = FLinkAllocator::Get().ClosedLink();
+			const FLargePointer LocalHeadToReplace = HeadToReplace;
 			const FLargePointer LocalEmptyLink  = FLargePointer();
 			FLargePointer LocalHeadPointer      = FLargePointer();
 
@@ -435,7 +450,7 @@ private:
 					return false;
 				}
 			}
-			while( FPlatformAtomics::InterlockedCompareExchange128( (volatile FInt128*)&HeadPointer, (FInt128&)LocalClosedLink, (FInt128*)&LocalHeadPointer ) == false );
+			while( FPlatformAtomics::InterlockedCompareExchange128( (volatile FInt128*)&HeadPointer, (FInt128&)LocalHeadToReplace, (FInt128*)&LocalHeadPointer ) == false );
 
 			return true;
 		}
