Stats - Fixed a warning about duplicated thread metadata

diff --git a/Engine/Source/Runtime/Core/Private/Misc/MonitoredProcess.cpp b/Engine/Source/Runtime/Core/Private/Misc/MonitoredProcess.cpp
--- a/Engine/Source/Runtime/Core/Private/Misc/MonitoredProcess.cpp
+++ b/Engine/Source/Runtime/Core/Private/Misc/MonitoredProcess.cpp
@@ -65,7 +65,11 @@ bool FMonitoredProcess::Launch( )
 		return false;
 	}
 
-	Thread = FRunnableThread::Create(this, TEXT("FMonitoredProcess"), 128 * 1024, TPri_AboveNormal);
+	static int32 MonitoredProcessIndex = 0;
+	const FString MonitoredProcessName = FString::Printf( TEXT( "FMonitoredProcess %d" ), MonitoredProcessIndex );
+	MonitoredProcessIndex++;
+
+	Thread = FRunnableThread::Create(this, *MonitoredProcessName, 128 * 1024, TPri_AboveNormal);
 
 	return true;
 }
