Checks to ensure that you don't insert map elements into itself, causing invalid pointers on reallocation.

diff --git a/Engine/Source/Runtime/Core/Public/Containers/Array.h b/Engine/Source/Runtime/Core/Public/Containers/Array.h
--- a/Engine/Source/Runtime/Core/Public/Containers/Array.h
+++ b/Engine/Source/Runtime/Core/Public/Containers/Array.h
@@ -970,10 +970,21 @@ public:
 		return Index;
 	}
 
+	/**
+	 * Checks that the specified address is not part of an element within the container.  Used for implementations
+	 * to check that reference arguments aren't going to be invalidated by possible reallocation.
+	 *
+	 * @param Addr The address to check.
+	 */
+	FORCEINLINE void CheckAddress(const void* Addr)
+	{
+		checkf(Addr < GetTypedData() || Addr >= GetTypedData() + ArrayMax, TEXT("Attempting to add a container element which already comes from the container!"));
+	}
+
 	int32 Insert( ElementType&& Item, int32 Index )
 	{
-		// It isn't valid to specify an Item that is in the array, since adding an item might resize the array, which would make the item invalid
-		check( ((&Item) < GetTypedData()) || ((&Item) >= GetTypedData()+ArrayMax) );
+		CheckAddress(&Item);
+
 		// construct a copy in place at Index (this new operator will insert at 
 		// Index, then construct that memory with Item)
 		InsertUninitialized(Index,1);
@@ -982,8 +993,8 @@ public:
 	}
 	int32 Insert( const ElementType& Item, int32 Index )
 	{
-		// It isn't valid to specify an Item that is in the array, since adding an item might resize the array, which would make the item invalid
-		check( ((&Item) < GetTypedData()) || ((&Item) >= GetTypedData()+ArrayMax) );
+		CheckAddress(&Item);
+
 		// construct a copy in place at Index (this new operator will insert at 
 		// Index, then construct that memory with Item)
 		InsertUninitialized(Index,1);
@@ -1252,8 +1263,8 @@ public:
 	 * @param Item	The item to add
 	 * @return		Index to the new item
 	 */
-	FORCEINLINE int32 Add(       ElementType&& Item ) { check( ((&Item) < GetTypedData()) || ((&Item) >= GetTypedData()+ArrayMax) ); return Emplace(MoveTemp(Item)); }
-	FORCEINLINE int32 Add( const ElementType&  Item ) { check( ((&Item) < GetTypedData()) || ((&Item) >= GetTypedData()+ArrayMax) ); return Emplace(         Item ); }
+	FORCEINLINE int32 Add(       ElementType&& Item ) { CheckAddress(&Item); return Emplace(MoveTemp(Item)); }
+	FORCEINLINE int32 Add( const ElementType&  Item ) { CheckAddress(&Item); return Emplace(         Item ); }
 
 	/** Caution, AddZeroed() will create elements without calling the constructor and this is not appropriate for element types that require a constructor to function properly. */
 	int32 AddZeroed( int32 Count=1 )
@@ -1316,8 +1327,7 @@ public:
 	 */
 	int32 RemoveSingle( const ElementType& Item )
 	{
-		// It isn't valid to specify an Item that is in the array, since removing that item will change Item's value.
-		check( ((&Item) < GetTypedData()) || ((&Item) >= GetTypedData()+ArrayMax) );
+		CheckAddress(&Item);
 
 		for( int32 Index=0; Index<ArrayNum; Index++ )
 		{
@@ -1347,8 +1357,7 @@ public:
 	/** Removes as many instances of Item as there are in the array, maintaining order but not indices. */
 	int32 Remove( const ElementType& Item )
 	{
-		// It isn't valid to specify an Item that is in the array, since removing that item will change Item's value.
-		check( ((&Item) < GetTypedData()) || ((&Item) >= GetTypedData()+ArrayMax) );
+		CheckAddress(&Item);
 
 		const int32 OriginalNum = ArrayNum;
 		if (!OriginalNum)
@@ -1442,7 +1451,8 @@ public:
 	 */
 	int32 RemoveSingleSwap( const ElementType& Item )
 	{
-		check( ((&Item) < (ElementType*)AllocatorInstance.GetAllocation()) || ((&Item) >= (ElementType*)AllocatorInstance.GetAllocation()+ArrayMax) );
+		CheckAddress(&Item);
+
 		for( int32 Index=0; Index<ArrayNum; Index++ )
 		{
 			if( (*this)[Index]==Item )
@@ -1461,7 +1471,8 @@ public:
 	/** RemoveItemSwap, this version is much more efficient O(Count) instead of O(ArrayNum), but does not preserve the order */
 	int32 RemoveSwap( const ElementType& Item )
 	{
-		check( ((&Item) < (ElementType*)AllocatorInstance.GetAllocation()) || ((&Item) >= (ElementType*)AllocatorInstance.GetAllocation()+ArrayMax) );
+		CheckAddress(&Item);
+
 		const int32 OriginalNum=ArrayNum;
 		for( int32 Index=0; Index<ArrayNum; Index++ )
 		{
@@ -2578,7 +2589,8 @@ public:
 	}
 	int32 Remove( const T& Item )
 	{
-		check( ((&Item) < (T*)this->AllocatorInstance.GetAllocation()) || ((&Item) >= (T*)this->AllocatorInstance.GetAllocation()+this->ArrayMax) );
+		this->CheckAddress(&Item);
+
 		const int32 OriginalNum=this->ArrayNum;
 		for( int32 Index=0; Index<this->ArrayNum; Index++ )
 		{
diff --git a/Engine/Source/Runtime/Core/Public/Containers/Map.h b/Engine/Source/Runtime/Core/Public/Containers/Map.h
--- a/Engine/Source/Runtime/Core/Public/Containers/Map.h
+++ b/Engine/Source/Runtime/Core/Public/Containers/Map.h
@@ -280,10 +280,10 @@ public:
 	 * @param InValue - The value to associate with the key.
 	 * @return A reference to the value as stored in the map.  The reference is only valid until the next change to any key in the map.
 	 */
-	FORCEINLINE ValueType& Add(const KeyType&  InKey, const ValueType&  InValue) { return Emplace(         InKey ,          InValue ); }
-	FORCEINLINE ValueType& Add(const KeyType&  InKey,       ValueType&& InValue) { return Emplace(         InKey , MoveTemp(InValue)); }
-	FORCEINLINE ValueType& Add(      KeyType&& InKey, const ValueType&  InValue) { return Emplace(MoveTemp(InKey),          InValue ); }
-	FORCEINLINE ValueType& Add(      KeyType&& InKey,       ValueType&& InValue) { return Emplace(MoveTemp(InKey), MoveTemp(InValue)); }
+	FORCEINLINE ValueType& Add(const KeyType&  InKey, const ValueType&  InValue) { Pairs.CheckAddress(&InKey); Pairs.CheckAddress(&InValue); return Emplace(         InKey ,          InValue ); }
+	FORCEINLINE ValueType& Add(const KeyType&  InKey,       ValueType&& InValue) { Pairs.CheckAddress(&InKey); Pairs.CheckAddress(&InValue); return Emplace(         InKey , MoveTemp(InValue)); }
+	FORCEINLINE ValueType& Add(      KeyType&& InKey, const ValueType&  InValue) { Pairs.CheckAddress(&InKey); Pairs.CheckAddress(&InValue); return Emplace(MoveTemp(InKey),          InValue ); }
+	FORCEINLINE ValueType& Add(      KeyType&& InKey,       ValueType&& InValue) { Pairs.CheckAddress(&InKey); Pairs.CheckAddress(&InValue); return Emplace(MoveTemp(InKey), MoveTemp(InValue)); }
 
 	/**
 	 * Sets a default value associated with a key.
@@ -292,8 +292,8 @@ public:
 	 * @param InKey - The key to associate the value with.
 	 * @return A reference to the value as stored in the map.  The reference is only valid until the next change to any key in the map.
 	 */
-	FORCEINLINE ValueType& Add(const KeyType&  InKey) { return Emplace(         InKey ); }
-	FORCEINLINE ValueType& Add(      KeyType&& InKey) { return Emplace(MoveTemp(InKey)); }
+	FORCEINLINE ValueType& Add(const KeyType&  InKey) { Pairs.CheckAddress(&InKey); return Emplace(         InKey ); }
+	FORCEINLINE ValueType& Add(      KeyType&& InKey) { Pairs.CheckAddress(&InKey); return Emplace(MoveTemp(InKey)); }
 
 	/**
 	 * Sets the value associated with a key.
diff --git a/Engine/Source/Runtime/Core/Public/Containers/Set.h b/Engine/Source/Runtime/Core/Public/Containers/Set.h
--- a/Engine/Source/Runtime/Core/Public/Containers/Set.h
+++ b/Engine/Source/Runtime/Core/Public/Containers/Set.h
@@ -776,6 +776,17 @@ public:
 		return Result;
 	}
 
+	/**
+	 * Checks that the specified address is not part of an element within the container.  Used for implementations
+	 * to check that reference arguments aren't going to be invalidated by possible reallocation.
+	 *
+	 * @param Addr The address to check.
+	 */
+	FORCEINLINE void CheckAddress(const void* Addr)
+	{
+		Elements.CheckAddress(Addr);
+	}
+
 private:
 	/** Extracts the element value from the set's element structure and passes it to the user provided comparison class. */
 	template <typename PREDICATE_CLASS>
diff --git a/Engine/Source/Runtime/Core/Public/Containers/SparseArray.h b/Engine/Source/Runtime/Core/Public/Containers/SparseArray.h
--- a/Engine/Source/Runtime/Core/Public/Containers/SparseArray.h
+++ b/Engine/Source/Runtime/Core/Public/Containers/SparseArray.h
@@ -568,6 +568,17 @@ public:
 	int32 GetMaxIndex() const { return Data.Num(); }
 	int32 Num() const { return Data.Num() - NumFreeIndices; }
 
+	/**
+	 * Checks that the specified address is not part of an element within the container.  Used for implementations
+	 * to check that reference arguments aren't going to be invalidated by possible reallocation.
+	 *
+	 * @param Addr The address to check.
+	 */
+	FORCEINLINE void CheckAddress(const void* Addr)
+	{
+		Data.CheckAddress(Addr);
+	}
+
 private:
 
 	/** The base class of sparse array iterators. */
