Named FMatrix/FTransform unsafe Inverse to InverseFast, and InverseSafe,InverseSlow to Inverse as people often use Inverse as default function.

- Inverse was one of the main reasons of crashes via lots of functions. Seems making default Inverse to be safer version seems better.

diff --git a/Engine/Source/Developer/MeshUtilities/Private/MeshUtilities.cpp b/Engine/Source/Developer/MeshUtilities/Private/MeshUtilities.cpp
--- a/Engine/Source/Developer/MeshUtilities/Private/MeshUtilities.cpp
+++ b/Engine/Source/Developer/MeshUtilities/Private/MeshUtilities.cpp
@@ -1470,8 +1470,8 @@ static void ComputeTriangleTangents(
 			FPlane(	0,				0,				0,	1	)
 			);
 
-		// Use InverseSlow to catch singular matrices.  InverseSafe can miss this sometimes.
-		const FMatrix TextureToLocal = ParameterToTexture.InverseSlow() * ParameterToLocal;
+		// Use InverseSlow to catch singular matrices.  Inverse can miss this sometimes.
+		const FMatrix TextureToLocal = ParameterToTexture.Inverse() * ParameterToLocal;
 
 		TriangleTangentX.Add(TextureToLocal.TransformVector(FVector(1,0,0)).SafeNormal());
 		TriangleTangentY.Add(TextureToLocal.TransformVector(FVector(0,1,0)).SafeNormal());
@@ -2596,7 +2596,7 @@ bool FMeshUtilities::BuildSkeletalMesh( FStaticLODModel& LODModel, const FRefere
 				FPlane(	0,			0,			0,	1	)
 				);
 
-			FMatrix	TextureToLocal = ParameterToTexture.InverseSlow() * ParameterToLocal;
+			FMatrix	TextureToLocal = ParameterToTexture.Inverse() * ParameterToLocal;
 			FVector	TangentX = TextureToLocal.TransformVector(FVector(1,0,0)).SafeNormal(),
 				TangentY = TextureToLocal.TransformVector(FVector(0,1,0)).SafeNormal(),
 				TangentZ;
diff --git a/Engine/Source/Editor/Kismet/Private/SCSEditorViewportClient.cpp b/Engine/Source/Editor/Kismet/Private/SCSEditorViewportClient.cpp
--- a/Engine/Source/Editor/Kismet/Private/SCSEditorViewportClient.cpp
+++ b/Engine/Source/Editor/Kismet/Private/SCSEditorViewportClient.cpp
@@ -410,12 +410,12 @@ bool FSCSEditorViewportClient::InputWidgetDelta( FViewport* Viewport, EAxisList:
 
 								if(!SceneComp->bAbsoluteLocation)
 								{
-									Drag = ParentToWorldSpace.InverseSafe().TransformVector(Drag);
+									Drag = ParentToWorldSpace.Inverse().TransformVector(Drag);
 								}
 								
 								if(!SceneComp->bAbsoluteRotation)
 								{
-									Rot = (ParentToWorldSpace.InverseSafe().GetRotation() * Rot.Quaternion() * ParentToWorldSpace.GetRotation()).Rotator();
+									Rot = (ParentToWorldSpace.Inverse().GetRotation() * Rot.Quaternion() * ParentToWorldSpace.GetRotation()).Rotator();
 								}
 							}
 
diff --git a/Engine/Source/Editor/LandscapeEditor/Private/LandscapeEdModeBrushes.cpp b/Engine/Source/Editor/LandscapeEditor/Private/LandscapeEdModeBrushes.cpp
--- a/Engine/Source/Editor/LandscapeEditor/Private/LandscapeEdModeBrushes.cpp
+++ b/Engine/Source/Editor/LandscapeEditor/Private/LandscapeEdModeBrushes.cpp
@@ -545,7 +545,7 @@ public:
 				{
 					float ScaleXY = LandscapeInfo->DrawScale.X;
 					FMatrix LToW = LandscapeInfo->GetLandscapeProxy()->LandscapeActorToWorld().ToMatrixWithScale();
-					FMatrix WToL = LToW.Inverse();
+					FMatrix WToL = LToW.InverseFast();
 
 					UTexture2D* DataTexture = Gizmo->GizmoTexture;
 					int32 MinX = MAX_int32, MaxX = MIN_int32, MinY = MAX_int32, MaxY = MIN_int32;
@@ -646,10 +646,10 @@ public:
 			const float LW = Gizmo->GetWidth() / (2 * ScaleXY);
 			const float LH = Gizmo->GetHeight() / (2 * ScaleXY);
 
-			FMatrix WToL = LandscapeInfo->GetLandscapeProxy()->LandscapeActorToWorld().ToMatrixWithScale().Inverse();
+			FMatrix WToL = LandscapeInfo->GetLandscapeProxy()->LandscapeActorToWorld().ToMatrixWithScale().InverseFast();
 			FVector BaseLocation = WToL.TransformPosition(Gizmo->GetActorLocation());
 			FMatrix LandscapeToGizmoLocal =
-				(FTranslationMatrix(FVector(-LW + 0.5, -LH + 0.5, 0)) * FRotationTranslationMatrix(FRotator(0, Gizmo->GetActorRotation().Yaw, 0), FVector(BaseLocation.X, BaseLocation.Y, 0))).Inverse();
+				(FTranslationMatrix(FVector(-LW + 0.5, -LH + 0.5, 0)) * FRotationTranslationMatrix(FRotator(0, Gizmo->GetActorRotation().Yaw, 0), FVector(BaseLocation.X, BaseLocation.Y, 0))).InverseFast();
 
 			float W = Gizmo->GetWidth() / ScaleXY; //Gizmo->GetWidth() / (Gizmo->DrawScale * Gizmo->DrawScale3D.X);
 			float H = Gizmo->GetHeight() / ScaleXY; //Gizmo->GetHeight() / (Gizmo->DrawScale * Gizmo->DrawScale3D.Y);
diff --git a/Engine/Source/Editor/LandscapeEditor/Private/LandscapeEdModeComponentTools.cpp b/Engine/Source/Editor/LandscapeEditor/Private/LandscapeEdModeComponentTools.cpp
--- a/Engine/Source/Editor/LandscapeEditor/Private/LandscapeEdModeComponentTools.cpp
+++ b/Engine/Source/Editor/LandscapeEditor/Private/LandscapeEdModeComponentTools.cpp
@@ -1217,7 +1217,7 @@ public:
 			const float W = (Width - ScaleXY) / (2 * ScaleXY);
 			const float H = (Height - ScaleXY) / (2 * ScaleXY);
 
-			FMatrix WToL = LandscapeInfo->GetLandscapeProxy()->LandscapeActorToWorld().ToMatrixWithScale().Inverse();
+			FMatrix WToL = LandscapeInfo->GetLandscapeProxy()->LandscapeActorToWorld().ToMatrixWithScale().InverseFast();
 			//FMatrix LToW = Landscape->LocalToWorld();
 
 			FVector BaseLocation = WToL.TransformPosition(Gizmo->GetActorLocation());
@@ -1579,12 +1579,12 @@ public:
 			const float ScaleX = Gizmo->CachedWidth / Width * ScaleXY / Gizmo->CachedScaleXY;
 			const float ScaleY = Gizmo->CachedHeight / Height * ScaleXY / Gizmo->CachedScaleXY;
 
-			FMatrix WToL = LandscapeInfo->GetLandscapeProxy()->LandscapeActorToWorld().ToMatrixWithScale().Inverse();
+			FMatrix WToL = LandscapeInfo->GetLandscapeProxy()->LandscapeActorToWorld().ToMatrixWithScale().InverseFast();
 			//FMatrix LToW = Landscape->LocalToWorld();
 			FVector BaseLocation = WToL.TransformPosition(Gizmo->GetActorLocation());
 			//FMatrix LandscapeLocalToGizmo = FRotationTranslationMatrix(FRotator(0, Gizmo->Rotation.Yaw, 0), FVector(BaseLocation.X - W + 0.5, BaseLocation.Y - H + 0.5, 0));
 			FMatrix LandscapeToGizmoLocal =
-				(FTranslationMatrix(FVector((-W + 0.5)*SignX, (-H + 0.5)*SignY, 0)) * FScaleRotationTranslationMatrix(FVector(SignX, SignY, 1.f), FRotator(0, Gizmo->GetActorRotation().Yaw, 0), FVector(BaseLocation.X, BaseLocation.Y, 0))).Inverse();
+				(FTranslationMatrix(FVector((-W + 0.5)*SignX, (-H + 0.5)*SignY, 0)) * FScaleRotationTranslationMatrix(FVector(SignX, SignY, 1.f), FRotator(0, Gizmo->GetActorRotation().Yaw, 0), FVector(BaseLocation.X, BaseLocation.Y, 0))).InverseFast();
 
 			for (auto It = BrushInfo.CreateIterator(); It; ++It)
 			{
diff --git a/Engine/Source/Editor/LevelEditor/Private/LevelEditorActions.cpp b/Engine/Source/Editor/LevelEditor/Private/LevelEditorActions.cpp
--- a/Engine/Source/Editor/LevelEditor/Private/LevelEditorActions.cpp
+++ b/Engine/Source/Editor/LevelEditor/Private/LevelEditorActions.cpp
@@ -2470,7 +2470,7 @@ void FLevelEditorActionCallbacks::OnSaveBrushAsCollision()
 
 		// Need the transform between builder brush space and static mesh actor space.
 		const FMatrix BrushL2W = BuildBrush->ActorToWorld().ToMatrixWithScale();
-		const FMatrix MeshW2L = MeshToWorld.Inverse();
+		const FMatrix MeshW2L = MeshToWorld.InverseFast();
 		const FMatrix SMToBB = BrushL2W * MeshW2L;
 		const FMatrix SMToBB_AT = SMToBB.TransposeAdjoint();
 
diff --git a/Engine/Source/Editor/MeshPaint/Private/MeshPaintEdMode.cpp b/Engine/Source/Editor/MeshPaint/Private/MeshPaintEdMode.cpp
--- a/Engine/Source/Editor/MeshPaint/Private/MeshPaintEdMode.cpp
+++ b/Engine/Source/Editor/MeshPaint/Private/MeshPaintEdMode.cpp
@@ -1237,7 +1237,7 @@ void FEdModeMeshPaint::DoPaint( const FVector& InCameraOrigin,
 				Params.InnerBrushDepth = BrushDepth - Params.BrushDepthFalloffRange;
 				Params.BrushStrength = BrushStrength;
 				Params.BrushToWorldMatrix = FMatrix( BrushXAxis, BrushYAxis, Params.BrushNormal, Params.BrushPosition );
-				Params.InverseBrushToWorldMatrix = Params.BrushToWorldMatrix.Inverse();
+				Params.InverseBrushToWorldMatrix = Params.BrushToWorldMatrix.InverseFast();
 				Params.bWriteRed = FMeshPaintSettings::Get().bWriteRed;
 				Params.bWriteGreen = FMeshPaintSettings::Get().bWriteGreen;
 				Params.bWriteBlue = FMeshPaintSettings::Get().bWriteBlue;
diff --git a/Engine/Source/Editor/Persona/Private/AnimationEditorViewportClient.cpp b/Engine/Source/Editor/Persona/Private/AnimationEditorViewportClient.cpp
--- a/Engine/Source/Editor/Persona/Private/AnimationEditorViewportClient.cpp
+++ b/Engine/Source/Editor/Persona/Private/AnimationEditorViewportClient.cpp
@@ -624,7 +624,7 @@ void FAnimationViewportClient::SetCameraTargetLocation(const FSphere &BoundSpher
 	FVector OldViewLoc = GetViewLocation();
 	FMatrix EpicMat = FTranslationMatrix(-GetViewLocation());
 	EpicMat = EpicMat * FInverseRotationMatrix(GetViewRotation());
-	FMatrix CamRotMat = EpicMat.Inverse();
+	FMatrix CamRotMat = EpicMat.InverseFast();
 	FVector CamDir = FVector(CamRotMat.M[0][0],CamRotMat.M[0][1],CamRotMat.M[0][2]);
 	FVector NewViewLocation = BoundSphere.Center - BoundSphere.W * 2 * CamDir;
 
@@ -1095,7 +1095,7 @@ bool FAnimationViewportClient::InputWidgetDelta( FViewport* Viewport, EAxisList:
 			}
 
 			// Remove SkelControl's orientation from BoneMatrix, as we need to translate/rotate in the non-SkelControlled space
-			BaseTM = BaseTM.InverseSafe() * CurrentSkelControlTM;
+			BaseTM = BaseTM.Inverse() * CurrentSkelControlTM;
 
 			const bool bDoRotation    = WidgetMode == FWidget::WM_Rotate    || WidgetMode == FWidget::WM_TranslateRotateZ;
 			const bool bDoTranslation = WidgetMode == FWidget::WM_Translate || WidgetMode == FWidget::WM_TranslateRotateZ;
diff --git a/Engine/Source/Editor/PhAT/Private/PhATPreviewViewportClient.cpp b/Engine/Source/Editor/PhAT/Private/PhATPreviewViewportClient.cpp
--- a/Engine/Source/Editor/PhAT/Private/PhATPreviewViewportClient.cpp
+++ b/Engine/Source/Editor/PhAT/Private/PhATPreviewViewportClient.cpp
@@ -735,7 +735,7 @@ void FPhATEdPreviewViewportClient::StartManipulating(EAxisList::Type Axis, const
 
 		const FTransform WParentFrame = SharedData->GetConstraintWorldTM(SharedData->GetSelectedConstraint(), EConstraintFrame::Frame2);
 		const FTransform WChildFrame = SharedData->GetConstraintWorldTM(SharedData->GetSelectedConstraint(), EConstraintFrame::Frame1);
-		StartManRelConTM = WChildFrame * WParentFrame.InverseSafe();
+		StartManRelConTM = WChildFrame * WParentFrame.Inverse();
 
 		UPhysicsConstraintTemplate* Setup = SharedData->PhysicsAsset->ConstraintSetup[SharedData->GetSelectedConstraint()->Index];
 
@@ -831,7 +831,7 @@ void FPhATEdPreviewViewportClient::SimMousePress(FViewport* Viewport, bool bCons
 			// Create handle to object.
 			SharedData->MouseHandle->GrabComponent(SharedData->EditorSkelComp, BoneName, Result.Location, bConstrainRotation);
 
-			FMatrix	InvViewMatrix = View->ViewMatrices.ViewMatrix.Inverse();
+			FMatrix	InvViewMatrix = View->ViewMatrices.ViewMatrix.InverseFast();
 
 			SimGrabMinPush = SimMinHoldDistance - (Result.Time * SimGrabCheckDistance);
 
diff --git a/Engine/Source/Editor/PhAT/Private/PhATSharedData.cpp b/Engine/Source/Editor/PhAT/Private/PhATSharedData.cpp
--- a/Engine/Source/Editor/PhAT/Private/PhATSharedData.cpp
+++ b/Engine/Source/Editor/PhAT/Private/PhATSharedData.cpp
@@ -760,7 +760,7 @@ void FPhATSharedData::InitConstraintSetup(UPhysicsConstraintTemplate* Constraint
 	FMatrix ParentBoneTM = EditorSkelComp->GetBoneMatrix(ParentBoneIndex);
 	ParentBoneTM.RemoveScaling();
 
-	FMatrix RelTM = ChildBoneTM * ParentBoneTM.InverseSafe();
+	FMatrix RelTM = ChildBoneTM * ParentBoneTM.Inverse();
 
 	// Place joint at origin of child
 	ConstraintSetup->DefaultInstance.ConstraintBone1 = ChildBodySetup->BoneName;
@@ -951,9 +951,9 @@ void FPhATSharedData::PasteConstraintProperties()
 		FTransform BodyC = GetConstraintBodyTM(ToConstraintSetup, EConstraintFrame::Frame1);
 		FTransform BodyD = GetConstraintBodyTM(ToConstraintSetup, EConstraintFrame::Frame2);
 
-		FTransform FromF1AToF2A = BodyBFrame2B * BodyA.InverseSafe() * Frame1A.InverseSafe();
+		FTransform FromF1AToF2A = BodyBFrame2B * BodyA.Inverse() * Frame1A.Inverse();
 		FTransform Frame2C = FromF1AToF2A * Frame1C;
-		FTransform Frame2D = Frame2C* BodyC * BodyD.InverseSafe();
+		FTransform Frame2D = Frame2C* BodyC * BodyD.Inverse();
 
 
 		ToConstraintSetup->Modify();
@@ -997,7 +997,7 @@ void FPhATSharedData::CycleCurrentConstraintOrientation()
 	FMatrix ConstraintTransform = ConstraintTemplate->DefaultInstance.GetRefFrame(EConstraintFrame::Frame2).ToMatrixWithScale();
 	FTransform WParentFrame = GetConstraintWorldTM(GetSelectedConstraint(), EConstraintFrame::Frame2);
 	FTransform WChildFrame = GetConstraintWorldTM(GetSelectedConstraint(), EConstraintFrame::Frame1);
-	FTransform RelativeTransform = WChildFrame * WParentFrame.InverseSafe();
+	FTransform RelativeTransform = WChildFrame * WParentFrame.Inverse();
 
 	CycleMatrixRows(&ConstraintTransform);
 
diff --git a/Engine/Source/Editor/StaticMeshEditor/Private/StaticMeshEditorViewportClient.cpp b/Engine/Source/Editor/StaticMeshEditor/Private/StaticMeshEditorViewportClient.cpp
--- a/Engine/Source/Editor/StaticMeshEditor/Private/StaticMeshEditorViewportClient.cpp
+++ b/Engine/Source/Editor/StaticMeshEditor/Private/StaticMeshEditorViewportClient.cpp
@@ -507,7 +507,7 @@ void FStaticMeshEditorViewportClient::Draw(const FSceneView* View,FPrimitiveDraw
 		FIndexArrayView Indices = LODModel.IndexBuffer.GetArrayView();
 		uint32 NumIndices = Indices.Num();
 
-		FMatrix LocalToWorldInverseTranspose = StaticMeshComponent->ComponentToWorld.ToMatrixWithScale().Inverse().GetTransposed();
+		FMatrix LocalToWorldInverseTranspose = StaticMeshComponent->ComponentToWorld.ToMatrixWithScale().InverseFast().GetTransposed();
 		for (uint32 i = 0; i < NumIndices; i++)
 		{
 			const FVector& VertexPos = LODModel.PositionVertexBuffer.VertexPosition( Indices[i] );
diff --git a/Engine/Source/Editor/UnrealEd/Private/EditorComponents.cpp b/Engine/Source/Editor/UnrealEd/Private/EditorComponents.cpp
--- a/Engine/Source/Editor/UnrealEd/Private/EditorComponents.cpp
+++ b/Engine/Source/Editor/UnrealEd/Private/EditorComponents.cpp
@@ -547,7 +547,7 @@ void FEditorCommonDrawHelper::DrawGridSection(float ViewportGridY,FVector* A,FVe
 	FLinearColor MajorColor = FMath::Lerp(Background, FLinearColor::White, 0.05f);
 	FLinearColor MinorColor = FMath::Lerp(Background, FLinearColor::White, 0.02f);
 
-	const FMatrix InvViewProjMatrix = View->ViewMatrices.ProjMatrix.Inverse() * View->ViewMatrices.ViewMatrix.Inverse();
+	const FMatrix InvViewProjMatrix = View->ViewMatrices.ProjMatrix.InverseFast() * View->ViewMatrices.ViewMatrix.InverseFast();
 	int32 FirstLine = FMath::TruncToInt(InvViewProjMatrix.TransformPosition(FVector(-1, -1, 0.5f)).Component(Axis) / ViewportGridY);
 	int32 LastLine = FMath::TruncToInt(InvViewProjMatrix.TransformPosition(FVector(+1, +1, 0.5f)).Component(Axis) / ViewportGridY);
 	if (FirstLine > LastLine)
@@ -601,7 +601,7 @@ void FEditorCommonDrawHelper::DrawOriginAxisLine(FVector* A,FVector* B,float* AX
 
 void FEditorCommonDrawHelper::DrawPivot(const FSceneView* View,FPrimitiveDrawInterface* PDI)
 {
-	const FMatrix CameraToWorld = View->ViewMatrices.ViewMatrix.Inverse();
+	const FMatrix CameraToWorld = View->ViewMatrices.ViewMatrix.InverseFast();
 
 	const FVector PivLoc = GLevelEditorModeTools().SnappedLocation;
 
diff --git a/Engine/Source/Editor/UnrealEd/Private/EditorExporters.cpp b/Engine/Source/Editor/UnrealEd/Private/EditorExporters.cpp
--- a/Engine/Source/Editor/UnrealEd/Private/EditorExporters.cpp
+++ b/Engine/Source/Editor/UnrealEd/Private/EditorExporters.cpp
@@ -1555,7 +1555,7 @@ static void AddActorToOBJs(AActor* Actor, TArray<FOBJGeom*>& Objects, TSet<UMate
 
 		check(VertexCount == RenderData->VertexBuffer.GetNumVertices());
 
-		FMatrix LocalToWorldInverseTranspose = LocalToWorld.Inverse().GetTransposed();
+		FMatrix LocalToWorldInverseTranspose = LocalToWorld.InverseFast().GetTransposed();
 		for(uint32 i = 0; i < VertexCount; i++)
 		{
 			// Vertices
diff --git a/Engine/Source/Editor/UnrealEd/Private/EditorViewportClient.cpp b/Engine/Source/Editor/UnrealEd/Private/EditorViewportClient.cpp
--- a/Engine/Source/Editor/UnrealEd/Private/EditorViewportClient.cpp
+++ b/Engine/Source/Editor/UnrealEd/Private/EditorViewportClient.cpp
@@ -336,7 +336,7 @@ void FEditorViewportClient::SetViewLocationForOrbiting(const FVector& LookAtPoin
 {
 	FMatrix Matrix = FTranslationMatrix(-GetViewLocation());
 	Matrix = Matrix * FInverseRotationMatrix(GetViewRotation());
-	FMatrix CamRotMat = Matrix.Inverse();
+	FMatrix CamRotMat = Matrix.InverseFast();
 	FVector CamDir = FVector(CamRotMat.M[0][0],CamRotMat.M[0][1],CamRotMat.M[0][2]);
 	SetViewLocation( LookAtPoint - DistanceToCamera * CamDir );
 	SetLookAtLocation( LookAtPoint );
@@ -363,7 +363,7 @@ void FEditorViewportClient::ToggleOrbitCamera( bool bEnableOrbitCamera )
 
 		// Convert orbit view to regular view
 		FMatrix OrbitMatrix = ViewTransform.ComputeOrbitMatrix();
-		OrbitMatrix = OrbitMatrix.Inverse();
+		OrbitMatrix = OrbitMatrix.InverseFast();
 
 		if( !bUsingOrbitCamera )
 		{
@@ -383,10 +383,10 @@ void FEditorViewportClient::ToggleOrbitCamera( bool bEnableOrbitCamera )
 				OrbitViewMatrix *= FRotationMatrix( FRotator(0,90.f,0) );
 
 				FMatrix RotMat = FTranslationMatrix( -ViewTransform.GetLookAt() ) * OrbitViewMatrix;
-				FMatrix RotMatInv = RotMat.Inverse();
+				FMatrix RotMatInv = RotMat.InverseFast();
 				FRotator RollVec = RotMatInv.Rotator();
 				FMatrix YawMat = RotMatInv * FInverseRotationMatrix( FRotator(0, 0, -RollVec.Roll));
-				FMatrix YawMatInv = YawMat.Inverse();
+				FMatrix YawMatInv = YawMat.InverseFast();
 				FRotator YawVec = YawMat.Rotator();
 				FRotator rotYawInv = YawMatInv.Rotator();
 				SetViewRotation( FRotator(-RollVec.Roll,YawVec.Yaw,0) );
@@ -1590,7 +1590,7 @@ void FEditorViewportClient::InputAxisForOrbit(FViewport* InViewport, const FVect
 			FRotationMatrix( FRotator(0,GetViewRotation().Yaw,0) ) * 
 			FRotationMatrix( FRotator(0, 0, GetViewRotation().Pitch));
 
-		FVector TransformedDelta = RotMat.Inverse().TransformVector(DeltaLocation);
+		FVector TransformedDelta = RotMat.InverseFast().TransformVector(DeltaLocation);
 
 		SetViewLocation( GetViewLocation() + TransformedDelta );
 		SetLookAtLocation( GetLookAtLocation() + TransformedDelta );
@@ -1599,7 +1599,7 @@ void FEditorViewportClient::InputAxisForOrbit(FViewport* InViewport, const FVect
 	}
 	else if ( IsOrbitZoomMode( InViewport ) )
 	{
-		FMatrix OrbitMatrix = ViewTransform.ComputeOrbitMatrix().Inverse();
+		FMatrix OrbitMatrix = ViewTransform.ComputeOrbitMatrix().InverseFast();
 
 		FVector DeltaLocation = FVector(0, Drag.X+ -Drag.Y, 0);
 
@@ -1611,7 +1611,7 @@ void FEditorViewportClient::InputAxisForOrbit(FViewport* InViewport, const FVect
 			FRotationMatrix( FRotator(0,GetViewRotation().Yaw,0) ) * 
 			FRotationMatrix( FRotator(0, 0, GetViewRotation().Pitch));
 
-		FVector TransformedDelta = RotMat.Inverse().TransformVector(DeltaLocation);
+		FVector TransformedDelta = RotMat.InverseFast().TransformVector(DeltaLocation);
 
 		SetViewLocation( OrbitMatrix.GetOrigin() + TransformedDelta );
 
@@ -2078,7 +2078,7 @@ void FEditorViewportClient::DrawAxes(FViewport* InViewport, FCanvas* Canvas, con
 	FMatrix ViewTM = FMatrix::Identity;
 	if ( bUsingOrbitCamera)
 	{
-		ViewTM = FRotationMatrix( ViewTransform.ComputeOrbitMatrix().Inverse().Rotator() );
+		ViewTM = FRotationMatrix( ViewTransform.ComputeOrbitMatrix().InverseFast().Rotator() );
 	}
 	else
 	{
diff --git a/Engine/Source/Editor/UnrealEd/Private/Fbx/FbxSkeletalMeshImport.cpp b/Engine/Source/Editor/UnrealEd/Private/Fbx/FbxSkeletalMeshImport.cpp
--- a/Engine/Source/Editor/UnrealEd/Private/Fbx/FbxSkeletalMeshImport.cpp
+++ b/Engine/Source/Editor/UnrealEd/Private/Fbx/FbxSkeletalMeshImport.cpp
@@ -2561,7 +2561,7 @@ bool UnFbx::FFbxImporter::ImportSkeletalMeshLOD(USkeletalMesh* InSkeletalMesh, U
 	NewLODModel.RequiredBones.Sort();
 
 	// To be extra-nice, we apply the difference between the root transform of the meshes to the verts.
-	FMatrix LODToBaseTransform = InSkeletalMesh->GetRefPoseMatrix(0).Inverse() * BaseSkeletalMesh->GetRefPoseMatrix(0);
+	FMatrix LODToBaseTransform = InSkeletalMesh->GetRefPoseMatrix(0).InverseFast() * BaseSkeletalMesh->GetRefPoseMatrix(0);
 
 	for (int32 ChunkIndex = 0; ChunkIndex < NewLODModel.Chunks.Num(); ChunkIndex++)
 	{
diff --git a/Engine/Source/Editor/UnrealEd/Private/GeomTools.cpp b/Engine/Source/Editor/UnrealEd/Private/GeomTools.cpp
--- a/Engine/Source/Editor/UnrealEd/Private/GeomTools.cpp
+++ b/Engine/Source/Editor/UnrealEd/Private/GeomTools.cpp
@@ -654,7 +654,7 @@ void FClipSMTriangle::ComputeGradientsAndNormal()
 {
 	// Compute the transform from triangle parameter space to local space.
 	const FMatrix ParameterToLocal = ComputeTriangleParameterToAttribute(Vertices[0].Pos,Vertices[1].Pos,Vertices[2].Pos);
-	const FMatrix LocalToParameter = ParameterToLocal.InverseSafe();
+	const FMatrix LocalToParameter = ParameterToLocal.Inverse();
 
 	// Compute the triangle's normal.
 	FaceNormal = ParameterToLocal.TransformVector(FVector(0,0,1));
diff --git a/Engine/Source/Editor/UnrealEd/Private/LevelEditorViewport.cpp b/Engine/Source/Editor/UnrealEd/Private/LevelEditorViewport.cpp
--- a/Engine/Source/Editor/UnrealEd/Private/LevelEditorViewport.cpp
+++ b/Engine/Source/Editor/UnrealEd/Private/LevelEditorViewport.cpp
@@ -3787,7 +3787,7 @@ static void RenderViewFrustum( FPrimitiveDrawInterface* PDI,
 
 	for( int32 x = 0 ; x < 8 ; ++x )
 	{
-		Verts[x] = InViewMatrix.Inverse().TransformPosition( Verts[x] );
+		Verts[x] = InViewMatrix.InverseFast().TransformPosition( Verts[x] );
 	}
 
 	const uint8 PrimitiveDPG = SDPG_Foreground;
@@ -3949,7 +3949,7 @@ void FLevelEditorViewportClient::UpdateAudioListener( const FSceneView& View )
 
 			class AReverbVolume* ReverbVolume = GetWorld()->GetAudioSettings( ViewLocation, &ReverbSettings, &InteriorSettings );
 
-			FMatrix CameraToWorld = View.ViewMatrices.ViewMatrix.Inverse();
+			FMatrix CameraToWorld = View.ViewMatrices.ViewMatrix.InverseFast();
 			FVector ProjUp = CameraToWorld.TransformVector(FVector(0, 1000, 0));
 			FVector ProjRight = CameraToWorld.TransformVector(FVector(1000, 0, 0));
 
diff --git a/Engine/Source/Editor/UnrealEd/Private/PhysicsAssetUtils.cpp b/Engine/Source/Editor/UnrealEd/Private/PhysicsAssetUtils.cpp
--- a/Engine/Source/Editor/UnrealEd/Private/PhysicsAssetUtils.cpp
+++ b/Engine/Source/Editor/UnrealEd/Private/PhysicsAssetUtils.cpp
@@ -541,7 +541,6 @@ void WeldBodies(UPhysicsAsset* PhysAsset, int32 BaseBodyIndex, int32 AddBodyInde
 	check(Bone1Index != INDEX_NONE);
 	FTransform Bone1TM = SkelComp->GetBoneTransform(Bone1Index);
 	Bone1TM.RemoveScaling();
-	FTransform InvBone1TM = Bone1TM.InverseSafe();
 
 	UBodySetup* Body2 = PhysAsset->BodySetup[AddBodyIndex];
 	int32 Bone2Index = SkelComp->SkeletalMesh->RefSkeleton.FindBoneIndex(Body2->BoneName);
@@ -549,7 +548,7 @@ void WeldBodies(UPhysicsAsset* PhysAsset, int32 BaseBodyIndex, int32 AddBodyInde
 	FTransform Bone2TM = SkelComp->GetBoneTransform(Bone2Index);
 	Bone2TM.RemoveScaling();
 
-	FTransform Bone2ToBone1TM = Bone2TM * InvBone1TM;
+	FTransform Bone2ToBone1TM = Bone2TM.GetRelativeTransform(Bone1TM);
 
 	// First copy all collision info over.
 	for(int32 i=0; i<Body2->AggGeom.SphereElems.Num(); i++)
diff --git a/Engine/Source/Editor/UnrealEd/Private/StaticLightingSystem/StaticLightingSystem.cpp b/Engine/Source/Editor/UnrealEd/Private/StaticLightingSystem/StaticLightingSystem.cpp
--- a/Engine/Source/Editor/UnrealEd/Private/StaticLightingSystem/StaticLightingSystem.cpp
+++ b/Engine/Source/Editor/UnrealEd/Private/StaticLightingSystem/StaticLightingSystem.cpp
@@ -1354,7 +1354,7 @@ void FStaticLightingSystem::AddBSPStaticLightingInfo(ULevel* Level, bool bBuildL
 
 			// fill out the NodeGroup/mapping, as UModelComponent::GetStaticLightingInfo did
 			SomeModelComponent->GetSurfaceLightMapResolution(SurfaceIndex, true, NodeGroup->SizeX, NodeGroup->SizeY, NodeGroup->WorldToMap, &NodeGroup->Nodes);
-			NodeGroup->MapToWorld = NodeGroup->WorldToMap.Inverse();
+			NodeGroup->MapToWorld = NodeGroup->WorldToMap.InverseFast();
 
 			// Cache the surface's vertices and triangles.
 			NodeGroup->BoundingBox.Init();
@@ -1510,7 +1510,7 @@ void FStaticLightingSystem::AddBSPStaticLightingInfo(ULevel* Level, TArray<FNode
 
 			// fill out the NodeGroup/mapping, as UModelComponent::GetStaticLightingInfo did
 			SomeModelComponent->GetSurfaceLightMapResolution(SurfaceIndex, true, NodeGroup->SizeX, NodeGroup->SizeY, NodeGroup->WorldToMap, &NodeGroup->Nodes);
-			NodeGroup->MapToWorld = NodeGroup->WorldToMap.Inverse();
+			NodeGroup->MapToWorld = NodeGroup->WorldToMap.InverseFast();
 
 			// Cache the surface's vertices and triangles.
 			NodeGroup->BoundingBox.Init();
diff --git a/Engine/Source/Editor/UnrealEd/Private/StaticMeshEdit.cpp b/Engine/Source/Editor/UnrealEd/Private/StaticMeshEdit.cpp
--- a/Engine/Source/Editor/UnrealEd/Private/StaticMeshEdit.cpp
+++ b/Engine/Source/Editor/UnrealEd/Private/StaticMeshEdit.cpp
@@ -725,7 +725,7 @@ static inline void FTexCoordsToVectors(const FVector& V0, const FVector& UV0,
 	TexEqu.SetAxis( 0, FVector(	V1.X - V0.X, V1.Y - V0.Y, V1.Z - V0.Z ) );
 	TexEqu.SetAxis( 1, FVector( V2.X - V0.X, V2.Y - V0.Y, V2.Z - V0.Z ) );
 	TexEqu.SetAxis( 2, FVector( PN.X,        PN.Y,        PN.Z        ) );
-	TexEqu = TexEqu.Inverse();
+	TexEqu = TexEqu.InverseFast();
 
 	const FVector UResult( UV1.X-UV0.X, UV2.X-UV0.X, 0.0f );
 	const FVector TUResult = TexEqu.TransformVector( UResult );
@@ -740,7 +740,7 @@ static inline void FTexCoordsToVectors(const FVector& V0, const FVector& UV0,
 	BaseEqu.SetAxis( 0, TUResult );
 	BaseEqu.SetAxis( 1, TVResult ); 
 	BaseEqu.SetAxis( 2, FVector( PN.X, PN.Y, PN.Z ) );
-	BaseEqu = BaseEqu.Inverse();
+	BaseEqu = BaseEqu.InverseFast();
 
 	const FVector BResult = FVector( UV0.X - ( TUResult|V0 ), UV0.Y - ( TVResult|V0 ),  0.0f );
 
diff --git a/Engine/Source/Editor/UnrealEd/Private/UnrealEdSrv.cpp b/Engine/Source/Editor/UnrealEd/Private/UnrealEdSrv.cpp
--- a/Engine/Source/Editor/UnrealEd/Private/UnrealEdSrv.cpp
+++ b/Engine/Source/Editor/UnrealEd/Private/UnrealEdSrv.cpp
@@ -2093,7 +2093,7 @@ bool UUnrealEdEngine::Exec_Actor( UWorld* InWorld, const TCHAR* Str, FOutputDevi
 			{
 				// Move the plane into the same coordinate space as the builder brush
 
-				*SplittingPlane = SplittingPlane->TransformBy(InWorld->GetDefaultBrush()->ActorToWorld().ToMatrixWithScale().Inverse());
+				*SplittingPlane = SplittingPlane->TransformBy(InWorld->GetDefaultBrush()->ActorToWorld().ToMatrixWithScale().InverseFast());
 
 				// Before keeping this plane, make sure there aren't any existing planes that have a normal within the rejection tolerance.
 
diff --git a/Engine/Source/Editor/UnrealEd/Public/EditorViewportClient.h b/Engine/Source/Editor/UnrealEd/Public/EditorViewportClient.h
--- a/Engine/Source/Editor/UnrealEd/Public/EditorViewportClient.h
+++ b/Engine/Source/Editor/UnrealEd/Public/EditorViewportClient.h
@@ -357,7 +357,7 @@ public:
 		if( bRecalculateView )
 		{
 			FMatrix OrbitMatrix = ViewTransform.ComputeOrbitMatrix();
-			OrbitMatrix = OrbitMatrix.Inverse();
+			OrbitMatrix = OrbitMatrix.InverseFast();
 
 			ViewTransform.SetRotation( OrbitMatrix.Rotator() );
 			ViewTransform.SetLocation( OrbitMatrix.GetOrigin() );
diff --git a/Engine/Source/Editor/WorldBrowser/Private/Tiles/SWorldComposition.cpp b/Engine/Source/Editor/WorldBrowser/Private/Tiles/SWorldComposition.cpp
--- a/Engine/Source/Editor/WorldBrowser/Private/Tiles/SWorldComposition.cpp
+++ b/Engine/Source/Editor/WorldBrowser/Private/Tiles/SWorldComposition.cpp
@@ -350,7 +350,7 @@ public:
 			// Draw a current camera position
 			{
 				const FSlateBrush* CameraImage = FEditorStyle::GetBrush(TEXT("WorldBrowser.SimulationViewPositon"));
-				FMatrix ObserverViewToWorld = WorldModel->GetObserverViewMatrix().InverseSafe();
+				FMatrix ObserverViewToWorld = WorldModel->GetObserverViewMatrix().Inverse();
 				FVector ObserverPosition	= ObserverViewToWorld.GetOrigin();
 				FRotator ObserverRotation	= ObserverViewToWorld.Rotator();
 								
diff --git a/Engine/Source/Programs/UnrealLightmass/Private/Launch/UnitTest.cpp b/Engine/Source/Programs/UnrealLightmass/Private/Launch/UnitTest.cpp
--- a/Engine/Source/Programs/UnrealLightmass/Private/Launch/UnitTest.cpp
+++ b/Engine/Source/Programs/UnrealLightmass/Private/Launch/UnitTest.cpp
@@ -198,7 +198,7 @@ void TestLightmass()
 		TestMatrix.M[3][0], TestMatrix.M[3][1], TestMatrix.M[3][2], TestMatrix.M[3][3]
 		);
 
-	TestMatrix = TestMatrix.GetTransposed().Inverse();
+	TestMatrix = TestMatrix.GetTransposed().InverseFast();
 
 	UE_LOG(LogLightmass, Display, TEXT("Inverted Mat=\n  [%0.2f, %0.2f, %0.2f, %0.2f]\n  [%0.2f, %0.2f, %0.2f, %0.2f]\n  [%0.2f, %0.2f, %0.2f, %0.2f]\n  [%0.2f, %0.2f, %0.2f, %0.2f]"), 
 		TestMatrix.M[0][0], TestMatrix.M[0][1], TestMatrix.M[0][2], TestMatrix.M[0][3], 
diff --git a/Engine/Source/Programs/UnrealLightmass/Private/Lighting/LightingSystem.cpp b/Engine/Source/Programs/UnrealLightmass/Private/Lighting/LightingSystem.cpp
--- a/Engine/Source/Programs/UnrealLightmass/Private/Lighting/LightingSystem.cpp
+++ b/Engine/Source/Programs/UnrealLightmass/Private/Lighting/LightingSystem.cpp
@@ -2086,7 +2086,7 @@ void FStaticLightingSystem::CalculateStaticShadowDepthMap(FGuid LightGuid)
 
 		{
 			const float InvDistanceRange = 1.0f / (LightSpaceImportanceBounds.Max.Z - LightSpaceImportanceBounds.Min.Z);
-			const FMatrix LightToWorld = ShadowDepthMap->WorldToLight.Inverse();
+			const FMatrix LightToWorld = ShadowDepthMap->WorldToLight.InverseFast();
 
 			for (int32 Y = 0; Y < ShadowDepthMap->ShadowMapSizeY; Y++)
 			{
@@ -2177,7 +2177,7 @@ void FStaticLightingSystem::CalculateStaticShadowDepthMap(FGuid LightGuid)
 
 		// Calculate the maximum possible distance for quantization
 		const float MaxPossibleDistance = LightSpaceImportanceBoundMax.Z - LightSpaceImportanceBoundMin.Z;
-		const FMatrix LightToWorld = ShadowDepthMap->WorldToLight.Inverse();
+		const FMatrix LightToWorld = ShadowDepthMap->WorldToLight.InverseFast();
 		const FBoxSphereBounds ImportanceVolume = GetImportanceBounds().SphereRadius > 0.0f ? GetImportanceBounds() : FBoxSphereBounds(AggregateMesh.GetBounds());
 
 		for (int32 Y = 0; Y < ShadowDepthMap->ShadowMapSizeY; Y++)
diff --git a/Engine/Source/Programs/UnrealLightmass/Private/Lighting/StaticMesh.cpp b/Engine/Source/Programs/UnrealLightmass/Private/Lighting/StaticMesh.cpp
--- a/Engine/Source/Programs/UnrealLightmass/Private/Lighting/StaticMesh.cpp
+++ b/Engine/Source/Programs/UnrealLightmass/Private/Lighting/StaticMesh.cpp
@@ -267,7 +267,7 @@ void FStaticMeshStaticLightingMesh::Import( FLightmassImporter& Importer )
 	SplineParameters = TempSMSLMD.SplineParameters;
 
 	// calculate the inverse transpose
-	LocalToWorldInverseTranspose = LocalToWorld.Inverse().GetTransposed();
+	LocalToWorldInverseTranspose = LocalToWorld.InverseFast().GetTransposed();
 
 	// we have the guid for the mesh, now hook it up to the actual static mesh
 	StaticMesh = Importer.ConditionalImportObject<FStaticMesh>(StaticMeshGuid, LM_STATICMESH_VERSION, LM_STATICMESH_EXTENSION, LM_STATICMESH_CHANNEL_FLAGS, Importer.GetStaticMeshes());
diff --git a/Engine/Source/Runtime/Core/Private/Math/Transform.cpp b/Engine/Source/Runtime/Core/Private/Math/Transform.cpp
--- a/Engine/Source/Runtime/Core/Private/Math/Transform.cpp
+++ b/Engine/Source/Runtime/Core/Private/Math/Transform.cpp
@@ -16,7 +16,7 @@ DEFINE_LOG_CATEGORY_STATIC(LogTransform, Log, All);
 const FTransform FTransform::Identity(FQuat(0.f,0.f,0.f,1.f), FVector::ZeroVector, FVector(1.f));
 
 
-// Replacement of InverseSafe of FMatrix
+// Replacement of Inverse of FMatrix
 
 /**
 * Does a debugf of the contents of this BoneAtom.
@@ -107,7 +107,7 @@ FTransform FTransform::GetRelativeTransformReverse(const FTransform& Other) cons
 	FMatrix AM = ToMatrixWithScale();
 	FMatrix BM = Other.ToMatrixWithScale();
 
-	Result.DebugEqualMatrix(AM.Inverse() *  BM);
+	Result.DebugEqualMatrix(AM.InverseFast() *  BM);
 #endif
 
 	return Result;
@@ -138,7 +138,7 @@ void FTransform::SetToRelativeTransform(const FTransform& ParentTransform)
 	Rotation = InverseRot * Rotation;
 
 #if DEBUG_INVERSE_TRANSFORM
- 	DebugEqualMatrix(AM *  BM.Inverse());
+ 	DebugEqualMatrix(AM *  BM.InverseFast());
 #endif
 }
 
@@ -160,7 +160,7 @@ FTransform FTransform::GetRelativeTransform(const FTransform& Other) const
 		return FTransform::Identity;
 	}
 
-	FQuat Inverse = Other.Rotation.Inverse();
+	FQuat Inverse = Other.Rotation.InverseFast();
 	Result.Rotation = Inverse*Rotation;
 
 	Result.Translation = (Inverse*(Translation - Other.Translation))*(SafeRecipScale3D);
@@ -169,7 +169,7 @@ FTransform FTransform::GetRelativeTransform(const FTransform& Other) const
  	FMatrix AM = ToMatrixWithScale();
  	FMatrix BM = Other.ToMatrixWithScale();
  
- 	Result.DebugEqualMatrix(AM *  BM.Inverse());
+ 	Result.DebugEqualMatrix(AM *  BM.InverseFast());
 
 #endif
 
diff --git a/Engine/Source/Runtime/Core/Private/Math/TransformVectorized.cpp b/Engine/Source/Runtime/Core/Private/Math/TransformVectorized.cpp
--- a/Engine/Source/Runtime/Core/Private/Math/TransformVectorized.cpp
+++ b/Engine/Source/Runtime/Core/Private/Math/TransformVectorized.cpp
@@ -16,7 +16,7 @@ DEFINE_LOG_CATEGORY_STATIC(LogTransform, Log, All);
 const FTransform FTransform::Identity(FQuat(0.f,0.f,0.f,1.f), FVector::ZeroVector, FVector(1.f));
 
 
-// Replacement of InverseSafe of FMatrix
+// Replacement of Inverse of FMatrix
 
 /**
 * Does a debugf of the contents of this BoneAtom.
@@ -122,7 +122,7 @@ FTransform FTransform::GetRelativeTransformReverse(const FTransform& Other) cons
 	FMatrix AM = ToMatrixWithScale();
 	FMatrix BM = Other.ToMatrixWithScale();
 
-	Result.DebugEqualMatrix(AM.Inverse() *  BM);
+	Result.DebugEqualMatrix(AM.InverseFast() *  BM);
 #endif
 
 	return Result;
@@ -166,7 +166,7 @@ void FTransform::SetToRelativeTransform(const FTransform& ParentTransform)
 	DiagnosticCheckNaN_All(); 
 
 #if DEBUG_INVERSE_TRANSFORM
-	DebugEqualMatrix(AM *  BM.Inverse());
+	DebugEqualMatrix(AM *  BM.InverseFast());
 #endif
 }
 
@@ -212,7 +212,7 @@ FTransform FTransform::GetRelativeTransform(const FTransform& Other) const
 	FMatrix AM = ToMatrixWithScale();
 	FMatrix BM = Other.ToMatrixWithScale();
 
-	Result.DebugEqualMatrix(AM *  BM.Inverse());
+	Result.DebugEqualMatrix(AM *  BM.InverseFast());
 
 #endif
 	return Result;
diff --git a/Engine/Source/Runtime/Core/Private/Math/UnrealMath.cpp b/Engine/Source/Runtime/Core/Private/Math/UnrealMath.cpp
--- a/Engine/Source/Runtime/Core/Private/Math/UnrealMath.cpp
+++ b/Engine/Source/Runtime/Core/Private/Math/UnrealMath.cpp
@@ -1624,7 +1624,7 @@ FVector4 FMath::ComputeBaryCentric3D(const FVector& Point, const FVector& A, con
 
 	//The point V can be expressed as Ax=v where x is the vector containing the weights {w1...wn}
 	//Solve for x by multiplying both sides by AInv   (AInv * A)x = AInv * v ==> x = AInv * v
-	const FMatrix InvSolvMat = SolvMat.InverseSafe();
+	const FMatrix InvSolvMat = SolvMat.Inverse();
 	const FPlane BaryCoords = InvSolvMat.TransformVector(V);	 
 
 	//Reorder the weights to be a, b, c, d
diff --git a/Engine/Source/Runtime/Core/Private/Tests/Math/UnrealMathTest.cpp b/Engine/Source/Runtime/Core/Private/Tests/Math/UnrealMathTest.cpp
--- a/Engine/Source/Runtime/Core/Private/Tests/Math/UnrealMathTest.cpp
+++ b/Engine/Source/Runtime/Core/Private/Tests/Math/UnrealMathTest.cpp
@@ -692,11 +692,11 @@ bool FVectorRegisterAbstractionTest::RunTest(const FString& Parameters)
 	TestVectorMatrixInverse( &M3, &M1 );
 	LogTest( TEXT("VectorMatrixInverse"), TestMatricesEqual( M2, M3 ) );
 
-	FTransform Transform;
-	Transform.SetFromMatrix(M1);
-	FTransform InvTransform = Transform.InverseSafe();
-	FTransform InvTransform2 = FTransform(Transform.ToMatrixWithScale().InverseSafe());
-	LogTest( TEXT("FTransform Inverse"), InvTransform.Equals(InvTransform2, 1e-3f ) );
+// 	FTransform Transform;
+// 	Transform.SetFromMatrix(M1);
+// 	FTransform InvTransform = Transform.Inverse();
+// 	FTransform InvTransform2 = FTransform(Transform.ToMatrixWithScale().Inverse());
+// 	LogTest( TEXT("FTransform Inverse"), InvTransform.Equals(InvTransform2, 1e-3f ) );
 
 	V0 = MakeVectorRegister( 100.0f, -100.0f, 200.0f, 1.0f );
 	V1 = VectorTransformVector(V0, &M0);
diff --git a/Engine/Source/Runtime/Core/Public/Math/Matrix.h b/Engine/Source/Runtime/Core/Public/Math/Matrix.h
--- a/Engine/Source/Runtime/Core/Public/Math/Matrix.h
+++ b/Engine/Source/Runtime/Core/Public/Math/Matrix.h
@@ -168,11 +168,9 @@ public:
 	inline float RotDeterminant() const;
 
 	/** Fast path, doesn't check for nil matrices in final release builds */
-	inline FMatrix Inverse() const;
+	inline FMatrix InverseFast() const;
 	/** Fast path, and handles nil matrices. */
-	inline FMatrix InverseSafe() const;
-	/** Slow and safe path */
-	inline FMatrix InverseSlow() const;
+	inline FMatrix Inverse() const;
 
 	inline FMatrix TransposeAdjoint() const;
 
diff --git a/Engine/Source/Runtime/Core/Public/Math/Transform.h b/Engine/Source/Runtime/Core/Public/Math/Transform.h
--- a/Engine/Source/Runtime/Core/Public/Math/Transform.h
+++ b/Engine/Source/Runtime/Core/Public/Math/Transform.h
@@ -267,16 +267,16 @@ public:
 	FORCEINLINE FMatrix ToInverseMatrixWithScale() const
 	{
 		// todo: optimize
-		return ToMatrixWithScale().InverseSafe();
+		return ToMatrixWithScale().Inverse();
 	}
 
 	/**
 	* Convert this Transform to inverse.
 	*/
-	FORCEINLINE FTransform InverseSafe() const
+	FORCEINLINE FTransform Inverse() const
 	{
 		// todo: optimize
-		return FTransform(ToMatrixWithScale().InverseSafe());
+		return FTransform(ToMatrixWithScale().Inverse());
 	}
 
 	/**
diff --git a/Engine/Source/Runtime/Core/Public/Math/TransformVectorized.h b/Engine/Source/Runtime/Core/Public/Math/TransformVectorized.h
--- a/Engine/Source/Runtime/Core/Public/Math/TransformVectorized.h
+++ b/Engine/Source/Runtime/Core/Public/Math/TransformVectorized.h
@@ -301,18 +301,18 @@ public:
 	FORCEINLINE FMatrix ToInverseMatrixWithScale() const
 	{
 		// todo: optimize
-		return ToMatrixWithScale().InverseSafe();
+		return ToMatrixWithScale().Inverse();
 	}
 
 	/**
 	* Convert this Transform to inverse.
 	*/
-	FORCEINLINE FTransform InverseSafe() const
+	FORCEINLINE FTransform Inverse() const
 	{
-		// Replacement of InverseSafe of FMatrix
+		// Replacement of Inverse of FMatrix
 		if (VectorAnyGreaterThan(VectorAbs(Scale3D), GlobalVectorConstants::SmallNumber))
 		{
-			return Inverse();
+			return InverseFast();
 		}
 		else
 		{
@@ -1231,7 +1231,7 @@ private:
 	}
 
 	/** Returns Inverse Transform of this FTransform **/
-	FORCEINLINE FTransform Inverse() const
+	FORCEINLINE FTransform InverseFast() const
 	{
 		// Inverse QST (A) = QST (~A)
 		// Since A*~A = Identity, 
diff --git a/Engine/Source/Runtime/Core/Public/Math/UnrealMatrix.h b/Engine/Source/Runtime/Core/Public/Math/UnrealMatrix.h
--- a/Engine/Source/Runtime/Core/Public/Math/UnrealMatrix.h
+++ b/Engine/Source/Runtime/Core/Public/Math/UnrealMatrix.h
@@ -166,7 +166,7 @@ FORCEINLINE FVector4 FMatrix::TransformPosition(const FVector &V) const
 /** Inverts the matrix and then transforms V - correctly handles scaling in this matrix. */
 FORCEINLINE FVector FMatrix::InverseTransformPosition(const FVector &V) const
 {
-	FMatrix InvSelf = this->InverseSafe();
+	FMatrix InvSelf = this->InverseFast();
 	return InvSelf.TransformPosition(V);
 }
 
@@ -184,7 +184,7 @@ FORCEINLINE FVector4 FMatrix::TransformVector(const FVector& V) const
 /** Faster version of InverseTransformVector that assumes no scaling. WARNING: Will NOT work correctly if there is scaling in the matrix. */
 FORCEINLINE FVector FMatrix::InverseTransformVector(const FVector &V) const
 {
-	FMatrix InvSelf = this->Inverse();
+	FMatrix InvSelf = this->InverseFast();
 	return InvSelf.TransformVector(V);
 }
 
@@ -255,7 +255,7 @@ inline float FMatrix::RotDeterminant() const
 
 // Inverse.
 /** Fast path, doesn't check for nil matrices in final release builds */
-inline FMatrix FMatrix::Inverse() const
+inline FMatrix FMatrix::InverseFast() const
 {
 	// If we're in non final release, then make sure we're not creating NaNs
 #if !(UE_BUILD_SHIPPING || UE_BUILD_TEST)
@@ -264,8 +264,8 @@ inline FMatrix FMatrix::Inverse() const
 		GetScaledAxis( EAxis::Y ).IsNearlyZero(SMALL_NUMBER) && 
 		GetScaledAxis( EAxis::Z ).IsNearlyZero(SMALL_NUMBER) ) 
 	{
-		UE_LOG(LogUnrealMath, Error, TEXT("FMatrix::Inverse(), trying to invert a NIL matrix, this results in NaNs! Use InverseSafe() instead."));
-		ensureMsgf(false, TEXT("FMatrix::Inverse(), trying to invert a NIL matrix, this results in NaNs! Use InverseSafe() instead."));
+		UE_LOG(LogUnrealMath, Error, TEXT("FMatrix::Inverse(), trying to invert a NIL matrix, this results in NaNs! Use Inverse() instead."));
+		ensureMsgf(false, TEXT("FMatrix::Inverse(), trying to invert a NIL matrix, this results in NaNs! Use Inverse() instead."));
 	}
 #endif
 	FMatrix Result;
@@ -274,7 +274,7 @@ inline FMatrix FMatrix::Inverse() const
 }
 
 // Inverse.
-inline FMatrix FMatrix::InverseSafe() const
+inline FMatrix FMatrix::Inverse() const
 {
 	FMatrix Result;
 
@@ -284,108 +284,21 @@ inline FMatrix FMatrix::InverseSafe() const
 		GetScaledAxis( EAxis::Z ).IsNearlyZero(SMALL_NUMBER) ) 
 	{
 		// just set to zero - avoids unsafe inverse of zero and duplicates what QNANs were resulting in before (scaling away all children)
-		Result = FMatrix(FVector::ZeroVector, FVector::ZeroVector, FVector::ZeroVector, FVector::ZeroVector);
+		Result = FMatrix::Identity;
 	}
 	else
 	{
-		VectorMatrixInverse( &Result, this );
-	}
-	return Result;
-}
+		const float	Det = Determinant();
 
-inline FMatrix FMatrix::InverseSlow() const
-{
-	FMatrix Result;
-	const float	Det = Determinant();
-
-	if(Det == 0.0f)
-		return FMatrix::Identity;
-
-	const float	RDet = 1.0f / Det;
-
-	Result.M[0][0] = RDet * (
-			M[1][1] * (M[2][2] * M[3][3] - M[2][3] * M[3][2]) -
-			M[2][1] * (M[1][2] * M[3][3] - M[1][3] * M[3][2]) +
-			M[3][1] * (M[1][2] * M[2][3] - M[1][3] * M[2][2])
-			);
-	Result.M[0][1] = -RDet * (
-			M[0][1] * (M[2][2] * M[3][3] - M[2][3] * M[3][2]) -
-			M[2][1] * (M[0][2] * M[3][3] - M[0][3] * M[3][2]) +
-			M[3][1] * (M[0][2] * M[2][3] - M[0][3] * M[2][2])
-			);
-	Result.M[0][2] = RDet * (
-			M[0][1] * (M[1][2] * M[3][3] - M[1][3] * M[3][2]) -
-			M[1][1] * (M[0][2] * M[3][3] - M[0][3] * M[3][2]) +
-			M[3][1] * (M[0][2] * M[1][3] - M[0][3] * M[1][2])
-			);
-	Result.M[0][3] = -RDet * (
-			M[0][1] * (M[1][2] * M[2][3] - M[1][3] * M[2][2]) -
-			M[1][1] * (M[0][2] * M[2][3] - M[0][3] * M[2][2]) +
-			M[2][1] * (M[0][2] * M[1][3] - M[0][3] * M[1][2])
-			);
-
-	Result.M[1][0] = -RDet * (
-			M[1][0] * (M[2][2] * M[3][3] - M[2][3] * M[3][2]) -
-			M[2][0] * (M[1][2] * M[3][3] - M[1][3] * M[3][2]) +
-			M[3][0] * (M[1][2] * M[2][3] - M[1][3] * M[2][2])
-			);
-	Result.M[1][1] = RDet * (
-			M[0][0] * (M[2][2] * M[3][3] - M[2][3] * M[3][2]) -
-			M[2][0] * (M[0][2] * M[3][3] - M[0][3] * M[3][2]) +
-			M[3][0] * (M[0][2] * M[2][3] - M[0][3] * M[2][2])
-			);
-	Result.M[1][2] = -RDet * (
-			M[0][0] * (M[1][2] * M[3][3] - M[1][3] * M[3][2]) -
-			M[1][0] * (M[0][2] * M[3][3] - M[0][3] * M[3][2]) +
-			M[3][0] * (M[0][2] * M[1][3] - M[0][3] * M[1][2])
-			);
-	Result.M[1][3] = RDet * (
-			M[0][0] * (M[1][2] * M[2][3] - M[1][3] * M[2][2]) -
-			M[1][0] * (M[0][2] * M[2][3] - M[0][3] * M[2][2]) +
-			M[2][0] * (M[0][2] * M[1][3] - M[0][3] * M[1][2])
-			);
-
-	Result.M[2][0] = RDet * (
-			M[1][0] * (M[2][1] * M[3][3] - M[2][3] * M[3][1]) -
-			M[2][0] * (M[1][1] * M[3][3] - M[1][3] * M[3][1]) +
-			M[3][0] * (M[1][1] * M[2][3] - M[1][3] * M[2][1])
-			);
-	Result.M[2][1] = -RDet * (
-			M[0][0] * (M[2][1] * M[3][3] - M[2][3] * M[3][1]) -
-			M[2][0] * (M[0][1] * M[3][3] - M[0][3] * M[3][1]) +
-			M[3][0] * (M[0][1] * M[2][3] - M[0][3] * M[2][1])
-			);
-	Result.M[2][2] = RDet * (
-			M[0][0] * (M[1][1] * M[3][3] - M[1][3] * M[3][1]) -
-			M[1][0] * (M[0][1] * M[3][3] - M[0][3] * M[3][1]) +
-			M[3][0] * (M[0][1] * M[1][3] - M[0][3] * M[1][1])
-			);
-	Result.M[2][3] = -RDet * (
-			M[0][0] * (M[1][1] * M[2][3] - M[1][3] * M[2][1]) -
-			M[1][0] * (M[0][1] * M[2][3] - M[0][3] * M[2][1]) +
-			M[2][0] * (M[0][1] * M[1][3] - M[0][3] * M[1][1])
-			);
-
-	Result.M[3][0] = -RDet * (
-			M[1][0] * (M[2][1] * M[3][2] - M[2][2] * M[3][1]) -
-			M[2][0] * (M[1][1] * M[3][2] - M[1][2] * M[3][1]) +
-			M[3][0] * (M[1][1] * M[2][2] - M[1][2] * M[2][1])
-			);
-	Result.M[3][1] = RDet * (
-			M[0][0] * (M[2][1] * M[3][2] - M[2][2] * M[3][1]) -
-			M[2][0] * (M[0][1] * M[3][2] - M[0][2] * M[3][1]) +
-			M[3][0] * (M[0][1] * M[2][2] - M[0][2] * M[2][1])
-			);
-	Result.M[3][2] = -RDet * (
-			M[0][0] * (M[1][1] * M[3][2] - M[1][2] * M[3][1]) -
-			M[1][0] * (M[0][1] * M[3][2] - M[0][2] * M[3][1]) +
-			M[3][0] * (M[0][1] * M[1][2] - M[0][2] * M[1][1])
-			);
-	Result.M[3][3] = RDet * (
-			M[0][0] * (M[1][1] * M[2][2] - M[1][2] * M[2][1]) -
-			M[1][0] * (M[0][1] * M[2][2] - M[0][2] * M[2][1]) +
-			M[2][0] * (M[0][1] * M[1][2] - M[0][2] * M[1][1])
-			);
+		if(Det == 0.0f)
+		{
+			Result = FMatrix::Identity;
+		}
+		else
+		{
+			VectorMatrixInverse( &Result, this );
+		}
+	}
 
 	return Result;
 }
diff --git a/Engine/Source/Runtime/Engine/Private/ActiveSound.cpp b/Engine/Source/Runtime/Engine/Private/ActiveSound.cpp
--- a/Engine/Source/Runtime/Engine/Private/ActiveSound.cpp
+++ b/Engine/Source/Runtime/Engine/Private/ActiveSound.cpp
@@ -196,7 +196,7 @@ void FActiveSound::UpdateWaveInstances( FAudioDevice* AudioDevice, TArray<FWaveI
 	// if the closest listener is not the primary one, transform CurrentLocation
 	if( ClosestListenerIndex != 0 )
 	{
-		ParseParams.Transform = ParseParams.Transform * ClosestListener.Transform.InverseSafe() * Listener.Transform;
+		ParseParams.Transform = ParseParams.Transform * ClosestListener.Transform.Inverse() * Listener.Transform;
 	}
 
 	// Recurse nodes, have SoundWave's create new wave instances and update bFinished unless we finished fading out.
diff --git a/Engine/Source/Runtime/Engine/Private/Animation/AnimNode_Trail.cpp b/Engine/Source/Runtime/Engine/Private/Animation/AnimNode_Trail.cpp
--- a/Engine/Source/Runtime/Engine/Private/Animation/AnimNode_Trail.cpp
+++ b/Engine/Source/Runtime/Engine/Private/Animation/AnimNode_Trail.cpp
@@ -96,7 +96,7 @@ void FAnimNode_Trail::EvaluateBoneTransforms(USkeletalMeshComponent* SkelComp, c
 	bHadValidStrength = bHasValidStrength;
 
 	// transform between last frame and now.
-	FMatrix OldToNewTM = OldLocalToWorld * SkelComp->GetTransformMatrix().Inverse();
+	FMatrix OldToNewTM = OldLocalToWorld * SkelComp->GetTransformMatrix().InverseFast();
 
 	// Add fake velocity if present to all but root bone
 	if(!FakeVelocity.IsZero())
diff --git a/Engine/Source/Runtime/Engine/Private/BatchedElements.cpp b/Engine/Source/Runtime/Engine/Private/BatchedElements.cpp
--- a/Engine/Source/Runtime/Engine/Private/BatchedElements.cpp
+++ b/Engine/Source/Runtime/Engine/Private/BatchedElements.cpp
@@ -612,7 +612,7 @@ bool FBatchedElements::Draw(FRHICommandList& RHICmdList, bool bNeedToSwitchVerti
 {
 	if( HasPrimsToDraw() )
 	{
-		FMatrix InvTransform = Transform.InverseSafe();
+		FMatrix InvTransform = Transform.Inverse();
 		FVector CameraX = InvTransform.TransformVector(FVector(1,0,0)).SafeNormal();
 		FVector CameraY = InvTransform.TransformVector(FVector(0,1,0)).SafeNormal();
 		FVector CameraZ = InvTransform.TransformVector(FVector(0,0,1)).SafeNormal();
diff --git a/Engine/Source/Runtime/Engine/Private/Camera/CameraModifier_CameraShake.cpp b/Engine/Source/Runtime/Engine/Private/Camera/CameraModifier_CameraShake.cpp
--- a/Engine/Source/Runtime/Engine/Private/Camera/CameraModifier_CameraShake.cpp
+++ b/Engine/Source/Runtime/Engine/Private/Camera/CameraModifier_CameraShake.cpp
@@ -153,11 +153,11 @@ void UCameraModifier_CameraShake::UpdateCameraShake(float DeltaTime, FCameraShak
 
 			// find transform from camera to the "play space"
 			FRotationMatrix const CamToWorld(InOutPOV.Rotation);
-			FMatrix const CameraToPlaySpace = CamToWorld * PlaySpaceToWorld.InverseSafe();			// CameraToWorld * WorldToPlaySpace
+			FMatrix const CameraToPlaySpace = CamToWorld * PlaySpaceToWorld.Inverse();			// CameraToWorld * WorldToPlaySpace
 
 			// find transform from anim (applied in playspace) back to camera
 			FRotationMatrix const AnimToPlaySpace(RotOffset);
-			FMatrix const AnimToCamera = AnimToPlaySpace * CameraToPlaySpace.InverseSafe();			// AnimToPlaySpace * PlaySpaceToCamera
+			FMatrix const AnimToCamera = AnimToPlaySpace * CameraToPlaySpace.Inverse();			// AnimToPlaySpace * PlaySpaceToCamera
 
 			// RCS = rotated camera space, meaning camera space after it's been animated
 			// this is what we're looking for, the diff between rotated cam space and regular cam space.
diff --git a/Engine/Source/Runtime/Engine/Private/Landscape/LandscapeGizmoActor.cpp b/Engine/Source/Runtime/Engine/Private/Landscape/LandscapeGizmoActor.cpp
--- a/Engine/Source/Runtime/Engine/Private/Landscape/LandscapeGizmoActor.cpp
+++ b/Engine/Source/Runtime/Engine/Private/Landscape/LandscapeGizmoActor.cpp
@@ -140,7 +140,7 @@ public:
 				SampleSizeY = Gizmo->SampleSizeY;
 				bHeightmapRendering = (Gizmo->DataType & LGT_Height);
 				FMatrix LToW = LandscapeInfo->GetLandscapeProxy()->LandscapeActorToWorld().ToMatrixWithScale();
-				FMatrix WToL = LToW.Inverse();
+				FMatrix WToL = LToW.InverseFast();
 				FVector BaseLocation = WToL.TransformPosition(Gizmo->GetActorLocation());
 				float ScaleXY = LandscapeInfo->DrawScale.X;
 				float ScaleZ = LandscapeInfo->DrawScale.Z;
@@ -172,7 +172,7 @@ public:
 				float ScaleX = Gizmo->GetWidth() / Gizmo->CachedWidth / ScaleXY * Gizmo->CachedScaleXY;
 				float ScaleY = Gizmo->GetHeight() / Gizmo->CachedHeight / ScaleXY * Gizmo->CachedScaleXY;
 				FScaleMatrix Mat(FVector(ScaleX, ScaleY, L));
-				FMatrix NormalM = Mat.Inverse().GetTransposed();
+				FMatrix NormalM = Mat.InverseFast().GetTransposed();
 
 				int32 SamplingSize = Gizmo->SampleSizeX * Gizmo->SampleSizeY;
 				SampledPositions.Empty(SamplingSize);
@@ -616,7 +616,7 @@ FRotator ALandscapeGizmoActiveActor::SnapToLandscapeGrid(const FRotator& GizmoRo
 	// Snap to multiples of 90 Yaw in landscape coordinate system
 	//check(TargetLandscapeInfo && TargetLandscapeInfo->LandscapeProxy);
 	//const FTransform LToW = TargetLandscapeInfo->LandscapeProxy->ActorToWorld();
-	//const FRotator LandscapeSpaceRotation = (LToW.GetRotation().Inverse() * GizmoRotation.Quaternion()).Rotator().GetNormalized();
+	//const FRotator LandscapeSpaceRotation = (LToW.GetRotation().InverseFast() * GizmoRotation.Quaternion()).Rotator().GetNormalized();
 	//const FRotator SnappedLandscapeSpaceRotation = FRotator(0, FMath::GridSnap(LandscapeSpaceRotation.Yaw, 90), 0);
 	//const FRotator ResultRotation = (SnappedLandscapeSpaceRotation.Quaternion() * LToW.GetRotation()).Rotator().GetNormalized();
 
diff --git a/Engine/Source/Runtime/Engine/Private/Landscape/LandscapeRender.cpp b/Engine/Source/Runtime/Engine/Private/Landscape/LandscapeRender.cpp
--- a/Engine/Source/Runtime/Engine/Private/Landscape/LandscapeRender.cpp
+++ b/Engine/Source/Runtime/Engine/Private/Landscape/LandscapeRender.cpp
@@ -976,7 +976,7 @@ void FLandscapeComponentSceneProxy::OnTransformChanged()
 	
 	// cache component's WorldToLocal
 	FMatrix LtoW = GetLocalToWorld();
-	WorldToLocal = LtoW.Inverse();
+	WorldToLocal = LtoW.InverseFast();
 
 	// cache component's LocalToWorldNoScaling
 	LocalToWorldNoScaling = LtoW;
diff --git a/Engine/Source/Runtime/Engine/Private/Level.cpp b/Engine/Source/Runtime/Engine/Private/Level.cpp
--- a/Engine/Source/Runtime/Engine/Private/Level.cpp
+++ b/Engine/Source/Runtime/Engine/Private/Level.cpp
@@ -989,7 +989,7 @@ void ULevel::CreateModelComponents()
 
 				// fill out the NodeGroup/mapping, as UModelComponent::GetStaticLightingInfo did
 				SomeModelComponent->GetSurfaceLightMapResolution(SurfaceIndex, true, NodeGroup->SizeX, NodeGroup->SizeY, NodeGroup->WorldToMap, &NodeGroup->Nodes);
-				NodeGroup->MapToWorld = NodeGroup->WorldToMap.Inverse();
+				NodeGroup->MapToWorld = NodeGroup->WorldToMap.InverseFast();
 
 				// Cache the surface's vertices and triangles.
 				NodeGroup->BoundingBox.Init();
diff --git a/Engine/Source/Runtime/Engine/Private/LevelUtils.cpp b/Engine/Source/Runtime/Engine/Private/LevelUtils.cpp
--- a/Engine/Source/Runtime/Engine/Private/LevelUtils.cpp
+++ b/Engine/Source/Runtime/Engine/Private/LevelUtils.cpp
@@ -312,7 +312,7 @@ void FLevelUtils::RemoveEditorTransform(const ULevelStreaming* StreamingLevel, b
 	ULevel* LoadedLevel = StreamingLevel->GetLoadedLevel();
 	if( LoadedLevel != NULL )
 	{
-		ApplyLevelTransform( LoadedLevel, StreamingLevel->LevelTransform.InverseSafe(), bDoPostEditMove );
+		ApplyLevelTransform( LoadedLevel, StreamingLevel->LevelTransform.Inverse(), bDoPostEditMove );
 	}
 }
 
diff --git a/Engine/Source/Runtime/Engine/Private/LightComponent.cpp b/Engine/Source/Runtime/Engine/Private/LightComponent.cpp
--- a/Engine/Source/Runtime/Engine/Private/LightComponent.cpp
+++ b/Engine/Source/Runtime/Engine/Private/LightComponent.cpp
@@ -257,7 +257,7 @@ bool FLightSceneProxy::ShouldCreatePerObjectShadowsForDynamicObjects() const
 void FLightSceneProxy::SetTransform(const FMatrix& InLightToWorld,const FVector4& InPosition)
 {
 	LightToWorld = InLightToWorld;
-	WorldToLight = InLightToWorld.Inverse();
+	WorldToLight = InLightToWorld.InverseFast();
 	Position = InPosition;
 }
 
diff --git a/Engine/Source/Runtime/Engine/Private/MaterialBillboardComponent.cpp b/Engine/Source/Runtime/Engine/Private/MaterialBillboardComponent.cpp
--- a/Engine/Source/Runtime/Engine/Private/MaterialBillboardComponent.cpp
+++ b/Engine/Source/Runtime/Engine/Private/MaterialBillboardComponent.cpp
@@ -136,7 +136,7 @@ public:
 				const FVector CameraUp      = -View->InvViewProjectionMatrix.TransformVector(FVector(1.0f,0.0f,0.0f));
 				const FVector CameraRight   = -View->InvViewProjectionMatrix.TransformVector(FVector(0.0f,1.0f,0.0f));
 				const FVector CameraForward = -View->InvViewProjectionMatrix.TransformVector(FVector(0.0f,0.0f,1.0f));
-				const FMatrix WorldToLocal = GetLocalToWorld().Inverse();
+				const FMatrix WorldToLocal = GetLocalToWorld().InverseFast();
 				const FVector LocalCameraUp = WorldToLocal.TransformVector(CameraUp);
 				const FVector LocalCameraRight = WorldToLocal.TransformVector(CameraRight);
 				const FVector LocalCameraForward = WorldToLocal.TransformVector(CameraForward);
@@ -243,7 +243,7 @@ public:
 		const FVector CameraUp      = -View->InvViewProjectionMatrix.TransformVector(FVector(1.0f,0.0f,0.0f));
 		const FVector CameraRight   = -View->InvViewProjectionMatrix.TransformVector(FVector(0.0f,1.0f,0.0f));
 		const FVector CameraForward = -View->InvViewProjectionMatrix.TransformVector(FVector(0.0f,0.0f,1.0f));
-		const FMatrix WorldToLocal = GetLocalToWorld().Inverse();
+		const FMatrix WorldToLocal = GetLocalToWorld().InverseFast();
 		const FVector LocalCameraUp = WorldToLocal.TransformVector(CameraUp);
 		const FVector LocalCameraRight = WorldToLocal.TransformVector(CameraRight);
 		const FVector LocalCameraForward = WorldToLocal.TransformVector(CameraForward);
diff --git a/Engine/Source/Runtime/Engine/Private/Particles/ParticleModules_Location.cpp b/Engine/Source/Runtime/Engine/Private/Particles/ParticleModules_Location.cpp
--- a/Engine/Source/Runtime/Engine/Private/Particles/ParticleModules_Location.cpp
+++ b/Engine/Source/Runtime/Engine/Private/Particles/ParticleModules_Location.cpp
@@ -233,7 +233,7 @@ void UParticleModuleLocationWorldOffset::SpawnEx(FParticleEmitterInstance* Owner
 	else
 	{
 		// We need to inverse transform the location so that the bUseLocalSpace transform uses the proper value
-		FMatrix InvMat = Owner->Component->ComponentToWorld.ToMatrixWithScale().Inverse();
+		FMatrix InvMat = Owner->Component->ComponentToWorld.ToMatrixWithScale().InverseFast();
 		FVector StartLoc = StartLocation.GetValue(Owner->EmitterTime, Owner->Component, 0, InRandomStream);
 		Particle.Location += InvMat.TransformVector(StartLoc);
 	}
diff --git a/Engine/Source/Runtime/Engine/Private/Particles/ParticleModules_Velocity.cpp b/Engine/Source/Runtime/Engine/Private/Particles/ParticleModules_Velocity.cpp
--- a/Engine/Source/Runtime/Engine/Private/Particles/ParticleModules_Velocity.cpp
+++ b/Engine/Source/Runtime/Engine/Private/Particles/ParticleModules_Velocity.cpp
@@ -329,7 +329,7 @@ void UParticleModuleVelocityOverLifetime::Update(FParticleEmitterInstance* Owner
 			{
 				FVector Vel;
 				const FMatrix LocalToWorld = Owner->Component->ComponentToWorld.ToMatrixNoScale();
-				const FMatrix InvMat = LocalToWorld.Inverse();
+				const FMatrix InvMat = LocalToWorld.InverseFast();
 				BEGIN_UPDATE_LOOP;
 				{
 					Vel = VelOverLife.GetValue(Particle.RelativeTime, Owner->Component);
@@ -377,7 +377,7 @@ void UParticleModuleVelocityOverLifetime::Update(FParticleEmitterInstance* Owner
 			{
 				FVector Vel;
 				const FMatrix LocalToWorld = Owner->Component->ComponentToWorld.ToMatrixNoScale();
-				const FMatrix InvMat = LocalToWorld.Inverse();
+				const FMatrix InvMat = LocalToWorld.InverseFast();
 				BEGIN_UPDATE_LOOP;
 				{
 					Vel = VelOverLife.GetValue(Particle.RelativeTime, Owner->Component);
@@ -607,7 +607,7 @@ void UParticleModuleVelocityCone::Render3DPreview(FParticleEmitterInstance* Owne
 	else if (LODLevel->RequiredModule->bUseLocalSpace && bInWorldSpace)
 	{
 		Transform *= OwnerRotation;
-		Transform *= LocalToWorld.Inverse();
+		Transform *= LocalToWorld.InverseFast();
 	}
 	else if (!bInWorldSpace)
 	{
diff --git a/Engine/Source/Runtime/Engine/Private/Particles/ParticleSystemRender.cpp b/Engine/Source/Runtime/Engine/Private/Particles/ParticleSystemRender.cpp
--- a/Engine/Source/Runtime/Engine/Private/Particles/ParticleSystemRender.cpp
+++ b/Engine/Source/Runtime/Engine/Private/Particles/ParticleSystemRender.cpp
@@ -645,7 +645,7 @@ void FDynamicSpriteEmitterDataBase::RenderDebug(const FParticleSystemSceneProxy*
 
 	const FMatrix& LocalToWorld = SpriteSource.bUseLocalSpace ? Proxy->GetLocalToWorld() : FMatrix::Identity;
 
-	FMatrix CameraToWorld = View->ViewMatrices.ViewMatrix.Inverse();
+	FMatrix CameraToWorld = View->ViewMatrices.ViewMatrix.InverseFast();
 	FVector CamX = CameraToWorld.TransformVector(FVector(1,0,0));
 	FVector CamY = CameraToWorld.TransformVector(FVector(0,1,0));
 
@@ -2299,7 +2299,7 @@ void FDynamicMeshEmitterData::GetParticleTransform(const FBaseParticle& InPartic
 			else if (CameraFacingOption >= VelocityAligned_ZAxisFacing)
 			{
 				bClearLocal2World = true;
-				VelocityDirection = LocalToWorld.Inverse().GetTransposed().TransformVector(VelocityDirection);
+				VelocityDirection = LocalToWorld.InverseFast().GetTransposed().TransformVector(VelocityDirection);
 			}
 
 			if (bClearLocal2World)
@@ -2601,7 +2601,7 @@ void FDynamicMeshEmitterData::GetInstanceData(void* InstanceData, void* DynamicP
 	FVector CameraPosition = View->ViewMatrices.ViewOrigin;
 	if (Source.bUseLocalSpace)
 	{
-		const FMatrix InvLocalToWorld = Proxy->GetLocalToWorld().InverseSafe();
+		const FMatrix InvLocalToWorld = Proxy->GetLocalToWorld().Inverse();
 		CameraPosition = InvLocalToWorld.TransformPosition(CameraPosition);
 	}
 
@@ -3237,7 +3237,7 @@ void FDynamicBeam2EmitterData::RenderLines(const FParticleSystemSceneProxy* Prox
 
 		FMatrix WorldToLocal = Proxy->GetWorldToLocal();
 		FMatrix LocalToWorld = Proxy->GetLocalToWorld();
-		FMatrix CameraToWorld = View->ViewMatrices.ViewMatrix.Inverse();
+		FMatrix CameraToWorld = View->ViewMatrices.ViewMatrix.InverseFast();
 		FVector	ViewOrigin = CameraToWorld.GetOrigin();
 
 		// NoiseTessellation is the amount of tessellation that should occur between noise points.
@@ -3587,7 +3587,7 @@ void FDynamicBeam2EmitterData::RenderLines(const FParticleSystemSceneProxy* Prox
 
 	if (Source.InterpolationPoints > 1)
 	{
-		FMatrix CameraToWorld = View->ViewMatrices.ViewMatrix.Inverse();
+		FMatrix CameraToWorld = View->ViewMatrices.ViewMatrix.InverseFast();
 		FVector	ViewOrigin = CameraToWorld.GetOrigin();
 		int32 TessFactor = Source.InterpolationPoints ? Source.InterpolationPoints : 1;
 
@@ -3792,7 +3792,7 @@ int32 FDynamicBeam2EmitterData::FillVertexData_NoNoise(FAsyncBufferFillData& Me)
 	int32	TrianglesToRender = 0;
 
 	FParticleBeamTrailVertex* Vertex = (FParticleBeamTrailVertex*)Me.VertexData;
-	FMatrix CameraToWorld = Me.View->ViewMatrices.ViewMatrix.InverseSafe();
+	FMatrix CameraToWorld = Me.View->ViewMatrices.ViewMatrix.Inverse();
 	FVector	ViewOrigin = CameraToWorld.GetOrigin();
 	int32 TessFactor = Source.InterpolationPoints ? Source.InterpolationPoints : 1;
 
@@ -4208,7 +4208,7 @@ int32 FDynamicBeam2EmitterData::FillData_Noise(FAsyncBufferFillData& Me) const
 	}
 
 	FParticleBeamTrailVertex* Vertex = (FParticleBeamTrailVertex*)Me.VertexData;
-	FMatrix CameraToWorld = Me.View->ViewMatrices.ViewMatrix.InverseSafe();
+	FMatrix CameraToWorld = Me.View->ViewMatrices.ViewMatrix.Inverse();
 
 	FVector	ViewOrigin	= CameraToWorld.GetOrigin();
 
@@ -5169,7 +5169,7 @@ int32 FDynamicBeam2EmitterData::FillData_InterpolatedNoise(FAsyncBufferFillData&
 	check(Source.Frequency > 0);
 
 	FParticleBeamTrailVertex* Vertex = (FParticleBeamTrailVertex*)Me.VertexData;
-	FMatrix CameraToWorld = Me.View->ViewMatrices.ViewMatrix.Inverse();
+	FMatrix CameraToWorld = Me.View->ViewMatrices.ViewMatrix.InverseFast();
 	
 	FVector	ViewOrigin	= CameraToWorld.GetOrigin();
 
@@ -6435,7 +6435,7 @@ int32 FDynamicRibbonEmitterData::FillVertexData(struct FAsyncBufferFillData& Dat
 	FParticleBeamTrailVertex* Vertex;
 	FParticleBeamTrailVertexDynamicParameter* DynParamVertex;
 
-	FMatrix CameraToWorld = Data.View->ViewMatrices.ViewMatrix.Inverse();
+	FMatrix CameraToWorld = Data.View->ViewMatrices.ViewMatrix.InverseFast();
 	FVector CameraUp = CameraToWorld.TransformVector(FVector(0,0,1));
 	FVector	ViewOrigin	= CameraToWorld.GetOrigin();
 
diff --git a/Engine/Source/Runtime/Engine/Private/PhysicsEngine/PhysAnim.cpp b/Engine/Source/Runtime/Engine/Private/PhysicsEngine/PhysAnim.cpp
--- a/Engine/Source/Runtime/Engine/Private/PhysicsEngine/PhysAnim.cpp
+++ b/Engine/Source/Runtime/Engine/Private/PhysicsEngine/PhysAnim.cpp
@@ -479,11 +479,11 @@ void USkeletalMeshComponent::UpdateRBJointMotors()
 
 					FMatrix Body1TM = CS->DefaultInstance.GetRefFrame(EConstraintFrame::Frame1).ToMatrixNoScale();
 					Body1TM.SetOrigin(FVector::ZeroVector);
-					FMatrix Body1TMInv = Body1TM.Inverse();
+					FMatrix Body1TMInv = Body1TM.InverseFast();
 
 					FMatrix Body2TM = CS->DefaultInstance.GetRefFrame(EConstraintFrame::Frame2).ToMatrixNoScale();
 					Body2TM.SetOrigin(FVector::ZeroVector);
-					FMatrix Body2TMInv = Body2TM.Inverse();
+					FMatrix Body2TMInv = Body2TM.InverseFast();
 
 					FMatrix JointRot = Body1TM * (LocalRot * ControlBodyToParentBoneTM) * Body2TMInv;
 					FQuat JointQuat(JointRot);
diff --git a/Engine/Source/Runtime/Engine/Private/PlayerCameraManager.cpp b/Engine/Source/Runtime/Engine/Private/PlayerCameraManager.cpp
--- a/Engine/Source/Runtime/Engine/Private/PlayerCameraManager.cpp
+++ b/Engine/Source/Runtime/Engine/Private/PlayerCameraManager.cpp
@@ -279,7 +279,7 @@ void APlayerCameraManager::ApplyAnimToCamera(ACameraActor const* AnimatedCamActo
 
 	// move animated cam actor to initial-relative position
 	FTransform const AnimatedCamToWorld = AnimatedCamActor->GetTransform();
-	FTransform const AnimatedCamToInitialCam = AnimatedCamToWorld * AnimInst->InitialCamToWorld.InverseSafe();
+	FTransform const AnimatedCamToInitialCam = AnimatedCamToWorld * AnimInst->InitialCamToWorld.Inverse();
 	ACameraActor* const MutableCamActor = const_cast<ACameraActor*>(AnimatedCamActor);
 	MutableCamActor->SetActorTransform(AnimatedCamToInitialCam);		// set it back because that's what the code below expects
 
@@ -308,11 +308,11 @@ void APlayerCameraManager::ApplyAnimToCamera(ACameraActor const* AnimatedCamActo
 
 		// rot
 		// find transform from camera to the "play space"
-		FMatrix const CameraToPlaySpace = CameraToWorld * PlaySpaceToWorld.InverseSafe();	// CameraToWorld * WorldToPlaySpace
+		FMatrix const CameraToPlaySpace = CameraToWorld * PlaySpaceToWorld.Inverse();	// CameraToWorld * WorldToPlaySpace
 
 		// find transform from anim (applied in playspace) back to camera
 		FRotationMatrix const AnimToPlaySpace(AnimatedCamActor->GetActorRotation()*Scale);
-		FMatrix const AnimToCamera = AnimToPlaySpace * CameraToPlaySpace.InverseSafe();			// AnimToPlaySpace * PlaySpaceToCamera
+		FMatrix const AnimToCamera = AnimToPlaySpace * CameraToPlaySpace.Inverse();			// AnimToPlaySpace * PlaySpaceToCamera
 
 		// RCS = rotated camera space, meaning camera space after it's been animated
 		// this is what we're looking for, the diff between rotated cam space and regular cam space.
diff --git a/Engine/Source/Runtime/Engine/Private/ReflectionCaptureComponent.cpp b/Engine/Source/Runtime/Engine/Private/ReflectionCaptureComponent.cpp
--- a/Engine/Source/Runtime/Engine/Private/ReflectionCaptureComponent.cpp
+++ b/Engine/Source/Runtime/Engine/Private/ReflectionCaptureComponent.cpp
@@ -1452,7 +1452,7 @@ FReflectionCaptureProxy::FReflectionCaptureProxy(const UReflectionCaptureCompone
 void FReflectionCaptureProxy::SetTransform(const FMatrix& InTransform)
 {
 	Position = InTransform.GetOrigin();
-	BoxTransform = InTransform.InverseSafe();
+	BoxTransform = InTransform.Inverse();
 
 	FVector ForwardVector(1.0f,0.0f,0.0f);
 	FVector RightVector(0.0f,-1.0f,0.0f);
diff --git a/Engine/Source/Runtime/Engine/Private/SceneView.cpp b/Engine/Source/Runtime/Engine/Private/SceneView.cpp
--- a/Engine/Source/Runtime/Engine/Private/SceneView.cpp
+++ b/Engine/Source/Runtime/Engine/Private/SceneView.cpp
@@ -216,7 +216,7 @@ FSceneView::FSceneView(const FSceneViewInitOptions& InitOptions)
 	ViewProjectionMatrix = ViewMatrices.GetViewProjMatrix();
 
 	// For precision reasons the view matrix inverse is calculated independently.
-	InvViewMatrix = ViewMatrices.ViewMatrix.InverseSafe();
+	InvViewMatrix = ViewMatrices.ViewMatrix.Inverse();
 	InvViewProjectionMatrix = ViewMatrices.GetInvProjMatrix() * InvViewMatrix;
 
 	bool ApplyPreViewTranslation = true;
@@ -272,7 +272,7 @@ FSceneView::FSceneView(const FSceneViewInitOptions& InitOptions)
 	
 	// Compute a transform from view origin centered world-space to clip space.
 	ViewMatrices.TranslatedViewProjectionMatrix = TranslatedViewMatrix * ViewMatrices.ProjMatrix;
-	ViewMatrices.InvTranslatedViewProjectionMatrix = ViewMatrices.TranslatedViewProjectionMatrix.InverseSafe();
+	ViewMatrices.InvTranslatedViewProjectionMatrix = ViewMatrices.TranslatedViewProjectionMatrix.Inverse();
 	
 	ShadowViewMatrices = ViewMatrices;
 
@@ -371,7 +371,7 @@ void FSceneView::UpdateViewMatrix()
 	
 	// Compute a transform from view origin centered world-space to clip space.
 	ViewMatrices.TranslatedViewProjectionMatrix = TranslatedViewMatrix * ViewMatrices.ProjMatrix;
-	ViewMatrices.InvTranslatedViewProjectionMatrix = ViewMatrices.TranslatedViewProjectionMatrix.InverseSafe();
+	ViewMatrices.InvTranslatedViewProjectionMatrix = ViewMatrices.TranslatedViewProjectionMatrix.Inverse();
 
 	ViewProjectionMatrix = ViewMatrices.GetViewProjMatrix();
 	InvViewMatrix = ViewMatrices.GetInvViewMatrix();
@@ -484,7 +484,7 @@ FVector FSceneView::Deproject(const FPlane& ScreenPoint) const
 
 void FSceneView::DeprojectFVector2D(const FVector2D& ScreenPos, FVector& out_WorldOrigin, FVector& out_WorldDirection) const
 {
-	const FMatrix InverseViewMatrix = ViewMatrices.ViewMatrix.Inverse();
+	const FMatrix InverseViewMatrix = ViewMatrices.ViewMatrix.InverseFast();
 	const FMatrix InvProjectionMatrix = ViewMatrices.GetInvProjMatrix();
 	
 	DeprojectScreenToWorld(ScreenPos, UnscaledViewRect, InverseViewMatrix, InvProjectionMatrix, out_WorldOrigin, out_WorldDirection);
diff --git a/Engine/Source/Runtime/Engine/Private/SkeletalMesh.cpp b/Engine/Source/Runtime/Engine/Private/SkeletalMesh.cpp
--- a/Engine/Source/Runtime/Engine/Private/SkeletalMesh.cpp
+++ b/Engine/Source/Runtime/Engine/Private/SkeletalMesh.cpp
@@ -2488,7 +2488,7 @@ void USkeletalMesh::CalculateInvRefMatrices()
 			}
 
 			// Precompute inverse so we can use from-refpose-skin vertices.
-			RefBasesInvMatrix[b] = RefBases[b].Inverse(); 
+			RefBasesInvMatrix[b] = RefBases[b].InverseFast(); 
 		}
 	}
 }
@@ -4488,7 +4488,7 @@ void FSkeletalMeshSceneProxy::DrawDynamicElementsSection(FPrimitiveDrawInterface
 void FSkeletalMeshSceneProxy::GetWorldMatrices( FMatrix& OutLocalToWorld, FMatrix& OutWorldToLocal ) const
 {
 	OutLocalToWorld = GetLocalToWorld();
-	OutWorldToLocal = GetLocalToWorld().Inverse();
+	OutWorldToLocal = GetLocalToWorld().InverseFast();
 }
 
 /**
diff --git a/Engine/Source/Runtime/Engine/Private/SkeletalMeshComponent.cpp b/Engine/Source/Runtime/Engine/Private/SkeletalMeshComponent.cpp
--- a/Engine/Source/Runtime/Engine/Private/SkeletalMeshComponent.cpp
+++ b/Engine/Source/Runtime/Engine/Private/SkeletalMeshComponent.cpp
@@ -1059,7 +1059,7 @@ void USkeletalMeshComponent::UpdateBounds()
 			Bounds = CalcBounds(ComponentToWorld);
 
 			bCachedLocalBoundsUpToDate = true;
-			CachedLocalBounds = Bounds.TransformBy(ComponentToWorld.InverseSafe());
+			CachedLocalBounds = Bounds.TransformBy(ComponentToWorld.Inverse());
 		}
 	}
 
diff --git a/Engine/Source/Runtime/Engine/Private/SkeletalMeshComponentPhysics.cpp b/Engine/Source/Runtime/Engine/Private/SkeletalMeshComponentPhysics.cpp
--- a/Engine/Source/Runtime/Engine/Private/SkeletalMeshComponentPhysics.cpp
+++ b/Engine/Source/Runtime/Engine/Private/SkeletalMeshComponentPhysics.cpp
@@ -1869,7 +1869,7 @@ bool USkeletalMeshComponent::ComponentOverlapMulti(TArray<struct FOverlapResult>
 		return false;
 	}
 
-	const FTransform WorldToComponent(ComponentToWorld.InverseSafe());
+	const FTransform WorldToComponent(ComponentToWorld.Inverse());
 	const FCollisionResponseParams ResponseParams(GetCollisionResponseToChannels());
 
 	bool bHaveBlockingHit = false;
@@ -3163,7 +3163,7 @@ void USkeletalMeshComponent::CheckClothTeleport(float DeltaTime)
 	{
 		// Detect whether teleportation is needed or not
 		// Rotation matrix's transpose means an inverse but can't use a transpose because this matrix includes scales
-		FMatrix AInvB = CurRootBoneMat * PrevRootBoneMatrix.Inverse();
+		FMatrix AInvB = CurRootBoneMat * PrevRootBoneMatrix.InverseFast();
 		float Trace = AInvB.M[0][0] + AInvB.M[1][1] + AInvB.M[2][2];
 		float CosineTheta = (Trace - 1.0f) / 2.0f; // trace = 1+2cos(theta) for a 3x3 matrix
 
diff --git a/Engine/Source/Runtime/Engine/Private/SkeletalRenderCPUSkin.cpp b/Engine/Source/Runtime/Engine/Private/SkeletalRenderCPUSkin.cpp
--- a/Engine/Source/Runtime/Engine/Private/SkeletalRenderCPUSkin.cpp
+++ b/Engine/Source/Runtime/Engine/Private/SkeletalRenderCPUSkin.cpp
@@ -403,7 +403,7 @@ void FSkeletalMeshObjectCPUSkin::DrawVertexElements(FPrimitiveDrawInterface* PDI
 {
 	uint32 NumIndices = CachedFinalVertices.Num();
 
-	FMatrix LocalToWorldInverseTranspose = ToWorldSpace.ToMatrixWithScale().Inverse().GetTransposed();
+	FMatrix LocalToWorldInverseTranspose = ToWorldSpace.ToMatrixWithScale().InverseFast().GetTransposed();
 
 	for (uint32 i = 0; i < NumIndices; i++)
 	{
diff --git a/Engine/Source/Runtime/Engine/Private/SkinnedMeshComponent.cpp b/Engine/Source/Runtime/Engine/Private/SkinnedMeshComponent.cpp
--- a/Engine/Source/Runtime/Engine/Private/SkinnedMeshComponent.cpp
+++ b/Engine/Source/Runtime/Engine/Private/SkinnedMeshComponent.cpp
@@ -1199,7 +1199,7 @@ void USkinnedMeshComponent::TransformToBoneSpace(FName BoneName, FVector InPosit
 	{
 		FMatrix BoneToWorldTM = GetBoneMatrix(BoneIndex);
 		FMatrix WorldTM = FRotationTranslationMatrix(InRotation, InPosition);
-		FMatrix LocalTM = WorldTM * BoneToWorldTM.InverseSafe();
+		FMatrix LocalTM = WorldTM * BoneToWorldTM.Inverse();
 
 		OutPosition = LocalTM.GetOrigin();
 		OutRotation = LocalTM.Rotator();
diff --git a/Engine/Source/Runtime/Engine/Private/StaticMeshLight.cpp b/Engine/Source/Runtime/Engine/Private/StaticMeshLight.cpp
--- a/Engine/Source/Runtime/Engine/Private/StaticMeshLight.cpp
+++ b/Engine/Source/Runtime/Engine/Private/StaticMeshLight.cpp
@@ -74,7 +74,7 @@ FStaticMeshStaticLightingMesh::FStaticMeshStaticLightingMesh(const UStaticMeshCo
 void FStaticMeshStaticLightingMesh::SetLocalToWorld(const FMatrix& InLocalToWorld)
 {
 	LocalToWorld = InLocalToWorld;
-	LocalToWorldInverseTranspose = LocalToWorld.Inverse().GetTransposed();
+	LocalToWorldInverseTranspose = LocalToWorld.InverseFast().GetTransposed();
 	LocalToWorldDeterminant = LocalToWorld.Determinant();
 }
 
diff --git a/Engine/Source/Runtime/Engine/Private/VectorField.cpp b/Engine/Source/Runtime/Engine/Private/VectorField.cpp
--- a/Engine/Source/Runtime/Engine/Private/VectorField.cpp
+++ b/Engine/Source/Runtime/Engine/Private/VectorField.cpp
@@ -77,7 +77,7 @@ void FVectorFieldInstance::UpdateTransforms(const FMatrix& LocalToWorld)
 	VolumeToWorldNoScale = LocalToWorld.GetMatrixWithoutScale().RemoveTranslation();
 	VolumeToWorld = FScaleMatrix(VolumeScale) * FTranslationMatrix(VolumeOffset)
 		* LocalToWorld;
-	WorldToVolume = VolumeToWorld.Inverse();
+	WorldToVolume = VolumeToWorld.InverseFast();
 }
 
 /*------------------------------------------------------------------------------
diff --git a/Engine/Source/Runtime/Engine/Public/CanvasTypes.h b/Engine/Source/Runtime/Engine/Public/CanvasTypes.h
--- a/Engine/Source/Runtime/Engine/Public/CanvasTypes.h
+++ b/Engine/Source/Runtime/Engine/Public/CanvasTypes.h
@@ -163,7 +163,7 @@ public:
 	*/
 	ENGINE_API FMatrix GetTransform() const
 	{ 
-		return TransformStack.Top().GetMatrix() * TransformStack[0].GetMatrix().Inverse(); 
+		return TransformStack.Top().GetMatrix() * TransformStack[0].GetMatrix().InverseFast(); 
 	}
 
 	/** 
diff --git a/Engine/Source/Runtime/Engine/Public/ParticleHelper.h b/Engine/Source/Runtime/Engine/Public/ParticleHelper.h
--- a/Engine/Source/Runtime/Engine/Public/ParticleHelper.h
+++ b/Engine/Source/Runtime/Engine/Public/ParticleHelper.h
@@ -2587,7 +2587,7 @@ public:
 
 	// While this isn't good OO design, access to everything is made public.
 	// This is to allow custom emitter instances to easily be written when extending the engine.
-	FMatrix GetWorldToLocal() const		{	return GetLocalToWorld().Inverse();	}
+	FMatrix GetWorldToLocal() const		{	return GetLocalToWorld().InverseFast();	}
 	bool GetCastShadow() const			{	return bCastShadow;				}
 	const FMaterialRelevance& GetMaterialRelevance() const
 	{
diff --git a/Engine/Source/Runtime/Engine/Public/PrimitiveUniformShaderParameters.h b/Engine/Source/Runtime/Engine/Public/PrimitiveUniformShaderParameters.h
--- a/Engine/Source/Runtime/Engine/Public/PrimitiveUniformShaderParameters.h
+++ b/Engine/Source/Runtime/Engine/Public/PrimitiveUniformShaderParameters.h
@@ -40,7 +40,7 @@ inline FPrimitiveUniformShaderParameters GetPrimitiveUniformShaderParameters(
 {
 	FPrimitiveUniformShaderParameters Result;
 	Result.LocalToWorld = LocalToWorld;
-	Result.WorldToLocal = LocalToWorld.InverseSafe();
+	Result.WorldToLocal = LocalToWorld.Inverse();
 	Result.ObjectWorldPositionAndRadius = FVector4(WorldBounds.Origin, WorldBounds.SphereRadius);
 	Result.ObjectBounds = WorldBounds.BoxExtent;
 	Result.LocalObjectBoundsMin = LocalBounds.GetBoxExtrema(0); // 0 == minimum
diff --git a/Engine/Source/Runtime/Engine/Public/SceneView.h b/Engine/Source/Runtime/Engine/Public/SceneView.h
--- a/Engine/Source/Runtime/Engine/Public/SceneView.h
+++ b/Engine/Source/Runtime/Engine/Public/SceneView.h
@@ -177,13 +177,13 @@ struct FViewMatrices
 
 	FMatrix GetInvProjMatrix() const
 	{
-		return ProjMatrix.InverseSafe();
+		return ProjMatrix.Inverse();
 	}
 
 	FMatrix GetInvViewMatrix() const
 	{
 		// can be optimized: it's not a perspective matrix so transpose would be enough
-		return ViewMatrix.InverseSafe();
+		return ViewMatrix.Inverse();
 	}
 
 	FMatrix GetInvViewProjMatrix() const
diff --git a/Engine/Source/Runtime/Mac/CoreAudio/Private/CoreAudioDevice.cpp b/Engine/Source/Runtime/Mac/CoreAudio/Private/CoreAudioDevice.cpp
--- a/Engine/Source/Runtime/Mac/CoreAudio/Private/CoreAudioDevice.cpp
+++ b/Engine/Source/Runtime/Mac/CoreAudio/Private/CoreAudioDevice.cpp
@@ -288,7 +288,7 @@ void FCoreAudioDevice::UpdateHardware()
 	// at the Y component after normalization to determine spatialization.
 	const FVector Up = Listeners[0].GetUp();
 	const FVector Right = Listeners[0].GetFront();
-	InverseTransform = FMatrix(Up, Right, Up ^ Right, Listeners[0].Transform.GetTranslation()).Inverse();
+	InverseTransform = FMatrix(Up, Right, Up ^ Right, Listeners[0].Transform.GetTranslation()).InverseFast();
 }
 
 FAudioEffectsManager* FCoreAudioDevice::CreateEffectsManager()
diff --git a/Engine/Source/Runtime/Renderer/Private/AtmosphereRendering.cpp b/Engine/Source/Runtime/Renderer/Private/AtmosphereRendering.cpp
--- a/Engine/Source/Runtime/Renderer/Private/AtmosphereRendering.cpp
+++ b/Engine/Source/Runtime/Renderer/Private/AtmosphereRendering.cpp
@@ -1590,7 +1590,7 @@ FAtmosphericFogSceneInfo::FAtmosphericFogSceneInfo(UAtmosphericFogComponent* InC
 	RenderFlag |= (InComponent->bDisableGroundScattering) ? EAtmosphereRenderFlag::E_DisableGroundScattering : EAtmosphereRenderFlag::E_EnableAll;
 	// Should be same as UpdateAtmosphericFogTransform
 	GroundOffset += InComponent->GetComponentLocation().Z;
-	FMatrix WorldToLight = InComponent->ComponentToWorld.ToMatrixNoScale().Inverse();
+	FMatrix WorldToLight = InComponent->ComponentToWorld.ToMatrixNoScale().InverseFast();
 	DefaultSunDirection = FVector(WorldToLight.M[0][0],WorldToLight.M[1][0],WorldToLight.M[2][0]);
 }
 
diff --git a/Engine/Source/Runtime/Renderer/Private/CompositionLighting/PostProcessDeferredDecals.cpp b/Engine/Source/Runtime/Renderer/Private/CompositionLighting/PostProcessDeferredDecals.cpp
--- a/Engine/Source/Runtime/Renderer/Private/CompositionLighting/PostProcessDeferredDecals.cpp
+++ b/Engine/Source/Runtime/Renderer/Private/CompositionLighting/PostProcessDeferredDecals.cpp
@@ -482,7 +482,7 @@ public:
 		// 1,1,1 requires no scale
 		//			ComponentTrans = ComponentTrans.GetScaled(GDefaultDecalSize);
 
-		FMatrix WorldToComponent = ComponentTrans.ToMatrixWithScale().Inverse();
+		FMatrix WorldToComponent = ComponentTrans.ToMatrixWithScale().InverseFast();
 
 		// Set the transform from screen space to light space.
 		if(ScreenToDecal.IsBound())
diff --git a/Engine/Source/Runtime/Renderer/Private/DistanceFieldSurfaceCacheLighting.cpp b/Engine/Source/Runtime/Renderer/Private/DistanceFieldSurfaceCacheLighting.cpp
--- a/Engine/Source/Runtime/Renderer/Private/DistanceFieldSurfaceCacheLighting.cpp
+++ b/Engine/Source/Runtime/Renderer/Private/DistanceFieldSurfaceCacheLighting.cpp
@@ -2756,7 +2756,7 @@ void UpdateVisibleObjectBuffers(const FScene* Scene, const FViewInfo& View, int3
 						ObjectData2.Add(*(FVector4*)&UniformScaleVolumeToWorld.M[2]);
 						ObjectData2.Add(*(FVector4*)&UniformScaleVolumeToWorld.M[3]);
 
-						const FMatrix WorldToVolume = UniformScaleVolumeToWorld.Inverse();
+						const FMatrix WorldToVolume = UniformScaleVolumeToWorld.InverseFast();
 						// WorldToVolume
 						ObjectData.Add(*(FVector4*)&WorldToVolume.M[0]);
 						ObjectData.Add(*(FVector4*)&WorldToVolume.M[1]);
diff --git a/Engine/Source/Runtime/Renderer/Private/LightPropagationVolume.cpp b/Engine/Source/Runtime/Renderer/Private/LightPropagationVolume.cpp
--- a/Engine/Source/Runtime/Renderer/Private/LightPropagationVolume.cpp
+++ b/Engine/Source/Runtime/Renderer/Private/LightPropagationVolume.cpp
@@ -988,7 +988,7 @@ void FLightPropagationVolume::GetShadowInfo( const FProjectedShadowInfo& Project
 {
 	FIntPoint ShadowBufferResolution( ProjectedShadowInfo.ResolutionX, ProjectedShadowInfo.ResolutionY );
 	RsmInfoOut.WorldToShadow = ProjectedShadowInfo.GetWorldToShadowMatrix(RsmInfoOut.ShadowmapMinMax, &ShadowBufferResolution );
-	RsmInfoOut.ShadowToWorld = RsmInfoOut.WorldToShadow.Inverse();
+	RsmInfoOut.ShadowToWorld = RsmInfoOut.WorldToShadow.InverseFast();
 
 	// Determine the shadow area in world space, so we can scale the brightness if needed. 
 	FVector ShadowUp    = FVector( 1.0f,0.0f,0.0f );
diff --git a/Engine/Source/Runtime/Renderer/Private/PostProcess/PostProcessTemporalAA.cpp b/Engine/Source/Runtime/Renderer/Private/PostProcess/PostProcessTemporalAA.cpp
--- a/Engine/Source/Runtime/Renderer/Private/PostProcess/PostProcessTemporalAA.cpp
+++ b/Engine/Source/Runtime/Renderer/Private/PostProcess/PostProcessTemporalAA.cpp
@@ -621,7 +621,7 @@ void FRCPassPostProcessTemporalAA::Process(FRenderingCompositePassContext& Conte
 
 	// Compute a transform from view origin centered world-space to clip space.
 	View.ViewMatrices.TranslatedViewProjectionMatrix = TranslatedViewMatrix * View.ViewMatrices.ProjMatrix;
-	View.ViewMatrices.InvTranslatedViewProjectionMatrix = View.ViewMatrices.TranslatedViewProjectionMatrix.InverseSafe();
+	View.ViewMatrices.InvTranslatedViewProjectionMatrix = View.ViewMatrices.TranslatedViewProjectionMatrix.Inverse();
 
 	View.InitRHIResources();
 }
diff --git a/Engine/Source/Runtime/Renderer/Private/RendererScene.cpp b/Engine/Source/Runtime/Renderer/Private/RendererScene.cpp
--- a/Engine/Source/Runtime/Renderer/Private/RendererScene.cpp
+++ b/Engine/Source/Runtime/Renderer/Private/RendererScene.cpp
@@ -1952,7 +1952,7 @@ void FScene::ApplyWorldOffset_RenderThread(FVector InOffset)
 	// Reflection captures
 	for (auto It = ReflectionSceneData.RegisteredReflectionCaptures.CreateIterator(); It; ++It)
 	{
-		FMatrix NewTransform = (*It)->BoxTransform.InverseSafe().ConcatTranslation(InOffset);
+		FMatrix NewTransform = (*It)->BoxTransform.Inverse().ConcatTranslation(InOffset);
 		(*It)->SetTransform(NewTransform);
 	}
 
diff --git a/Engine/Source/Runtime/Renderer/Private/SceneRendering.cpp b/Engine/Source/Runtime/Renderer/Private/SceneRendering.cpp
--- a/Engine/Source/Runtime/Renderer/Private/SceneRendering.cpp
+++ b/Engine/Source/Runtime/Renderer/Private/SceneRendering.cpp
@@ -254,7 +254,7 @@ TUniformBufferRef<FViewUniformShaderParameters> FViewInfo::CreateUniformBuffer(
 	ViewUniformShaderParameters.PrevViewOrigin = PrevViewMatrices.ViewOrigin;
 	ViewUniformShaderParameters.PrevPreViewTranslation = PrevViewMatrices.PreViewTranslation;
 	// can be optimized
-	ViewUniformShaderParameters.PrevInvViewProj = PrevViewProjMatrix.InverseSafe();
+	ViewUniformShaderParameters.PrevInvViewProj = PrevViewProjMatrix.Inverse();
 
 	ViewUniformShaderParameters.ScreenToWorld = FMatrix(
 		FPlane(1,0,0,0),
diff --git a/Engine/Source/Runtime/Renderer/Private/SceneVisibility.cpp b/Engine/Source/Runtime/Renderer/Private/SceneVisibility.cpp
--- a/Engine/Source/Runtime/Renderer/Private/SceneVisibility.cpp
+++ b/Engine/Source/Runtime/Renderer/Private/SceneVisibility.cpp
@@ -1172,7 +1172,7 @@ void FSceneRenderer::PreVisibilityFrameSetup(FRHICommandListImmediate& RHICmdLis
 
 				// Compute a transform from view origin centered world-space to clip space.
 				View.ViewMatrices.TranslatedViewProjectionMatrix = TranslatedViewMatrix * View.ViewMatrices.ProjMatrix;
-				View.ViewMatrices.InvTranslatedViewProjectionMatrix = View.ViewMatrices.TranslatedViewProjectionMatrix.InverseSafe();
+				View.ViewMatrices.InvTranslatedViewProjectionMatrix = View.ViewMatrices.TranslatedViewProjectionMatrix.Inverse();
 			}
 		}
 		else if(ViewState)
diff --git a/Engine/Source/Runtime/Renderer/Private/ShadowRendering.cpp b/Engine/Source/Runtime/Renderer/Private/ShadowRendering.cpp
--- a/Engine/Source/Runtime/Renderer/Private/ShadowRendering.cpp
+++ b/Engine/Source/Runtime/Renderer/Private/ShadowRendering.cpp
@@ -2218,7 +2218,7 @@ void FProjectedShadowInfo::RenderFrustumWireframe(FPrimitiveDrawInterface* PDI)
 		SubjectPrimitiveId = SubjectPrimitives[0]->GetIndex();
 	}
 
-	const FMatrix InvShadowTransform = (bWholeSceneShadow || bPreShadow) ? SubjectAndReceiverMatrix.Inverse() : InvReceiverMatrix;
+	const FMatrix InvShadowTransform = (bWholeSceneShadow || bPreShadow) ? SubjectAndReceiverMatrix.InverseFast() : InvReceiverMatrix;
 
 	FColor Color;
 
diff --git a/Engine/Source/Runtime/Renderer/Private/ShadowSetup.cpp b/Engine/Source/Runtime/Renderer/Private/ShadowSetup.cpp
--- a/Engine/Source/Runtime/Renderer/Private/ShadowSetup.cpp
+++ b/Engine/Source/Runtime/Renderer/Private/ShadowSetup.cpp
@@ -459,7 +459,7 @@ FProjectedShadowInfo::FProjectedShadowInfo(
 
 		GetViewFrustumBounds(CasterFrustum,SubjectAndReceiverMatrix,true);
 
-		InvReceiverMatrix = ReceiverMatrix.Inverse();
+		InvReceiverMatrix = ReceiverMatrix.InverseFast();
 		GetViewFrustumBounds(ReceiverFrustum,ReceiverMatrix,true);
 	}
 	else
@@ -539,7 +539,7 @@ FProjectedShadowInfo::FProjectedShadowInfo(
 			// Snap the shadow's position and transform it back into world space
 			// This snapping prevents sub-texel camera movements which removes view dependent aliasing from the final shadow result
 			// This only maintains stable shadows under camera translation and rotation
-			const FVector SnappedWorldPosition = WorldToFace.Inverse().TransformPosition(TransformedPosition - FVector(SnapX, SnapY, 0.0f));
+			const FVector SnappedWorldPosition = WorldToFace.InverseFast().TransformPosition(TransformedPosition - FVector(SnapX, SnapY, 0.0f));
 			PreShadowTranslation = -SnappedWorldPosition;
 		}
 
@@ -574,7 +574,7 @@ FProjectedShadowInfo::FProjectedShadowInfo(
 			// Snap the shadow's position and transform it back into world space
 			// This snapping prevents sub-texel camera movements which removes view dependent aliasing from the final shadow result
 			// This only maintains stable shadows under camera translation and rotation
-			const FVector SnappedWorldPosition = WorldToFace.Inverse().TransformPosition(TransformedPosition - FVector(SnapX, SnapY, 0.0f));
+			const FVector SnappedWorldPosition = WorldToFace.InverseFast().TransformPosition(TransformedPosition - FVector(SnapX, SnapY, 0.0f));
 			PreShadowTranslation = -SnappedWorldPosition;
 		}
 
@@ -603,7 +603,7 @@ FProjectedShadowInfo::FProjectedShadowInfo(
 
 	float MaxSubjectDepth = SubjectAndReceiverMatrix.TransformPosition(
 		Initializer.SubjectBounds.Origin
-		+ WorldToLightScaled.Inverse().TransformVector(Initializer.FaceDirection) * Initializer.SubjectBounds.SphereRadius
+		+ WorldToLightScaled.InverseFast().TransformVector(Initializer.FaceDirection) * Initializer.SubjectBounds.SphereRadius
 		).Z;
 
 	if (Initializer.bOnePassPointLightShadow)
@@ -622,7 +622,7 @@ FProjectedShadowInfo::FProjectedShadowInfo(
 		FPlane(0,	1,	0,	0),
 		FPlane(0,	0,	0,	1));
 
-	InvReceiverMatrix = ReceiverMatrix.Inverse();
+	InvReceiverMatrix = ReceiverMatrix.InverseFast();
 
 	GetViewFrustumBounds(ReceiverFrustum, ReceiverMatrix, true);
 
@@ -1666,14 +1666,14 @@ void FDeferredShadingSceneRenderer::InitProjectedShadowVisibility(FRHICommandLis
 								float Far = ProjectedShadowInfo.CascadeSettings.SplitFar;
 
 								// Camera Subfrustum
-								DrawFrustumWireframe(&ShadowFrustumPDI, (ViewMatrix * FPerspectiveMatrix(ActualFOV, AspectRatio, 1.0f, Near, Mid)).Inverse(), Color, 0);
-								DrawFrustumWireframe(&ShadowFrustumPDI, (ViewMatrix * FPerspectiveMatrix(ActualFOV, AspectRatio, 1.0f, Mid, Far)).Inverse(), FColor::White, 0);
+								DrawFrustumWireframe(&ShadowFrustumPDI, (ViewMatrix * FPerspectiveMatrix(ActualFOV, AspectRatio, 1.0f, Near, Mid)).InverseFast(), Color, 0);
+								DrawFrustumWireframe(&ShadowFrustumPDI, (ViewMatrix * FPerspectiveMatrix(ActualFOV, AspectRatio, 1.0f, Mid, Far)).InverseFast(), FColor::White, 0);
 
 								// Subfrustum Sphere Bounds
 								DrawWireSphere(&ShadowFrustumPDI, FTransform(ProjectedShadowInfo.ShadowBounds.Center), Color, ProjectedShadowInfo.ShadowBounds.W, 40, 0);
 
 								// Shadow Map Projection Bounds
-								DrawFrustumWireframe(&ShadowFrustumPDI, ProjectedShadowInfo.SubjectAndReceiverMatrix.Inverse() * FTranslationMatrix(-ProjectedShadowInfo.PreShadowTranslation), Color, 0);
+								DrawFrustumWireframe(&ShadowFrustumPDI, ProjectedShadowInfo.SubjectAndReceiverMatrix.InverseFast() * FTranslationMatrix(-ProjectedShadowInfo.PreShadowTranslation), Color, 0);
 							}
 							else
 							{
diff --git a/Engine/Source/Runtime/UMG/Public/Components/Viewport.h b/Engine/Source/Runtime/UMG/Public/Components/Viewport.h
--- a/Engine/Source/Runtime/UMG/Public/Components/Viewport.h
+++ b/Engine/Source/Runtime/UMG/Public/Components/Viewport.h
@@ -154,7 +154,7 @@ public:
 		if ( bRecalculateView )
 		{
 			FMatrix OrbitMatrix = ViewTransform.ComputeOrbitMatrix();
-			OrbitMatrix = OrbitMatrix.Inverse();
+			OrbitMatrix = OrbitMatrix.InverseFast();
 
 			ViewTransform.SetRotation(OrbitMatrix.Rotator());
 			ViewTransform.SetLocation(OrbitMatrix.GetOrigin());
diff --git a/Engine/Source/Runtime/Windows/XAudio2/Private/XAudio2Device.cpp b/Engine/Source/Runtime/Windows/XAudio2/Private/XAudio2Device.cpp
--- a/Engine/Source/Runtime/Windows/XAudio2/Private/XAudio2Device.cpp
+++ b/Engine/Source/Runtime/Windows/XAudio2/Private/XAudio2Device.cpp
@@ -213,7 +213,7 @@ void FXAudio2Device::UpdateHardware()
 		// at the Y component after normalization to determine spatialization.
 		const FVector Up = Listeners[0].GetUp();
 		const FVector Right = Listeners[0].GetFront();
-		InverseTransform = FMatrix(Up, Right, Up ^ Right, Listeners[0].Transform.GetTranslation()).Inverse();
+		InverseTransform = FMatrix(Up, Right, Up ^ Right, Listeners[0].Transform.GetTranslation()).InverseFast();
 	}
 }
 
