Added CulturePointer.h, which typedefs ThreadSAfe TSharedRef and TSharePtr for FCulture. These should be used instead of the manual template instantiations for consistency.

diff --git a/Engine/Plugins/Editor/EpicSurvey/Source/EpicSurvey/Private/EpicSurvey.cpp b/Engine/Plugins/Editor/EpicSurvey/Source/EpicSurvey/Private/EpicSurvey.cpp
--- a/Engine/Plugins/Editor/EpicSurvey/Source/EpicSurvey/Private/EpicSurvey.cpp
+++ b/Engine/Plugins/Editor/EpicSurvey/Source/EpicSurvey/Private/EpicSurvey.cpp
@@ -154,7 +154,7 @@ void FEpicSurvey::Initialize()
 
 void FEpicSurvey::InitializeTitleCloud()
 {
-	TSharedRef<FCulture, ESPMode::ThreadSafe> Culture = FInternationalization::Get().GetCurrentCulture();
+	FCultureRef Culture = FInternationalization::Get().GetCurrentCulture();
 
 	FString CultureString;
 	switch (CurrentCulture)
diff --git a/Engine/Source/Developer/InternationalizationSettings/Private/InternationalizationSettingsModelDetails.cpp b/Engine/Source/Developer/InternationalizationSettings/Private/InternationalizationSettingsModelDetails.cpp
--- a/Engine/Source/Developer/InternationalizationSettings/Private/InternationalizationSettingsModelDetails.cpp
+++ b/Engine/Source/Developer/InternationalizationSettings/Private/InternationalizationSettingsModelDetails.cpp
@@ -7,14 +7,14 @@
 /** Functions for sorting the languages */
 struct FCompareCultureByNativeLanguage
 {
-	static FText GetCultureNativeLanguageText( const TSharedPtr<FCulture, ESPMode::ThreadSafe> Culture )
+	static FText GetCultureNativeLanguageText( const FCulturePtr Culture )
 	{
 		check( Culture.IsValid() );
 		const FString Language = Culture->GetNativeLanguage();
 		return FText::FromString(Language);
 	}
 
-	FORCEINLINE bool operator()( const TSharedPtr<FCulture, ESPMode::ThreadSafe> A, const TSharedPtr<FCulture, ESPMode::ThreadSafe> B ) const
+	FORCEINLINE bool operator()( const FCulturePtr A, const FCulturePtr B ) const
 	{
 		check( A.IsValid() );
 		check( B.IsValid() );
@@ -26,7 +26,7 @@ struct FCompareCultureByNativeLanguage
 /** Functions for sorting the regions */
 struct FCompareCultureByNativeRegion
 {
-	static FText GetCultureNativeRegionText( const TSharedPtr<FCulture, ESPMode::ThreadSafe> Culture )
+	static FText GetCultureNativeRegionText( const FCulturePtr Culture )
 	{
 		check( Culture.IsValid() );
 		FString Region = Culture->GetNativeRegion();
@@ -38,7 +38,7 @@ struct FCompareCultureByNativeRegion
 		return FText::FromString(Region);
 	}
 
-	FORCEINLINE bool operator()( const TSharedPtr<FCulture, ESPMode::ThreadSafe> A, const TSharedPtr<FCulture, ESPMode::ThreadSafe> B ) const
+	FORCEINLINE bool operator()( const FCulturePtr A, const FCulturePtr B ) const
 	{
 		check( A.IsValid() );
 		check( B.IsValid() );
@@ -128,7 +128,7 @@ void FInternationalizationSettingsModelDetails::CustomizeDetails( IDetailLayoutB
 	const FText LanguageToolTipText = LOCTEXT("EditorLanguageTooltip", "Change the Editor language (requires restart to take effect)");
 
 	// For use in the Slate macros below, the type must be typedef'd to compile.
-	typedef TSharedPtr<FCulture, ESPMode::ThreadSafe> ThreadSafeCulturePtr;
+	typedef FCulturePtr ThreadSafeCulturePtr;
 
 	CategoryBuilder.AddCustomRow("Language")
 	.NameContent()
@@ -295,7 +295,7 @@ void FInternationalizationSettingsModelDetails::RefreshAvailableRegions()
 {
 	AvailableRegions.Empty();
 
-	TSharedPtr<FCulture, ESPMode::ThreadSafe> DefaultCulture = NULL;
+	FCulturePtr DefaultCulture = NULL;
 	if ( SelectedLanguage.IsValid() )
 	{
 		const FString SelectedLanguageName = SelectedLanguage->GetNativeLanguage();
@@ -323,7 +323,7 @@ void FInternationalizationSettingsModelDetails::RefreshAvailableRegions()
 	// If we have a preferred default (or there's only one in the list), select that now
 	if ( DefaultCulture.IsValid() || AvailableCultures.Num() == 1 )
 	{
-		TSharedPtr<FCulture, ESPMode::ThreadSafe> Culture = DefaultCulture.IsValid() ? DefaultCulture : AvailableCultures.Last();
+		FCulturePtr Culture = DefaultCulture.IsValid() ? DefaultCulture : AvailableCultures.Last();
 		// Set it as our default region, if one hasn't already been chosen
 		if ( !SelectedCulture.IsValid() && RegionComboBox.IsValid() )
 		{
@@ -347,14 +347,14 @@ FText FInternationalizationSettingsModelDetails::GetCurrentLanguageText() const
 	return FText::GetEmpty();
 }
 
-TSharedRef<SWidget> FInternationalizationSettingsModelDetails::OnLanguageGenerateWidget( TSharedPtr<FCulture, ESPMode::ThreadSafe> Culture, IDetailLayoutBuilder* DetailBuilder ) const
+TSharedRef<SWidget> FInternationalizationSettingsModelDetails::OnLanguageGenerateWidget( FCulturePtr Culture, IDetailLayoutBuilder* DetailBuilder ) const
 {
 	return SNew(STextBlock)
 		.Text(FCompareCultureByNativeLanguage::GetCultureNativeLanguageText(Culture))
 		.Font(DetailBuilder->GetDetailFont());
 }
 
-void FInternationalizationSettingsModelDetails::OnLanguageSelectionChanged( TSharedPtr<FCulture, ESPMode::ThreadSafe> Culture, ESelectInfo::Type SelectionType )
+void FInternationalizationSettingsModelDetails::OnLanguageSelectionChanged( FCulturePtr Culture, ESelectInfo::Type SelectionType )
 {
 	SelectedLanguage = Culture;
 	SelectedCulture = NULL;
@@ -372,14 +372,14 @@ FText FInternationalizationSettingsModelDetails::GetCurrentRegionText() const
 	return FText::GetEmpty();
 }
 
-TSharedRef<SWidget> FInternationalizationSettingsModelDetails::OnRegionGenerateWidget( TSharedPtr<FCulture, ESPMode::ThreadSafe> Culture, IDetailLayoutBuilder* DetailBuilder ) const
+TSharedRef<SWidget> FInternationalizationSettingsModelDetails::OnRegionGenerateWidget( FCulturePtr Culture, IDetailLayoutBuilder* DetailBuilder ) const
 {
 	return SNew(STextBlock)
 		.Text(FCompareCultureByNativeRegion::GetCultureNativeRegionText(Culture))
 		.Font(DetailBuilder->GetDetailFont());
 }
 
-void FInternationalizationSettingsModelDetails::OnRegionSelectionChanged( TSharedPtr<FCulture, ESPMode::ThreadSafe> Culture, ESelectInfo::Type SelectionType )
+void FInternationalizationSettingsModelDetails::OnRegionSelectionChanged( FCulturePtr Culture, ESelectInfo::Type SelectionType )
 {
 	SelectedCulture = Culture;
 
diff --git a/Engine/Source/Developer/InternationalizationSettings/Public/InternationalizationSettingsModelDetails.h b/Engine/Source/Developer/InternationalizationSettings/Public/InternationalizationSettingsModelDetails.h
--- a/Engine/Source/Developer/InternationalizationSettings/Public/InternationalizationSettingsModelDetails.h
+++ b/Engine/Source/Developer/InternationalizationSettings/Public/InternationalizationSettingsModelDetails.h
@@ -42,19 +42,19 @@ private:
 	FText GetCurrentLanguageText() const;
 
 	/** Generate a widget for the language combo */
-	TSharedRef<SWidget> OnLanguageGenerateWidget( TSharedPtr<FCulture, ESPMode::ThreadSafe> Culture, IDetailLayoutBuilder* DetailBuilder ) const;
+	TSharedRef<SWidget> OnLanguageGenerateWidget( FCulturePtr Culture, IDetailLayoutBuilder* DetailBuilder ) const;
 
 	/** Delegate called when the language selection has changed */
-	void OnLanguageSelectionChanged( TSharedPtr<FCulture, ESPMode::ThreadSafe> Culture, ESelectInfo::Type SelectionType );
+	void OnLanguageSelectionChanged( FCulturePtr Culture, ESelectInfo::Type SelectionType );
 
 	/** Generate a widget for the region combo */
-	TSharedRef<SWidget> OnRegionGenerateWidget( TSharedPtr<FCulture, ESPMode::ThreadSafe> Culture, IDetailLayoutBuilder* DetailBuilder ) const;
+	TSharedRef<SWidget> OnRegionGenerateWidget( FCulturePtr Culture, IDetailLayoutBuilder* DetailBuilder ) const;
 
 	/** Delegate called to display the current region */
 	FText GetCurrentRegionText() const;
 
 	/** Delegate called when the region selection has changed */
-	void OnRegionSelectionChanged( TSharedPtr<FCulture, ESPMode::ThreadSafe> Culture, ESelectInfo::Type SelectionType );
+	void OnRegionSelectionChanged( FCulturePtr Culture, ESelectInfo::Type SelectionType );
 
 	/** Delegate called to see if we're allowed to change the region */
 	bool IsRegionSelectionAllowed() const;
@@ -72,25 +72,25 @@ private:
 	TWeakObjectPtr<UInternationalizationSettingsModel> Model;
 
 	/** The culture selected at present */
-	TSharedPtr<FCulture, ESPMode::ThreadSafe> SelectedCulture;
+	FCulturePtr SelectedCulture;
 
 	/** The language selected at present */
-	TSharedPtr<FCulture, ESPMode::ThreadSafe> SelectedLanguage;
+	FCulturePtr SelectedLanguage;
 
 	/** The cultures we have available to us */
-	TArray< TSharedPtr<FCulture, ESPMode::ThreadSafe> > AvailableCultures;
+	TArray< FCulturePtr > AvailableCultures;
 
 	/** The languages we have available to us */
-	TArray< TSharedPtr<FCulture, ESPMode::ThreadSafe> > AvailableLanguages;
+	TArray< FCulturePtr > AvailableLanguages;
 
 	/** The dropdown for the languages */
-	TSharedPtr<SComboBox< TSharedPtr<FCulture, ESPMode::ThreadSafe> > > LanguageComboBox;
+	TSharedPtr<SComboBox< FCulturePtr > > LanguageComboBox;
 
 	/** The regions we have available to us */
-	TArray< TSharedPtr<FCulture, ESPMode::ThreadSafe> > AvailableRegions;
+	TArray< FCulturePtr > AvailableRegions;
 
 	/** The dropdown for the regions */
-	TSharedPtr<SComboBox< TSharedPtr<FCulture, ESPMode::ThreadSafe> > > RegionComboBox;
+	TSharedPtr<SComboBox< FCulturePtr > > RegionComboBox;
 
 	/** The check box for using localized field names */
 	TSharedPtr<SCheckBox> FieldNamesCheckBox;
diff --git a/Engine/Source/Editor/Documentation/Private/Documentation.cpp b/Engine/Source/Editor/Documentation/Private/Documentation.cpp
--- a/Engine/Source/Editor/Documentation/Private/Documentation.cpp
+++ b/Engine/Source/Editor/Documentation/Private/Documentation.cpp
@@ -30,7 +30,7 @@ bool FDocumentation::OpenHome() const
 	return Open( TEXT("%ROOT%") );
 }
 
-bool FDocumentation::OpenHome(const TSharedRef<FCulture, ESPMode::ThreadSafe>& Culture) const
+bool FDocumentation::OpenHome(const FCultureRef& Culture) const
 {
 	return Open(TEXT("%ROOT%"), Culture);
 }
@@ -84,7 +84,7 @@ bool FDocumentation::Open( const FString& Link ) const
 	return !DocumentationUrl.IsEmpty();
 }
 
-bool FDocumentation::Open(const FString& Link, const TSharedRef<FCulture, ESPMode::ThreadSafe>& Culture) const
+bool FDocumentation::Open(const FString& Link, const FCultureRef& Culture) const
 {
 	FString DocumentationUrl;
 
@@ -158,7 +158,7 @@ bool FDocumentation::PageExists(const FString& Link) const
 	return FPaths::FileExists(SourcePath);
 }
 
-bool FDocumentation::PageExists(const FString& Link, const TSharedRef<FCulture, ESPMode::ThreadSafe>& Culture) const
+bool FDocumentation::PageExists(const FString& Link, const FCultureRef& Culture) const
 {
 	const TWeakPtr< IDocumentationPage >* ExistingPagePtr = LoadedPages.Find(Link);
 	if (ExistingPagePtr != NULL)
diff --git a/Engine/Source/Editor/Documentation/Private/Documentation.h b/Engine/Source/Editor/Documentation/Private/Documentation.h
--- a/Engine/Source/Editor/Documentation/Private/Documentation.h
+++ b/Engine/Source/Editor/Documentation/Private/Documentation.h
@@ -15,13 +15,13 @@ public:
 
 	virtual bool OpenHome() const override;
 
-	virtual bool OpenHome(const TSharedRef<FCulture, ESPMode::ThreadSafe>& Culture) const override;
+	virtual bool OpenHome(const FCultureRef& Culture) const override;
 
 	virtual bool OpenAPIHome() const override;
 
 	virtual bool Open( const FString& Link ) const override;
 
-	virtual bool Open(const FString& Link, const TSharedRef<FCulture, ESPMode::ThreadSafe>& Culture) const override;
+	virtual bool Open(const FString& Link, const FCultureRef& Culture) const override;
 
 	virtual TSharedRef< SWidget > CreateAnchor( const TAttribute<FString>& Link, const FString& PreviewLink = FString(), const FString& PreviewExcerptName = FString() ) const override;
 
@@ -29,7 +29,7 @@ public:
 
 	virtual bool PageExists(const FString& Link) const override;
 
-	virtual bool PageExists(const FString& Link, const TSharedRef<FCulture, ESPMode::ThreadSafe>& Culture) const override;
+	virtual bool PageExists(const FString& Link, const FCultureRef& Culture) const override;
 
 	virtual TSharedRef< class SToolTip > CreateToolTip( const TAttribute<FText>& Text, const TSharedPtr<SWidget>& OverrideContent, const FString& Link, const FString& ExcerptName ) const override;
 
diff --git a/Engine/Source/Editor/Documentation/Private/DocumentationLink.h b/Engine/Source/Editor/Documentation/Private/DocumentationLink.h
--- a/Engine/Source/Editor/Documentation/Private/DocumentationLink.h
+++ b/Engine/Source/Editor/Documentation/Private/DocumentationLink.h
@@ -23,7 +23,7 @@ public:
 		return GetHomeUrl(FInternationalization::Get().GetCurrentCulture());
 	}
 
-	static FString GetHomeUrl(const TSharedRef<FCulture, ESPMode::ThreadSafe>& Culture)
+	static FString GetHomeUrl(const FCultureRef& Culture)
 	{
 		FString Url;
 		FUnrealEdMisc::Get().GetURL( TEXT("UDNURL"), Url, true );
@@ -37,7 +37,7 @@ public:
 		return ToUrl(Link, FInternationalization::Get().GetCurrentCulture());
 	}
 
-	static FString ToUrl(const FString& Link, const TSharedRef<FCulture, ESPMode::ThreadSafe>& Culture)
+	static FString ToUrl(const FString& Link, const FCultureRef& Culture)
 	{
 		FString Path;
 		FString Anchor;
@@ -56,7 +56,7 @@ public:
 
 		if (!FPaths::FileExists(FilePath))
 		{
-			const TSharedPtr<FCulture, ESPMode::ThreadSafe> FallbackCulture = I18N.GetCulture(TEXT("en"));
+			const FCulturePtr FallbackCulture = I18N.GetCulture(TEXT("en"));
 			if (FallbackCulture.IsValid())
 			{
 				const FString FallbackFilePath = ToFilePath(Link, FallbackCulture.ToSharedRef());
@@ -70,7 +70,7 @@ public:
 		return FilePath;
 	}
 
-	static FString ToFilePath(const FString& Link, const TSharedRef<FCulture, ESPMode::ThreadSafe>& Culture)
+	static FString ToFilePath(const FString& Link, const FCultureRef& Culture)
 	{
 		FString Path;
 		FString Anchor;
@@ -84,12 +84,12 @@ public:
 	{
 		FInternationalization& I18N = FInternationalization::Get();
 
-		TSharedRef<FCulture, ESPMode::ThreadSafe> Culture = I18N.GetCurrentCulture();
+		FCultureRef Culture = I18N.GetCurrentCulture();
 		FString FilePath = ToFilePath(Link, Culture);
 
 		if (!FPaths::FileExists(FilePath))
 		{
-			const TSharedPtr<FCulture, ESPMode::ThreadSafe> FallbackCulture = I18N.GetCulture(TEXT("en"));
+			const FCulturePtr FallbackCulture = I18N.GetCulture(TEXT("en"));
 			if (FallbackCulture.IsValid())
 			{
 				const FString FallbackFilePath = ToFilePath(Link, FallbackCulture.ToSharedRef());
@@ -103,7 +103,7 @@ public:
 		return ToFileUrl(Link, Culture);
 	}
 
-	static FString ToFileUrl(const FString& Link, const TSharedRef<FCulture, ESPMode::ThreadSafe>& Culture)
+	static FString ToFileUrl(const FString& Link, const FCultureRef& Culture)
 	{
 		FString Path;
 		FString Anchor;
@@ -120,7 +120,7 @@ public:
 
 		if (!FPaths::FileExists(SourcePath))
 		{
-			const TSharedPtr<FCulture, ESPMode::ThreadSafe> FallbackCulture = I18N.GetCulture(TEXT("en"));
+			const FCulturePtr FallbackCulture = I18N.GetCulture(TEXT("en"));
 			if (FallbackCulture.IsValid())
 			{
 				const FString FallbackSourcePath = ToSourcePath(Link, FallbackCulture.ToSharedRef());
@@ -134,7 +134,7 @@ public:
 		return SourcePath;
 	}
 
-	static FString ToSourcePath(const FString& Link, const TSharedRef<FCulture, ESPMode::ThreadSafe>& Culture)
+	static FString ToSourcePath(const FString& Link, const FCultureRef& Culture)
 	{
 		FString Path;
 		FString Anchor;
diff --git a/Engine/Source/Editor/Documentation/Public/IDocumentation.h b/Engine/Source/Editor/Documentation/Public/IDocumentation.h
--- a/Engine/Source/Editor/Documentation/Public/IDocumentation.h
+++ b/Engine/Source/Editor/Documentation/Public/IDocumentation.h
@@ -139,13 +139,13 @@ public:
 
 	virtual bool OpenHome() const = 0;
 
-	virtual bool OpenHome(const TSharedRef<FCulture, ESPMode::ThreadSafe>& Culture) const = 0;
+	virtual bool OpenHome(const FCultureRef& Culture) const = 0;
 
 	virtual bool OpenAPIHome() const = 0;
 
 	virtual bool Open( const FString& Link ) const = 0;
 
-	virtual bool Open(const FString& Link, const TSharedRef<FCulture, ESPMode::ThreadSafe>& Culture) const = 0;
+	virtual bool Open(const FString& Link, const FCultureRef& Culture) const = 0;
 
 	virtual TSharedRef< class SWidget > CreateAnchor( const TAttribute<FString>& Link, const FString& PreviewLink = FString(), const FString& PreviewExcerptName = FString() ) const = 0;
 
@@ -153,7 +153,7 @@ public:
 
 	virtual bool PageExists(const FString& Link) const = 0;
 
-	virtual bool PageExists(const FString& Link, const TSharedRef<FCulture, ESPMode::ThreadSafe>& Culture) const = 0;
+	virtual bool PageExists(const FString& Link, const FCultureRef& Culture) const = 0;
 
 	virtual TSharedRef< class SToolTip > CreateToolTip( const TAttribute<FText>& Text, const TSharedPtr<SWidget>& OverrideContent, const FString& Link, const FString& ExcerptName ) const = 0;
 };
diff --git a/Engine/Source/Editor/MainFrame/Private/Menus/TranslationEditorMenu.h b/Engine/Source/Editor/MainFrame/Private/Menus/TranslationEditorMenu.h
--- a/Engine/Source/Editor/MainFrame/Private/Menus/TranslationEditorMenu.h
+++ b/Engine/Source/Editor/MainFrame/Private/Menus/TranslationEditorMenu.h
@@ -152,7 +152,7 @@ public:
 				FString ArchiveName = FPaths::GetBaseFilename(FPaths::GetPath(ArchiveFileName));
 				FString ManifestName = FPaths::GetBaseFilename(ManifestFileName);
 
-				TSharedPtr<FCulture, ESPMode::ThreadSafe> Culture = FInternationalization::Get().GetCulture(ArchiveName);
+				FCulturePtr Culture = FInternationalization::Get().GetCulture(ArchiveName);
 				if (Culture.IsValid())
 				{
 					ArchiveName = Culture->GetDisplayName();
diff --git a/Engine/Source/Editor/NewsFeed/Private/Implementation/NewsFeedCache.cpp b/Engine/Source/Editor/NewsFeed/Private/Implementation/NewsFeedCache.cpp
--- a/Engine/Source/Editor/NewsFeed/Private/Implementation/NewsFeedCache.cpp
+++ b/Engine/Source/Editor/NewsFeed/Private/Implementation/NewsFeedCache.cpp
@@ -89,7 +89,7 @@ void FNewsFeedCache::LoadTitleFile( )
 	FString CultureString;
 
 	// get current culture
-	TSharedRef<FCulture, ESPMode::ThreadSafe> Culture = FInternationalization::Get().GetCurrentCulture();
+	FCultureRef Culture = FInternationalization::Get().GetCurrentCulture();
 
 	if (CurrentCultureSpec == ENewsFeedCultureSpec::Full)
 	{
diff --git a/Engine/Source/Editor/UnrealEd/Classes/Commandlets/InternationalizationExportCommandlet.h b/Engine/Source/Editor/UnrealEd/Classes/Commandlets/InternationalizationExportCommandlet.h
--- a/Engine/Source/Editor/UnrealEd/Classes/Commandlets/InternationalizationExportCommandlet.h
+++ b/Engine/Source/Editor/UnrealEd/Classes/Commandlets/InternationalizationExportCommandlet.h
@@ -74,7 +74,7 @@ public:
 	/**
 	 * Retrieves the culture equivalent that matches the set language code.
 	 */
-	TSharedPtr<FCulture, ESPMode::ThreadSafe> GetCulture() const;
+	FCulturePtr GetCulture() const;
 
 	/**
 	 * Retrieves the display name for the language.
@@ -97,7 +97,7 @@ private:
 	FString LanguageCode;
 	FString LanguagePluralForms;
 
-	TSharedPtr<FCulture, ESPMode::ThreadSafe> Culture;
+	FCulturePtr Culture;
 };
 
 
diff --git a/Engine/Source/Runtime/Core/Private/Internationalization/Cultures/AmericanEnglish.h b/Engine/Source/Runtime/Core/Private/Internationalization/Cultures/AmericanEnglish.h
--- a/Engine/Source/Runtime/Core/Private/Internationalization/Cultures/AmericanEnglish.h
+++ b/Engine/Source/Runtime/Core/Private/Internationalization/Cultures/AmericanEnglish.h
@@ -8,10 +8,10 @@
 class FAmericanEnglishCulture
 {
 public:
-	static TSharedRef<FCulture, ESPMode::ThreadSafe> Create()
+	static FCultureRef Create()
 	{
 #if UE_ENABLE_ICU
-		TSharedRef<FCulture, ESPMode::ThreadSafe> Culture = MakeShareable( new FCulture( FString("en-US") ) );
+		FCultureRef Culture = MakeShareable( new FCulture( FString("en-US") ) );
 #else
 		TArray<FString> InNativeDigits;
 
@@ -131,7 +131,7 @@ public:
 			);
 
 
-		TSharedRef<FCulture, ESPMode::ThreadSafe> Culture = MakeShareable( new FCulture(
+		FCultureRef Culture = MakeShareable( new FCulture(
 			  LOCTEXT("AmericanEnglishDisplayName", "English (United States)")		//const FText DisplayName
 			, FString(TEXT("English (United States)"))			//const FString EnglishName
 			, 1033												//const int KeyboardLayoutId
diff --git a/Engine/Source/Runtime/Core/Private/Internationalization/Cultures/IndiaHindi.h b/Engine/Source/Runtime/Core/Private/Internationalization/Cultures/IndiaHindi.h
--- a/Engine/Source/Runtime/Core/Private/Internationalization/Cultures/IndiaHindi.h
+++ b/Engine/Source/Runtime/Core/Private/Internationalization/Cultures/IndiaHindi.h
@@ -8,10 +8,10 @@
 class FIndiaHindiCulture
 {
 public:
-	static TSharedRef<FCulture, ESPMode::ThreadSafe> Create()
+	static FCultureRef Create()
 	{
 #if UE_ENABLE_ICU
-		TSharedRef<FCulture, ESPMode::ThreadSafe> Culture = MakeShareable( new FCulture( FString("hi-IN") ) );
+		FCultureRef Culture = MakeShareable( new FCulture( FString("hi-IN") ) );
 #else
 		TArray<FString> InNativeDigits;
 
@@ -133,7 +133,7 @@ public:
 			);
 
 
-		TSharedRef<FCulture, ESPMode::ThreadSafe> Culture = MakeShareable( new FCulture(
+		FCultureRef Culture = MakeShareable( new FCulture(
 			  LOCTEXT("HindiDisplayName", "Hindi (India)")					//const FText DisplayName
 			, FString(TEXT("Hindi (India)"))		//const FString EnglishName
 			, 1081									//const int KeyboardLayoutId
diff --git a/Engine/Source/Runtime/Core/Private/Internationalization/Cultures/InvariantCulture.h b/Engine/Source/Runtime/Core/Private/Internationalization/Cultures/InvariantCulture.h
--- a/Engine/Source/Runtime/Core/Private/Internationalization/Cultures/InvariantCulture.h
+++ b/Engine/Source/Runtime/Core/Private/Internationalization/Cultures/InvariantCulture.h
@@ -8,10 +8,10 @@
 class FInvariantCulture
 {
 public:
-	static TSharedRef<FCulture, ESPMode::ThreadSafe> Create()
+	static FCultureRef Create()
 	{
 #if UE_ENABLE_ICU
-		TSharedRef<FCulture, ESPMode::ThreadSafe> Culture = MakeShareable( new FCulture( FString() ) );
+		FCultureRef Culture = MakeShareable( new FCulture( FString() ) );
 #else
 		TArray<FString> InNativeDigits;
 
@@ -131,7 +131,7 @@ public:
 			);
 
 
-		TSharedRef<FCulture, ESPMode::ThreadSafe> Culture = MakeShareable( new FCulture(
+		FCultureRef Culture = MakeShareable( new FCulture(
 			  LOCTEXT("InvariantCultureDisplayName", "Invariant Language (Invariant Country)")		//const FText DisplayName
 			, FString(TEXT("Invariant Language (Invariant Country)"))			//const FString EnglishName
 			, 1033												//const int KeyboardLayoutId
diff --git a/Engine/Source/Runtime/Core/Private/Internationalization/Cultures/JapaneseCulture.h b/Engine/Source/Runtime/Core/Private/Internationalization/Cultures/JapaneseCulture.h
--- a/Engine/Source/Runtime/Core/Private/Internationalization/Cultures/JapaneseCulture.h
+++ b/Engine/Source/Runtime/Core/Private/Internationalization/Cultures/JapaneseCulture.h
@@ -8,10 +8,10 @@
 class FJapaneseCulture
 {
 public:
-	static TSharedRef<FCulture, ESPMode::ThreadSafe> Create()
+	static FCultureRef Create()
 	{
 #if UE_ENABLE_ICU
-		TSharedRef<FCulture, ESPMode::ThreadSafe> Culture = MakeShareable( new FCulture( FString("ja-JP") ) );
+		FCultureRef Culture = MakeShareable( new FCulture( FString("ja-JP") ) );
 #else
 		TArray<FString> NativeDigits;
 
@@ -147,7 +147,7 @@ public:
 			);
 
 
-		TSharedRef<FCulture, ESPMode::ThreadSafe> Culture = MakeShareable( new FCulture(
+		FCultureRef Culture = MakeShareable( new FCulture(
 			  LOCTEXT("JapaneseDisplayName", "Japanese (Japan)")	//const FText DisplayName
 			, FString(TEXT("Japanese (Japan)"))	//const FString EnglishName
 			, 1041	//const int KeyboardLayoutId
diff --git a/Engine/Source/Runtime/Core/Private/Internationalization/Cultures/KoreanCulture.h b/Engine/Source/Runtime/Core/Private/Internationalization/Cultures/KoreanCulture.h
--- a/Engine/Source/Runtime/Core/Private/Internationalization/Cultures/KoreanCulture.h
+++ b/Engine/Source/Runtime/Core/Private/Internationalization/Cultures/KoreanCulture.h
@@ -8,10 +8,10 @@
 class FKoreanCulture
 {
 public:
-	static TSharedRef<FCulture, ESPMode::ThreadSafe> Create()
+	static FCultureRef Create()
 	{
 #if UE_ENABLE_ICU
-		TSharedRef<FCulture, ESPMode::ThreadSafe> Culture = MakeShareable( new FCulture( FString("ko-KR") ) );
+		FCultureRef Culture = MakeShareable( new FCulture( FString("ko-KR") ) );
 #else
 		TArray<FString> NativeDigits;
 
@@ -147,7 +147,7 @@ public:
 			);
 
 
-		TSharedRef<FCulture, ESPMode::ThreadSafe> Culture = MakeShareable( new FCulture(
+		FCultureRef Culture = MakeShareable( new FCulture(
 			  LOCTEXT("KoreanDisplayName", "Korean (Korea)")	//const FText DisplayName
 			, FString(TEXT("Korean (Korea)"))	//const FString EnglishName
 			, 1042	//const int KeyboardLayoutId
diff --git a/Engine/Source/Runtime/Core/Private/Internationalization/Cultures/LEETCulture.h b/Engine/Source/Runtime/Core/Private/Internationalization/Cultures/LEETCulture.h
--- a/Engine/Source/Runtime/Core/Private/Internationalization/Cultures/LEETCulture.h
+++ b/Engine/Source/Runtime/Core/Private/Internationalization/Cultures/LEETCulture.h
@@ -8,10 +8,10 @@
 class FLEETCulture
 {
 public:
-	static TSharedRef<FCulture, ESPMode::ThreadSafe> Create()
+	static FCultureRef Create()
 	{
 #if UE_ENABLE_ICU
-		TSharedRef<FCulture, ESPMode::ThreadSafe> Culture = MakeShareable( new FCulture( FString() ) );
+		FCultureRef Culture = MakeShareable( new FCulture( FString() ) );
 #else
 		TArray<FString> InNativeDigits;
 
@@ -135,7 +135,7 @@ public:
 			);
 
 
-		TSharedRef<FCulture, ESPMode::ThreadSafe> Culture = MakeShareable( new FCulture(
+		FCultureRef Culture = MakeShareable( new FCulture(
 			  LOCTEXT("LEETCultureDisplayName", "LEET Language (LEET Country)")		//const FText DisplayName
 			, FString(TEXT("LEET Language (LEET Country)"))			//const FString EnglishName
 			, 1033												//const int KeyboardLayoutId
diff --git a/Engine/Source/Runtime/Core/Private/Internationalization/ICUText.cpp b/Engine/Source/Runtime/Core/Private/Internationalization/ICUText.cpp
--- a/Engine/Source/Runtime/Core/Private/Internationalization/ICUText.cpp
+++ b/Engine/Source/Runtime/Core/Private/Internationalization/ICUText.cpp
@@ -21,11 +21,11 @@ bool FText::IsWhitespace( const TCHAR Char )
 	return u_isWhitespace(ICUChar) != 0;
 }
 
-FText FText::AsDate(const FDateTime& DateTime, const EDateTimeStyle::Type DateStyle, const TSharedPtr<FCulture, ESPMode::ThreadSafe>& TargetCulture)
+FText FText::AsDate(const FDateTime& DateTime, const EDateTimeStyle::Type DateStyle, const FCulturePtr& TargetCulture)
 {
 	FInternationalization& I18N = FInternationalization::Get();
 	checkf(I18N.IsInitialized() == true, TEXT("FInternationalization is not initialized. An FText formatting method was likely used in static object initialization - this is not supported."));
-	TSharedPtr<FCulture, ESPMode::ThreadSafe> Culture = TargetCulture.IsValid() ? TargetCulture : I18N.GetCurrentCulture();
+	FCulturePtr Culture = TargetCulture.IsValid() ? TargetCulture : I18N.GetCurrentCulture();
 
 	int32 UNIXTimestamp = DateTime.ToUnixTimestamp();
 	UDate ICUDate = static_cast<double>(UNIXTimestamp) * U_MILLIS_PER_SECOND;
@@ -44,11 +44,11 @@ FText FText::AsDate(const FDateTime& DateTime, const EDateTimeStyle::Type DateSt
 	return ResultText;
 }
 
-FText FText::AsTime(const FDateTime& DateTime, const EDateTimeStyle::Type TimeStyle, const FString& TimeZone, const TSharedPtr<FCulture, ESPMode::ThreadSafe>& TargetCulture)
+FText FText::AsTime(const FDateTime& DateTime, const EDateTimeStyle::Type TimeStyle, const FString& TimeZone, const FCulturePtr& TargetCulture)
 {
 	FInternationalization& I18N = FInternationalization::Get();
 	checkf(I18N.IsInitialized() == true, TEXT("FInternationalization is not initialized. An FText formatting method was likely used in static object initialization - this is not supported."));
-	TSharedPtr<FCulture, ESPMode::ThreadSafe> Culture = TargetCulture.IsValid() ? TargetCulture : I18N.GetCurrentCulture();
+	FCulturePtr Culture = TargetCulture.IsValid() ? TargetCulture : I18N.GetCurrentCulture();
 
 	int32 UNIXTimestamp = DateTime.ToUnixTimestamp();
 	UDate ICUDate = static_cast<double>(UNIXTimestamp) * U_MILLIS_PER_SECOND;
@@ -67,11 +67,11 @@ FText FText::AsTime(const FDateTime& DateTime, const EDateTimeStyle::Type TimeSt
 	return ResultText;
 }
 
-FText FText::AsTimespan(const FTimespan& Time, const TSharedPtr<FCulture, ESPMode::ThreadSafe>& TargetCulture)
+FText FText::AsTimespan(const FTimespan& Time, const FCulturePtr& TargetCulture)
 {
 	FInternationalization& I18N = FInternationalization::Get();
 	checkf(I18N.IsInitialized() == true, TEXT("FInternationalization is not initialized. An FText formatting method was likely used in static object initialization - this is not supported."));
-	TSharedPtr<FCulture, ESPMode::ThreadSafe> Culture = TargetCulture.IsValid() ? TargetCulture : I18N.GetCurrentCulture();
+	FCulturePtr Culture = TargetCulture.IsValid() ? TargetCulture : I18N.GetCurrentCulture();
 
 	FText TimespanFormatPattern = NSLOCTEXT("Timespan", "FormatPattern", "{Hours}:{Minutes}:{Seconds}");
 
@@ -91,11 +91,11 @@ FText FText::AsTimespan(const FTimespan& Time, const TSharedPtr<FCulture, ESPMod
 	return FText::Format(TimespanFormatPattern, TimeArguments);
 }
 
-FText FText::AsDateTime(const FDateTime& DateTime, const EDateTimeStyle::Type DateStyle, const EDateTimeStyle::Type TimeStyle, const FString& TimeZone, const TSharedPtr<FCulture, ESPMode::ThreadSafe>& TargetCulture)
+FText FText::AsDateTime(const FDateTime& DateTime, const EDateTimeStyle::Type DateStyle, const EDateTimeStyle::Type TimeStyle, const FString& TimeZone, const FCulturePtr& TargetCulture)
 {
 	FInternationalization& I18N = FInternationalization::Get();
 	checkf(I18N.IsInitialized() == true, TEXT("FInternationalization is not initialized. An FText formatting method was likely used in static object initialization - this is not supported."));
-	TSharedPtr<FCulture, ESPMode::ThreadSafe> Culture = TargetCulture.IsValid() ? TargetCulture : I18N.GetCurrentCulture();
+	FCulturePtr Culture = TargetCulture.IsValid() ? TargetCulture : I18N.GetCurrentCulture();
 
 	int32 UNIXTimestamp = DateTime.ToUnixTimestamp();
 	UDate ICUDate = static_cast<double>(UNIXTimestamp) * U_MILLIS_PER_SECOND;
@@ -114,7 +114,7 @@ FText FText::AsDateTime(const FDateTime& DateTime, const EDateTimeStyle::Type Da
 	return ResultText;
 }
 
-FText FText::AsMemory(SIZE_T NumBytes, const FNumberFormattingOptions* const Options, const TSharedPtr<FCulture, ESPMode::ThreadSafe>& TargetCulture)
+FText FText::AsMemory(SIZE_T NumBytes, const FNumberFormattingOptions* const Options, const FCulturePtr& TargetCulture)
 {
 	checkf(FInternationalization::Get().IsInitialized() == true, TEXT("FInternationalization is not initialized. An FText formatting method was likely used in static object initialization - this is not supported."));
 	FFormatNamedArguments Args;
diff --git a/Engine/Source/Runtime/Core/Private/Internationalization/ICUText.h b/Engine/Source/Runtime/Core/Private/Internationalization/ICUText.h
--- a/Engine/Source/Runtime/Core/Private/Internationalization/ICUText.h
+++ b/Engine/Source/Runtime/Core/Private/Internationalization/ICUText.h
@@ -11,11 +11,11 @@
 #include "unicode/unistr.h"
 
 template<typename T1, typename T2>
-FText FText::AsNumberTemplate(T1 Val, const FNumberFormattingOptions* const Options, const TSharedPtr<FCulture, ESPMode::ThreadSafe>& TargetCulture)
+FText FText::AsNumberTemplate(T1 Val, const FNumberFormattingOptions* const Options, const FCulturePtr& TargetCulture)
 {
 	FInternationalization& I18N = FInternationalization::Get();
 	checkf(I18N.IsInitialized() == true, TEXT("FInternationalization is not initialized. An FText formatting method was likely used in static object initialization - this is not supported."));
-	const TSharedRef<FCulture, ESPMode::ThreadSafe> Culture = TargetCulture.IsValid() ? TargetCulture.ToSharedRef() : I18N.GetCurrentCulture();
+	const FCultureRef Culture = TargetCulture.IsValid() ? TargetCulture.ToSharedRef() : I18N.GetCurrentCulture();
 	UErrorCode ICUStatus = U_ZERO_ERROR;
 	const TSharedRef<const icu::DecimalFormat> ICUDecimalFormat( Culture->Implementation->GetDecimalFormatter(Options) );
 	icu::Formattable FormattableVal(static_cast<T2>(Val));
@@ -31,11 +31,11 @@ FText FText::AsNumberTemplate(T1 Val, const FNumberFormattingOptions* const Opti
 }
 
 template<typename T1, typename T2>
-FText FText::AsCurrencyTemplate(T1 Val, const FString& CurrencyCode, const FNumberFormattingOptions* const Options, const TSharedPtr<FCulture, ESPMode::ThreadSafe>& TargetCulture)
+FText FText::AsCurrencyTemplate(T1 Val, const FString& CurrencyCode, const FNumberFormattingOptions* const Options, const FCulturePtr& TargetCulture)
 {
 	FInternationalization& I18N = FInternationalization::Get();
 	checkf(I18N.IsInitialized() == true, TEXT("FInternationalization is not initialized. An FText formatting method was likely used in static object initialization - this is not supported."));
-	const TSharedRef<FCulture, ESPMode::ThreadSafe> Culture = TargetCulture.IsValid() ? TargetCulture.ToSharedRef() : I18N.GetCurrentCulture();
+	const FCultureRef Culture = TargetCulture.IsValid() ? TargetCulture.ToSharedRef() : I18N.GetCurrentCulture();
 	UErrorCode ICUStatus = U_ZERO_ERROR;
 	const TSharedRef<const icu::DecimalFormat> ICUDecimalFormat(Culture->Implementation->GetCurrencyFormatter(CurrencyCode, Options));
 	icu::Formattable FormattableVal(static_cast<T2>(Val));
@@ -51,11 +51,11 @@ FText FText::AsCurrencyTemplate(T1 Val, const FString& CurrencyCode, const FNumb
 }
 
 template<typename T1, typename T2>
-FText FText::AsPercentTemplate(T1 Val, const FNumberFormattingOptions* const Options, const TSharedPtr<FCulture, ESPMode::ThreadSafe>& TargetCulture)
+FText FText::AsPercentTemplate(T1 Val, const FNumberFormattingOptions* const Options, const FCulturePtr& TargetCulture)
 {
 	FInternationalization& I18N = FInternationalization::Get();
 	checkf(I18N.IsInitialized() == true, TEXT("FInternationalization is not initialized. An FText formatting method was likely used in static object initialization - this is not supported."));
-	const TSharedRef<FCulture, ESPMode::ThreadSafe> Culture = TargetCulture.IsValid() ? TargetCulture.ToSharedRef() : I18N.GetCurrentCulture();
+	const FCultureRef Culture = TargetCulture.IsValid() ? TargetCulture.ToSharedRef() : I18N.GetCurrentCulture();
 	UErrorCode ICUStatus = U_ZERO_ERROR;
 	const TSharedRef<const icu::DecimalFormat> ICUDecimalFormat( Culture->Implementation->GetPercentFormatter(Options) );
 	icu::Formattable FormattableVal(static_cast<T2>(Val));
diff --git a/Engine/Source/Runtime/Core/Private/Internationalization/Internationalization.cpp b/Engine/Source/Runtime/Core/Private/Internationalization/Internationalization.cpp
--- a/Engine/Source/Runtime/Core/Private/Internationalization/Internationalization.cpp
+++ b/Engine/Source/Runtime/Core/Private/Internationalization/Internationalization.cpp
@@ -107,15 +107,15 @@ void FInternationalization::SetCurrentCulture(const FString& Name)
 	}
 }
 
-TSharedRef<FCulture, ESPMode::ThreadSafe> FInternationalization::GetCurrentCulture() const
+FCultureRef FInternationalization::GetCurrentCulture() const
 {
 	return AllCultures[CurrentCultureIndex];
 }
 
-TSharedPtr<FCulture, ESPMode::ThreadSafe> FInternationalization::GetCulture(const FString& Name) const
+FCulturePtr FInternationalization::GetCulture(const FString& Name) const
 {
 	int32 CultureIndex = GetCultureIndex(Name);
-	return CultureIndex != -1 ? AllCultures[CultureIndex] : TSharedPtr<FCulture, ESPMode::ThreadSafe>();
+	return CultureIndex != -1 ? AllCultures[CultureIndex] : FCulturePtr();
 }
 
 int32 FInternationalization::GetCultureIndex(const FString& Name) const
@@ -566,7 +566,7 @@ void FInternationalization::GetCultureNames(TArray<FString>& CultureNames) const
 	}
 }
 
-void FInternationalization::GetCulturesWithAvailableLocalization(const TArray<FString>& InLocalizationPaths, TArray< TSharedPtr<FCulture, ESPMode::ThreadSafe> >& OutAvailableCultures, const bool bIncludeDerivedCultures) const
+void FInternationalization::GetCulturesWithAvailableLocalization(const TArray<FString>& InLocalizationPaths, TArray< FCulturePtr >& OutAvailableCultures, const bool bIncludeDerivedCultures) const
 {
 	OutAvailableCultures.Reset();
 
@@ -624,7 +624,7 @@ void FInternationalization::GetCulturesWithAvailableLocalization(const TArray<FS
 	{
 		for(const FString& LocalizationFolder : AllLocalizationFolders)
 		{
-			TSharedPtr<FCulture, ESPMode::ThreadSafe> Culture = GetCulture(LocalizationFolder);
+			FCulturePtr Culture = GetCulture(LocalizationFolder);
 			if(Culture.IsValid() && AllCultures.Contains(Culture.ToSharedRef()))
 			{
 				OutAvailableCultures.AddUnique(Culture);
@@ -656,11 +656,11 @@ void FInternationalization::PopulateAllCultures(void)
 	const FString DefaultLocaleName = FPlatformMisc::GetDefaultLocale();
 
 	// Find the default locale as a culture, if present..
-	auto FindDefaultPredicate = [DefaultLocaleName](const TSharedRef<FCulture, ESPMode::ThreadSafe>& Culture) -> bool
+	auto FindDefaultPredicate = [DefaultLocaleName](const FCultureRef& Culture) -> bool
 	{
 		return Culture->GetName() == DefaultLocaleName;
 	};
-	TSharedRef<FCulture, ESPMode::ThreadSafe>* FoundDefaultCulture = AllCultures.FindByPredicate(FindDefaultPredicate);
+	FCultureRef* FoundDefaultCulture = AllCultures.FindByPredicate(FindDefaultPredicate);
 	// If present, set default culture variable.
 	if(FoundDefaultCulture)
 	{
diff --git a/Engine/Source/Runtime/Core/Private/Internationalization/LegacyText.cpp b/Engine/Source/Runtime/Core/Private/Internationalization/LegacyText.cpp
--- a/Engine/Source/Runtime/Core/Private/Internationalization/LegacyText.cpp
+++ b/Engine/Source/Runtime/Core/Private/Internationalization/LegacyText.cpp
@@ -10,32 +10,32 @@ bool FText::IsWhitespace( const TCHAR Char )
 	return FChar::IsWhitespace(Char);
 }
 
-FText FText::AsDate( const FDateTime& DateTime, const EDateTimeStyle::Type DateStyle, const TSharedPtr<FCulture, ESPMode::ThreadSafe>& TargetCulture )
+FText FText::AsDate( const FDateTime& DateTime, const EDateTimeStyle::Type DateStyle, const FCulturePtr& TargetCulture )
 {
 	checkf(FInternationalization::Get().IsInitialized() == true, TEXT("FInternationalization is not initialized. An FText formatting method was likely used in static object initialization - this is not supported."));
 	return FText::FromString( DateTime.ToString( TEXT("%Y.%m.%d") ) );
 }
 
-FText FText::AsTime( const FDateTime& DateTime, const EDateTimeStyle::Type TimeStyle, const FString& TimeZone, const TSharedPtr<FCulture, ESPMode::ThreadSafe>& TargetCulture )
+FText FText::AsTime( const FDateTime& DateTime, const EDateTimeStyle::Type TimeStyle, const FString& TimeZone, const FCulturePtr& TargetCulture )
 {
 	checkf(FInternationalization::Get().IsInitialized() == true, TEXT("FInternationalization is not initialized. An FText formatting method was likely used in static object initialization - this is not supported."));
 	return FText::FromString( DateTime.ToString( TEXT("%H.%M.%S") ) );
 }
 
-FText FText::AsTimespan( const FTimespan& Time, const TSharedPtr<FCulture, ESPMode::ThreadSafe>& TargetCulture)
+FText FText::AsTimespan( const FTimespan& Time, const FCulturePtr& TargetCulture)
 {
 	checkf(FInternationalization::Get().IsInitialized() == true, TEXT("FInternationalization is not initialized. An FText formatting method was likely used in static object initialization - this is not supported."));
 	FDateTime DateTime(Time.GetTicks());
 	return FText::FromString( DateTime.ToString( TEXT("%H.%M.%S") ) );
 }
 
-FText FText::AsDateTime( const FDateTime& DateTime, const EDateTimeStyle::Type DateStyle, const EDateTimeStyle::Type TimeStyle, const FString& TimeZone, const TSharedPtr<FCulture, ESPMode::ThreadSafe>& TargetCulture )
+FText FText::AsDateTime( const FDateTime& DateTime, const EDateTimeStyle::Type DateStyle, const EDateTimeStyle::Type TimeStyle, const FString& TimeZone, const FCulturePtr& TargetCulture )
 {
 	checkf(FInternationalization::Get().IsInitialized() == true, TEXT("FInternationalization is not initialized. An FText formatting method was likely used in static object initialization - this is not supported."));
 	return FText::FromString(DateTime.ToString(TEXT("%Y.%m.%d-%H.%M.%S")));
 }
 
-FText FText::AsMemory( SIZE_T NumBytes, const FNumberFormattingOptions* const Options, const TSharedPtr<FCulture, ESPMode::ThreadSafe>& TargetCulture )
+FText FText::AsMemory( SIZE_T NumBytes, const FNumberFormattingOptions* const Options, const FCulturePtr& TargetCulture )
 {
 	checkf(FInternationalization::Get().IsInitialized() == true, TEXT("FInternationalization is not initialized. An FText formatting method was likely used in static object initialization - this is not supported."));
 	FFormatNamedArguments Args;
diff --git a/Engine/Source/Runtime/Core/Private/Internationalization/LegacyText.h b/Engine/Source/Runtime/Core/Private/Internationalization/LegacyText.h
--- a/Engine/Source/Runtime/Core/Private/Internationalization/LegacyText.h
+++ b/Engine/Source/Runtime/Core/Private/Internationalization/LegacyText.h
@@ -3,28 +3,28 @@
 #pragma once
 
 template<typename T1, typename T2>
-FText FText::AsNumberTemplate(T1 Val, const FNumberFormattingOptions* const Options, const TSharedPtr<FCulture, ESPMode::ThreadSafe>& TargetCulture)
+FText FText::AsNumberTemplate(T1 Val, const FNumberFormattingOptions* const Options, const FCulturePtr& TargetCulture)
 {
 	FInternationalization& I18N = FInternationalization::Get();
 	checkf(I18N.IsInitialized() == true, TEXT("FInternationalization is not initialized. An FText formatting method was likely used in static object initialization - this is not supported."));
-	const TSharedRef<FCulture, ESPMode::ThreadSafe> Culture = TargetCulture.IsValid() ? TargetCulture.ToSharedRef() : I18N.GetCurrentCulture();
+	const FCultureRef Culture = TargetCulture.IsValid() ? TargetCulture.ToSharedRef() : I18N.GetCurrentCulture();
 	return FText::CreateNumericalText( Culture->NumberFormattingRule.AsNumber(Val) );
 }
 
 template<typename T1, typename T2>
-FText FText::AsCurrencyTemplate(T1 Val, const FString& CurrencyCode, const FNumberFormattingOptions* const Options, const TSharedPtr<FCulture, ESPMode::ThreadSafe>& TargetCulture)
+FText FText::AsCurrencyTemplate(T1 Val, const FString& CurrencyCode, const FNumberFormattingOptions* const Options, const FCulturePtr& TargetCulture)
 {
 	FInternationalization& I18N = FInternationalization::Get();
 	checkf(I18N.IsInitialized() == true, TEXT("FInternationalization is not initialized. An FText formatting method was likely used in static object initialization - this is not supported."));
-	const TSharedRef<FCulture, ESPMode::ThreadSafe> Culture = TargetCulture.IsValid() ? TargetCulture.ToSharedRef() : I18N.GetCurrentCulture();
+	const FCultureRef Culture = TargetCulture.IsValid() ? TargetCulture.ToSharedRef() : I18N.GetCurrentCulture();
 	return FText::CreateNumericalText( Culture->NumberFormattingRule.AsCurrency(Val) );
 }
 
 template<typename T1, typename T2>
-FText FText::AsPercentTemplate(T1 Val, const FNumberFormattingOptions* const Options, const TSharedPtr<FCulture, ESPMode::ThreadSafe>& TargetCulture)
+FText FText::AsPercentTemplate(T1 Val, const FNumberFormattingOptions* const Options, const FCulturePtr& TargetCulture)
 {
 	FInternationalization& I18N = FInternationalization::Get();
 	checkf(I18N.IsInitialized() == true, TEXT("FInternationalization is not initialized. An FText formatting method was likely used in static object initialization - this is not supported."));
-	const TSharedRef<FCulture, ESPMode::ThreadSafe> Culture = TargetCulture.IsValid() ? TargetCulture.ToSharedRef() : I18N.GetCurrentCulture();
+	const FCultureRef Culture = TargetCulture.IsValid() ? TargetCulture.ToSharedRef() : I18N.GetCurrentCulture();
 	return FText::CreateNumericalText( Culture->NumberFormattingRule.AsPercent(Val) );
 }
\ No newline at end of file
diff --git a/Engine/Source/Runtime/Core/Private/Internationalization/Text.cpp b/Engine/Source/Runtime/Core/Private/Internationalization/Text.cpp
--- a/Engine/Source/Runtime/Core/Private/Internationalization/Text.cpp
+++ b/Engine/Source/Runtime/Core/Private/Internationalization/Text.cpp
@@ -468,7 +468,7 @@ FText FText::Format(const FText& Fmt,const FText& v1,const FText& v2,const FText
 // on some platforms (PS4) int64_t is a typedef of long.  However, UE4 typedefs int64 as long long.  Since these are distinct types, and ICU only has a constructor for int64_t, casting to int64 causes a compile error from ambiguous function call, 
 // so cast to int64_t's where appropriate here to avoid problems.
 
-#define DEF_ASNUMBER_CAST(T1, T2) FText FText::AsNumber(T1 Val, const FNumberFormattingOptions* const Options, const TSharedPtr<FCulture, ESPMode::ThreadSafe>& TargetCulture) { return FText::AsNumberTemplate<T1, T2>(Val, Options, TargetCulture); }
+#define DEF_ASNUMBER_CAST(T1, T2) FText FText::AsNumber(T1 Val, const FNumberFormattingOptions* const Options, const FCulturePtr& TargetCulture) { return FText::AsNumberTemplate<T1, T2>(Val, Options, TargetCulture); }
 #define DEF_ASNUMBER(T) DEF_ASNUMBER_CAST(T, T)
 DEF_ASNUMBER(float)
 DEF_ASNUMBER(double)
@@ -486,7 +486,7 @@ DEF_ASNUMBER_CAST(uint64, int64_t)
 /**
  * Generate an FText that represents the passed number as currency in the current culture
  */
-#define DEF_ASCURRENCY_CAST(T1, T2) FText FText::AsCurrency(T1 Val, const FString& CurrencyCode, const FNumberFormattingOptions* const Options, const TSharedPtr<FCulture, ESPMode::ThreadSafe>& TargetCulture) { return FText::AsCurrencyTemplate<T1, T2>(Val, CurrencyCode, Options, TargetCulture); }
+#define DEF_ASCURRENCY_CAST(T1, T2) FText FText::AsCurrency(T1 Val, const FString& CurrencyCode, const FNumberFormattingOptions* const Options, const FCulturePtr& TargetCulture) { return FText::AsCurrencyTemplate<T1, T2>(Val, CurrencyCode, Options, TargetCulture); }
 #define DEF_ASCURRENCY(T) DEF_ASCURRENCY_CAST(T, T)
 DEF_ASCURRENCY(float)
 	DEF_ASCURRENCY(double)
@@ -505,7 +505,7 @@ DEF_ASCURRENCY(float)
 * Generate an FText that represents the passed number as a percentage in the current culture
 */
 
-#define DEF_ASPERCENT_CAST(T1, T2) FText FText::AsPercent(T1 Val, const FNumberFormattingOptions* const Options, const TSharedPtr<FCulture, ESPMode::ThreadSafe>& TargetCulture) { return FText::AsPercentTemplate<T1, T2>(Val, Options, TargetCulture); }
+#define DEF_ASPERCENT_CAST(T1, T2) FText FText::AsPercent(T1 Val, const FNumberFormattingOptions* const Options, const FCulturePtr& TargetCulture) { return FText::AsPercentTemplate<T1, T2>(Val, Options, TargetCulture); }
 #define DEF_ASPERCENT(T) DEF_ASPERCENT_CAST(T, T)
 DEF_ASPERCENT(double)
 DEF_ASPERCENT(float)
diff --git a/Engine/Source/Runtime/Core/Private/Internationalization/TextHistory.cpp b/Engine/Source/Runtime/Core/Private/Internationalization/TextHistory.cpp
--- a/Engine/Source/Runtime/Core/Private/Internationalization/TextHistory.cpp
+++ b/Engine/Source/Runtime/Core/Private/Internationalization/TextHistory.cpp
@@ -259,7 +259,7 @@ FTextHistory_FormatNumber::FTextHistory_FormatNumber()
 {
 }
 
-FTextHistory_FormatNumber::FTextHistory_FormatNumber(const FFormatArgumentValue& InSourceValue, const FNumberFormattingOptions* const InFormatOptions, const TSharedPtr<FCulture, ESPMode::ThreadSafe> InTargetCulture)
+FTextHistory_FormatNumber::FTextHistory_FormatNumber(const FFormatArgumentValue& InSourceValue, const FNumberFormattingOptions* const InFormatOptions, const FCulturePtr InTargetCulture)
 	: SourceValue(InSourceValue)
 	, FormatOptions(nullptr)
 	, TargetCulture(InTargetCulture)
@@ -313,7 +313,7 @@ void FTextHistory_FormatNumber::Serialize(FArchive& Ar)
 ///////////////////////////////////////
 // FTextHistory_AsNumber
 
-FTextHistory_AsNumber::FTextHistory_AsNumber(const FFormatArgumentValue& InSourceValue, const FNumberFormattingOptions* const InFormatOptions, const TSharedPtr<FCulture, ESPMode::ThreadSafe> InTargetCulture)
+FTextHistory_AsNumber::FTextHistory_AsNumber(const FFormatArgumentValue& InSourceValue, const FNumberFormattingOptions* const InFormatOptions, const FCulturePtr InTargetCulture)
 	: FTextHistory_FormatNumber(InSourceValue, InFormatOptions, InTargetCulture)
 {
 }
@@ -363,7 +363,7 @@ void FTextHistory_AsNumber::Serialize( FArchive& Ar )
 ///////////////////////////////////////
 // FTextHistory_AsPercent
 
-FTextHistory_AsPercent::FTextHistory_AsPercent(const FFormatArgumentValue& InSourceValue, const FNumberFormattingOptions* const InFormatOptions, const TSharedPtr<FCulture, ESPMode::ThreadSafe> InTargetCulture)
+FTextHistory_AsPercent::FTextHistory_AsPercent(const FFormatArgumentValue& InSourceValue, const FNumberFormattingOptions* const InFormatOptions, const FCulturePtr InTargetCulture)
 	: FTextHistory_FormatNumber(InSourceValue, InFormatOptions, InTargetCulture)
 {
 }
@@ -405,7 +405,7 @@ void FTextHistory_AsPercent::Serialize( FArchive& Ar )
 ///////////////////////////////////////
 // FTextHistory_AsCurrency
 
-FTextHistory_AsCurrency::FTextHistory_AsCurrency(const FFormatArgumentValue& InSourceValue, const FString& CurrencyCode, const FNumberFormattingOptions* const InFormatOptions, const TSharedPtr<FCulture, ESPMode::ThreadSafe> InTargetCulture)
+FTextHistory_AsCurrency::FTextHistory_AsCurrency(const FFormatArgumentValue& InSourceValue, const FString& CurrencyCode, const FNumberFormattingOptions* const InFormatOptions, const FCulturePtr InTargetCulture)
 : FTextHistory_FormatNumber(InSourceValue, InFormatOptions, InTargetCulture)
 , CurrencyCode(CurrencyCode)
 {
@@ -461,7 +461,7 @@ void FTextHistory_AsCurrency::Serialize( FArchive& Ar )
 ///////////////////////////////////////
 // FTextHistory_AsDate
 
-FTextHistory_AsDate::FTextHistory_AsDate(const FDateTime& InSourceDateTime, const EDateTimeStyle::Type InDateStyle, const TSharedPtr<FCulture, ESPMode::ThreadSafe> InTargetCulture)
+FTextHistory_AsDate::FTextHistory_AsDate(const FDateTime& InSourceDateTime, const EDateTimeStyle::Type InDateStyle, const FCulturePtr InTargetCulture)
 	: SourceDateTime(InSourceDateTime)
 	, DateStyle(InDateStyle)
 	, TargetCulture(InTargetCulture)
@@ -508,7 +508,7 @@ FText FTextHistory_AsDate::ToText(bool bInAsSource) const
 ///////////////////////////////////////
 // FTextHistory_AsTime
 
-FTextHistory_AsTime::FTextHistory_AsTime(const FDateTime& InSourceDateTime, const EDateTimeStyle::Type InTimeStyle, const FString& InTimeZone, const TSharedPtr<FCulture, ESPMode::ThreadSafe> InTargetCulture)
+FTextHistory_AsTime::FTextHistory_AsTime(const FDateTime& InSourceDateTime, const EDateTimeStyle::Type InTimeStyle, const FString& InTimeZone, const FCulturePtr InTargetCulture)
 	: SourceDateTime(InSourceDateTime)
 	, TimeStyle(InTimeStyle)
 	, TimeZone(InTimeZone)
@@ -559,7 +559,7 @@ FText FTextHistory_AsTime::ToText(bool bInAsSource) const
 ///////////////////////////////////////
 // FTextHistory_AsDateTime
 
-FTextHistory_AsDateTime::FTextHistory_AsDateTime(const FDateTime& InSourceDateTime, const EDateTimeStyle::Type InDateStyle, const EDateTimeStyle::Type InTimeStyle, const FString& InTimeZone, const TSharedPtr<FCulture, ESPMode::ThreadSafe> InTargetCulture)
+FTextHistory_AsDateTime::FTextHistory_AsDateTime(const FDateTime& InSourceDateTime, const EDateTimeStyle::Type InDateStyle, const EDateTimeStyle::Type InTimeStyle, const FString& InTimeZone, const FCulturePtr InTargetCulture)
 	: SourceDateTime(InSourceDateTime)
 	, DateStyle(InDateStyle)
 	, TimeStyle(InTimeStyle)
diff --git a/Engine/Source/Runtime/Core/Private/Internationalization/TextHistory.h b/Engine/Source/Runtime/Core/Private/Internationalization/TextHistory.h
--- a/Engine/Source/Runtime/Core/Private/Internationalization/TextHistory.h
+++ b/Engine/Source/Runtime/Core/Private/Internationalization/TextHistory.h
@@ -132,7 +132,7 @@ class CORE_API FTextHistory_FormatNumber : public FTextHistory
 {
 public:
 	FTextHistory_FormatNumber();
-	FTextHistory_FormatNumber(const FFormatArgumentValue& InSourceValue, const FNumberFormattingOptions* const InFormatOptions, const TSharedPtr<FCulture, ESPMode::ThreadSafe> InTargetCulture);
+	FTextHistory_FormatNumber(const FFormatArgumentValue& InSourceValue, const FNumberFormattingOptions* const InFormatOptions, const FCulturePtr InTargetCulture);
 
 	~FTextHistory_FormatNumber();
 
@@ -145,7 +145,7 @@ protected:
 	/** All the formatting options available to format using. This can be NULL */
 	FNumberFormattingOptions* FormatOptions;
 	/** The culture to format using */
-	TSharedPtr<FCulture, ESPMode::ThreadSafe> TargetCulture;
+	FCulturePtr TargetCulture;
 };
 
 /**  Handles history for formatting using AsNumber */
@@ -153,7 +153,7 @@ class CORE_API FTextHistory_AsNumber : public FTextHistory_FormatNumber
 {
 public:
 	FTextHistory_AsNumber() {};
-	FTextHistory_AsNumber(const FFormatArgumentValue& InSourceValue, const FNumberFormattingOptions* const InFormatOptions, const TSharedPtr<FCulture, ESPMode::ThreadSafe> InTargetCulture);
+	FTextHistory_AsNumber(const FFormatArgumentValue& InSourceValue, const FNumberFormattingOptions* const InFormatOptions, const FCulturePtr InTargetCulture);
 
 	// Begin FTextHistory interface
 	virtual FText ToText(bool bInAsSource) const override;
@@ -166,7 +166,7 @@ class CORE_API FTextHistory_AsPercent : public FTextHistory_FormatNumber
 {
 public:
 	FTextHistory_AsPercent() {};
-	FTextHistory_AsPercent(const FFormatArgumentValue& InSourceValue, const FNumberFormattingOptions* const InFormatOptions, const TSharedPtr<FCulture, ESPMode::ThreadSafe> InTargetCulture);
+	FTextHistory_AsPercent(const FFormatArgumentValue& InSourceValue, const FNumberFormattingOptions* const InFormatOptions, const FCulturePtr InTargetCulture);
 
 	// Begin FTextHistory interface
 	virtual FText ToText(bool bInAsSource) const override;
@@ -179,7 +179,7 @@ class CORE_API FTextHistory_AsCurrency : public FTextHistory_FormatNumber
 {
 public:
 	FTextHistory_AsCurrency() {};
-	FTextHistory_AsCurrency(const FFormatArgumentValue& InSourceValue, const FString& CurrencyCode, const FNumberFormattingOptions* const InFormatOptions, const TSharedPtr<FCulture, ESPMode::ThreadSafe> InTargetCulture);
+	FTextHistory_AsCurrency(const FFormatArgumentValue& InSourceValue, const FString& CurrencyCode, const FNumberFormattingOptions* const InFormatOptions, const FCulturePtr InTargetCulture);
 
 	// Begin FTextHistory interface
 	virtual FText ToText(bool bInAsSource) const override;
@@ -196,7 +196,7 @@ class CORE_API FTextHistory_AsDate : public FTextHistory
 {
 public:
 	FTextHistory_AsDate() {};
-	FTextHistory_AsDate(const FDateTime& InSourceDateTime, const EDateTimeStyle::Type InDateStyle, const TSharedPtr<FCulture, ESPMode::ThreadSafe> InTargetCulture);
+	FTextHistory_AsDate(const FDateTime& InSourceDateTime, const EDateTimeStyle::Type InDateStyle, const FCulturePtr InTargetCulture);
 
 	// Begin FTextHistory interface
 	virtual FText ToText(bool bInAsSource) const override;
@@ -209,7 +209,7 @@ private:
 	/** Style to format the date using */
 	EDateTimeStyle::Type DateStyle;
 	/** Culture to format the date in */
-	TSharedPtr<FCulture, ESPMode::ThreadSafe> TargetCulture;
+	FCulturePtr TargetCulture;
 };
 
 /**  Handles history for formatting using AsTime */
@@ -217,7 +217,7 @@ class CORE_API FTextHistory_AsTime : public FTextHistory
 {
 public:
 	FTextHistory_AsTime() {};
-	FTextHistory_AsTime(const FDateTime& InSourceDateTime, const EDateTimeStyle::Type InTimeStyle, const FString& InTimeZone, const TSharedPtr<FCulture, ESPMode::ThreadSafe> InTargetCulture);
+	FTextHistory_AsTime(const FDateTime& InSourceDateTime, const EDateTimeStyle::Type InTimeStyle, const FString& InTimeZone, const FCulturePtr InTargetCulture);
 
 	// Begin FTextHistory interface
 	virtual FText ToText(bool bInAsSource) const override;
@@ -232,7 +232,7 @@ private:
 	/** Timezone to put the time in */
 	FString TimeZone;
 	/** Culture to format the time in */
-	TSharedPtr<FCulture, ESPMode::ThreadSafe> TargetCulture;
+	FCulturePtr TargetCulture;
 };
 
 /**  Handles history for formatting using AsDateTime */
@@ -240,7 +240,7 @@ class CORE_API FTextHistory_AsDateTime : public FTextHistory
 {
 public:
 	FTextHistory_AsDateTime() {};
-	FTextHistory_AsDateTime(const FDateTime& InSourceDateTime, const EDateTimeStyle::Type InDateStyle, const EDateTimeStyle::Type InTimeStyle, const FString& InTimeZone, const TSharedPtr<FCulture, ESPMode::ThreadSafe> InTargetCulture);
+	FTextHistory_AsDateTime(const FDateTime& InSourceDateTime, const EDateTimeStyle::Type InDateStyle, const EDateTimeStyle::Type InTimeStyle, const FString& InTimeZone, const FCulturePtr InTargetCulture);
 
 	// Begin FTextHistory interface
 	virtual FText ToText(bool bInAsSource) const override;
@@ -257,5 +257,5 @@ private:
 	/** Timezone to put the time in */
 	FString TimeZone;
 	/** Culture to format the time in */
-	TSharedPtr<FCulture, ESPMode::ThreadSafe> TargetCulture;
+	FCulturePtr TargetCulture;
 };
diff --git a/Engine/Source/Runtime/Core/Private/Internationalization/TextLocalizationManager.cpp b/Engine/Source/Runtime/Core/Private/Internationalization/TextLocalizationManager.cpp
--- a/Engine/Source/Runtime/Core/Private/Internationalization/TextLocalizationManager.cpp
+++ b/Engine/Source/Runtime/Core/Private/Internationalization/TextLocalizationManager.cpp
@@ -85,14 +85,14 @@ void EndInitTextLocalization()
 			}
 			LocalizationPaths += FPaths::GetEngineLocalizationPaths();
 
-			TArray< TSharedPtr<FCulture, ESPMode::ThreadSafe> > AvailableCultures;
+			TArray< FCulturePtr > AvailableCultures;
 			I18N.GetCulturesWithAvailableLocalization(LocalizationPaths, AvailableCultures, false);
 
 			// Query for parent culture name based on request culture name.
 			FString ParentCultureName = RequestedCultureName;
 				
 			// If we do not have localization data, try the parent culture.
-			for(TSharedPtr<FCulture, ESPMode::ThreadSafe> ParentCulture = I18N.GetCulture(ParentCultureName); !AvailableCultures.Contains(ParentCulture); ParentCulture = I18N.GetCulture(ParentCultureName))
+			for(FCulturePtr ParentCulture = I18N.GetCulture(ParentCultureName); !AvailableCultures.Contains(ParentCulture); ParentCulture = I18N.GetCulture(ParentCultureName))
 			{
 				ParentCultureName = FCulture::GetParentName(ParentCultureName);
 
diff --git a/Engine/Source/Runtime/Core/Private/Tests/Internationalization/NumberFormatingRulesTest.cpp b/Engine/Source/Runtime/Core/Private/Tests/Internationalization/NumberFormatingRulesTest.cpp
--- a/Engine/Source/Runtime/Core/Private/Tests/Internationalization/NumberFormatingRulesTest.cpp
+++ b/Engine/Source/Runtime/Core/Private/Tests/Internationalization/NumberFormatingRulesTest.cpp
@@ -253,7 +253,7 @@ bool FNumberFormattingRulesTest::RunTest (const FString& Parameters)
 
 #if UE_ENABLE_ICU
 	{
-		const TSharedPtr<FCulture, ESPMode::ThreadSafe>& InvariantCulture = I18N.GetInvariantCulture();
+		const FCulturePtr& InvariantCulture = I18N.GetInvariantCulture();
 		Test(this, TEXT("Convert a Double to a number formatted correct for hi-IN but as invariant"),				FText::AsNumber(DoubleValue, NULL, InvariantCulture),			FText::FromString(TEXT("12345678.901")));
 		Test(this, TEXT("Convert a Float to a number formatted correct for hi-IN but as invariant"),				FText::AsNumber(FloatValue, NULL, InvariantCulture),			FText::FromString(TEXT("1234.567017")));
 		Test(this, TEXT("Convert a Negative Double to a number formatted correct for hi-IN but as invariant"),		FText::AsNumber(DoubleNegativeValue, NULL, InvariantCulture),	FText::FromString(TEXT("-12345678.901")));
diff --git a/Engine/Source/Runtime/Core/Public/Internationalization/Culture.h b/Engine/Source/Runtime/Core/Public/Internationalization/Culture.h
--- a/Engine/Source/Runtime/Core/Public/Internationalization/Culture.h
+++ b/Engine/Source/Runtime/Core/Public/Internationalization/Culture.h
@@ -80,3 +80,5 @@ public:
 	const FNumberFormattingRules NumberFormattingRule;
 #endif
 };
+
+#include "CulturePointer.h"
\ No newline at end of file
diff --git a/Engine/Source/Runtime/Core/Public/Internationalization/CulturePointer.h b/Engine/Source/Runtime/Core/Public/Internationalization/CulturePointer.h
new file mode 100644
--- /dev/null
+++ b/Engine/Source/Runtime/Core/Public/Internationalization/CulturePointer.h
@@ -0,0 +1,9 @@
+// Copyright 1998-2014 Epic Games, Inc. All Rights Reserved.
+#pragma once
+
+#include "SharedPointer.h"
+
+class FCulture;
+
+typedef TSharedPtr<FCulture, ESPMode::ThreadSafe> FCulturePtr;
+typedef TSharedRef<FCulture, ESPMode::ThreadSafe> FCultureRef;
\ No newline at end of file
diff --git a/Engine/Source/Runtime/Core/Public/Internationalization/Internationalization.h b/Engine/Source/Runtime/Core/Public/Internationalization/Internationalization.h
--- a/Engine/Source/Runtime/Core/Public/Internationalization/Internationalization.h
+++ b/Engine/Source/Runtime/Core/Public/Internationalization/Internationalization.h
@@ -42,10 +42,10 @@ public:
 	CORE_API void SetCurrentCulture(const FString& Name);
 
 	//@return the current culture
-	CORE_API TSharedRef<FCulture, ESPMode::ThreadSafe> GetCurrentCulture() const;
+	CORE_API FCultureRef GetCurrentCulture() const;
 
 	//@return culture object by given name, or NULL if not found
-	CORE_API TSharedPtr<FCulture, ESPMode::ThreadSafe> GetCulture(const FString& Name) const;
+	CORE_API FCulturePtr GetCulture(const FString& Name) const;
 
 	//@return the default culture
 	CORE_API TSharedRef< FCulture, ESPMode::ThreadSafe > GetDefaultCulture() const
@@ -69,7 +69,7 @@ public:
 	CORE_API void GetCultureNames(TArray<FString>& CultureNames) const;
 
 	// Given some paths to look at, populate a list of cultures that we have available localization information for. If bIncludeDerivedCultures, include cultures that are derived from those we have localization data for.
-	CORE_API void GetCulturesWithAvailableLocalization(const TArray<FString>& InLocalizationPaths, TArray< TSharedPtr<FCulture, ESPMode::ThreadSafe> >& OutAvailableCultures, const bool bIncludeDerivedCultures) const;
+	CORE_API void GetCulturesWithAvailableLocalization(const TArray<FString>& InLocalizationPaths, TArray< FCulturePtr >& OutAvailableCultures, const bool bIncludeDerivedCultures) const;
 
 private:
 	FInternationalization();
diff --git a/Engine/Source/Runtime/Core/Public/Internationalization/Text.h b/Engine/Source/Runtime/Core/Public/Internationalization/Text.h
--- a/Engine/Source/Runtime/Core/Public/Internationalization/Text.h
+++ b/Engine/Source/Runtime/Core/Public/Internationalization/Text.h
@@ -2,6 +2,7 @@
 #pragma once
 
 #include "SharedPointer.h"
+#include "CulturePointer.h"
 
 #define ENABLE_TEXT_ERROR_CHECKING_RESULTS (UE_BUILD_DEBUG | UE_BUILD_DEVELOPMENT | UE_BUILD_TEST )
 
@@ -123,51 +124,51 @@ public:
 	/**
 	 * Generate an FText that represents the passed number in the current culture
 	 */
-	static FText AsNumber(float Val,	const FNumberFormattingOptions* const Options = NULL, const TSharedPtr<FCulture, ESPMode::ThreadSafe>& TargetCulture = NULL);
-	static FText AsNumber(double Val,	const FNumberFormattingOptions* const Options = NULL, const TSharedPtr<FCulture, ESPMode::ThreadSafe>& TargetCulture = NULL);
-	static FText AsNumber(int8 Val,		const FNumberFormattingOptions* const Options = NULL, const TSharedPtr<FCulture, ESPMode::ThreadSafe>& TargetCulture = NULL);
-	static FText AsNumber(int16 Val,	const FNumberFormattingOptions* const Options = NULL, const TSharedPtr<FCulture, ESPMode::ThreadSafe>& TargetCulture = NULL);
-	static FText AsNumber(int32 Val,	const FNumberFormattingOptions* const Options = NULL, const TSharedPtr<FCulture, ESPMode::ThreadSafe>& TargetCulture = NULL);
-	static FText AsNumber(int64 Val,	const FNumberFormattingOptions* const Options = NULL, const TSharedPtr<FCulture, ESPMode::ThreadSafe>& TargetCulture = NULL);
-	static FText AsNumber(uint8 Val,	const FNumberFormattingOptions* const Options = NULL, const TSharedPtr<FCulture, ESPMode::ThreadSafe>& TargetCulture = NULL);
-	static FText AsNumber(uint16 Val,	const FNumberFormattingOptions* const Options = NULL, const TSharedPtr<FCulture, ESPMode::ThreadSafe>& TargetCulture = NULL);
-	static FText AsNumber(uint32 Val,	const FNumberFormattingOptions* const Options = NULL, const TSharedPtr<FCulture, ESPMode::ThreadSafe>& TargetCulture = NULL);
-	static FText AsNumber(uint64 Val,	const FNumberFormattingOptions* const Options = NULL, const TSharedPtr<FCulture, ESPMode::ThreadSafe>& TargetCulture = NULL);
-	static FText AsNumber(long Val,		const FNumberFormattingOptions* const Options = NULL, const TSharedPtr<FCulture, ESPMode::ThreadSafe>& TargetCulture = NULL);
+	static FText AsNumber(float Val,	const FNumberFormattingOptions* const Options = NULL, const FCulturePtr& TargetCulture = nullptr);
+	static FText AsNumber(double Val,	const FNumberFormattingOptions* const Options = NULL, const FCulturePtr& TargetCulture = nullptr);
+	static FText AsNumber(int8 Val,		const FNumberFormattingOptions* const Options = NULL, const FCulturePtr& TargetCulture = nullptr);
+	static FText AsNumber(int16 Val,	const FNumberFormattingOptions* const Options = NULL, const FCulturePtr& TargetCulture = nullptr);
+	static FText AsNumber(int32 Val,	const FNumberFormattingOptions* const Options = NULL, const FCulturePtr& TargetCulture = nullptr);
+	static FText AsNumber(int64 Val,	const FNumberFormattingOptions* const Options = NULL, const FCulturePtr& TargetCulture = nullptr);
+	static FText AsNumber(uint8 Val,	const FNumberFormattingOptions* const Options = NULL, const FCulturePtr& TargetCulture = nullptr);
+	static FText AsNumber(uint16 Val,	const FNumberFormattingOptions* const Options = NULL, const FCulturePtr& TargetCulture = nullptr);
+	static FText AsNumber(uint32 Val,	const FNumberFormattingOptions* const Options = NULL, const FCulturePtr& TargetCulture = nullptr);
+	static FText AsNumber(uint64 Val,	const FNumberFormattingOptions* const Options = NULL, const FCulturePtr& TargetCulture = nullptr);
+	static FText AsNumber(long Val,		const FNumberFormattingOptions* const Options = NULL, const FCulturePtr& TargetCulture = nullptr);
 
 	/**
 	 * Generate an FText that represents the passed number as currency in the current culture
 	 */
-	static FText AsCurrency(float Val,  const FString& CurrencyCode = FString(), const FNumberFormattingOptions* const Options = NULL, const TSharedPtr<FCulture, ESPMode::ThreadSafe>& TargetCulture = NULL);
-	static FText AsCurrency(double Val, const FString& CurrencyCode = FString(), const FNumberFormattingOptions* const Options = NULL, const TSharedPtr<FCulture, ESPMode::ThreadSafe>& TargetCulture = NULL);
-	static FText AsCurrency(int8 Val,   const FString& CurrencyCode = FString(), const FNumberFormattingOptions* const Options = NULL, const TSharedPtr<FCulture, ESPMode::ThreadSafe>& TargetCulture = NULL);
-	static FText AsCurrency(int16 Val,  const FString& CurrencyCode = FString(), const FNumberFormattingOptions* const Options = NULL, const TSharedPtr<FCulture, ESPMode::ThreadSafe>& TargetCulture = NULL);
-	static FText AsCurrency(int32 Val,  const FString& CurrencyCode = FString(), const FNumberFormattingOptions* const Options = NULL, const TSharedPtr<FCulture, ESPMode::ThreadSafe>& TargetCulture = NULL);
-	static FText AsCurrency(int64 Val,  const FString& CurrencyCode = FString(), const FNumberFormattingOptions* const Options = NULL, const TSharedPtr<FCulture, ESPMode::ThreadSafe>& TargetCulture = NULL);
-	static FText AsCurrency(uint8 Val,  const FString& CurrencyCode = FString(), const FNumberFormattingOptions* const Options = NULL, const TSharedPtr<FCulture, ESPMode::ThreadSafe>& TargetCulture = NULL);
-	static FText AsCurrency(uint16 Val, const FString& CurrencyCode = FString(), const FNumberFormattingOptions* const Options = NULL, const TSharedPtr<FCulture, ESPMode::ThreadSafe>& TargetCulture = NULL);
-	static FText AsCurrency(uint32 Val, const FString& CurrencyCode = FString(), const FNumberFormattingOptions* const Options = NULL, const TSharedPtr<FCulture, ESPMode::ThreadSafe>& TargetCulture = NULL);
-	static FText AsCurrency(uint64 Val, const FString& CurrencyCode = FString(), const FNumberFormattingOptions* const Options = NULL, const TSharedPtr<FCulture, ESPMode::ThreadSafe>& TargetCulture = NULL);
-	static FText AsCurrency(long Val,   const FString& CurrencyCode = FString(), const FNumberFormattingOptions* const Options = NULL, const TSharedPtr<FCulture, ESPMode::ThreadSafe>& TargetCulture = NULL);
+	static FText AsCurrency(float Val,  const FString& CurrencyCode = FString(), const FNumberFormattingOptions* const Options = nullptr, const FCulturePtr& TargetCulture = nullptr);
+	static FText AsCurrency(double Val, const FString& CurrencyCode = FString(), const FNumberFormattingOptions* const Options = nullptr, const FCulturePtr& TargetCulture = nullptr);
+	static FText AsCurrency(int8 Val,   const FString& CurrencyCode = FString(), const FNumberFormattingOptions* const Options = nullptr, const FCulturePtr& TargetCulture = nullptr);
+	static FText AsCurrency(int16 Val,  const FString& CurrencyCode = FString(), const FNumberFormattingOptions* const Options = nullptr, const FCulturePtr& TargetCulture = nullptr);
+	static FText AsCurrency(int32 Val,  const FString& CurrencyCode = FString(), const FNumberFormattingOptions* const Options = nullptr, const FCulturePtr& TargetCulture = nullptr);
+	static FText AsCurrency(int64 Val,  const FString& CurrencyCode = FString(), const FNumberFormattingOptions* const Options = nullptr, const FCulturePtr& TargetCulture = nullptr);
+	static FText AsCurrency(uint8 Val,  const FString& CurrencyCode = FString(), const FNumberFormattingOptions* const Options = nullptr, const FCulturePtr& TargetCulture = nullptr);
+	static FText AsCurrency(uint16 Val, const FString& CurrencyCode = FString(), const FNumberFormattingOptions* const Options = nullptr, const FCulturePtr& TargetCulture = nullptr);
+	static FText AsCurrency(uint32 Val, const FString& CurrencyCode = FString(), const FNumberFormattingOptions* const Options = nullptr, const FCulturePtr& TargetCulture = nullptr);
+	static FText AsCurrency(uint64 Val, const FString& CurrencyCode = FString(), const FNumberFormattingOptions* const Options = nullptr, const FCulturePtr& TargetCulture = nullptr);
+	static FText AsCurrency(long Val,   const FString& CurrencyCode = FString(), const FNumberFormattingOptions* const Options = nullptr, const FCulturePtr& TargetCulture = nullptr);
 
 	/**
 	 * Generate an FText that represents the passed number as a percentage in the current culture
 	 */
-	static FText AsPercent(float Val,	const FNumberFormattingOptions* const Options = NULL, const TSharedPtr<FCulture, ESPMode::ThreadSafe>& TargetCulture = NULL);
-	static FText AsPercent(double Val,	const FNumberFormattingOptions* const Options = NULL, const TSharedPtr<FCulture, ESPMode::ThreadSafe>& TargetCulture = NULL);
+	static FText AsPercent(float Val,	const FNumberFormattingOptions* const Options = NULL, const FCulturePtr& TargetCulture = NULL);
+	static FText AsPercent(double Val,	const FNumberFormattingOptions* const Options = NULL, const FCulturePtr& TargetCulture = NULL);
 
 	/**
 	 * Generate an FText that represents the passed number as a date and/or time in the current culture
 	 */
-	static FText AsDate(const FDateTime& DateTime, const EDateTimeStyle::Type DateStyle = EDateTimeStyle::Default, const TSharedPtr<FCulture, ESPMode::ThreadSafe>& TargetCulture = NULL);
-	static FText AsDateTime(const FDateTime& DateTime, const EDateTimeStyle::Type DateStyle = EDateTimeStyle::Default, const EDateTimeStyle::Type TimeStyle = EDateTimeStyle::Default, const FString& TimeZone = TEXT(""), const TSharedPtr<FCulture, ESPMode::ThreadSafe>& TargetCulture = NULL);
-	static FText AsTime(const FDateTime& DateTime, const EDateTimeStyle::Type TimeStyle = EDateTimeStyle::Default, const FString& TimeZone = TEXT(""), const TSharedPtr<FCulture, ESPMode::ThreadSafe>& TargetCulture = NULL);
-	static FText AsTimespan(const FTimespan& Time, const TSharedPtr<FCulture, ESPMode::ThreadSafe>& TargetCulture = NULL);
+	static FText AsDate(const FDateTime& DateTime, const EDateTimeStyle::Type DateStyle = EDateTimeStyle::Default, const FCulturePtr& TargetCulture = NULL);
+	static FText AsDateTime(const FDateTime& DateTime, const EDateTimeStyle::Type DateStyle = EDateTimeStyle::Default, const EDateTimeStyle::Type TimeStyle = EDateTimeStyle::Default, const FString& TimeZone = TEXT(""), const FCulturePtr& TargetCulture = NULL);
+	static FText AsTime(const FDateTime& DateTime, const EDateTimeStyle::Type TimeStyle = EDateTimeStyle::Default, const FString& TimeZone = TEXT(""), const FCulturePtr& TargetCulture = NULL);
+	static FText AsTimespan(const FTimespan& Time, const FCulturePtr& TargetCulture = NULL);
 
 	/**
 	 * Generate an FText that represents the passed number as a memory size in the current culture
 	 */
-	static FText AsMemory(SIZE_T NumBytes, const FNumberFormattingOptions* const Options = NULL, const TSharedPtr<FCulture, ESPMode::ThreadSafe>& TargetCulture = NULL);
+	static FText AsMemory(SIZE_T NumBytes, const FNumberFormattingOptions* const Options = NULL, const FCulturePtr& TargetCulture = NULL);
 
 	/**
 	 * Attempts to find an existing FText using the representation found in the loc tables for the specified namespace and key
@@ -319,11 +320,11 @@ private:
 
 private:
 	template<typename T1, typename T2>
-	static FText AsNumberTemplate(T1 Val, const FNumberFormattingOptions* const Options = NULL, const TSharedPtr<FCulture, ESPMode::ThreadSafe>& TargetCulture = NULL);
+	static FText AsNumberTemplate(T1 Val, const FNumberFormattingOptions* const Options = NULL, const FCulturePtr& TargetCulture = NULL);
 	template<typename T1, typename T2>
-	static FText AsCurrencyTemplate(T1 Val, const FString& CurrencyCode, const FNumberFormattingOptions* const Options = NULL, const TSharedPtr<FCulture, ESPMode::ThreadSafe>& TargetCulture = NULL);
+	static FText AsCurrencyTemplate(T1 Val, const FString& CurrencyCode, const FNumberFormattingOptions* const Options = NULL, const FCulturePtr& TargetCulture = NULL);
 	template<typename T1, typename T2>
-	static FText AsPercentTemplate(T1 Val, const FNumberFormattingOptions* const Options = NULL, const TSharedPtr<FCulture, ESPMode::ThreadSafe>& TargetCulture = NULL);
+	static FText AsPercentTemplate(T1 Val, const FNumberFormattingOptions* const Options = NULL, const FCulturePtr& TargetCulture = NULL);
 
 private:
 	/** The visible display string for this FText */
