Changed nullptr default parameter values to NULL in Text.h due to errors on some compilers.

diff --git a/Engine/Source/Runtime/Core/Public/Internationalization/Text.h b/Engine/Source/Runtime/Core/Public/Internationalization/Text.h
--- a/Engine/Source/Runtime/Core/Public/Internationalization/Text.h
+++ b/Engine/Source/Runtime/Core/Public/Internationalization/Text.h
@@ -124,32 +124,32 @@ public:
 	/**
 	 * Generate an FText that represents the passed number in the current culture
 	 */
-	static FText AsNumber(float Val,	const FNumberFormattingOptions* const Options = NULL, const FCulturePtr& TargetCulture = nullptr);
-	static FText AsNumber(double Val,	const FNumberFormattingOptions* const Options = NULL, const FCulturePtr& TargetCulture = nullptr);
-	static FText AsNumber(int8 Val,		const FNumberFormattingOptions* const Options = NULL, const FCulturePtr& TargetCulture = nullptr);
-	static FText AsNumber(int16 Val,	const FNumberFormattingOptions* const Options = NULL, const FCulturePtr& TargetCulture = nullptr);
-	static FText AsNumber(int32 Val,	const FNumberFormattingOptions* const Options = NULL, const FCulturePtr& TargetCulture = nullptr);
-	static FText AsNumber(int64 Val,	const FNumberFormattingOptions* const Options = NULL, const FCulturePtr& TargetCulture = nullptr);
-	static FText AsNumber(uint8 Val,	const FNumberFormattingOptions* const Options = NULL, const FCulturePtr& TargetCulture = nullptr);
-	static FText AsNumber(uint16 Val,	const FNumberFormattingOptions* const Options = NULL, const FCulturePtr& TargetCulture = nullptr);
-	static FText AsNumber(uint32 Val,	const FNumberFormattingOptions* const Options = NULL, const FCulturePtr& TargetCulture = nullptr);
-	static FText AsNumber(uint64 Val,	const FNumberFormattingOptions* const Options = NULL, const FCulturePtr& TargetCulture = nullptr);
-	static FText AsNumber(long Val,		const FNumberFormattingOptions* const Options = NULL, const FCulturePtr& TargetCulture = nullptr);
+	static FText AsNumber(float Val,	const FNumberFormattingOptions* const Options = NULL, const FCulturePtr& TargetCulture = NULL);
+	static FText AsNumber(double Val,	const FNumberFormattingOptions* const Options = NULL, const FCulturePtr& TargetCulture = NULL);
+	static FText AsNumber(int8 Val,		const FNumberFormattingOptions* const Options = NULL, const FCulturePtr& TargetCulture = NULL);
+	static FText AsNumber(int16 Val,	const FNumberFormattingOptions* const Options = NULL, const FCulturePtr& TargetCulture = NULL);
+	static FText AsNumber(int32 Val,	const FNumberFormattingOptions* const Options = NULL, const FCulturePtr& TargetCulture = NULL);
+	static FText AsNumber(int64 Val,	const FNumberFormattingOptions* const Options = NULL, const FCulturePtr& TargetCulture = NULL);
+	static FText AsNumber(uint8 Val,	const FNumberFormattingOptions* const Options = NULL, const FCulturePtr& TargetCulture = NULL);
+	static FText AsNumber(uint16 Val,	const FNumberFormattingOptions* const Options = NULL, const FCulturePtr& TargetCulture = NULL);
+	static FText AsNumber(uint32 Val,	const FNumberFormattingOptions* const Options = NULL, const FCulturePtr& TargetCulture = NULL);
+	static FText AsNumber(uint64 Val,	const FNumberFormattingOptions* const Options = NULL, const FCulturePtr& TargetCulture = NULL);
+	static FText AsNumber(long Val,		const FNumberFormattingOptions* const Options = NULL, const FCulturePtr& TargetCulture = NULL);
 
 	/**
 	 * Generate an FText that represents the passed number as currency in the current culture
 	 */
-	static FText AsCurrency(float Val,  const FString& CurrencyCode = FString(), const FNumberFormattingOptions* const Options = nullptr, const FCulturePtr& TargetCulture = nullptr);
-	static FText AsCurrency(double Val, const FString& CurrencyCode = FString(), const FNumberFormattingOptions* const Options = nullptr, const FCulturePtr& TargetCulture = nullptr);
-	static FText AsCurrency(int8 Val,   const FString& CurrencyCode = FString(), const FNumberFormattingOptions* const Options = nullptr, const FCulturePtr& TargetCulture = nullptr);
-	static FText AsCurrency(int16 Val,  const FString& CurrencyCode = FString(), const FNumberFormattingOptions* const Options = nullptr, const FCulturePtr& TargetCulture = nullptr);
-	static FText AsCurrency(int32 Val,  const FString& CurrencyCode = FString(), const FNumberFormattingOptions* const Options = nullptr, const FCulturePtr& TargetCulture = nullptr);
-	static FText AsCurrency(int64 Val,  const FString& CurrencyCode = FString(), const FNumberFormattingOptions* const Options = nullptr, const FCulturePtr& TargetCulture = nullptr);
-	static FText AsCurrency(uint8 Val,  const FString& CurrencyCode = FString(), const FNumberFormattingOptions* const Options = nullptr, const FCulturePtr& TargetCulture = nullptr);
-	static FText AsCurrency(uint16 Val, const FString& CurrencyCode = FString(), const FNumberFormattingOptions* const Options = nullptr, const FCulturePtr& TargetCulture = nullptr);
-	static FText AsCurrency(uint32 Val, const FString& CurrencyCode = FString(), const FNumberFormattingOptions* const Options = nullptr, const FCulturePtr& TargetCulture = nullptr);
-	static FText AsCurrency(uint64 Val, const FString& CurrencyCode = FString(), const FNumberFormattingOptions* const Options = nullptr, const FCulturePtr& TargetCulture = nullptr);
-	static FText AsCurrency(long Val,   const FString& CurrencyCode = FString(), const FNumberFormattingOptions* const Options = nullptr, const FCulturePtr& TargetCulture = nullptr);
+	static FText AsCurrency(float Val,  const FString& CurrencyCode = FString(), const FNumberFormattingOptions* const Options = NULL, const FCulturePtr& TargetCulture = NULL);
+	static FText AsCurrency(double Val, const FString& CurrencyCode = FString(), const FNumberFormattingOptions* const Options = NULL, const FCulturePtr& TargetCulture = NULL);
+	static FText AsCurrency(int8 Val,   const FString& CurrencyCode = FString(), const FNumberFormattingOptions* const Options = NULL, const FCulturePtr& TargetCulture = NULL);
+	static FText AsCurrency(int16 Val,  const FString& CurrencyCode = FString(), const FNumberFormattingOptions* const Options = NULL, const FCulturePtr& TargetCulture = NULL);
+	static FText AsCurrency(int32 Val,  const FString& CurrencyCode = FString(), const FNumberFormattingOptions* const Options = NULL, const FCulturePtr& TargetCulture = NULL);
+	static FText AsCurrency(int64 Val,  const FString& CurrencyCode = FString(), const FNumberFormattingOptions* const Options = NULL, const FCulturePtr& TargetCulture = NULL);
+	static FText AsCurrency(uint8 Val,  const FString& CurrencyCode = FString(), const FNumberFormattingOptions* const Options = NULL, const FCulturePtr& TargetCulture = NULL);
+	static FText AsCurrency(uint16 Val, const FString& CurrencyCode = FString(), const FNumberFormattingOptions* const Options = NULL, const FCulturePtr& TargetCulture = NULL);
+	static FText AsCurrency(uint32 Val, const FString& CurrencyCode = FString(), const FNumberFormattingOptions* const Options = NULL, const FCulturePtr& TargetCulture = NULL);
+	static FText AsCurrency(uint64 Val, const FString& CurrencyCode = FString(), const FNumberFormattingOptions* const Options = NULL, const FCulturePtr& TargetCulture = NULL);
+	static FText AsCurrency(long Val,   const FString& CurrencyCode = FString(), const FNumberFormattingOptions* const Options = NULL, const FCulturePtr& TargetCulture = NULL);
 
 	/**
 	 * Generate an FText that represents the passed number as a percentage in the current culture
