CORE: Potentially unsafe code in FWordWrapper::ProcessLine()

Added simple checks for bounds.

#ttp 340092

diff --git a/Engine/Source/Runtime/Core/Private/Internationalization/WordWrapper.cpp b/Engine/Source/Runtime/Core/Private/Internationalization/WordWrapper.cpp
--- a/Engine/Source/Runtime/Core/Private/Internationalization/WordWrapper.cpp
+++ b/Engine/Source/Runtime/Core/Private/Internationalization/WordWrapper.cpp
@@ -69,7 +69,7 @@ bool FWordWrapper::ProcessLine()
 			--BreakIndex;			
 		}
 
-		while (FText::IsWhitespace(String[BreakIndex - 1]))
+		while (BreakIndex > 0 && FText::IsWhitespace(String[BreakIndex - 1]))
 		{
 			--BreakIndex;
 		}
@@ -81,7 +81,7 @@ bool FWordWrapper::ProcessLine()
 			bHasAddedLine = true;
 		}
 
-		while (FText::IsWhitespace(String[NextStartIndex]))
+		while (NextStartIndex < StringLength && FText::IsWhitespace(String[NextStartIndex]))
 		{
 			++NextStartIndex;
 		}
