Fix compile error.

diff --git a/Engine/Source/Runtime/Core/Private/Math/Transform.cpp b/Engine/Source/Runtime/Core/Private/Math/Transform.cpp
--- a/Engine/Source/Runtime/Core/Private/Math/Transform.cpp
+++ b/Engine/Source/Runtime/Core/Private/Math/Transform.cpp
@@ -160,7 +160,7 @@ FTransform FTransform::GetRelativeTransform(const FTransform& Other) const
 		return FTransform::Identity;
 	}
 
-	FQuat Inverse = Other.Rotation.InverseFast();
+	FQuat Inverse = Other.Rotation.Inverse();
 	Result.Rotation = Inverse*Rotation;
 
 	Result.Translation = (Inverse*(Translation - Other.Translation))*(SafeRecipScale3D);
