fix for logs not being written for iOS
#ue4
#ios

diff --git a/Engine/Source/Runtime/Core/Private/IOS/IOSPlatformOutputDevices.cpp b/Engine/Source/Runtime/Core/Private/IOS/IOSPlatformOutputDevices.cpp
--- a/Engine/Source/Runtime/Core/Private/IOS/IOSPlatformOutputDevices.cpp
+++ b/Engine/Source/Runtime/Core/Private/IOS/IOSPlatformOutputDevices.cpp
@@ -19,7 +19,7 @@ class FOutputDeviceError* FIOSPlatformOutputDevices::GetError()
 	return &Singleton;
 }
 
-class FOutputDevice* FIOSPlatformOutputDevices::GetLog()
+/*class FOutputDevice* FIOSPlatformOutputDevices::GetLog()
 {
 	static FOutputDeviceIOSDebug Singleton;
 
@@ -33,7 +33,7 @@ class FOutputDevice* FIOSPlatformOutputDevices::GetLog()
 		//@todo-mobile: Use a real log-to-disk system instead?
 		return &Singleton;
 	}
-}
+}*/
 
 
 //class FFeedbackContext*				FIOSPlatformOutputDevices::GetWarn()
diff --git a/Engine/Source/Runtime/Core/Public/IOS/IOSPlatformOutputDevices.h b/Engine/Source/Runtime/Core/Public/IOS/IOSPlatformOutputDevices.h
--- a/Engine/Source/Runtime/Core/Public/IOS/IOSPlatformOutputDevices.h
+++ b/Engine/Source/Runtime/Core/Public/IOS/IOSPlatformOutputDevices.h
@@ -8,7 +8,7 @@
 
 struct CORE_API FIOSPlatformOutputDevices : public FGenericPlatformOutputDevices
 {
-	static class FOutputDevice*			GetLog();
+//	static class FOutputDevice*			GetLog();
 	static class FOutputDeviceError*	GetError();
 };
 
