removed the old code instead of just commenting out
#ue4
#ios

diff --git a/Engine/Source/Runtime/Core/Private/IOS/IOSPlatformOutputDevices.cpp b/Engine/Source/Runtime/Core/Private/IOS/IOSPlatformOutputDevices.cpp
--- a/Engine/Source/Runtime/Core/Private/IOS/IOSPlatformOutputDevices.cpp
+++ b/Engine/Source/Runtime/Core/Private/IOS/IOSPlatformOutputDevices.cpp
@@ -19,23 +19,6 @@ class FOutputDeviceError* FIOSPlatformOutputDevices::GetError()
 	return &Singleton;
 }
 
-/*class FOutputDevice* FIOSPlatformOutputDevices::GetLog()
-{
-	static FOutputDeviceIOSDebug Singleton;
-
-	// Don't use the dummy Debug log if the engine is going to add a proper Debug output device anyway.
-	if (!FPlatformProperties::SupportsWindowedMode() || FPlatformMisc::IsDebuggerPresent() || GIsBuildMachine)
-	{
-		return NULL;
-	}
-	else
-	{
-		//@todo-mobile: Use a real log-to-disk system instead?
-		return &Singleton;
-	}
-}*/
-
-
 //class FFeedbackContext*				FIOSPlatformOutputDevices::GetWarn()
 //{
 //	static FOutputDeviceIOSDebug Singleton;
diff --git a/Engine/Source/Runtime/Core/Public/IOS/IOSPlatformOutputDevices.h b/Engine/Source/Runtime/Core/Public/IOS/IOSPlatformOutputDevices.h
--- a/Engine/Source/Runtime/Core/Public/IOS/IOSPlatformOutputDevices.h
+++ b/Engine/Source/Runtime/Core/Public/IOS/IOSPlatformOutputDevices.h
@@ -8,7 +8,6 @@
 
 struct CORE_API FIOSPlatformOutputDevices : public FGenericPlatformOutputDevices
 {
-//	static class FOutputDevice*			GetLog();
 	static class FOutputDeviceError*	GetError();
 };
 
