When displaying the native console log on OS X be sure to dispatch all UI modifications to the Cocoa main-thread - it is unsafe to manipulate an NSView from any other thread.

diff --git a/Engine/Source/Runtime/Core/Private/Mac/MacPlatformOutputDevices.cpp b/Engine/Source/Runtime/Core/Private/Mac/MacPlatformOutputDevices.cpp
--- a/Engine/Source/Runtime/Core/Private/Mac/MacPlatformOutputDevices.cpp
+++ b/Engine/Source/Runtime/Core/Private/Mac/MacPlatformOutputDevices.cpp
@@ -248,6 +248,11 @@ void FOutputDeviceConsoleMac::DestroyConsole()
 	if (ConsoleHandle)
 	{
 		SCOPED_AUTORELEASE_POOL;
+		
+		do
+		{
+			FPlatformMisc::PumpMessages( true );
+		} while(OutstandingTasks);
 
 		SaveToINI();
 		
@@ -311,24 +316,40 @@ void FOutputDeviceConsoleMac::Serialize( const TCHAR* Data, ELogVerbosity::Type
 						*S -= '0';
 					}
 					
-					NSColor* Colors[2];
-					NSString* AttributeKeys[2];
+					NSMutableArray* Colors = [[NSMutableArray alloc] init];
+					NSMutableArray* AttributeKeys = [[NSMutableArray alloc] init];
 					
 					// Get FOREGROUND_INTENSITY and calculate final color
 					CGFloat Intensity = String[3] ? 1.0 : 0.5;
-					Colors[0] = [NSColor colorWithSRGBRed:(String[0] ? 1.0 * Intensity : 0.0) green:(String[1] ? 1.0 * Intensity : 0.0) blue:(String[2] ? 1.0 * Intensity : 0.0) alpha:1.0];
+					[Colors addObject:[NSColor colorWithSRGBRed:(String[0] ? 1.0 * Intensity : 0.0) green:(String[1] ? 1.0 * Intensity : 0.0) blue:(String[2] ? 1.0 * Intensity : 0.0) alpha:1.0]];
 					
 					// Get BACKGROUND_INTENSITY and calculate final color
 					Intensity = String[7] ? 1.0 : 0.5;
-					Colors[1] = [NSColor colorWithSRGBRed:(String[4] ? 1.0 * Intensity : 0.0) green:(String[5] ? 1.0 * Intensity : 0.0) blue:(String[6] ? 1.0 * Intensity : 0.0) alpha:1.0];
+					[Colors addObject:[NSColor colorWithSRGBRed:(String[4] ? 1.0 * Intensity : 0.0) green:(String[5] ? 1.0 * Intensity : 0.0) blue:(String[6] ? 1.0 * Intensity : 0.0) alpha:1.0]];
 					
-					AttributeKeys[0] = NSForegroundColorAttributeName;
-					AttributeKeys[1] = NSBackgroundColorAttributeName;
+					[AttributeKeys addObject:NSForegroundColorAttributeName];
+					[AttributeKeys addObject:NSBackgroundColorAttributeName];
 					
-					if( TextViewTextColor )
-						[TextViewTextColor release];
+					dispatch_block_t Block = ^{
+						if( TextViewTextColor )
+							[TextViewTextColor release];
+						
+						TextViewTextColor = [[NSDictionary alloc] initWithObjects:Colors forKeys:AttributeKeys];
+						
+						[Colors release];
+						[AttributeKeys release];
+						OutstandingTasks--;
+					};
 					
-					TextViewTextColor = [[NSDictionary alloc] initWithObjects:(id *)Colors forKeys:(id *)AttributeKeys count:2];
+					OutstandingTasks++;
+					if([NSThread isMainThread])
+					{
+						Block();
+					}
+					else
+					{
+						dispatch_async(dispatch_get_main_queue(), Block);
+					}
 				}
 			}
 			else
@@ -339,11 +360,25 @@ void FOutputDeviceConsoleMac::Serialize( const TCHAR* Data, ELogVerbosity::Type
 				FCString::Sprintf(OutputString,TEXT("%s%s"),*FOutputDevice::FormatLogLine(Verbosity, Category, Data, GPrintLogTimes),LINE_TERMINATOR);
 
 				CFStringRef CocoaText = FPlatformString::TCHARToCFString(OutputString);
-				NSAttributedString *AttributedString = [[NSAttributedString alloc] initWithString:(NSString*)CocoaText attributes:TextViewTextColor];
-				[[TextView textStorage] appendAttributedString:AttributedString];
-				[TextView scrollRangeToVisible:NSMakeRange([[TextView string] length], 0)];
-				CFRelease(CocoaText);
-				[AttributedString release];
+				
+				dispatch_block_t Block = ^{
+					NSAttributedString *AttributedString = [[NSAttributedString alloc] initWithString:(NSString*)CocoaText attributes:TextViewTextColor];
+					[[TextView textStorage] appendAttributedString:AttributedString];
+					[TextView scrollRangeToVisible:NSMakeRange([[TextView string] length], 0)];
+					[AttributedString release];
+					CFRelease(CocoaText);
+					OutstandingTasks--;
+				};
+				
+				OutstandingTasks++;
+				if([NSThread isMainThread])
+				{
+					Block();
+				}
+				else
+				{
+					dispatch_async(dispatch_get_main_queue(), Block);
+				}
 				
 				if(!MacApplication)
 				{
@@ -369,18 +404,35 @@ void FOutputDeviceConsoleMac::Serialize( const TCHAR* Data, ELogVerbosity::Type
 void FOutputDeviceConsoleMac::SetDefaultTextColor()
 {
 	SCOPED_AUTORELEASE_POOL;
+	FScopeLock ScopeLock( &CriticalSection );
 
-	NSColor* Colors[2];
-	NSString* AttributeKeys[2];
+	NSMutableArray* Colors = [[NSMutableArray alloc] init];
+	NSMutableArray* AttributeKeys = [[NSMutableArray alloc] init];
 	
-	Colors[0] = [NSColor grayColor];
-	Colors[1] = [NSColor blackColor];
+	[Colors addObject:[NSColor grayColor]];
+	[Colors addObject:[NSColor blackColor]];
 	
-	AttributeKeys[0] = NSForegroundColorAttributeName;
-	AttributeKeys[1] = NSBackgroundColorAttributeName;
+	[AttributeKeys addObject:NSForegroundColorAttributeName];
+	[AttributeKeys addObject:NSBackgroundColorAttributeName];
 
-	if( TextViewTextColor )
-		[TextViewTextColor release];
+	dispatch_block_t Block = ^{
+		if( TextViewTextColor )
+			[TextViewTextColor release];
+		
+		TextViewTextColor = [[NSDictionary alloc] initWithObjects:Colors forKeys:AttributeKeys];
+		
+		[Colors release];
+		[AttributeKeys release];
+		OutstandingTasks--;
+	};
 	
-	TextViewTextColor = [[NSDictionary alloc] initWithObjects:(id *)Colors forKeys:(id *)AttributeKeys count:2];
+	OutstandingTasks++;
+	if([NSThread isMainThread])
+	{
+		Block();
+	}
+	else
+	{
+		dispatch_async(dispatch_get_main_queue(), Block);
+	}
 }
diff --git a/Engine/Source/Runtime/Core/Private/Mac/MacPlatformOutputDevicesPrivate.h b/Engine/Source/Runtime/Core/Private/Mac/MacPlatformOutputDevicesPrivate.h
--- a/Engine/Source/Runtime/Core/Private/Mac/MacPlatformOutputDevicesPrivate.h
+++ b/Engine/Source/Runtime/Core/Private/Mac/MacPlatformOutputDevicesPrivate.h
@@ -50,6 +50,9 @@ private:
 	
 	/** Critical section for Serialize() */
 	FCriticalSection CriticalSection;
+	
+	/** Number of outstanding dispatched Cocoa tasks */
+	uint64 OutstandingTasks;
 		
 	/**
 	 * Saves the console window's position and size to the game .ini
