Merging over John Barrett's memory usage api's

diff --git a/Engine/Source/Runtime/Core/Private/GenericPlatform/GenericPlatformProcess.cpp b/Engine/Source/Runtime/Core/Private/GenericPlatform/GenericPlatformProcess.cpp
--- a/Engine/Source/Runtime/Core/Private/GenericPlatform/GenericPlatformProcess.cpp
+++ b/Engine/Source/Runtime/Core/Private/GenericPlatform/GenericPlatformProcess.cpp
@@ -181,6 +181,12 @@ bool FGenericPlatformProcess::GetProcReturnCode( FProcHandle & ProcHandle, int32
 	return false;
 }
 
+bool FGenericPlatformProcess::GetApplicationMemoryUsage(uint32 ProcessId, SIZE_T* OutMemoryUsage)
+{
+	UE_LOG(LogHAL, Fatal, TEXT("FGenericPlatformProcess::GetApplicationMemoryUsage: not implemented on this platform"));
+	return false;
+}
+
 bool FGenericPlatformProcess::IsApplicationRunning( uint32 ProcessId )
 {
 	UE_LOG(LogHAL, Fatal, TEXT("FGenericPlatformProcess::IsApplicationRunning not implemented on this platform"));
diff --git a/Engine/Source/Runtime/Core/Private/Windows/WindowsPlatformProcess.cpp b/Engine/Source/Runtime/Core/Private/Windows/WindowsPlatformProcess.cpp
--- a/Engine/Source/Runtime/Core/Private/Windows/WindowsPlatformProcess.cpp
+++ b/Engine/Source/Runtime/Core/Private/Windows/WindowsPlatformProcess.cpp
@@ -13,10 +13,7 @@
 	#include <shlobj.h>
 	#include <LM.h>
 	#include <tlhelp32.h>
-
-#if WINVER == 0x0502
 	#include <Psapi.h>
-#endif
 
 	namespace ProcessConstants
 	{
@@ -29,9 +26,7 @@
 #include "HideWindowsPlatformTypes.h"
 #include "WindowsPlatformMisc.h"
 
-#if WINVER == 0x0502
-	#pragma comment(lib, "psapi.lib")
-#endif
+#pragma comment(lib, "psapi.lib")
 
 // static variables
 TArray<FString> FWindowsPlatformProcess::DllDirectoryStack;
@@ -377,6 +372,27 @@ bool FWindowsPlatformProcess::GetProcReturnCode( FProcHandle & ProcHandle, int32
 	return ::GetExitCodeProcess( ProcHandle.Get(), (::DWORD *)ReturnCode ) && *((uint32*)ReturnCode) != ProcessConstants::WIN_STILL_ACTIVE;
 }
 
+bool FWindowsPlatformProcess::GetApplicationMemoryUsage(uint32 ProcessId, SIZE_T* OutMemoryUsage)
+{
+	bool bSuccess = false;
+	HANDLE ProcessHandle = OpenProcess(PROCESS_QUERY_INFORMATION | PROCESS_VM_READ, false, ProcessId);
+
+	if (ProcessHandle != NULL)
+	{
+		PROCESS_MEMORY_COUNTERS_EX MemoryInfo;
+
+		if (GetProcessMemoryInfo(ProcessHandle, (PROCESS_MEMORY_COUNTERS*)&MemoryInfo, sizeof(MemoryInfo)))
+		{
+			*OutMemoryUsage = MemoryInfo.PrivateUsage;
+			bSuccess = true;
+		}
+
+		::CloseHandle(ProcessHandle);
+	}
+
+	return bSuccess;
+}
+
 bool FWindowsPlatformProcess::IsApplicationRunning( uint32 ProcessId )
 {
 	bool bApplicationRunning = true;
diff --git a/Engine/Source/Runtime/Core/Public/GenericPlatform/GenericPlatformProcess.h b/Engine/Source/Runtime/Core/Public/GenericPlatform/GenericPlatformProcess.h
--- a/Engine/Source/Runtime/Core/Public/GenericPlatform/GenericPlatformProcess.h
+++ b/Engine/Source/Runtime/Core/Public/GenericPlatform/GenericPlatformProcess.h
@@ -322,6 +322,8 @@ struct CORE_API FGenericPlatformProcess
 	static bool IsApplicationRunning( const TCHAR* ProcName );
 	/** Returns the Name of process given by the PID.  Returns Empty string "" if PID not found. */
 	static FString GetApplicationName( uint32 ProcessId );
+	/** Outputs the virtual memory usage, of the process with the specified PID */
+	static bool GetApplicationMemoryUsage(uint32 ProcessId, SIZE_T* OutMemoryUsage);
 	/** Returns true if the specified application has a visible window, and that window is active/has focus/is selected */
 	static bool IsThisApplicationForeground();	
 
diff --git a/Engine/Source/Runtime/Core/Public/Windows/WindowsPlatformProcess.h b/Engine/Source/Runtime/Core/Public/Windows/WindowsPlatformProcess.h
--- a/Engine/Source/Runtime/Core/Public/Windows/WindowsPlatformProcess.h
+++ b/Engine/Source/Runtime/Core/Public/Windows/WindowsPlatformProcess.h
@@ -93,6 +93,7 @@ struct CORE_API FWindowsPlatformProcess : public FGenericPlatformProcess
 	static void WaitForProc( FProcHandle & ProcessHandle );
 	static void TerminateProc( FProcHandle & ProcessHandle, bool KillTree = false );
 	static bool GetProcReturnCode( FProcHandle & ProcHandle, int32* ReturnCode );
+	static bool GetApplicationMemoryUsage(uint32 ProcessId, SIZE_T* OutMemoryUsage);
 	static bool IsApplicationRunning( uint32 ProcessId );
 	static bool IsApplicationRunning( const TCHAR* ProcName );
 	static FString GetApplicationName( uint32 ProcessId );	
