Fixed bas use of const literal string  in WindowsPlatformProcess (clang warning)

diff --git a/Engine/Source/Runtime/Core/Private/Windows/WindowsPlatformProcess.cpp b/Engine/Source/Runtime/Core/Private/Windows/WindowsPlatformProcess.cpp
--- a/Engine/Source/Runtime/Core/Private/Windows/WindowsPlatformProcess.cpp
+++ b/Engine/Source/Runtime/Core/Private/Windows/WindowsPlatformProcess.cpp
@@ -798,7 +798,7 @@ const TCHAR* FWindowsPlatformProcess::GetBinariesSubdirectory()
 
 void FWindowsPlatformProcess::LaunchFileInDefaultExternalApplication( const TCHAR* FileName, const TCHAR* Parms /*= NULL*/, ELaunchVerb::Type Verb /*= ELaunchVerb::Open*/ )
 {
-	TCHAR* VerbString = Verb == ELaunchVerb::Edit ? TEXT("edit") : TEXT("open");
+	const TCHAR* VerbString = Verb == ELaunchVerb::Edit ? TEXT("edit") : TEXT("open");
 
 	// First attempt to open the file in its default application
 	UE_LOG(LogWindows, Log,  TEXT("LaunchFileInExternalEditor %s %s"), FileName, Parms ? Parms : TEXT("") );
