WindowsTextInputMethodSystem: Fixed member init order issues and missing virtual

diff --git a/Engine/Source/Runtime/Core/Public/Windows/WindowsTextInputMethodSystem.h b/Engine/Source/Runtime/Core/Public/Windows/WindowsTextInputMethodSystem.h
--- a/Engine/Source/Runtime/Core/Public/Windows/WindowsTextInputMethodSystem.h
+++ b/Engine/Source/Runtime/Core/Public/Windows/WindowsTextInputMethodSystem.h
@@ -16,10 +16,14 @@ class FTSFActivationProxy : public ITfInputProcessorProfileActivationSink, publi
 {
 public:
 	FTSFActivationProxy(FWindowsTextInputMethodSystem* InOwner) 
-		: Owner(InOwner)
-		, ReferenceCount(1)
-		, TSFProfileCookie(TF_INVALID_COOKIE)
+		: TSFProfileCookie(TF_INVALID_COOKIE)
 		, TSFLanguageCookie(TF_INVALID_COOKIE)
+		, Owner(InOwner)
+		, ReferenceCount(1)
+	{
+	}
+
+	virtual ~FTSFActivationProxy()
 	{
 	}
 
@@ -53,6 +57,8 @@ class FWindowsTextInputMethodSystem : public ITextInputMethodSystem
 	friend class FTSFActivationProxy;
 
 public:
+	virtual ~FWindowsTextInputMethodSystem() {}
+
 	bool Initialize();
 	void Terminate();
 
