Workaround for Clang compile error on Windows
- System headers intsafe.h and strsafe.h will redeclare INT/DWORD/etc on their own
- Clang doesn't allow typedefs with global scope operators, but we have redefined these to ::INT/::DWORD/etc
- We now include these headers very early on to avoid having to disable our DoNotUseOldUE4Type trick under Clang

diff --git a/Engine/Source/Runtime/Core/Public/Windows/WindowsPlatformIncludes.h b/Engine/Source/Runtime/Core/Public/Windows/WindowsPlatformIncludes.h
--- a/Engine/Source/Runtime/Core/Public/Windows/WindowsPlatformIncludes.h
+++ b/Engine/Source/Runtime/Core/Public/Windows/WindowsPlatformIncludes.h
@@ -28,6 +28,15 @@ extern "C" CORE_API HINSTANCE hInstance;
 #include <stdint.h>
 #include <tchar.h>
 
+// @todo clang: When compiling under Clang for Windows, these headers cause us particular problems.  We need to make sure they're included
+// before we pull in our 'DoNotUseOldUE4Type' namespace.  This is because these headers will redeclare various numeric typedefs, but under
+// the Clang compiler it is not allowed to define a typedef with a global scope operator in it (such as ::INT).  So we'll get these headers
+// included early on to avoid compiler errors with that.
+#if defined(__clang__)
+	#include <intsafe.h>
+	#include <strsafe.h>
+#endif
+
 #if USING_CODE_ANALYSIS
 // Source annotation support
 #include <CodeAnalysis/SourceAnnotations.h>
