Workaround for Clang bug on Windows with non-type template arguments
- Clang for Windows triggers bogus compile errors with some non-type template arguments

diff --git a/Engine/Source/Runtime/Core/Public/GenericPlatform/GenericPlatformProcess.h b/Engine/Source/Runtime/Core/Public/GenericPlatform/GenericPlatformProcess.h
--- a/Engine/Source/Runtime/Core/Public/GenericPlatform/GenericPlatformProcess.h
+++ b/Engine/Source/Runtime/Core/Public/GenericPlatform/GenericPlatformProcess.h
@@ -34,7 +34,12 @@ namespace ELaunchVerb
 }
 
 /** Generic implementation for the process handle. */
+#if PLATFORM_WINDOWS && defined(__clang__)	// @todo clang: Clang on Windows has a bug with non-type template arguments
+static void* InvalidHandleValue = nullptr;
+template< typename T >
+#else
 template< typename T, T InvalidHandleValue >
+#endif
 struct TProcHandle
 {
 	typedef T HandleType;
diff --git a/Engine/Source/Runtime/Core/Public/Windows/WindowsPlatformProcess.h b/Engine/Source/Runtime/Core/Public/Windows/WindowsPlatformProcess.h
--- a/Engine/Source/Runtime/Core/Public/Windows/WindowsPlatformProcess.h
+++ b/Engine/Source/Runtime/Core/Public/Windows/WindowsPlatformProcess.h
@@ -8,7 +8,11 @@
 #pragma once
 
 /** Windows implementation of the process handle. */
+#if PLATFORM_WINDOWS && defined(__clang__)	// @todo clang: Clang on Windows has a bug with non-type template arguments
+struct FProcHandle : public TProcHandle<HANDLE>
+#else
 struct FProcHandle : public TProcHandle<HANDLE, nullptr>
+#endif
 {
 public:
 	/** Default constructor. */
