Update version numbers for the 4.4.1 hotfix.

diff --git a/Engine/Source/Runtime/Core/Public/Modules/ModuleVersion.h b/Engine/Source/Runtime/Core/Public/Modules/ModuleVersion.h
--- a/Engine/Source/Runtime/Core/Public/Modules/ModuleVersion.h
+++ b/Engine/Source/Runtime/Core/Public/Modules/ModuleVersion.h
@@ -6,7 +6,11 @@
 
 // This number identifies a particular API revision, and is used to determine module compatibility. Hotfixes should retain the API version of the original release.
 // This define is parsed by the build tools, and should be a number or BUILT_FROM_CHANGELIST.
+#if BUILT_FROM_CHANGELIST > 0
 #define MODULE_API_VERSION BUILT_FROM_CHANGELIST
+#else
+#define MODULE_API_VERSION 0
+#endif
 
 // Check that the API version has been set manually for a hotfix release
 #if ENGINE_PATCH_VERSION > 0 && BUILT_FROM_CHANGELIST > 0 && MODULE_API_VERSION == BUILT_FROM_CHANGELIST
