Clang for Windows improvements
- UnrealHeaderTool now compiles and runs(*) under Clang on Windows
- Minor cleanups to Clang-related code in WindowsPlatform.h

(* Crashes at startup unless run under a debugger for some reason)

diff --git a/Engine/Source/Runtime/Core/Public/Windows/WIndowsPlatform.h b/Engine/Source/Runtime/Core/Public/Windows/WIndowsPlatform.h
--- a/Engine/Source/Runtime/Core/Public/Windows/WIndowsPlatform.h
+++ b/Engine/Source/Runtime/Core/Public/Windows/WIndowsPlatform.h
@@ -6,9 +6,6 @@
 
 #pragma once
 
-// Setup Clang compilation support on Windows
-#define PLATFORM_WINDOWS_CLANG defined(__clang__)
-
 /**
 * Windows specific types
 **/
@@ -38,7 +35,7 @@ typedef FWindowsPlatformTypes FPlatformTypes;
 // Base defines, defaults are commented out
 
 #define PLATFORM_LITTLE_ENDIAN								1
-#if PLATFORM_WINDOWS_CLANG
+#if defined(__clang__)
 	// @todo Clang Clang emits a compile error on Windows when __try/__except blocks are used (missing functionality in Clang.)
 	//             After Clang adds proper support for SEH exceptions, we should delete this definition and all use cases of it!
 	#define PLATFORM_SEH_EXCEPTIONS_DISABLED				1
@@ -50,7 +47,7 @@ typedef FWindowsPlatformTypes FPlatformTypes;
 #endif
 
 #define PLATFORM_SUPPORTS_PRAGMA_PACK						1
-#if PLATFORM_WINDOWS_CLANG
+#if defined(__clang__)
 	#define PLATFORM_ENABLE_VECTORINTRINSICS				0
 #else
 	#define PLATFORM_ENABLE_VECTORINTRINSICS				1
@@ -91,36 +88,38 @@ typedef FWindowsPlatformTypes FPlatformTypes;
 #define FORCENOINLINE __declspec(noinline)	/* Force code to NOT be inline */
 
 // Hints compiler that expression is true; generally restricted to comparisons against constants
-#if !PLATFORM_WINDOWS_CLANG		// Clang doesn't support __assume (Microsoft specific)
+#if !defined(__clang__)		// Clang doesn't support __assume (Microsoft specific)
 	#define ASSUME(expr) __assume(expr)
 #endif
 
 #define DECLARE_UINT64(x)	x
 
 // Optimization macros (uses __pragma to enable inside a #define).
-#if !PLATFORM_WINDOWS_CLANG		// @todo clang: Clang doesn't appear to support optimization pragmas yet
+#if !defined(__clang__)		// @todo clang: Clang doesn't appear to support optimization pragmas yet
 	#define PRAGMA_DISABLE_OPTIMIZATION_ACTUAL __pragma(optimize("",off))
 	#define PRAGMA_ENABLE_OPTIMIZATION_ACTUAL  __pragma(optimize("",on))
 #endif
 
 // Backwater of the spec. All compilers support this except microsoft, and they will soon
-#if !PLATFORM_WINDOWS_CLANG		// Clang expects typename outside template
+#if !defined(__clang__)		// Clang expects typename outside template
 	#define TYPENAME_OUTSIDE_TEMPLATE
 #endif
 
 #pragma warning(disable : 4481) // nonstandard extension used: override specifier 'override'
-#if PLATFORM_WINDOWS_CLANG
+
+#if defined(__clang__)
 	#define CONSTEXPR constexpr
+#else
+	#define CONSTEXPR
 #endif
 #define ABSTRACT abstract
-#define CONSTEXPR
 
 // Strings.
 #define LINE_TERMINATOR TEXT("\r\n")
 #define LINE_TERMINATOR_ANSI "\r\n"
 
 // Alignment.
-#if PLATFORM_WINDOWS_CLANG
+#if defined(__clang__)
 	#define GCC_PACK(n) __attribute__((packed,aligned(n)))
 	#define GCC_ALIGN(n) __attribute__((aligned(n)))
 #else
