#HTML5 Fix up usage of CONSTEXPR for win32 simulator build.

diff --git a/Engine/Source/Runtime/Core/Public/HTML5/HTML5Platform.h b/Engine/Source/Runtime/Core/Public/HTML5/HTML5Platform.h
--- a/Engine/Source/Runtime/Core/Public/HTML5/HTML5Platform.h
+++ b/Engine/Source/Runtime/Core/Public/HTML5/HTML5Platform.h
@@ -56,6 +56,7 @@ typedef FHTML5Types FPlatformTypes;
 #define FORCEINLINE _forceinline
 #define FORCENOINLINE __declspec(noinline)	/* Force code to NOT be inline */
 #define PLATFORM_COMPILER_HAS_EXPLICIT_OPERATORS 0
+#define CONSTEXPR    
 #else
 #define FORCEINLINE		inline __attribute__((__always_inline__))		/* Force code to be inline */
 #define FORCENOINLINE	__attribute__((noinline))			/* Force code to NOT be inline */
