Bug with Fwindow windows QueryInterface function returning incorrect type.

diff --git a/Engine/Source/Runtime/Core/Private/Windows/WindowsWindow.cpp b/Engine/Source/Runtime/Core/Private/Windows/WindowsWindow.cpp
--- a/Engine/Source/Runtime/Core/Private/Windows/WindowsWindow.cpp
+++ b/Engine/Source/Runtime/Core/Private/Windows/WindowsWindow.cpp
@@ -665,7 +665,7 @@ HRESULT STDCALL FWindowsWindow::QueryInterface( REFIID iid, void ** ppvObject )
 	if ( IID_IDropTarget == iid || IID_IUnknown == iid )
 	{
 		AddRef();
-		*ppvObject = this;
+		*ppvObject = (IDropTarget*)(this);
 		return S_OK;
 	}
 	else
