Move the Cocoa window & text-input classes into their own files to clarify whether they are primarily Cocoa-facing or UE4-facing in preparation for them to execute on separate threads.

diff --git a/Engine/Source/Developer/StandaloneRenderer/Private/Mac/OpenGL/SlateOpenGLContext.cpp b/Engine/Source/Developer/StandaloneRenderer/Private/Mac/OpenGL/SlateOpenGLContext.cpp
--- a/Engine/Source/Developer/StandaloneRenderer/Private/Mac/OpenGL/SlateOpenGLContext.cpp
+++ b/Engine/Source/Developer/StandaloneRenderer/Private/Mac/OpenGL/SlateOpenGLContext.cpp
@@ -9,6 +9,7 @@
 
 #include "MacWindow.h"
 #include "MacTextInputMethodSystem.h"
+#include "CocoaTextView.h"
 
 char const* const CompositedBlitVertexShader = "#version 120\n"
 "#extension GL_EXT_gpu_shader4 : require\n"
@@ -40,7 +41,7 @@ char const* const CompositedBlitFragmentShader = "#version 120\n"
 /**
  * Custom view class used to provide window OpenGL-enabled content view.
  */
-@interface FSlateCocoaView : FSlateTextView
+@interface FSlateCocoaView : FCocoaTextView
 {
 	@private
 
diff --git a/Engine/Source/Developer/StandaloneRenderer/Private/Mac/OpenGL/SlateOpenGLViewport.cpp b/Engine/Source/Developer/StandaloneRenderer/Private/Mac/OpenGL/SlateOpenGLViewport.cpp
--- a/Engine/Source/Developer/StandaloneRenderer/Private/Mac/OpenGL/SlateOpenGLViewport.cpp
+++ b/Engine/Source/Developer/StandaloneRenderer/Private/Mac/OpenGL/SlateOpenGLViewport.cpp
@@ -129,7 +129,7 @@ void FSlateOpenGLViewport::SwapBuffers()
 	}
 	
 	NSWindow* Window = [[RenderingContext.Context view] window];
-	FSlateCocoaWindow* SlateCocoaWindow = [Window isKindOfClass:[FSlateCocoaWindow class]] ? (FSlateCocoaWindow*)Window : nil;
+	FCocoaWindow* SlateCocoaWindow = [Window isKindOfClass:[FCocoaWindow class]] ? (FCocoaWindow*)Window : nil;
 	if(SlateCocoaWindow)
 	{
 		[SlateCocoaWindow performDeferredOrderFront];
diff --git a/Engine/Source/Runtime/Core/Private/Mac/CocoaTextView.cpp b/Engine/Source/Runtime/Core/Private/Mac/CocoaTextView.cpp
new file mode 100644
--- /dev/null
+++ b/Engine/Source/Runtime/Core/Private/Mac/CocoaTextView.cpp
@@ -0,0 +1,441 @@
+// Copyright 1998-2014 Epic Games, Inc. All Rights Reserved.
+
+#include "Core.h"
+#include "CocoaTextView.h"
+#include "CocoaWindow.h"
+
+@implementation FCocoaTextView
+
+- (id)initWithFrame:(NSRect)frame
+{
+    self = [super initWithFrame:frame];
+    if (self)
+	{
+        // Initialization code here.
+		IMMContext = nil;
+		markedRange = {NSNotFound, 0};
+		reallyHandledEvent = false;
+    }
+    return self;
+}
+
+- (void)drawRect:(NSRect)dirtyRect
+{
+	[super drawRect:dirtyRect];
+	
+    // Drawing code here.
+}
+
+- (bool)imkKeyDown:(NSEvent *)theEvent
+{
+	if (IMMContext.IsValid())
+	{
+		reallyHandledEvent = true;
+		return [[self inputContext] handleEvent:theEvent] && reallyHandledEvent;
+	}
+	else
+	{
+		return false;
+	}
+}
+
+/** Forward mouse events up to the window rather than through the responder chain - thus avoiding
+ *	the hidden titlebar controls. Normal windows just use the responder chain as usual.
+ */
+- (BOOL)acceptsFirstMouse:(NSEvent *)Event
+{
+	return YES;
+}
+
+- (void)mouseDown:(NSEvent *)theEvent
+{
+	if (IMMContext.IsValid())
+	{
+		[[self inputContext] handleEvent:theEvent];
+	}
+	
+	FCocoaWindow* CocoaWindow = [[self window] isKindOfClass:[FCocoaWindow class]] ? (FCocoaWindow*)[self window] : nil;
+	if (CocoaWindow)
+	{
+		[CocoaWindow mouseDown:theEvent];
+	}
+}
+
+- (void)mouseDragged:(NSEvent *)theEvent
+{
+	if (IMMContext.IsValid())
+	{
+		[[self inputContext] handleEvent:theEvent];
+	}
+}
+
+- (void)mouseUp:(NSEvent *)theEvent
+{
+	if (IMMContext.IsValid())
+	{
+		[[self inputContext] handleEvent:theEvent];
+	}
+	
+	FCocoaWindow* CocoaWindow = [[self window] isKindOfClass:[FCocoaWindow class]] ? (FCocoaWindow*)[self window] : nil;
+	if (CocoaWindow)
+	{
+		[CocoaWindow mouseUp:theEvent];
+	}
+}
+
+- (void)rightMouseDown:(NSEvent*)Event
+{
+	FCocoaWindow* CocoaWindow = [[self window] isKindOfClass:[FCocoaWindow class]] ? (FCocoaWindow*)[self window] : nil;
+	if (CocoaWindow)
+	{
+		[CocoaWindow rightMouseDown:Event];
+	}
+	else
+	{
+		[super rightMouseDown:Event];
+	}
+}
+
+- (void)otherMouseDown:(NSEvent*)Event
+{
+	FCocoaWindow* CocoaWindow = [[self window] isKindOfClass:[FCocoaWindow class]] ? (FCocoaWindow*)[self window] : nil;
+	if (CocoaWindow)
+	{
+		[CocoaWindow otherMouseDown:Event];
+	}
+	else
+	{
+		[super otherMouseDown:Event];
+	}
+}
+
+- (void)rightMouseUp:(NSEvent*)Event
+{
+	FCocoaWindow* CocoaWindow = [[self window] isKindOfClass:[FCocoaWindow class]] ? (FCocoaWindow*)[self window] : nil;
+	if (CocoaWindow)
+	{
+		[CocoaWindow rightMouseUp:Event];
+	}
+	else
+	{
+		[super rightMouseUp:Event];
+	}
+}
+
+- (void)otherMouseUp:(NSEvent*)Event
+{
+	FCocoaWindow* CocoaWindow = [[self window] isKindOfClass:[FCocoaWindow class]] ? (FCocoaWindow*)[self window] : nil;
+	if (CocoaWindow)
+	{
+		[CocoaWindow otherMouseUp:Event];
+	}
+	else
+	{
+		[super otherMouseUp:Event];
+	}
+}
+
+- (void)activateInputMethod:(const TSharedRef<ITextInputMethodContext>&)InContext
+{
+	if (IMMContext.IsValid())
+	{
+		[self unmarkText];
+		[[self inputContext] deactivate];
+		[[self inputContext] discardMarkedText];
+	}
+	IMMContext = InContext;
+	[[self inputContext] activate];
+}
+
+- (void)deactivateInputMethod
+{
+	[self unmarkText];
+	IMMContext = NULL;
+	[[self inputContext] deactivate];
+	[[self inputContext] discardMarkedText];
+}
+
+//@protocol NSTextInputClient
+//@required
+/* The receiver inserts aString replacing the content specified by replacementRange. aString can be either an NSString or NSAttributedString instance.
+ */
+- (void)insertText:(id)aString replacementRange:(NSRange)replacementRange
+{
+	if (IMMContext.IsValid() && [self hasMarkedText])
+	{
+		uint32 SelectionLocation, SelectionLength;
+		if (replacementRange.location == NSNotFound)
+		{
+			if (markedRange.location != NSNotFound)
+			{
+				SelectionLocation = markedRange.location;
+				SelectionLength = markedRange.length;
+			}
+			else
+			{
+				ITextInputMethodContext::ECaretPosition CaretPosition;
+				IMMContext->GetSelectionRange(SelectionLocation, SelectionLength, CaretPosition);
+			}
+		}
+		else
+		{
+			SelectionLocation = replacementRange.location;
+			SelectionLength = replacementRange.length;
+		}
+		
+		NSString* TheString;
+		if ([aString isKindOfClass:[NSAttributedString class]])
+		{
+			TheString = [(NSAttributedString*)aString string];
+		}
+		else
+		{
+			TheString = (NSString*)aString;
+		}
+		
+		FString TheFString(TheString);
+		IMMContext->SetTextInRange(SelectionLocation, SelectionLength, TheFString);
+		IMMContext->SetSelectionRange(SelectionLocation+TheFString.Len(), 0, ITextInputMethodContext::ECaretPosition::Ending);
+		
+		[self unmarkText];
+		[[self inputContext] invalidateCharacterCoordinates]; // recentering
+	}
+	else
+	{
+		reallyHandledEvent = false;
+	}
+}
+
+/* The receiver invokes the action specified by aSelector.
+ */
+- (void)doCommandBySelector:(SEL)aSelector
+{
+	reallyHandledEvent = false;
+}
+
+/* The receiver inserts aString replacing the content specified by replacementRange. aString can be either an NSString or NSAttributedString instance. selectedRange specifies the selection inside the string being inserted; hence, the location is relative to the beginning of aString. When aString is an NSString, the receiver is expected to render the marked text with distinguishing appearance (i.e. NSTextView renders with -markedTextAttributes).
+ */
+- (void)setMarkedText:(id)aString selectedRange:(NSRange)selectedRange replacementRange:(NSRange)replacementRange
+{
+	if (IMMContext.IsValid())
+	{
+		uint32 SelectionLocation, SelectionLength;
+		if (replacementRange.location == NSNotFound)
+		{
+			if (markedRange.location != NSNotFound)
+			{
+				SelectionLocation = markedRange.location;
+				SelectionLength = markedRange.length;
+			}
+			else
+			{
+				ITextInputMethodContext::ECaretPosition CaretPosition;
+				IMMContext->GetSelectionRange(SelectionLocation, SelectionLength, CaretPosition);
+			}
+		}
+		else
+		{
+			SelectionLocation = replacementRange.location;
+			SelectionLength = replacementRange.length;
+		}
+		
+		if ([aString length] == 0)
+		{
+			IMMContext->SetTextInRange(SelectionLocation, SelectionLength, FString());
+			[self unmarkText];
+		}
+		else
+		{
+			if(markedRange.location == NSNotFound)
+			{
+				IMMContext->BeginComposition();
+			}
+			markedRange = NSMakeRange(SelectionLocation, [aString length]);
+			
+			__block NSRange CompositionRange = markedRange;
+			
+			NSString* TheString;
+			if ([aString isKindOfClass:[NSAttributedString class]])
+			{
+				// While the whole string is being composed NSUnderlineStyleAttributeName is 1 to show a single line below the whole string.
+				// When using the pop-up glyph selection window in some IME's the NSAttributedString is broken up into separate glyph ranges, each with its own set of attributes.
+				// Each range specifies NSMarkedClauseSegment, incrementing the NSNumber value from 0 as well as NSUnderlineStyleAttributeName, which makes the underlining show the different ranges.
+				// The subrange being edited by the pop-up glyph selection window will set NSUnderlineStyleAttributeName to a value >1, while all other ranges will be set NSUnderlineStyleAttributeName to 1.
+				NSAttributedString* AttributedString = (NSAttributedString*)aString;
+				[AttributedString enumerateAttribute:NSUnderlineStyleAttributeName inRange:NSMakeRange(0, [aString length]) options:0 usingBlock:^(id Value, NSRange Range, BOOL* bStop) {
+					 if(Value && [Value isKindOfClass:[NSNumber class]])
+					 {
+						 NSNumber* NumberValue = (NSNumber*)Value;
+						 const int UnderlineValue = [NumberValue intValue];
+						 if(UnderlineValue > 1)
+						 {
+							 // Found the active range, stop enumeration.
+							 *bStop = YES;
+							 CompositionRange.location += Range.location;
+							 CompositionRange.length = Range.length;
+						 }
+					 }
+				 }];
+				
+				TheString = [AttributedString string];
+			}
+			else
+			{
+				TheString = (NSString*)aString;
+			}
+			
+			FString TheFString(TheString);
+			IMMContext->SetTextInRange(SelectionLocation, SelectionLength, TheFString);
+			IMMContext->UpdateCompositionRange(CompositionRange.location, CompositionRange.length);
+			IMMContext->SetSelectionRange(markedRange.location + selectedRange.location, 0, ITextInputMethodContext::ECaretPosition::Ending);
+		}
+		[[self inputContext] invalidateCharacterCoordinates]; // recentering
+	}
+	else
+	{
+		reallyHandledEvent = false;
+	}
+}
+
+/* The receiver unmarks the marked text. If no marked text, the invocation of this method has no effect.
+ */
+- (void)unmarkText
+{
+	if(markedRange.location != NSNotFound)
+	{
+		markedRange = {NSNotFound, 0};
+		if (IMMContext.IsValid())
+		{
+			IMMContext->UpdateCompositionRange(0, 0);
+			IMMContext->EndComposition();
+		}
+	}
+}
+
+/* Returns the selection range. The valid location is from 0 to the document length.
+ */
+- (NSRange)selectedRange
+{
+	if (IMMContext.IsValid())
+	{
+		uint32 SelectionLocation, SelectionLength;
+		ITextInputMethodContext::ECaretPosition CaretPosition;
+		IMMContext->GetSelectionRange(SelectionLocation, SelectionLength, CaretPosition);
+		
+		return {SelectionLocation, SelectionLength};
+	}
+	else
+	{
+		return {NSNotFound, 0};
+	}
+}
+
+/* Returns the marked range. Returns {NSNotFound, 0} if no marked range.
+ */
+- (NSRange)markedRange
+{
+	if (IMMContext.IsValid())
+	{
+		return markedRange;
+	}
+	else
+	{
+		return {NSNotFound, 0};
+	}
+}
+
+/* Returns whether or not the receiver has marked text.
+ */
+- (BOOL)hasMarkedText
+{
+	return IMMContext.IsValid() && (markedRange.location != NSNotFound);
+}
+
+/* Returns attributed string specified by aRange. It may return nil. If non-nil return value and actualRange is non-NULL, it contains the actual range for the return value. The range can be adjusted from various reasons (i.e. adjust to grapheme cluster boundary, performance optimization, etc).
+ */
+- (NSAttributedString *)attributedSubstringForProposedRange:(NSRange)aRange actualRange:(NSRangePointer)actualRange
+{
+	NSAttributedString* AttributedString = nil;
+	if (IMMContext.IsValid())
+	{
+		FString String;
+		IMMContext->GetTextInRange(aRange.location, aRange.length, String);
+		CFStringRef CFString = FPlatformString::TCHARToCFString(*String);
+		if(CFString)
+		{
+			AttributedString = [[[NSAttributedString alloc] initWithString:(NSString*)CFString] autorelease];
+			CFRelease(CFString);
+			if(actualRange)
+			{
+				*actualRange = aRange;
+			}
+		}
+	}
+	return AttributedString;
+}
+
+/* Returns an array of attribute names recognized by the receiver.
+ */
+- (NSArray*)validAttributesForMarkedText
+{
+    // We only allow these attributes to be set on our marked text (plus standard attributes)
+    // NSMarkedClauseSegmentAttributeName is important for CJK input, among other uses
+    // NSGlyphInfoAttributeName allows alternate forms of characters
+    return [NSArray arrayWithObjects:NSMarkedClauseSegmentAttributeName, NSGlyphInfoAttributeName, nil];
+}
+
+/* Returns the first logical rectangular area for aRange. The return value is in the screen coordinate. The size value can be negative if the text flows to the left. If non-NULL, actuallRange contains the character range corresponding to the returned area.
+ */
+- (NSRect)firstRectForCharacterRange:(NSRange)aRange actualRange:(NSRangePointer)actualRange
+{
+	if (IMMContext.IsValid())
+	{
+		FVector2D Position, Size;
+		IMMContext->GetTextBounds(aRange.location, aRange.length, Position, Size);
+		
+		if(actualRange)
+		{
+			*actualRange = aRange;
+		}
+		
+		NSScreen* PrimaryScreen = [[self window] screen];
+		const float PrimaryScreenHeight = [PrimaryScreen visibleFrame].size.height;
+		Position.Y = -(Position.Y - PrimaryScreenHeight + 1);
+		
+		NSRect GlyphBox = {{Position.X,Position.Y},{Size.X,Size.Y}};
+		return GlyphBox;
+	}
+	else
+	{
+		NSRect GlyphBox = {{0,0},{0,0}};
+		return GlyphBox;
+	}
+}
+
+/* Returns the index for character that is nearest to aPoint. aPoint is in the screen coordinate system.
+ */
+- (NSUInteger)characterIndexForPoint:(NSPoint)aPoint
+{
+	if (IMMContext.IsValid())
+	{
+		FVector2D Pos(aPoint.x, aPoint.y);
+		int32 Index = IMMContext->GetCharacterIndexFromPoint(Pos);
+		NSUInteger Idx = Index == INDEX_NONE ? NSNotFound : (NSUInteger)Index;
+		return Idx;
+	}
+	else
+	{
+		return NSNotFound;
+	}
+}
+
+//@optional
+
+/* Returns the window level of the receiver. An NSTextInputClient can implement this interface to specify its window level if it is higher than NSFloatingWindowLevel.
+ */
+- (NSInteger)windowLevel
+{
+	return [[self window] level];
+}
+
+@end
diff --git a/Engine/Source/Runtime/Core/Private/Mac/CocoaWindow.cpp b/Engine/Source/Runtime/Core/Private/Mac/CocoaWindow.cpp
new file mode 100644
--- /dev/null
+++ b/Engine/Source/Runtime/Core/Private/Mac/CocoaWindow.cpp
@@ -0,0 +1,641 @@
+// Copyright 1998-2014 Epic Games, Inc. All Rights Reserved.
+
+#include "Core.h"
+#include "CocoaWindow.h"
+#include "CocoaTextView.h"
+#include "MacApplication.h"
+#include "MacWindow.h"
+
+/**
+ * Custom window class used for input handling
+ */
+@implementation FCocoaWindow
+
+@synthesize bForwardEvents;
+@synthesize TargetWindowMode;
+@synthesize PreFullScreenRect;
+
+- (id)initWithContentRect:(NSRect)ContentRect styleMask:(NSUInteger)Style backing:(NSBackingStoreType)BufferingType defer:(BOOL)Flag
+{
+	WindowMode = EWindowMode::Windowed;
+	bAcceptsInput = false;
+	bRoundedCorners = false;
+	bDisplayReconfiguring = false;
+	bDeferOrderFront = false;
+	DeferOpacity = 0.0f;
+	bRenderInitialised = false;
+	bDeferSetFrame = false;
+	bDeferSetOrigin = false;
+
+	id NewSelf = [super initWithContentRect:ContentRect styleMask:Style backing:BufferingType defer:Flag];
+	if(NewSelf)
+	{
+		bZoomed = [super isZoomed];
+		self.bForwardEvents = true;
+		self.TargetWindowMode = EWindowMode::Windowed;
+		[super setAlphaValue:DeferOpacity];
+		DeferFrame = [super frame];
+		self.PreFullScreenRect = DeferFrame;
+	}
+	return NewSelf;
+}
+
+- (NSRect)openGLFrame
+{
+	if(self.TargetWindowMode == EWindowMode::Fullscreen || WindowMode == EWindowMode::Fullscreen)
+	{
+		return self.PreFullScreenRect;
+	}
+	else if([self styleMask] & (NSTexturedBackgroundWindowMask))
+	{
+		return (!bDeferSetFrame ? [self frame] : DeferFrame);
+	}
+	else
+	{
+		return (!bDeferSetFrame ? [[self contentView] frame] : [self contentRectForFrameRect:DeferFrame]);
+	}
+}
+
+
+- (NSView*)openGLView
+{
+	if([self styleMask] & (NSTexturedBackgroundWindowMask))
+	{
+		NSView* SuperView = [[self contentView] superview];
+		for(NSView* View in [SuperView subviews])
+		{
+			if([View isKindOfClass:[FCocoaTextView class]])
+			{
+				return View;
+			}
+		}
+		return nil;
+	}
+	else
+	{
+		return [self contentView];
+	}
+}
+
+- (void)performDeferredOrderFront
+{
+	if(!bRenderInitialised)
+	{
+		bRenderInitialised = true;
+	}
+	
+	if(bDeferOrderFront)
+	{
+		if(!(bDeferSetFrame || bDeferSetOrigin))
+		{
+			bDeferOrderFront = false;
+			[super setAlphaValue:DeferOpacity];
+		}
+		else
+		{
+			[self performDeferredSetFrame];
+		}
+	}
+}
+
+- (void)performDeferredSetFrame
+{
+	if(bRenderInitialised && (bDeferSetFrame || bDeferSetOrigin))
+	{
+		dispatch_block_t Block = ^{
+			if(!bDeferSetFrame && bDeferSetOrigin)
+			{
+				DeferFrame.size = [self frame].size;
+			}
+			
+			[super setFrame:DeferFrame display:YES];
+		};
+		
+		if([NSThread isMainThread])
+		{
+			Block();
+		}
+		else
+		{
+			dispatch_async(dispatch_get_main_queue(), Block);
+		}
+		
+		bDeferSetFrame = false;
+		bDeferSetOrigin = false;
+	}
+}
+
+- (void)orderWindow:(NSWindowOrderingMode)OrderingMode relativeTo:(NSInteger)OtherWindowNumber
+{
+	bool bModal = FMacWindow::CurrentModalWindow() == nil || FMacWindow::CurrentModalWindow() == self || [self styleMask] == NSBorderlessWindowMask;
+	if(OrderingMode == NSWindowOut || bModal)
+	{
+		if([self alphaValue] > 0.0f)
+		{
+			[self performDeferredSetFrame];
+		}
+		[super orderWindow:OrderingMode relativeTo:OtherWindowNumber];
+	}
+}
+
+- (bool)roundedCorners
+{
+    return bRoundedCorners;
+}
+
+- (void)setRoundedCorners:(bool)bUseRoundedCorners
+{
+	bRoundedCorners = bUseRoundedCorners;
+}
+
+- (void)setAcceptsInput:(bool)InAcceptsInput
+{
+	bAcceptsInput = InAcceptsInput;
+}
+
+- (void)redrawContents
+{
+	if(bNeedsRedraw && bForwardEvents && ([self isVisible] && [super alphaValue] > 0.0f))
+	{
+		MacApplication->OnWindowRedrawContents( self );
+	}
+	bNeedsRedraw = false;
+}
+
+- (void)setWindowMode:(EWindowMode::Type)NewWindowMode
+{
+	WindowMode = NewWindowMode;
+}
+
+- (EWindowMode::Type)windowMode
+{
+	return WindowMode;
+}
+
+- (void)setDisplayReconfiguring:(bool)bIsDisplayReconfiguring
+{
+	bDisplayReconfiguring = bIsDisplayReconfiguring;
+}
+
+- (void)orderFrontAndMakeMain:(bool)bMain andKey:(bool)bKey
+{
+	if ([NSApp isHidden] == NO)
+	{
+		bool bBringToFront = FMacWindow::CurrentModalWindow() == nil || FMacWindow::CurrentModalWindow() == self || [self styleMask] == NSBorderlessWindowMask;
+		if (bBringToFront)
+		{
+			[self orderFront:nil];
+		}
+		
+		if (bMain && [self canBecomeMainWindow] && self != [NSApp mainWindow])
+		{
+			[self makeMainWindow];
+		}
+		if (bKey && [self canBecomeKeyWindow] && self != [NSApp keyWindow])
+		{
+			[self makeKeyWindow];
+		}
+	}
+}
+
+// Following few methods overload NSWindow's methods from Cocoa API, so have to use Cocoa's BOOL (signed char), not bool (unsigned int)
+- (BOOL)canBecomeMainWindow
+{
+	bool bNoModalOrCurrent = FMacWindow::CurrentModalWindow() == nil || FMacWindow::CurrentModalWindow() == self;
+	return bAcceptsInput && ([self styleMask] != NSBorderlessWindowMask) && bNoModalOrCurrent;
+}
+
+- (BOOL)canBecomeKeyWindow
+{
+	return bAcceptsInput && ![self ignoresMouseEvents];
+}
+
+- (BOOL)validateMenuItem:(NSMenuItem *)MenuItem
+{
+	// Borderless windows we use do not automatically handle first responder's actions, so we force it here
+	return ([MenuItem action] == @selector(performClose:) || [MenuItem action] == @selector(performMiniaturize:) || [MenuItem action] == @selector(performZoom:)) ? YES : [super validateMenuItem:MenuItem];
+}
+
+- (void)setAlphaValue:(CGFloat)WindowAlpha
+{
+	if(!bRenderInitialised)
+	{
+		DeferOpacity = WindowAlpha;
+		bDeferOrderFront = true;
+	}
+	else
+	{
+		if([self isVisible] && WindowAlpha > 0.0f)
+		{
+			[self performDeferredSetFrame];
+		}
+		[super setAlphaValue:WindowAlpha];
+	}
+}
+
+- (void)orderOut:(id)Sender
+{
+	bDeferOrderFront = false;
+	
+	[super orderOut:Sender];
+}
+
+- (void)performClose:(id)Sender
+{
+	bDeferOrderFront = false;
+	
+	[self close];
+}
+
+- (void)performMiniaturize:(id)Sender
+{
+	[self miniaturize: self];
+}
+
+- (void)performZoom:(id)Sender
+{
+	bZoomed = !bZoomed;
+	[self zoom: self];
+}
+
+- (void)setFrame:(NSRect)FrameRect display:(BOOL)Flag
+{
+	NSSize Size = [self frame].size;
+	NSSize NewSize = FrameRect.size;
+	if(!bRenderInitialised || ([self isVisible] && [super alphaValue] > 0.0f && (Size.width > 1 || Size.height > 1 || NewSize.width > 1 || NewSize.height > 1)))
+	{
+		[super setFrame:FrameRect display:Flag];
+		bDeferSetFrame = false;
+	}
+	else
+	{
+		bDeferSetFrame = true;
+		DeferFrame = FrameRect;
+		if(self.bForwardEvents)
+		{
+			MacApplication->OnWindowDidResize( self );
+		}
+	}
+}
+
+- (void)setFrameOrigin:(NSPoint)Point
+{
+	NSSize Size = [self frame].size;
+	if(!bRenderInitialised || ([self isVisible] && [super alphaValue] > 0.0f && (Size.width > 1 || Size.height > 1)))
+	{
+		[super setFrameOrigin:Point];
+		bDeferSetOrigin = false;
+	}
+	else
+	{
+		bDeferSetOrigin = true;
+		DeferFrame.origin = Point;
+	}
+}
+
+- (void)keyDown:(NSEvent *)Event
+{
+	if(self.bForwardEvents)
+	{
+		MacApplication->ProcessEvent( Event );
+	}
+}
+
+- (void)keyUp:(NSEvent *)Event
+{
+	if(self.bForwardEvents)
+	{
+		MacApplication->ProcessEvent( Event );
+	}
+}
+
+- (void)windowWillEnterFullScreen:(NSNotification *)notification
+{
+	// Handle clicking on the titlebar fullscreen item
+	if(self.TargetWindowMode == EWindowMode::Windowed)
+	{
+		// @todo: Fix fullscreen mode mouse coordinate handling - for now default to windowed fullscreen
+		self.TargetWindowMode = EWindowMode::WindowedFullscreen;
+	}
+}
+
+- (void)windowDidEnterFullScreen:(NSNotification *)notification
+{
+	WindowMode = self.TargetWindowMode;
+	if(self.bForwardEvents)
+	{
+		MacApplication->OnWindowDidResize( self );
+	}
+}
+
+- (void)windowDidExitFullScreen:(NSNotification *)notification
+{
+	WindowMode = EWindowMode::Windowed;
+	self.TargetWindowMode = EWindowMode::Windowed;
+	if(self.bForwardEvents)
+	{
+		MacApplication->OnWindowDidResize( self );
+	}
+}
+
+- (void)windowDidBecomeKey:(NSNotification *)Notification
+{
+	if([NSApp isHidden] == NO)
+	{
+		if(FMacWindow::CurrentModalWindow() == nil || FMacWindow::CurrentModalWindow() == self || [self styleMask] == NSBorderlessWindowMask)
+		{
+			[self orderFrontAndMakeMain:false andKey:false];
+		}
+		else
+		{
+			[FMacWindow::CurrentModalWindow() orderFrontAndMakeMain:true andKey:true];
+		}
+	}
+	
+	if(self.bForwardEvents)
+	{
+		MacApplication->OnWindowDidBecomeKey( self );
+	}
+}
+
+- (void)windowDidResignKey:(NSNotification *)Notification
+{
+	[self setMovable: YES];
+	[self setMovableByWindowBackground: NO];
+	
+	if(self.bForwardEvents)
+	{
+		MacApplication->OnWindowDidResignKey( self );
+	}
+}
+
+- (void)windowWillMove:(NSNotification *)Notification
+{
+	if(self.bForwardEvents)
+	{
+		MacApplication->OnWindowWillMove( self );
+	}
+}
+
+- (void)windowDidMove:(NSNotification *)Notification
+{
+	bZoomed = [self isZoomed];
+	
+	NSView* OpenGLView = [self openGLView];
+	[[NSNotificationCenter defaultCenter] postNotificationName:NSViewGlobalFrameDidChangeNotification object:OpenGLView];
+	
+	if(self.bForwardEvents)
+	{
+		MacApplication->OnWindowDidMove( self );
+	}
+}
+
+- (void)windowDidChangeScreen:(NSNotification *)notification
+{
+	// You'd think this would be a good place to handle un/parenting...
+	// However, do that here and your dragged window will disappear when you drag it onto another screen!
+	// The windowdidChangeScreen notification only comes after you finish dragging.
+	// It does however, work fine for handling display arrangement changes that cause a window to go offscreen.
+	if(bDisplayReconfiguring)
+	{
+		NSScreen* Screen = [self screen];
+		NSRect Frame = [self frame];
+		NSRect VisibleFrame = [Screen visibleFrame];
+		if(NSContainsRect(VisibleFrame, Frame) == NO)
+		{
+			// May need to scale the window to fit if it is larger than the new display.
+			if (Frame.size.width > VisibleFrame.size.width || Frame.size.height > VisibleFrame.size.height)
+			{
+				NSRect NewFrame;
+				NewFrame.size.width = Frame.size.width > VisibleFrame.size.width ? VisibleFrame.size.width : Frame.size.width;
+				NewFrame.size.height = Frame.size.height > VisibleFrame.size.height ? VisibleFrame.size.height : Frame.size.height;
+				NewFrame.origin = VisibleFrame.origin;
+				
+				[self setFrame:NewFrame display:NO];
+			}
+			else
+			{
+				NSRect Intersection = NSIntersectionRect(VisibleFrame, Frame);
+				NSPoint Origin = Frame.origin;
+				
+				// If there's at least something on screen, try shifting it entirely on screen.
+				if(Intersection.size.width > 0 && Intersection.size.height > 0)
+				{
+					CGFloat X = Frame.size.width - Intersection.size.width;
+					CGFloat Y = Frame.size.height - Intersection.size.height;
+					
+					if(Intersection.size.width+Intersection.origin.x >= VisibleFrame.size.width+VisibleFrame.origin.x)
+					{
+						Origin.x -= X;
+					}
+					else if(Origin.x < VisibleFrame.origin.x)
+					{
+						Origin.x += X;
+					}
+					
+					if(Intersection.size.height+Intersection.origin.y >= VisibleFrame.size.height+VisibleFrame.origin.y)
+					{
+						Origin.y -= Y;
+					}
+					else if(Origin.y < VisibleFrame.origin.y)
+					{
+						Origin.y += Y;
+					}
+				}
+				else
+				{
+					Origin = VisibleFrame.origin;
+				}
+				
+				[self setFrameOrigin:Origin];
+			}
+		}
+	}
+}
+
+- (void)windowDidResize:(NSNotification *)Notification
+{
+	bZoomed = [self isZoomed];
+	if(self.bForwardEvents)
+	{
+		MacApplication->OnWindowDidResize( self );
+	}
+	bNeedsRedraw = true;
+}
+
+- (void)windowWillClose:(NSNotification *)notification
+{
+	if(self.bForwardEvents && MacApplication)
+	{
+		MacApplication->OnWindowDidClose( self );
+	}
+}
+
+- (void)mouseDown:(NSEvent*)Event
+{
+	if(self.bForwardEvents)
+	{
+		MacApplication->ProcessEvent( Event );
+	}
+}
+
+- (void)rightMouseDown:(NSEvent*)Event
+{
+	// Really we shouldn't be doing this - on OS X only left-click changes focus,
+	// but for the moment it is easier than changing Slate.
+	if([self canBecomeKeyWindow])
+	{
+		[self makeKeyWindow];
+	}
+	
+	if(self.bForwardEvents)
+	{
+		MacApplication->ProcessEvent( Event );
+	}
+}
+
+- (void)otherMouseDown:(NSEvent*)Event
+{
+	if(self.bForwardEvents)
+	{
+		MacApplication->ProcessEvent( Event );
+	}
+}
+
+- (void)mouseUp:(NSEvent*)Event
+{
+	if(self.bForwardEvents)
+	{
+		MacApplication->ProcessEvent( Event );
+	}
+}
+
+- (void)rightMouseUp:(NSEvent*)Event
+{
+	if(self.bForwardEvents)
+	{
+		MacApplication->ProcessEvent( Event );
+	}
+}
+
+- (void)otherMouseUp:(NSEvent*)Event
+{
+	if(self.bForwardEvents)
+	{
+		MacApplication->ProcessEvent( Event );
+	}
+}
+
+- (NSDragOperation)draggingEntered:(id < NSDraggingInfo >)Sender
+{
+	return NSDragOperationGeneric;
+}
+
+- (void)draggingExited:(id < NSDraggingInfo >)Sender
+{
+	if(self.bForwardEvents)
+	{
+		MacApplication->OnDragOut( self );
+	}
+}
+
+- (NSDragOperation)draggingUpdated:(id < NSDraggingInfo >)Sender
+{
+	if(self.bForwardEvents)
+	{
+		MacApplication->OnDragOver( self );
+	}
+	return NSDragOperationGeneric;
+}
+
+- (BOOL)prepareForDragOperation:(id < NSDraggingInfo >)Sender
+{
+	if(self.bForwardEvents)
+	{
+		MacApplication->OnDragEnter( self, [Sender draggingPasteboard] );
+	}
+	return YES;
+}
+
+- (BOOL)performDragOperation:(id < NSDraggingInfo >)Sender
+{
+	if(self.bForwardEvents)
+	{
+		MacApplication->OnDragDrop( self );
+	}
+	return YES;
+}
+
+- (BOOL)isMovable
+{
+	BOOL Movable = [super isMovable];
+	if(Movable && bRenderInitialised && MacApplication)
+	{
+		Movable &= MacApplication->IsWindowMovable(self, NULL);
+	}
+	return Movable;
+}
+
+@end
+
+/**
+ * Custom window class used for mouse capture
+ */
+@implementation FMouseCaptureWindow
+
+- (id)initWithTargetWindow: (FCocoaWindow*)Window
+{
+	self = [super initWithContentRect: [[Window screen] frame] styleMask: NSBorderlessWindowMask backing: NSBackingStoreBuffered defer: NO];
+	[self setBackgroundColor: [NSColor clearColor]];
+	[self setOpaque: NO];
+	[self setLevel: NSMainMenuWindowLevel + 1];
+	[self setIgnoresMouseEvents: NO];
+	[self setAcceptsMouseMovedEvents: YES];
+	[self setHidesOnDeactivate: YES];
+	
+	TargetWindow = Window;
+	
+	return self;
+}
+
+- (FCocoaWindow*)targetWindow
+{
+	return TargetWindow;
+}
+
+- (void)setTargetWindow: (FCocoaWindow*)Window
+{
+	TargetWindow = Window;
+}
+
+- (void)mouseDown:(NSEvent*)Event
+{
+	MacApplication->ProcessEvent( Event );
+}
+
+- (void)rightMouseDown:(NSEvent*)Event
+{
+	MacApplication->ProcessEvent( Event );
+}
+
+- (void)otherMouseDown:(NSEvent*)Event
+{
+	MacApplication->ProcessEvent( Event );
+}
+
+- (void)mouseUp:(NSEvent*)Event
+{
+	MacApplication->ProcessEvent( Event );
+}
+
+- (void)rightMouseUp:(NSEvent*)Event
+{
+	MacApplication->ProcessEvent( Event );
+}
+
+- (void)otherMouseUp:(NSEvent*)Event
+{
+	MacApplication->ProcessEvent( Event );
+}
+
+@end
diff --git a/Engine/Source/Runtime/Core/Private/Mac/MacApplication.cpp b/Engine/Source/Runtime/Core/Private/Mac/MacApplication.cpp
--- a/Engine/Source/Runtime/Core/Private/Mac/MacApplication.cpp
+++ b/Engine/Source/Runtime/Core/Private/Mac/MacApplication.cpp
@@ -172,7 +172,7 @@ void FMacApplication::SetMessageHandler( const TSharedRef< FGenericApplicationMe
 	HIDInput->SetMessageHandler( InMessageHandler );
 }
 
-static TSharedPtr< FMacWindow > FindWindowByNSWindow( const TArray< TSharedRef< FMacWindow > >& WindowsToSearch, FSlateCocoaWindow* const WindowHandle )
+static TSharedPtr< FMacWindow > FindWindowByNSWindow( const TArray< TSharedRef< FMacWindow > >& WindowsToSearch, FCocoaWindow* const WindowHandle )
 {
 	for (int32 WindowIndex=0; WindowIndex < WindowsToSearch.Num(); ++WindowIndex)
 	{
@@ -200,14 +200,14 @@ void FMacApplication::OnWindowDraggingFinished()
 	}
 }
 
-bool FMacApplication::IsWindowMovable(FSlateCocoaWindow* Win, bool* OutMovableByBackground)
+bool FMacApplication::IsWindowMovable(FCocoaWindow* Win, bool* OutMovableByBackground)
 {
 	if(OutMovableByBackground)
 	{
 		*OutMovableByBackground = false;
 	}
 
-	FSlateCocoaWindow* NativeWindow = (FSlateCocoaWindow*)Win;
+	FCocoaWindow* NativeWindow = (FCocoaWindow*)Win;
 	TSharedPtr< FMacWindow > CurrentEventWindow = FindWindowByNSWindow( Windows, NativeWindow );
 	if( CurrentEventWindow.IsValid() )
 	{
@@ -278,7 +278,7 @@ void FMacApplication::ProcessEvent( NSEvent* Event )
 
 	const NSEventType EventType = [Event type];
 
-	FSlateCocoaWindow* NativeWindow = FindEventWindow( Event );
+	FCocoaWindow* NativeWindow = FindEventWindow( Event );
 	TSharedPtr< FMacWindow > CurrentEventWindow = FindWindowByNSWindow( Windows, NativeWindow );
 	if( NativeWindow && !CurrentEventWindow.IsValid() && LastEventWindow.IsValid() )
 	{
@@ -409,6 +409,14 @@ void FMacApplication::ProcessEvent( NSEvent* Event )
 		case NSRightMouseDown:
 		case NSOtherMouseDown:
 		{
+			if( CurrentEventWindow.IsValid() && CurrentEventWindow->IsRegularWindow() )
+			{
+				bool IsMouseOverTitleBar = false;
+				bool IsMovable = IsWindowMovable(NativeWindow, &IsMouseOverTitleBar);
+				[NativeWindow setMovable: IsMovable];
+				[NativeWindow setMovableByWindowBackground: IsMouseOverTitleBar];
+			}
+			
 			EMouseButtons::Type Button = [Event type] == NSLeftMouseDown ? EMouseButtons::Left : EMouseButtons::Right;
 			if ([Event type] == NSOtherMouseDown)
 			{
@@ -614,7 +622,7 @@ void FMacApplication::ProcessEvent( NSEvent* Event )
 	}
 }
 
-FSlateCocoaWindow* FMacApplication::FindEventWindow( NSEvent* Event )
+FCocoaWindow* FMacApplication::FindEventWindow( NSEvent* Event )
 {
 	SCOPED_AUTORELEASE_POOL;
 
@@ -636,7 +644,7 @@ FSlateCocoaWindow* FMacApplication::FindEventWindow( NSEvent* Event )
 			break;
 	}
 
-	FSlateCocoaWindow* EventWindow = (FSlateCocoaWindow*)[Event window];
+	FCocoaWindow* EventWindow = (FCocoaWindow*)[Event window];
 
 	if ([Event type] == NSMouseMoved)
 	{
@@ -690,7 +698,7 @@ TSharedPtr<FGenericWindow> FMacApplication::LocateWindowUnderCursor( const FVect
 	for (int32 Index = 0; Index < [AllWindows count]; Index++)
 	{
 		NSWindow* NativeWindow = (NSWindow*)[AllWindows objectAtIndex: Index];
-		if ([NativeWindow isMiniaturized] || ![NativeWindow isVisible] || ![NativeWindow isOnActiveSpace] || ![NativeWindow isKindOfClass: [FSlateCocoaWindow class]])
+		if ([NativeWindow isMiniaturized] || ![NativeWindow isVisible] || ![NativeWindow isOnActiveSpace] || ![NativeWindow isKindOfClass: [FCocoaWindow class]])
 		{
 			continue;
 		}
@@ -719,7 +727,7 @@ TSharedPtr<FGenericWindow> FMacApplication::LocateWindowUnderCursor( const FVect
 
 		if (NSPointInRect(Position, VisibleFrame))
 		{
-			TSharedPtr<FMacWindow> MacWindow = FindWindowByNSWindow(Windows, (FSlateCocoaWindow*)NativeWindow);
+			TSharedPtr<FMacWindow> MacWindow = FindWindowByNSWindow(Windows, (FCocoaWindow*)NativeWindow);
 			return MacWindow;
 		}
 	}
@@ -765,7 +773,7 @@ void FMacApplication::UseMouseCaptureWindow( bool bUseMouseCaptureWindow )
 	}
 }
 
-void FMacApplication::UpdateMouseCaptureWindow( FSlateCocoaWindow* TargetWindow )
+void FMacApplication::UpdateMouseCaptureWindow( FCocoaWindow* TargetWindow )
 {
 	SCOPED_AUTORELEASE_POOL;
 
@@ -782,9 +790,9 @@ void FMacApplication::UpdateMouseCaptureWindow( FSlateCocoaWindow* TargetWindow
 			{
 				TargetWindow = [MouseCaptureWindow targetWindow];
 			}
-			else if( [[NSApp keyWindow] isKindOfClass: [FSlateCocoaWindow class]] )
+			else if( [[NSApp keyWindow] isKindOfClass: [FCocoaWindow class]] )
 			{
-				TargetWindow = (FSlateCocoaWindow*)[NSApp keyWindow];
+				TargetWindow = (FCocoaWindow*)[NSApp keyWindow];
 			}
 		}
 
@@ -913,7 +921,7 @@ void FMacApplication::GetDisplayMetrics( FDisplayMetrics& OutDisplayMetrics ) co
 	OutDisplayMetrics.PrimaryDisplayWorkAreaRect.Bottom = OutDisplayMetrics.PrimaryDisplayWorkAreaRect.Top + VisibleFrame.size.height;
 }
 
-void FMacApplication::OnDragEnter( FSlateCocoaWindow* Window, void *InPasteboard )
+void FMacApplication::OnDragEnter( FCocoaWindow* Window, void *InPasteboard )
 {
 	SCOPED_AUTORELEASE_POOL;
 
@@ -956,7 +964,7 @@ void FMacApplication::OnDragEnter( FSlateCocoaWindow* Window, void *InPasteboard
 	}
 }
 
-void FMacApplication::OnDragOver( FSlateCocoaWindow* Window )
+void FMacApplication::OnDragOver( FCocoaWindow* Window )
 {
 	TSharedPtr< FMacWindow > EventWindow = FindWindowByNSWindow( Windows, Window );
 	if( EventWindow.IsValid() )
@@ -965,7 +973,7 @@ void FMacApplication::OnDragOver( FSlateCocoaWindow* Window )
 	}
 }
 
-void FMacApplication::OnDragOut( FSlateCocoaWindow* Window )
+void FMacApplication::OnDragOut( FCocoaWindow* Window )
 {
 	TSharedPtr< FMacWindow > EventWindow = FindWindowByNSWindow( Windows, Window );
 	if( EventWindow.IsValid() )
@@ -974,7 +982,7 @@ void FMacApplication::OnDragOut( FSlateCocoaWindow* Window )
 	}
 }
 
-void FMacApplication::OnDragDrop( FSlateCocoaWindow* Window )
+void FMacApplication::OnDragDrop( FCocoaWindow* Window )
 {
 	TSharedPtr< FMacWindow > EventWindow = FindWindowByNSWindow( Windows, Window );
 	if( EventWindow.IsValid() )
@@ -983,7 +991,7 @@ void FMacApplication::OnDragDrop( FSlateCocoaWindow* Window )
 	}
 }
 
-void FMacApplication::OnWindowDidBecomeKey( FSlateCocoaWindow* Window )
+void FMacApplication::OnWindowDidBecomeKey( FCocoaWindow* Window )
 {
 	TSharedPtr< FMacWindow > EventWindow = FindWindowByNSWindow( Windows, Window );
 	if( EventWindow.IsValid() )
@@ -993,7 +1001,7 @@ void FMacApplication::OnWindowDidBecomeKey( FSlateCocoaWindow* Window )
 	}
 }
 
-void FMacApplication::OnWindowDidResignKey( FSlateCocoaWindow* Window )
+void FMacApplication::OnWindowDidResignKey( FCocoaWindow* Window )
 {
 	TSharedPtr< FMacWindow > EventWindow = FindWindowByNSWindow( Windows, Window );
 	if( EventWindow.IsValid() )
@@ -1002,14 +1010,14 @@ void FMacApplication::OnWindowDidResignKey( FSlateCocoaWindow* Window )
 	}
 }
 
-void FMacApplication::OnWindowWillMove( FSlateCocoaWindow* Window )
+void FMacApplication::OnWindowWillMove( FCocoaWindow* Window )
 {
 	SCOPED_AUTORELEASE_POOL;
 
 	DraggedWindow = Window;
 }
 
-void FMacApplication::OnWindowDidMove( FSlateCocoaWindow* Window )
+void FMacApplication::OnWindowDidMove( FCocoaWindow* Window )
 {
 	SCOPED_AUTORELEASE_POOL;
 
@@ -1028,7 +1036,7 @@ void FMacApplication::OnWindowDidMove( FSlateCocoaWindow* Window )
 	}
 }
 
-void FMacApplication::OnWindowDidResize( FSlateCocoaWindow* Window )
+void FMacApplication::OnWindowDidResize( FCocoaWindow* Window )
 {
 	SCOPED_AUTORELEASE_POOL;
 
@@ -1052,7 +1060,7 @@ void FMacApplication::OnWindowDidResize( FSlateCocoaWindow* Window )
 	}
 }
 
-void FMacApplication::OnWindowRedrawContents( FSlateCocoaWindow* Window )
+void FMacApplication::OnWindowRedrawContents( FCocoaWindow* Window )
 {
 	TSharedPtr< FMacWindow > EventWindow = FindWindowByNSWindow( Windows, Window );
 	if( EventWindow.IsValid() )
@@ -1062,7 +1070,7 @@ void FMacApplication::OnWindowRedrawContents( FSlateCocoaWindow* Window )
 	}
 }
 
-void FMacApplication::OnWindowDidClose( FSlateCocoaWindow* Window )
+void FMacApplication::OnWindowDidClose( FCocoaWindow* Window )
 {
 	TSharedPtr< FMacWindow > EventWindow = FindWindowByNSWindow( Windows, Window );
 	if( EventWindow.IsValid() )
@@ -1075,7 +1083,7 @@ void FMacApplication::OnWindowDidClose( FSlateCocoaWindow* Window )
 	}
 }
 
-bool FMacApplication::OnWindowDestroyed( FSlateCocoaWindow* Window )
+bool FMacApplication::OnWindowDestroyed( FCocoaWindow* Window )
 {
 	TSharedPtr< FMacWindow > EventWindow = FindWindowByNSWindow( Windows, Window );
 	if( EventWindow.IsValid() )
diff --git a/Engine/Source/Runtime/Core/Private/Mac/MacTextInputMethodSystem.cpp b/Engine/Source/Runtime/Core/Private/Mac/MacTextInputMethodSystem.cpp
--- a/Engine/Source/Runtime/Core/Private/Mac/MacTextInputMethodSystem.cpp
+++ b/Engine/Source/Runtime/Core/Private/Mac/MacTextInputMethodSystem.cpp
@@ -4,445 +4,10 @@
 
 #include "MacWindow.h"
 #include "MacTextInputMethodSystem.h"
+#include "CocoaTextView.h"
 
 DEFINE_LOG_CATEGORY_STATIC(LogMacTextInputMethodSystem, Log, All);
 
-@implementation FSlateTextView
-
-- (id)initWithFrame:(NSRect)frame
-{
-    self = [super initWithFrame:frame];
-    if (self)
-	{
-        // Initialization code here.
-		IMMContext = nil;
-		markedRange = {NSNotFound, 0};
-		reallyHandledEvent = false;
-    }
-    return self;
-}
-
-- (void)drawRect:(NSRect)dirtyRect
-{
-	[super drawRect:dirtyRect];
-	
-    // Drawing code here.
-}
-
-- (bool)imkKeyDown:(NSEvent *)theEvent
-{
-	if (IMMContext.IsValid())
-	{
-		reallyHandledEvent = true;
-		return [[self inputContext] handleEvent:theEvent] && reallyHandledEvent;
-	}
-	else
-	{
-		return false;
-	}
-}
-
-/** Forward mouse events up to the window rather than through the responder chain - thus avoiding
- *	the hidden titlebar controls. Normal windows just use the responder chain as usual.
- */
-- (BOOL)acceptsFirstMouse:(NSEvent *)Event
-{
-	return YES;
-}
-
-- (void)mouseDown:(NSEvent *)theEvent
-{
-	if (IMMContext.IsValid())
-	{
-		[[self inputContext] handleEvent:theEvent];
-	}
-	
-	FSlateCocoaWindow* SlateCocoaWindow = [[self window] isKindOfClass:[FSlateCocoaWindow class]] ? (FSlateCocoaWindow*)[self window] : nil;
-	if (SlateCocoaWindow)
-	{
-		[SlateCocoaWindow mouseDown:theEvent];
-	}
-}
-
-- (void)mouseDragged:(NSEvent *)theEvent
-{
-	if (IMMContext.IsValid())
-	{
-		[[self inputContext] handleEvent:theEvent];
-	}
-}
-
-- (void)mouseUp:(NSEvent *)theEvent
-{
-	if (IMMContext.IsValid())
-	{
-		[[self inputContext] handleEvent:theEvent];
-	}
-	
-	FSlateCocoaWindow* SlateCocoaWindow = [[self window] isKindOfClass:[FSlateCocoaWindow class]] ? (FSlateCocoaWindow*)[self window] : nil;
-	if (SlateCocoaWindow)
-	{
-		[SlateCocoaWindow mouseUp:theEvent];
-	}
-}
-
-- (void)rightMouseDown:(NSEvent*)Event
-{
-	FSlateCocoaWindow* SlateCocoaWindow = [[self window] isKindOfClass:[FSlateCocoaWindow class]] ? (FSlateCocoaWindow*)[self window] : nil;
-	if (SlateCocoaWindow)
-	{
-		[SlateCocoaWindow rightMouseDown:Event];
-	}
-	else
-	{
-		[super rightMouseDown:Event];
-	}
-}
-
-- (void)otherMouseDown:(NSEvent*)Event
-{
-	FSlateCocoaWindow* SlateCocoaWindow = [[self window] isKindOfClass:[FSlateCocoaWindow class]] ? (FSlateCocoaWindow*)[self window] : nil;
-	if (SlateCocoaWindow)
-	{
-		[SlateCocoaWindow otherMouseDown:Event];
-	}
-	else
-	{
-		[super otherMouseDown:Event];
-	}
-}
-
-- (void)rightMouseUp:(NSEvent*)Event
-{
-	FSlateCocoaWindow* SlateCocoaWindow = [[self window] isKindOfClass:[FSlateCocoaWindow class]] ? (FSlateCocoaWindow*)[self window] : nil;
-	if (SlateCocoaWindow)
-	{
-		[SlateCocoaWindow rightMouseUp:Event];
-	}
-	else
-	{
-		[super rightMouseUp:Event];
-	}
-}
-
-- (void)otherMouseUp:(NSEvent*)Event
-{
-	FSlateCocoaWindow* SlateCocoaWindow = [[self window] isKindOfClass:[FSlateCocoaWindow class]] ? (FSlateCocoaWindow*)[self window] : nil;
-	if (SlateCocoaWindow)
-	{
-		[SlateCocoaWindow otherMouseUp:Event];
-	}
-	else
-	{
-		[super otherMouseUp:Event];
-	}
-}
-
-- (void)activateInputMethod:(const TSharedRef<ITextInputMethodContext>&)InContext
-{
-	if (IMMContext.IsValid())
-	{
-		[self unmarkText];
-		[[self inputContext] deactivate];
-		[[self inputContext] discardMarkedText];
-	}
-	IMMContext = InContext;
-	[[self inputContext] activate];
-}
-
-- (void)deactivateInputMethod
-{
-	[self unmarkText];
-	IMMContext = NULL;
-	[[self inputContext] deactivate];
-	[[self inputContext] discardMarkedText];
-}
-
-//@protocol NSTextInputClient
-//@required
-/* The receiver inserts aString replacing the content specified by replacementRange. aString can be either an NSString or NSAttributedString instance.
- */
-- (void)insertText:(id)aString replacementRange:(NSRange)replacementRange
-{
-	if (IMMContext.IsValid() && [self hasMarkedText])
-	{
-		uint32 SelectionLocation, SelectionLength;
-		if (replacementRange.location == NSNotFound)
-		{
-			if (markedRange.location != NSNotFound)
-			{
-				SelectionLocation = markedRange.location;
-				SelectionLength = markedRange.length;
-			}
-			else
-			{
-				ITextInputMethodContext::ECaretPosition CaretPosition;
-				IMMContext->GetSelectionRange(SelectionLocation, SelectionLength, CaretPosition);
-			}
-		}
-		else
-		{
-			SelectionLocation = replacementRange.location;
-			SelectionLength = replacementRange.length;
-		}
-		
-		NSString* TheString;
-		if ([aString isKindOfClass:[NSAttributedString class]])
-		{
-			TheString = [(NSAttributedString*)aString string];
-		}
-		else
-		{
-			TheString = (NSString*)aString;
-		}
-		
-		FString TheFString(TheString);
-		IMMContext->SetTextInRange(SelectionLocation, SelectionLength, TheFString);
-		IMMContext->SetSelectionRange(SelectionLocation+TheFString.Len(), 0, ITextInputMethodContext::ECaretPosition::Ending);
-		
-		[self unmarkText];
-		[[self inputContext] invalidateCharacterCoordinates]; // recentering
-	}
-	else
-	{
-		reallyHandledEvent = false;
-	}
-}
-
-/* The receiver invokes the action specified by aSelector.
- */
-- (void)doCommandBySelector:(SEL)aSelector
-{
-	reallyHandledEvent = false;
-}
-
-/* The receiver inserts aString replacing the content specified by replacementRange. aString can be either an NSString or NSAttributedString instance. selectedRange specifies the selection inside the string being inserted; hence, the location is relative to the beginning of aString. When aString is an NSString, the receiver is expected to render the marked text with distinguishing appearance (i.e. NSTextView renders with -markedTextAttributes).
- */
-- (void)setMarkedText:(id)aString selectedRange:(NSRange)selectedRange replacementRange:(NSRange)replacementRange
-{
-	if (IMMContext.IsValid())
-	{
-		uint32 SelectionLocation, SelectionLength;
-		if (replacementRange.location == NSNotFound)
-		{
-			if (markedRange.location != NSNotFound)
-			{
-				SelectionLocation = markedRange.location;
-				SelectionLength = markedRange.length;
-			}
-			else
-			{
-				ITextInputMethodContext::ECaretPosition CaretPosition;
-				IMMContext->GetSelectionRange(SelectionLocation, SelectionLength, CaretPosition);
-			}
-		}
-		else
-		{
-			SelectionLocation = replacementRange.location;
-			SelectionLength = replacementRange.length;
-		}
-		
-		if ([aString length] == 0)
-		{
-			IMMContext->SetTextInRange(SelectionLocation, SelectionLength, FString());
-			[self unmarkText];
-		}
-		else
-		{
-			if(markedRange.location == NSNotFound)
-			{
-				IMMContext->BeginComposition();
-			}
-			markedRange = NSMakeRange(SelectionLocation, [aString length]);
-			
-			__block NSRange CompositionRange = markedRange;
-			
-			NSString* TheString;
-			if ([aString isKindOfClass:[NSAttributedString class]])
-			{
-				// While the whole string is being composed NSUnderlineStyleAttributeName is 1 to show a single line below the whole string.
-				// When using the pop-up glyph selection window in some IME's the NSAttributedString is broken up into separate glyph ranges, each with its own set of attributes.
-				// Each range specifies NSMarkedClauseSegment, incrementing the NSNumber value from 0 as well as NSUnderlineStyleAttributeName, which makes the underlining show the different ranges.
-				// The subrange being edited by the pop-up glyph selection window will set NSUnderlineStyleAttributeName to a value >1, while all other ranges will be set NSUnderlineStyleAttributeName to 1.
-				NSAttributedString* AttributedString = (NSAttributedString*)aString;
-				[AttributedString enumerateAttribute:NSUnderlineStyleAttributeName inRange:NSMakeRange(0, [aString length]) options:0 usingBlock:^(id Value, NSRange Range, BOOL* bStop) {
-					if(Value && [Value isKindOfClass:[NSNumber class]])
-					{
-						NSNumber* NumberValue = (NSNumber*)Value;
-						const int UnderlineValue = [NumberValue intValue];
-						if(UnderlineValue > 1)
-						{
-							// Found the active range, stop enumeration.
-							*bStop = YES;
-							CompositionRange.location += Range.location;
-							CompositionRange.length = Range.length;
-						}
-					}
-				}];
-				
-				TheString = [AttributedString string];
-			}
-			else
-			{
-				TheString = (NSString*)aString;
-			}
-			
-			FString TheFString(TheString);
-			IMMContext->SetTextInRange(SelectionLocation, SelectionLength, TheFString);
-			IMMContext->UpdateCompositionRange(CompositionRange.location, CompositionRange.length);
-			IMMContext->SetSelectionRange(markedRange.location + selectedRange.location, 0, ITextInputMethodContext::ECaretPosition::Ending);
-		}
-		[[self inputContext] invalidateCharacterCoordinates]; // recentering
-	}
-	else
-	{
-		reallyHandledEvent = false;
-	}
-}
-
-/* The receiver unmarks the marked text. If no marked text, the invocation of this method has no effect.
- */
-- (void)unmarkText
-{
-	if(markedRange.location != NSNotFound)
-	{
-		markedRange = {NSNotFound, 0};
-		if (IMMContext.IsValid())
-		{
-			IMMContext->UpdateCompositionRange(0, 0);
-			IMMContext->EndComposition();
-		}
-	}
-}
-
-/* Returns the selection range. The valid location is from 0 to the document length.
- */
-- (NSRange)selectedRange
-{
-	if (IMMContext.IsValid())
-	{
-		uint32 SelectionLocation, SelectionLength;
-		ITextInputMethodContext::ECaretPosition CaretPosition;
-		IMMContext->GetSelectionRange(SelectionLocation, SelectionLength, CaretPosition);
-		
-		return {SelectionLocation, SelectionLength};
-	}
-	else
-	{
-		return {NSNotFound, 0};
-	}
-}
-
-/* Returns the marked range. Returns {NSNotFound, 0} if no marked range.
- */
-- (NSRange)markedRange
-{
-	if (IMMContext.IsValid())
-	{
-		return markedRange;
-	}
-	else
-	{
-		return {NSNotFound, 0};
-	}
-}
-
-/* Returns whether or not the receiver has marked text.
- */
-- (BOOL)hasMarkedText
-{
-	return IMMContext.IsValid() && (markedRange.location != NSNotFound);
-}
-
-/* Returns attributed string specified by aRange. It may return nil. If non-nil return value and actualRange is non-NULL, it contains the actual range for the return value. The range can be adjusted from various reasons (i.e. adjust to grapheme cluster boundary, performance optimization, etc).
- */
-- (NSAttributedString *)attributedSubstringForProposedRange:(NSRange)aRange actualRange:(NSRangePointer)actualRange
-{
-	NSAttributedString* AttributedString = nil;
-	if (IMMContext.IsValid())
-	{
-		FString String;
-		IMMContext->GetTextInRange(aRange.location, aRange.length, String);
-		CFStringRef CFString = FPlatformString::TCHARToCFString(*String);
-		if(CFString)
-		{
-			AttributedString = [[[NSAttributedString alloc] initWithString:(NSString*)CFString] autorelease];
-			CFRelease(CFString);
-			if(actualRange)
-			{
-				*actualRange = aRange;
-			}
-		}
-	}
-	return AttributedString;
-}
-
-/* Returns an array of attribute names recognized by the receiver.
- */
-- (NSArray*)validAttributesForMarkedText
-{
-    // We only allow these attributes to be set on our marked text (plus standard attributes)
-    // NSMarkedClauseSegmentAttributeName is important for CJK input, among other uses
-    // NSGlyphInfoAttributeName allows alternate forms of characters
-    return [NSArray arrayWithObjects:NSMarkedClauseSegmentAttributeName, NSGlyphInfoAttributeName, nil];
-}
-
-/* Returns the first logical rectangular area for aRange. The return value is in the screen coordinate. The size value can be negative if the text flows to the left. If non-NULL, actuallRange contains the character range corresponding to the returned area.
- */
-- (NSRect)firstRectForCharacterRange:(NSRange)aRange actualRange:(NSRangePointer)actualRange
-{
-	if (IMMContext.IsValid())
-	{
-		FVector2D Position, Size;
-		IMMContext->GetTextBounds(aRange.location, aRange.length, Position, Size);
-		
-		if(actualRange)
-		{
-			*actualRange = aRange;
-		}
-		
-		NSScreen* PrimaryScreen = [[self window] screen];
-		const float PrimaryScreenHeight = [PrimaryScreen visibleFrame].size.height;
-		Position.Y = -(Position.Y - PrimaryScreenHeight + 1);
-		
-		NSRect GlyphBox = {{Position.X,Position.Y},{Size.X,Size.Y}};
-		return GlyphBox;
-	}
-	else
-	{
-		NSRect GlyphBox = {{0,0},{0,0}};
-		return GlyphBox;
-	}
-}
-
-/* Returns the index for character that is nearest to aPoint. aPoint is in the screen coordinate system.
- */
-- (NSUInteger)characterIndexForPoint:(NSPoint)aPoint
-{
-	if (IMMContext.IsValid())
-	{
-		FVector2D Pos(aPoint.x, aPoint.y);
-		int32 Index = IMMContext->GetCharacterIndexFromPoint(Pos);
-		NSUInteger Idx = Index == INDEX_NONE ? NSNotFound : (NSUInteger)Index;
-		return Idx;
-	}
-	else
-	{
-		return NSNotFound;
-	}
-}
-
-//@optional
-
-/* Returns the window level of the receiver. An NSTextInputClient can implement this interface to specify its window level if it is higher than NSFloatingWindowLevel.
- */
-- (NSInteger)windowLevel
-{
-	return [[self window] level];
-}
-
-@end
-
 namespace
 {
 	class FTextInputMethodChangeNotifier : public ITextInputMethodChangeNotifier
@@ -483,10 +48,10 @@ namespace
 	{
 		if(ContextWindow.IsValid())
 		{
-			FSlateCocoaWindow* CocoaWindow = (FSlateCocoaWindow*)ContextWindow->GetOSWindowHandle();
+			FCocoaWindow* CocoaWindow = (FCocoaWindow*)ContextWindow->GetOSWindowHandle();
 			if(CocoaWindow && [CocoaWindow openGLView])
 			{
-				FSlateTextView* TextView = (FSlateTextView*)[CocoaWindow openGLView];
+				FCocoaTextView* TextView = (FCocoaTextView*)[CocoaWindow openGLView];
 				[[TextView inputContext] invalidateCharacterCoordinates];
 			}
 		}
@@ -496,10 +61,10 @@ namespace
 	{
 		if(ContextWindow.IsValid())
 		{
-			FSlateCocoaWindow* CocoaWindow = (FSlateCocoaWindow*)ContextWindow->GetOSWindowHandle();
+			FCocoaWindow* CocoaWindow = (FCocoaWindow*)ContextWindow->GetOSWindowHandle();
 			if(CocoaWindow && [CocoaWindow openGLView])
 			{
-				FSlateTextView* TextView = (FSlateTextView*)[CocoaWindow openGLView];
+				FCocoaTextView* TextView = (FCocoaTextView*)[CocoaWindow openGLView];
 				[[TextView inputContext] invalidateCharacterCoordinates];
 			}
 		}
@@ -509,10 +74,10 @@ namespace
 	{
 		if(ContextWindow.IsValid())
 		{
-			FSlateCocoaWindow* CocoaWindow = (FSlateCocoaWindow*)ContextWindow->GetOSWindowHandle();
+			FCocoaWindow* CocoaWindow = (FCocoaWindow*)ContextWindow->GetOSWindowHandle();
 			if(CocoaWindow && [CocoaWindow openGLView])
 			{
-				FSlateTextView* TextView = (FSlateTextView*)[CocoaWindow openGLView];
+				FCocoaTextView* TextView = (FCocoaTextView*)[CocoaWindow openGLView];
 				[[TextView inputContext] invalidateCharacterCoordinates];
 			}
 		}
@@ -522,10 +87,10 @@ namespace
 	{
 		if(ContextWindow.IsValid())
 		{
-			FSlateCocoaWindow* CocoaWindow = (FSlateCocoaWindow*)ContextWindow->GetOSWindowHandle();
+			FCocoaWindow* CocoaWindow = (FCocoaWindow*)ContextWindow->GetOSWindowHandle();
 			if(CocoaWindow && [CocoaWindow openGLView])
 			{
-				FSlateTextView* TextView = (FSlateTextView*)[CocoaWindow openGLView];
+				FCocoaTextView* TextView = (FCocoaTextView*)[CocoaWindow openGLView];
 				[[TextView inputContext] discardMarkedText];
 				[TextView unmarkText];
 			}
@@ -580,13 +145,13 @@ void FMacTextInputMethodSystem::ActivateContext(const TSharedRef<ITextInputMetho
 	if(GenericWindow.IsValid())
 	{
 		MacNotifier->SetContextWindow(GenericWindow);
-		FSlateCocoaWindow* CocoaWindow = (FSlateCocoaWindow*)GenericWindow->GetOSWindowHandle();
+		FCocoaWindow* CocoaWindow = (FCocoaWindow*)GenericWindow->GetOSWindowHandle();
 		if(CocoaWindow && [CocoaWindow openGLView])
 		{
 			NSView* GLView = [CocoaWindow openGLView];
-			if([GLView isKindOfClass:[FSlateTextView class]])
+			if([GLView isKindOfClass:[FCocoaTextView class]])
 			{
-				FSlateTextView* TextView = (FSlateTextView*)GLView;
+				FCocoaTextView* TextView = (FCocoaTextView*)GLView;
 				[TextView activateInputMethod:Context];
 				return;
 			}
@@ -609,13 +174,13 @@ void FMacTextInputMethodSystem::DeactivateContext(const TSharedRef<ITextInputMet
 	TSharedPtr<FGenericWindow> GenericWindow = MacNotifier->GetContextWindow();
 	if(GenericWindow.IsValid())
 	{
-		FSlateCocoaWindow* CocoaWindow = (FSlateCocoaWindow*)GenericWindow->GetOSWindowHandle();
+		FCocoaWindow* CocoaWindow = (FCocoaWindow*)GenericWindow->GetOSWindowHandle();
 		if(CocoaWindow && [CocoaWindow openGLView])
 		{
 			NSView* GLView = [CocoaWindow openGLView];
-			if([GLView isKindOfClass:[FSlateTextView class]])
+			if([GLView isKindOfClass:[FCocoaTextView class]])
 			{
-				FSlateTextView* TextView = (FSlateTextView*)GLView;
+				FCocoaTextView* TextView = (FCocoaTextView*)GLView;
 				[TextView deactivateInputMethod];
 				return;
 			}
diff --git a/Engine/Source/Runtime/Core/Private/Mac/MacWindow.cpp b/Engine/Source/Runtime/Core/Private/Mac/MacWindow.cpp
--- a/Engine/Source/Runtime/Core/Private/Mac/MacWindow.cpp
+++ b/Engine/Source/Runtime/Core/Private/Mac/MacWindow.cpp
@@ -3,637 +3,9 @@
 #include "Core.h"
 #include "MacWindow.h"
 #include "MacApplication.h"
+#include "CocoaTextView.h"
 
-TArray< FSlateCocoaWindow* > GRunningModalWindows;
-
-/**
- * Custom window class used for input handling
- */
-@implementation FSlateCocoaWindow
-
-@synthesize bForwardEvents;
-@synthesize TargetWindowMode;
-@synthesize PreFullScreenRect;
-
-- (id)initWithContentRect:(NSRect)ContentRect styleMask:(NSUInteger)Style backing:(NSBackingStoreType)BufferingType defer:(BOOL)Flag
-{
-	id NewSelf = [super initWithContentRect:ContentRect styleMask:Style backing:BufferingType defer:Flag];
-	if(NewSelf)
-	{
-		WindowMode = EWindowMode::Windowed;
-		bAcceptsInput = false;
-		bRoundedCorners = false;
-		bDisplayReconfiguring = false;
-		bDeferOrderFront = false;
-		DeferOpacity = 0.0f;
-		bRenderInitialised = false;
-		bZoomed = [super isZoomed];
-		self.bForwardEvents = true;
-		self.TargetWindowMode = EWindowMode::Windowed;
-		[super setAlphaValue:DeferOpacity];
-		DeferFrame = [super frame];
-		self.PreFullScreenRect = DeferFrame;
-		bDeferSetFrame = false;
-		bDeferSetOrigin = false;
-	}
-	return NewSelf;
-}
-
-- (NSRect)openGLFrame
-{
-	if(self.TargetWindowMode == EWindowMode::Fullscreen || WindowMode == EWindowMode::Fullscreen)
-	{
-		return self.PreFullScreenRect;
-	}
-	else if([self styleMask] & (NSTexturedBackgroundWindowMask))
-	{
-		return (!bDeferSetFrame ? [self frame] : DeferFrame);
-	}
-	else
-	{
-		return (!bDeferSetFrame ? [[self contentView] frame] : [self contentRectForFrameRect:DeferFrame]);
-	}
-}
-
-
-- (NSView*)openGLView
-{
-	if([self styleMask] & (NSTexturedBackgroundWindowMask))
-	{
-		NSView* SuperView = [[self contentView] superview];
-		for(NSView* View in [SuperView subviews])
-		{
-			if([View isKindOfClass:[FSlateTextView class]])
-			{
-				return View;
-			}
-		}
-		return nil;
-	}
-	else
-	{
-		return [self contentView];
-	}
-}
-
-- (void)performDeferredOrderFront
-{
-	if(!bRenderInitialised)
-	{
-		bRenderInitialised = true;
-	}
-	
-	if(bDeferOrderFront)
-	{
-		if(!(bDeferSetFrame || bDeferSetOrigin))
-		{
-			bDeferOrderFront = false;
-			[super setAlphaValue:DeferOpacity];
-		}
-		else
-		{
-			[self performDeferredSetFrame];
-		}
-	}
-}
-
-- (void)performDeferredSetFrame
-{
-	if(bRenderInitialised && (bDeferSetFrame || bDeferSetOrigin))
-	{
-		dispatch_block_t Block = ^{
-			if(!bDeferSetFrame && bDeferSetOrigin)
-			{
-				DeferFrame.size = [self frame].size;
-			}
-			
-			[super setFrame:DeferFrame display:YES];
-		};
-		
-		if([NSThread isMainThread])
-		{
-			Block();
-		}
-		else
-		{
-			dispatch_async(dispatch_get_main_queue(), Block);
-		}
-		
-		bDeferSetFrame = false;
-		bDeferSetOrigin = false;
-	}
-}
-
-- (void)orderWindow:(NSWindowOrderingMode)OrderingMode relativeTo:(NSInteger)OtherWindowNumber
-{
-	bool bModal = GRunningModalWindows.Num() == 0 || GRunningModalWindows.Last() == self || [self styleMask] == NSBorderlessWindowMask;
-	if(OrderingMode == NSWindowOut || bModal)
-	{
-		if([self alphaValue] > 0.0f)
-		{
-			[self performDeferredSetFrame];
-		}
-		[super orderWindow:OrderingMode relativeTo:OtherWindowNumber];
-	}
-}
-
-- (bool)roundedCorners
-{
-    return bRoundedCorners;
-}
-
-- (void)setRoundedCorners:(bool)bUseRoundedCorners
-{
-	bRoundedCorners = bUseRoundedCorners;
-}
-
-- (void)setAcceptsInput:(bool)InAcceptsInput
-{
-	bAcceptsInput = InAcceptsInput;
-}
-
-- (void)redrawContents
-{
-	if(bNeedsRedraw && bForwardEvents && ([self isVisible] && [super alphaValue] > 0.0f))
-	{
-		MacApplication->OnWindowRedrawContents( self );
-	}
-	bNeedsRedraw = false;
-}
-
-- (void)setWindowMode:(EWindowMode::Type)NewWindowMode
-{
-	WindowMode = NewWindowMode;
-}
-
-- (EWindowMode::Type)windowMode
-{
-	return WindowMode;
-}
-
-- (void)setDisplayReconfiguring:(bool)bIsDisplayReconfiguring
-{
-	bDisplayReconfiguring = bIsDisplayReconfiguring;
-}
-
-- (void)orderFrontAndMakeMain:(bool)bMain andKey:(bool)bKey
-{
-	if ([NSApp isHidden] == NO)
-	{
-		bool bBringToFront = GRunningModalWindows.Num() == 0 || GRunningModalWindows.Last() == self || [self styleMask] == NSBorderlessWindowMask;
-		if (bBringToFront)
-		{
-			[self orderFront:nil];
-		}
-		
-		if (bMain && [self canBecomeMainWindow] && self != [NSApp mainWindow])
-		{
-			[self makeMainWindow];
-		}
-		if (bKey && [self canBecomeKeyWindow] && self != [NSApp keyWindow])
-		{
-			[self makeKeyWindow];
-		}
-	}
-}
-
-// Following few methods overload NSWindow's methods from Cocoa API, so have to use Cocoa's BOOL (signed char), not bool (unsigned int)
-- (BOOL)canBecomeMainWindow
-{
-	bool bNoModalOrCurrent = GRunningModalWindows.Num() == 0 || GRunningModalWindows.Last() == self;
-	return bAcceptsInput && ([self styleMask] != NSBorderlessWindowMask) && bNoModalOrCurrent;
-}
-
-- (BOOL)canBecomeKeyWindow
-{
-	return bAcceptsInput && ![self ignoresMouseEvents];
-}
-
-- (BOOL)validateMenuItem:(NSMenuItem *)MenuItem
-{
-	// Borderless windows we use do not automatically handle first responder's actions, so we force it here
-	return ([MenuItem action] == @selector(performClose:) || [MenuItem action] == @selector(performMiniaturize:) || [MenuItem action] == @selector(performZoom:)) ? YES : [super validateMenuItem:MenuItem];
-}
-
-- (void)setAlphaValue:(CGFloat)WindowAlpha
-{
-	if(!bRenderInitialised)
-	{
-		DeferOpacity = WindowAlpha;
-		bDeferOrderFront = true;
-	}
-	else
-	{
-		if([self isVisible] && WindowAlpha > 0.0f)
-		{
-			[self performDeferredSetFrame];
-		}
-		[super setAlphaValue:WindowAlpha];
-	}
-}
-
-- (void)orderOut:(id)Sender
-{
-	bDeferOrderFront = false;
-	
-	[super orderOut:Sender];
-}
-
-- (void)performClose:(id)Sender
-{
-	bDeferOrderFront = false;
-	
-	[self close];
-}
-
-- (void)performMiniaturize:(id)Sender
-{
-	[self miniaturize: self];
-}
-
-- (void)performZoom:(id)Sender
-{
-	bZoomed = !bZoomed;
-	[self zoom: self];
-}
-
-- (void)setFrame:(NSRect)FrameRect display:(BOOL)Flag
-{
-	NSSize Size = [self frame].size;
-	NSSize NewSize = FrameRect.size;
-	if(!bRenderInitialised || ([self isVisible] && [super alphaValue] > 0.0f && (Size.width > 1 || Size.height > 1 || NewSize.width > 1 || NewSize.height > 1)))
-	{
-		[super setFrame:FrameRect display:Flag];
-		bDeferSetFrame = false;
-	}
-	else
-	{
-		bDeferSetFrame = true;
-		DeferFrame = FrameRect;
-		if(self.bForwardEvents)
-		{
-			MacApplication->OnWindowDidResize( self );
-		}
-	}
-}
-
-- (void)setFrameOrigin:(NSPoint)Point
-{
-	NSSize Size = [self frame].size;
-	if(!bRenderInitialised || ([self isVisible] && [super alphaValue] > 0.0f && (Size.width > 1 || Size.height > 1)))
-	{
-		[super setFrameOrigin:Point];
-		bDeferSetOrigin = false;
-	}
-	else
-	{
-		bDeferSetOrigin = true;
-		DeferFrame.origin = Point;
-	}
-}
-
-- (void)keyDown:(NSEvent *)Event
-{
-	if(self.bForwardEvents)
-	{
-		MacApplication->ProcessEvent( Event );
-	}
-}
-
-- (void)keyUp:(NSEvent *)Event
-{
-	if(self.bForwardEvents)
-	{
-		MacApplication->ProcessEvent( Event );
-	}
-}
-
-- (void)windowWillEnterFullScreen:(NSNotification *)notification
-{
-	// Handle clicking on the titlebar fullscreen item
-	if(self.TargetWindowMode == EWindowMode::Windowed)
-	{
-		// @todo: Fix fullscreen mode mouse coordinate handling - for now default to windowed fullscreen
-		self.TargetWindowMode = EWindowMode::WindowedFullscreen;
-	}
-}
-
-- (void)windowDidEnterFullScreen:(NSNotification *)notification
-{
-	WindowMode = self.TargetWindowMode;
-	if(self.bForwardEvents)
-	{
-		MacApplication->OnWindowDidResize( self );
-	}
-}
-
-- (void)windowDidExitFullScreen:(NSNotification *)notification
-{
-	WindowMode = EWindowMode::Windowed;
-	self.TargetWindowMode = EWindowMode::Windowed;
-	if(self.bForwardEvents)
-	{
-		MacApplication->OnWindowDidResize( self );
-	}
-}
-
-- (void)windowDidBecomeKey:(NSNotification *)Notification
-{
-	if([NSApp isHidden] == NO)
-	{
-		if(GRunningModalWindows.Num() == 0 || GRunningModalWindows.Last() == self || [self styleMask] == NSBorderlessWindowMask)
-		{
-			[self orderFrontAndMakeMain:false andKey:false];
-		}
-		else
-		{
-			[GRunningModalWindows.Last() orderFrontAndMakeMain:true andKey:true];
-		}
-	}
-
-	if(self.bForwardEvents)
-	{
-		MacApplication->OnWindowDidBecomeKey( self );
-	}
-}
-
-- (void)windowDidResignKey:(NSNotification *)Notification
-{
-	[self setMovable: YES];
-	[self setMovableByWindowBackground: NO];
-	
-	if(self.bForwardEvents)
-	{
-		MacApplication->OnWindowDidResignKey( self );
-	}
-}
-
-- (void)windowWillMove:(NSNotification *)Notification
-{
-	if(self.bForwardEvents)
-	{
-		MacApplication->OnWindowWillMove( self );
-	}
-}
-
-- (void)windowDidMove:(NSNotification *)Notification
-{
-	bZoomed = [self isZoomed];
-	
-	NSView* OpenGLView = [self openGLView];
-	[[NSNotificationCenter defaultCenter] postNotificationName:NSViewGlobalFrameDidChangeNotification object:OpenGLView];
-	
-	if(self.bForwardEvents)
-	{
-		MacApplication->OnWindowDidMove( self );
-	}
-}
-
-- (void)windowDidChangeScreen:(NSNotification *)notification
-{
-	// You'd think this would be a good place to handle un/parenting...
-	// However, do that here and your dragged window will disappear when you drag it onto another screen!
-	// The windowdidChangeScreen notification only comes after you finish dragging.
-	// It does however, work fine for handling display arrangement changes that cause a window to go offscreen.
-	if(bDisplayReconfiguring)
-	{
-		NSScreen* Screen = [self screen];
-		NSRect Frame = [self frame];
-		NSRect VisibleFrame = [Screen visibleFrame];
-		if(NSContainsRect(VisibleFrame, Frame) == NO)
-		{
-			// May need to scale the window to fit if it is larger than the new display.
-			if (Frame.size.width > VisibleFrame.size.width || Frame.size.height > VisibleFrame.size.height)
-			{
-				NSRect NewFrame;
-				NewFrame.size.width = Frame.size.width > VisibleFrame.size.width ? VisibleFrame.size.width : Frame.size.width;
-				NewFrame.size.height = Frame.size.height > VisibleFrame.size.height ? VisibleFrame.size.height : Frame.size.height;
-				NewFrame.origin = VisibleFrame.origin;
-				
-				[self setFrame:NewFrame display:NO];
-			}
-			else
-			{
-				NSRect Intersection = NSIntersectionRect(VisibleFrame, Frame);
-				NSPoint Origin = Frame.origin;
-				
-				// If there's at least something on screen, try shifting it entirely on screen.
-				if(Intersection.size.width > 0 && Intersection.size.height > 0)
-				{
-					CGFloat X = Frame.size.width - Intersection.size.width;
-					CGFloat Y = Frame.size.height - Intersection.size.height;
-					
-					if(Intersection.size.width+Intersection.origin.x >= VisibleFrame.size.width+VisibleFrame.origin.x)
-					{
-						Origin.x -= X;
-					}
-					else if(Origin.x < VisibleFrame.origin.x)
-					{
-						Origin.x += X;
-					}
-					
-					if(Intersection.size.height+Intersection.origin.y >= VisibleFrame.size.height+VisibleFrame.origin.y)
-					{
-						Origin.y -= Y;
-					}
-					else if(Origin.y < VisibleFrame.origin.y)
-					{
-						Origin.y += Y;
-					}
-				}
-				else
-				{
-					Origin = VisibleFrame.origin;
-				}
-				
-				[self setFrameOrigin:Origin];
-			}
-		}
-	}
-}
-
-- (void)windowDidResize:(NSNotification *)Notification
-{
-	bZoomed = [self isZoomed];
-	if(self.bForwardEvents)
-	{
-		MacApplication->OnWindowDidResize( self );
-	}
-	bNeedsRedraw = true;
-}
-
-- (void)windowWillClose:(NSNotification *)notification
-{
-	if(self.bForwardEvents && MacApplication)
-	{
-		MacApplication->OnWindowDidClose( self );
-	}
-}
-
-- (void)mouseDown:(NSEvent*)Event
-{
-	if(self.bForwardEvents)
-	{
-		MacApplication->ProcessEvent( Event );
-	}
-}
-
-- (void)rightMouseDown:(NSEvent*)Event
-{
-	// Really we shouldn't be doing this - on OS X only left-click changes focus,
-	// but for the moment it is easier than changing Slate.
-	if([self canBecomeKeyWindow])
-	{
-		[self makeKeyWindow];
-	}
-
-	if(self.bForwardEvents)
-	{
-		MacApplication->ProcessEvent( Event );
-	}
-}
-
-- (void)otherMouseDown:(NSEvent*)Event
-{
-	if(self.bForwardEvents)
-	{
-		MacApplication->ProcessEvent( Event );
-	}
-}
-
-- (void)mouseUp:(NSEvent*)Event
-{
-	if(self.bForwardEvents)
-	{
-		MacApplication->ProcessEvent( Event );
-	}
-}
-
-- (void)rightMouseUp:(NSEvent*)Event
-{
-	if(self.bForwardEvents)
-	{
-		MacApplication->ProcessEvent( Event );
-	}
-}
-
-- (void)otherMouseUp:(NSEvent*)Event
-{
-	if(self.bForwardEvents)
-	{
-		MacApplication->ProcessEvent( Event );
-	}
-}
-
-- (NSDragOperation)draggingEntered:(id < NSDraggingInfo >)Sender
-{
-	return NSDragOperationGeneric;
-}
-
-- (void)draggingExited:(id < NSDraggingInfo >)Sender
-{
-	if(self.bForwardEvents)
-	{
-		MacApplication->OnDragOut( self );
-	}
-}
-
-- (NSDragOperation)draggingUpdated:(id < NSDraggingInfo >)Sender
-{
-	if(self.bForwardEvents)
-	{
-		MacApplication->OnDragOver( self );
-	}
-	return NSDragOperationGeneric;
-}
-
-- (BOOL)prepareForDragOperation:(id < NSDraggingInfo >)Sender
-{
-	if(self.bForwardEvents)
-	{
-		MacApplication->OnDragEnter( self, [Sender draggingPasteboard] );
-	}
-	return YES;
-}
-
-- (BOOL)performDragOperation:(id < NSDraggingInfo >)Sender
-{
-	if(self.bForwardEvents)
-	{
-		MacApplication->OnDragDrop( self );
-	}
-	return YES;
-}
-
-- (BOOL)isMovable
-{
-	BOOL Movable = (BOOL)(MacApplication->IsWindowMovable(self, NULL) && [super isMovable]);
-	return Movable;
-}
-
-@end
-
-/**
- * Custom window class used for mouse capture
- */
-@implementation FMouseCaptureWindow
-
-- (id)initWithTargetWindow: (FSlateCocoaWindow*)Window
-{
-	self = [super initWithContentRect: [[Window screen] frame] styleMask: NSBorderlessWindowMask backing: NSBackingStoreBuffered defer: NO];
-	[self setBackgroundColor: [NSColor clearColor]];
-	[self setOpaque: NO];
-	[self setLevel: NSMainMenuWindowLevel + 1];
-	[self setIgnoresMouseEvents: NO];
-	[self setAcceptsMouseMovedEvents: YES];
-	[self setHidesOnDeactivate: YES];
-
-	TargetWindow = Window;
-
-	return self;
-}
-
-- (FSlateCocoaWindow*)targetWindow
-{
-	return TargetWindow;
-}
-
-- (void)setTargetWindow: (FSlateCocoaWindow*)Window
-{
-	TargetWindow = Window;
-}
-
-- (void)mouseDown:(NSEvent*)Event
-{
-	MacApplication->ProcessEvent( Event );
-}
-
-- (void)rightMouseDown:(NSEvent*)Event
-{
-	MacApplication->ProcessEvent( Event );
-}
-
-- (void)otherMouseDown:(NSEvent*)Event
-{
-	MacApplication->ProcessEvent( Event );
-}
-
-- (void)mouseUp:(NSEvent*)Event
-{
-	MacApplication->ProcessEvent( Event );
-}
-
-- (void)rightMouseUp:(NSEvent*)Event
-{
-	MacApplication->ProcessEvent( Event );
-}
-
-- (void)otherMouseUp:(NSEvent*)Event
-{
-	MacApplication->ProcessEvent( Event );
-}
-
-@end
+TArray< FCocoaWindow* > FMacWindow::RunningModalWindows;
 
 FMacWindow::~FMacWindow()
 {
@@ -712,7 +84,7 @@ void FMacWindow::Initialize( FMacApplication* const Application, const TSharedRe
 		WindowStyle &= ~(NSTexturedBackgroundWindowMask);
 	}
 
-	WindowHandle = [[FSlateCocoaWindow alloc] initWithContentRect: ViewRect styleMask: WindowStyle backing: NSBackingStoreBuffered defer: NO];
+	WindowHandle = [[FCocoaWindow alloc] initWithContentRect: ViewRect styleMask: WindowStyle backing: NSBackingStoreBuffered defer: NO];
 
 	if( WindowHandle == NULL )
 	{
@@ -803,7 +175,7 @@ FMacWindow::FMacWindow()
 	PreFullscreenWindowRect.origin.x = PreFullscreenWindowRect.origin.y = PreFullscreenWindowRect.size.width = PreFullscreenWindowRect.size.height = 0.0f;
 }
 
-FSlateCocoaWindow* FMacWindow::GetWindowHandle() const
+FCocoaWindow* FMacWindow::GetWindowHandle() const
 {
 	return WindowHandle;
 }
@@ -880,10 +252,10 @@ void FMacWindow::Destroy()
 		// The destructor will release the window once, so we need to still have a reference.
 		// This makes it possible to correctly handle any NSEvent's sent to the window during destruction
 		// as the WindowHandle is still a valid NSWindow. Unlike HWNDs accessing a destructed NSWindow* is fatal.
-		FSlateCocoaWindow* Window = [WindowHandle retain];
+		FCocoaWindow* Window = [WindowHandle retain];
 		if(Definition->IsModalWindow)
 		{
-			GRunningModalWindows.Remove(WindowHandle);
+			RemoveModalWindow(WindowHandle);
 		}
 		Window.bForwardEvents = false;
 		if( MacApplication->OnWindowDestroyed( Window ) )
@@ -892,7 +264,7 @@ void FMacWindow::Destroy()
 			
 			// Then change the focus to something useful, either the previous in the stack
 			TSharedPtr<FMacWindow> KeyWindow = MacApplication->GetKeyWindow();
-			if( KeyWindow.IsValid() && [Window isKeyWindow] && KeyWindow->GetOSWindowHandle() && ![(FSlateCocoaWindow*)KeyWindow->GetOSWindowHandle() isMiniaturized] )
+			if( KeyWindow.IsValid() && [Window isKeyWindow] && KeyWindow->GetOSWindowHandle() && ![(FCocoaWindow*)KeyWindow->GetOSWindowHandle() isMiniaturized] )
 			{
 				// Activate specified previous window if still present, provided it isn't minimized
 				KeyWindow->SetWindowFocus();
@@ -947,9 +319,9 @@ void FMacWindow::Show()
 	SCOPED_AUTORELEASE_POOL;
 	if (!bIsVisible)
 	{
-		if(Definition->IsModalWindow && !GRunningModalWindows.Contains(WindowHandle))
+		if(Definition->IsModalWindow)
 		{
-			GRunningModalWindows.Add(WindowHandle);
+			AddModalWindow(WindowHandle);
 		}
 		bool bMakeMainAndKey = ([WindowHandle canBecomeKeyWindow] && Definition->ActivateWhenFirstShown);
 		[WindowHandle orderFrontAndMakeMain:bMakeMainAndKey andKey:bMakeMainAndKey];
@@ -981,7 +353,7 @@ void FMacWindow::Hide()
 		bIsVisible = false;
 		if(Definition->IsModalWindow)
 		{
-			GRunningModalWindows.Remove(WindowHandle);
+			RemoveModalWindow(WindowHandle);
 		}
 		[WindowHandle orderOut:nil];
 	}
@@ -1177,7 +549,7 @@ bool FMacWindow::OnIMKKeyDown(NSEvent* Event)
 {
 	if(WindowHandle && [WindowHandle openGLView])
 	{
-		FSlateTextView* View = (FSlateTextView*)[WindowHandle openGLView];
+		FCocoaTextView* View = (FCocoaTextView*)[WindowHandle openGLView];
 		return View && [View imkKeyDown:Event];
 	}
 	else
@@ -1185,3 +557,26 @@ bool FMacWindow::OnIMKKeyDown(NSEvent* Event)
 		return false;
 	}
 }
+
+void FMacWindow::AddModalWindow(FCocoaWindow* Window)
+{
+	if(!RunningModalWindows.Contains(Window))
+	{
+		RunningModalWindows.Add(Window);
+	}
+}
+
+void FMacWindow::RemoveModalWindow(FCocoaWindow* Window)
+{
+	RunningModalWindows.Remove(Window);
+}
+
+FCocoaWindow* FMacWindow::CurrentModalWindow(void)
+{
+	FCocoaWindow* Window = nil;
+	if(RunningModalWindows.Num() > 0)
+	{
+		Window = RunningModalWindows.Last();
+	}
+	return Window;
+}
diff --git a/Engine/Source/Runtime/Core/Public/Mac/CocoaTextView.h b/Engine/Source/Runtime/Core/Public/Mac/CocoaTextView.h
new file mode 100644
--- /dev/null
+++ b/Engine/Source/Runtime/Core/Public/Mac/CocoaTextView.h
@@ -0,0 +1,94 @@
+// Copyright 1998-2014 Epic Games, Inc. All Rights Reserved.
+#pragma once
+
+#include "Core.h"
+#include "ITextInputMethodSystem.h"
+
+#ifdef __OBJC__
+
+@interface FCocoaTextView : NSView <NSTextInputClient>
+{
+	TSharedPtr<ITextInputMethodContext> IMMContext;
+	NSRange markedRange;
+	bool reallyHandledEvent;
+}
+
+- (void)activateInputMethod:(const TSharedRef<ITextInputMethodContext>&)InContext;
+- (void)deactivateInputMethod;
+
+- (bool)imkKeyDown:(NSEvent*)Event;
+
+//@protocol NSTextInputClient
+//@required
+/* The receiver inserts aString replacing the content specified by replacementRange. aString can be either an NSString or NSAttributedString instance.
+ */
+- (void)insertText:(id)aString replacementRange:(NSRange)replacementRange;
+
+/* The receiver invokes the action specified by aSelector.
+ */
+- (void)doCommandBySelector:(SEL)aSelector;
+
+/* The receiver inserts aString replacing the content specified by replacementRange. aString can be either an NSString or NSAttributedString instance. selectedRange specifies the selection inside the string being inserted; hence, the location is relative to the beginning of aString. When aString is an NSString, the receiver is expected to render the marked text with distinguishing appearance (i.e. NSTextView renders with -markedTextAttributes).
+ */
+- (void)setMarkedText:(id)aString selectedRange:(NSRange)selectedRange replacementRange:(NSRange)replacementRange;
+
+/* The receiver unmarks the marked text. If no marked text, the invocation of this method has no effect.
+ */
+- (void)unmarkText;
+
+/* Returns the selection range. The valid location is from 0 to the document length.
+ */
+- (NSRange)selectedRange;
+
+/* Returns the marked range. Returns {NSNotFound, 0} if no marked range.
+ */
+- (NSRange)markedRange;
+
+/* Returns whether or not the receiver has marked text.
+ */
+- (BOOL)hasMarkedText;
+
+/* Returns attributed string specified by aRange. It may return nil. If non-nil return value and actualRange is non-NULL, it contains the actual range for the return value. The range can be adjusted from various reasons (i.e. adjust to grapheme cluster boundary, performance optimization, etc).
+ */
+- (NSAttributedString *)attributedSubstringForProposedRange:(NSRange)aRange actualRange:(NSRangePointer)actualRange;
+
+/* Returns an array of attribute names recognized by the receiver.
+ */
+- (NSArray*)validAttributesForMarkedText;
+
+/* Returns the first logical rectangular area for aRange. The return value is in the screen coordinate. The size value can be negative if the text flows to the left. If non-NULL, actuallRange contains the character range corresponding to the returned area.
+ */
+- (NSRect)firstRectForCharacterRange:(NSRange)aRange actualRange:(NSRangePointer)actualRange;
+
+/* Returns the index for character that is nearest to aPoint. aPoint is in the screen coordinate system.
+ */
+- (NSUInteger)characterIndexForPoint:(NSPoint)aPoint;
+
+//@optional
+/* Returns an attributed string representing the receiver's document content. An NSTextInputClient can implement this interface if can be done efficiently. The caller of this interface can random access arbitrary portions of the receiver's content more efficiently.
+ */
+//- (NSAttributedString *)attributedString;
+
+/* Returns the fraction of distance for aPoint from the left side of the character. This allows caller to perform precise selection handling.
+ */
+//- (CGFloat)fractionOfDistanceThroughGlyphForPoint:(NSPoint)aPoint;
+
+/* Returns the baseline position relative to the origin of rectangle returned by -firstRectForCharacterRange:actualRange:. This information allows the caller to access finer-grained character position inside the NSTextInputClient document.
+ */
+//- (CGFloat)baselineDeltaForCharacterAtIndex:(NSUInteger)anIndex;
+
+/* Returns the window level of the receiver. An NSTextInputClient can implement this interface to specify its window level if it is higher than NSFloatingWindowLevel.
+ */
+- (NSInteger)windowLevel;
+
+/* Returns if the marked text is in vertical layout.
+ */
+//- (BOOL)drawsVerticallyForCharacterAtIndex:(NSUInteger)charIndex NS_AVAILABLE_MAC(10_6);
+
+@end
+
+#else // __OBJC__
+
+class FCocoaTextView;
+
+#endif // __OBJC__
diff --git a/Engine/Source/Runtime/Core/Public/Mac/CocoaWindow.h b/Engine/Source/Runtime/Core/Public/Mac/CocoaWindow.h
new file mode 100644
--- /dev/null
+++ b/Engine/Source/Runtime/Core/Public/Mac/CocoaWindow.h
@@ -0,0 +1,93 @@
+// Copyright 1998-2014 Epic Games, Inc. All Rights Reserved.
+
+#pragma once
+
+#ifdef __OBJC__
+
+/**
+ * Custom window class used for input handling
+ */
+@interface FCocoaWindow : NSWindow <NSWindowDelegate>
+{
+	EWindowMode::Type WindowMode;
+	NSRect DeferFrame;
+	CGFloat DeferOpacity;
+	bool bAcceptsInput;
+	bool bRoundedCorners;
+	bool bDisplayReconfiguring;
+	bool bDeferOrderFront;
+	bool bDeferSetFrame;
+	bool bDeferSetOrigin;
+	bool bRenderInitialised;
+	bool bNeedsRedraw;
+@public
+	bool bZoomed;
+}
+
+@property (assign) NSRect PreFullScreenRect;
+@property (assign) EWindowMode::Type TargetWindowMode;
+@property (assign) bool bForwardEvents;
+
+/** Get the frame filled by a child OpenGL view, which may cover the window or fill the content view depending upon the window style.
+ @return The NSRect for a child OpenGL view. */
+- (NSRect)openGLFrame;
+
+/** Get the view used for OpenGL rendering. @return The OpenGL view for rendering. */
+- (NSView*)openGLView;
+
+/** Perform render thread deferred order front operation. */
+- (void)performDeferredOrderFront;
+
+/** Perform deferred set frame operation */
+- (void)performDeferredSetFrame;
+
+/** Set whether the window should display with rounded corners. */
+- (void)setRoundedCorners:(bool)bUseRoundedCorners;
+
+/** Get whether the window should display with rounded corners.
+ @return True when window corners should be rounded, else false. */
+- (bool)roundedCorners;
+
+/** Lets window know if its owner (SWindow) accepts input */
+- (void)setAcceptsInput:(bool)InAcceptsInput;
+
+/** Redraws window's contents. */
+- (void)redrawContents;
+
+/** Set the initial window mode. */
+- (void)setWindowMode:(EWindowMode::Type)WindowMode;
+
+/**	@return The current mode for this Cocoa window. */
+- (EWindowMode::Type)windowMode;
+
+/** Mutator that specifies that the display arrangement is being reconfigured when bIsDisplayReconfiguring is true. */
+- (void)setDisplayReconfiguring:(bool)bIsDisplayReconfiguring;
+
+/** Order window to the front. */
+- (void)orderFrontAndMakeMain:(bool)bMain andKey:(bool)bKey;
+
+@end
+
+/**
+ * Custom window class used for mouse capture
+ */
+@interface FMouseCaptureWindow : NSWindow <NSWindowDelegate>
+{
+	FCocoaWindow*	TargetWindow;
+}
+
+- (id)initWithTargetWindow: (FCocoaWindow*)Window;
+- (FCocoaWindow*)targetWindow;
+- (void)setTargetWindow: (FCocoaWindow*)Window;
+
+@end
+
+#else // __OBJC__
+
+class FCocoaWindow;
+class FMouseCaptureWindow;
+class NSWindow;
+class NSEvent;
+class NSScreen;
+
+#endif // __OBJC__
diff --git a/Engine/Source/Runtime/Core/Public/Mac/MacApplication.h b/Engine/Source/Runtime/Core/Public/Mac/MacApplication.h
--- a/Engine/Source/Runtime/Core/Public/Mac/MacApplication.h
+++ b/Engine/Source/Runtime/Core/Public/Mac/MacApplication.h
@@ -65,7 +65,7 @@ public:
 
 	void OnWindowDraggingFinished();
 
-	bool IsWindowMovable(FSlateCocoaWindow* Win, bool* OutMovableByBackground);
+	bool IsWindowMovable(FCocoaWindow* Win, bool* OutMovableByBackground);
 
 	void ResetModifierKeys() { ModifierKeysFlags = 0; }
 
@@ -87,19 +87,19 @@ public:
 
 public:
 
-	void OnDragEnter( FSlateCocoaWindow* Window, void *InPasteboard );
-	void OnDragOver( FSlateCocoaWindow* Window );
-	void OnDragOut( FSlateCocoaWindow* Window );
-	void OnDragDrop( FSlateCocoaWindow* Window );
+	void OnDragEnter( FCocoaWindow* Window, void *InPasteboard );
+	void OnDragOver( FCocoaWindow* Window );
+	void OnDragOut( FCocoaWindow* Window );
+	void OnDragDrop( FCocoaWindow* Window );
 
-	void OnWindowDidBecomeKey( FSlateCocoaWindow* Window );
-	void OnWindowDidResignKey( FSlateCocoaWindow* Window );
-	void OnWindowWillMove( FSlateCocoaWindow* Window );
-	void OnWindowDidMove( FSlateCocoaWindow* Window );
-	void OnWindowDidResize( FSlateCocoaWindow* Window );
-	void OnWindowRedrawContents( FSlateCocoaWindow* Window );
-	void OnWindowDidClose( FSlateCocoaWindow* Window );
-	bool OnWindowDestroyed( FSlateCocoaWindow* Window );
+	void OnWindowDidBecomeKey( FCocoaWindow* Window );
+	void OnWindowDidResignKey( FCocoaWindow* Window );
+	void OnWindowWillMove( FCocoaWindow* Window );
+	void OnWindowDidMove( FCocoaWindow* Window );
+	void OnWindowDidResize( FCocoaWindow* Window );
+	void OnWindowRedrawContents( FCocoaWindow* Window );
+	void OnWindowDidClose( FCocoaWindow* Window );
+	bool OnWindowDestroyed( FCocoaWindow* Window );
 
 	void OnMouseCursorLock( bool bLockEnabled );
 
@@ -113,12 +113,12 @@ private:
 	TCHAR ConvertChar( TCHAR Character );
 	TCHAR TranslateCharCode( TCHAR CharCode, uint32 KeyCode );
 
-	FSlateCocoaWindow* FindEventWindow( NSEvent* CocoaEvent );
+	FCocoaWindow* FindEventWindow( NSEvent* CocoaEvent );
 	TSharedPtr<FGenericWindow> LocateWindowUnderCursor( const FVector2D& CursorPos );
 
 	NSScreen* FindScreenByPoint( int32 X, int32 Y ) const;
 
-	void UpdateMouseCaptureWindow( FSlateCocoaWindow* TargetWindow );
+	void UpdateMouseCaptureWindow( FCocoaWindow* TargetWindow );
 
 	void HandleModifierChange(TSharedPtr< FMacWindow > CurrentEventWindow, NSUInteger NewModifierFlags, NSUInteger FlagsShift, NSUInteger UE4Shift, EMacModifierKeys TranslatedCode);
 
@@ -142,7 +142,7 @@ private:
 
 	TSharedRef< class HIDInputInterface > HIDInput;
 
-	FSlateCocoaWindow* DraggedWindow;
+	FCocoaWindow* DraggedWindow;
 
 	FMouseCaptureWindow* MouseCaptureWindow;
 	bool bIsMouseCaptureEnabled;
diff --git a/Engine/Source/Runtime/Core/Public/Mac/MacTextInputMethodSystem.h b/Engine/Source/Runtime/Core/Public/Mac/MacTextInputMethodSystem.h
--- a/Engine/Source/Runtime/Core/Public/Mac/MacTextInputMethodSystem.h
+++ b/Engine/Source/Runtime/Core/Public/Mac/MacTextInputMethodSystem.h
@@ -4,95 +4,6 @@
 #include "Core.h"
 #include "ITextInputMethodSystem.h"
 
-#ifdef __OBJC__
-
-@interface FSlateTextView : NSView <NSTextInputClient>
-{
-	TSharedPtr<ITextInputMethodContext> IMMContext;
-	NSRange markedRange;
-	bool reallyHandledEvent;
-}
-
-- (void)activateInputMethod:(const TSharedRef<ITextInputMethodContext>&)InContext;
-- (void)deactivateInputMethod;
-
-- (bool)imkKeyDown:(NSEvent*)Event;
-
-//@protocol NSTextInputClient
-//@required
-/* The receiver inserts aString replacing the content specified by replacementRange. aString can be either an NSString or NSAttributedString instance.
- */
-- (void)insertText:(id)aString replacementRange:(NSRange)replacementRange;
-
-/* The receiver invokes the action specified by aSelector.
- */
-- (void)doCommandBySelector:(SEL)aSelector;
-
-/* The receiver inserts aString replacing the content specified by replacementRange. aString can be either an NSString or NSAttributedString instance. selectedRange specifies the selection inside the string being inserted; hence, the location is relative to the beginning of aString. When aString is an NSString, the receiver is expected to render the marked text with distinguishing appearance (i.e. NSTextView renders with -markedTextAttributes).
- */
-- (void)setMarkedText:(id)aString selectedRange:(NSRange)selectedRange replacementRange:(NSRange)replacementRange;
-
-/* The receiver unmarks the marked text. If no marked text, the invocation of this method has no effect.
- */
-- (void)unmarkText;
-
-/* Returns the selection range. The valid location is from 0 to the document length.
- */
-- (NSRange)selectedRange;
-
-/* Returns the marked range. Returns {NSNotFound, 0} if no marked range.
- */
-- (NSRange)markedRange;
-
-/* Returns whether or not the receiver has marked text.
- */
-- (BOOL)hasMarkedText;
-
-/* Returns attributed string specified by aRange. It may return nil. If non-nil return value and actualRange is non-NULL, it contains the actual range for the return value. The range can be adjusted from various reasons (i.e. adjust to grapheme cluster boundary, performance optimization, etc).
- */
-- (NSAttributedString *)attributedSubstringForProposedRange:(NSRange)aRange actualRange:(NSRangePointer)actualRange;
-
-/* Returns an array of attribute names recognized by the receiver.
- */
-- (NSArray*)validAttributesForMarkedText;
-
-/* Returns the first logical rectangular area for aRange. The return value is in the screen coordinate. The size value can be negative if the text flows to the left. If non-NULL, actuallRange contains the character range corresponding to the returned area.
- */
-- (NSRect)firstRectForCharacterRange:(NSRange)aRange actualRange:(NSRangePointer)actualRange;
-
-/* Returns the index for character that is nearest to aPoint. aPoint is in the screen coordinate system.
- */
-- (NSUInteger)characterIndexForPoint:(NSPoint)aPoint;
-
-//@optional
-/* Returns an attributed string representing the receiver's document content. An NSTextInputClient can implement this interface if can be done efficiently. The caller of this interface can random access arbitrary portions of the receiver's content more efficiently.
- */
-//- (NSAttributedString *)attributedString;
-
-/* Returns the fraction of distance for aPoint from the left side of the character. This allows caller to perform precise selection handling.
- */
-//- (CGFloat)fractionOfDistanceThroughGlyphForPoint:(NSPoint)aPoint;
-
-/* Returns the baseline position relative to the origin of rectangle returned by -firstRectForCharacterRange:actualRange:. This information allows the caller to access finer-grained character position inside the NSTextInputClient document.
- */
-//- (CGFloat)baselineDeltaForCharacterAtIndex:(NSUInteger)anIndex;
-
-/* Returns the window level of the receiver. An NSTextInputClient can implement this interface to specify its window level if it is higher than NSFloatingWindowLevel.
- */
-- (NSInteger)windowLevel;
-
-/* Returns if the marked text is in vertical layout.
- */
-//- (BOOL)drawsVerticallyForCharacterAtIndex:(NSUInteger)charIndex NS_AVAILABLE_MAC(10_6);
-
-@end
-
-#else // __OBJC__
-
-class FSlateTextView;
-
-#endif // __OBJC__
-
 class FMacTextInputMethodSystem : public ITextInputMethodSystem
 {
 public:
diff --git a/Engine/Source/Runtime/Core/Public/Mac/MacWindow.h b/Engine/Source/Runtime/Core/Public/Mac/MacWindow.h
--- a/Engine/Source/Runtime/Core/Public/Mac/MacWindow.h
+++ b/Engine/Source/Runtime/Core/Public/Mac/MacWindow.h
@@ -4,96 +4,7 @@
 
 #include "GenericWindow.h"
 #include "SharedPointer.h"
-
-#ifdef __OBJC__
-
-/**
- * Custom window class used for input handling
- */
-@interface FSlateCocoaWindow : NSWindow <NSWindowDelegate>
-{
-	EWindowMode::Type WindowMode;
-	NSRect DeferFrame;
-	CGFloat DeferOpacity;
-	bool bAcceptsInput;
-	bool bRoundedCorners;
-	bool bDisplayReconfiguring;
-	bool bDeferOrderFront;
-	bool bDeferSetFrame;
-	bool bDeferSetOrigin;
-	bool bRenderInitialised;
-	bool bNeedsRedraw;
-@public
-	bool bZoomed;
-}
-
-@property (assign) NSRect PreFullScreenRect;
-@property (assign) EWindowMode::Type TargetWindowMode;
-@property (assign) bool bForwardEvents;
-
-/** Get the frame filled by a child OpenGL view, which may cover the window or fill the content view depending upon the window style.
- @return The NSRect for a child OpenGL view. */
-- (NSRect)openGLFrame;
-
-/** Get the view used for OpenGL rendering. @return The OpenGL view for rendering. */
-- (NSView*)openGLView;
-
-/** Perform render thread deferred order front operation. */
-- (void)performDeferredOrderFront;
-
-/** Perform deferred set frame operation */
-- (void)performDeferredSetFrame;
-
-/** Set whether the window should display with rounded corners. */
-- (void)setRoundedCorners:(bool)bUseRoundedCorners;
-
-/** Get whether the window should display with rounded corners. 
- @return True when window corners should be rounded, else false. */
-- (bool)roundedCorners;
-
-/** Lets window know if its owner (SWindow) accepts input */
-- (void)setAcceptsInput:(bool)InAcceptsInput;
-
-/** Redraws window's contents. */
-- (void)redrawContents;
-
-/** Set the initial window mode. */
-- (void)setWindowMode:(EWindowMode::Type)WindowMode;
-
-/**	@return The current mode for this Cocoa window. */
-- (EWindowMode::Type)windowMode;
-
-/** Mutator that specifies that the display arrangement is being reconfigured when bIsDisplayReconfiguring is true. */
-- (void)setDisplayReconfiguring:(bool)bIsDisplayReconfiguring;
-
-/** Order window to the front. */
-- (void)orderFrontAndMakeMain:(bool)bMain andKey:(bool)bKey;
-
-@end
-
-/**
- * Custom window class used for mouse capture
- */
-@interface FMouseCaptureWindow : NSWindow <NSWindowDelegate>
-{
-	FSlateCocoaWindow*	TargetWindow;
-}
-
-- (id)initWithTargetWindow: (FSlateCocoaWindow*)Window;
-- (FSlateCocoaWindow*)targetWindow;
-- (void)setTargetWindow: (FSlateCocoaWindow*)Window;
-
-@end
-
-#else // __OBJC__
-
-class FSlateCocoaWindow;
-class FMouseCaptureWindow;
-class NSWindow;
-class NSEvent;
-class NSScreen;
-
-#endif // __OBJC__
+#include "CocoaWindow.h"
 
 /**
  * A platform specific implementation of FNativeWindow.
@@ -106,13 +17,17 @@ public:
 
 	static TSharedRef< FMacWindow > Make();
 
-	FSlateCocoaWindow* GetWindowHandle() const;
+	FCocoaWindow* GetWindowHandle() const;
 
 	void Initialize( class FMacApplication* const Application, const TSharedRef< FGenericWindowDefinition >& InDefinition, const TSharedPtr< FMacWindow >& InParent, const bool bShowImmediately );
 	
 	void OnDisplayReconfiguration(CGDirectDisplayID Display, CGDisplayChangeSummaryFlags Flags);
 	
 	bool OnIMKKeyDown(NSEvent* Event);
+	
+	static void AddModalWindow(FCocoaWindow* Window);
+	static void RemoveModalWindow(FCocoaWindow* Window);
+	static FCocoaWindow* CurrentModalWindow(void);
 
 public:
 
@@ -186,11 +101,13 @@ private:
 
 
 private:
+	/** Currently open modal windows. */
+	static TArray< FCocoaWindow* > RunningModalWindows;
 
 	FMacApplication* OwningApplication;
 
 	/** Mac window handle */
-	FSlateCocoaWindow* WindowHandle;
+	FCocoaWindow* WindowHandle;
 	
 	/** The mode that the window is in (windowed, fullscreen, windowedfullscreen ) */
 	EWindowMode::Type WindowMode;
diff --git a/Engine/Source/Runtime/OpenGLDrv/Private/Mac/OpenGLMac.cpp b/Engine/Source/Runtime/OpenGLDrv/Private/Mac/OpenGLMac.cpp
--- a/Engine/Source/Runtime/OpenGLDrv/Private/Mac/OpenGLMac.cpp
+++ b/Engine/Source/Runtime/OpenGLDrv/Private/Mac/OpenGLMac.cpp
@@ -4,6 +4,8 @@
 
 #include "MacWindow.h"
 #include "MacTextInputMethodSystem.h"
+#include "CocoaTextView.h"
+#include "CocoaWindow.h"
 
 /*------------------------------------------------------------------------------
  OpenGL static variables.
@@ -209,7 +211,7 @@ bool GMacEnableCocoaScreenUpdates = true;
 /**
  * Custom view class
  */
-@interface FCocoaOpenGLView : FSlateTextView
+@interface FCocoaOpenGLView : FCocoaTextView
 {
 }
 @property (atomic) bool bNeedsUpdate;
@@ -245,10 +247,10 @@ bool GMacEnableCocoaScreenUpdates = true;
 {
 	SCOPED_AUTORELEASE_POOL;
 	
-	FSlateCocoaWindow* SlateCocoaWindow = [[self window] isKindOfClass:[FSlateCocoaWindow class]] ? (FSlateCocoaWindow*)[self window] : nil;
-	if (SlateCocoaWindow)
+	FCocoaWindow* CocoaWindow = [[self window] isKindOfClass:[FCocoaWindow class]] ? (FCocoaWindow*)[self window] : nil;
+	if (CocoaWindow)
 	{
-		[SlateCocoaWindow redrawContents];
+		[CocoaWindow redrawContents];
 	}
 }
 
@@ -744,12 +746,12 @@ bool PlatformBlitToViewport( FPlatformOpenGLDevice* Device, const FOpenGLViewpor
 				
 				NSWindow* Window  = [Context->OpenGLView window];
 				NSView* SuperView = [[Window contentView] superview];
-				FSlateCocoaWindow* SlateCocoaWindow = [Window isKindOfClass:[FSlateCocoaWindow class]] ? (FSlateCocoaWindow*)Window : nil;
-				if([SuperView respondsToSelector:@selector(roundedCornerRadius)] && SlateCocoaWindow)
+				FCocoaWindow* CocoaWindow = [Window isKindOfClass:[FCocoaWindow class]] ? (FCocoaWindow*)Window : nil;
+				if([SuperView respondsToSelector:@selector(roundedCornerRadius)] && CocoaWindow)
 				{
-					bool bRoundedCorners = [SlateCocoaWindow roundedCorners];
-					bool bFullWindowRendering = ([SlateCocoaWindow styleMask] & (NSTexturedBackgroundWindowMask));
-					EWindowMode::Type WindowMode = [SlateCocoaWindow windowMode];
+					bool bRoundedCorners = [CocoaWindow roundedCorners];
+					bool bFullWindowRendering = ([CocoaWindow styleMask] & (NSTexturedBackgroundWindowMask));
+					EWindowMode::Type WindowMode = [CocoaWindow windowMode];
 					bRoundedBlit = (WindowMode == EWindowMode::Windowed && bRoundedCorners && bFullWindowRendering);
 				}
 				
@@ -871,12 +873,12 @@ bool PlatformBlitToViewport( FPlatformOpenGLDevice* Device, const FOpenGLViewpor
 				REPORT_GL_END_BUFFER_EVENT_FOR_FRAME_DUMP();
 //				INITIATE_GL_FRAME_DUMP_EVERY_X_CALLS( 1000 );
 				
-				if(SlateCocoaWindow)
+				if(CocoaWindow)
 				{
 					// Using dispatch was safer - but during loading we aren't responsive on the main thread
 					// this makes it impossible to open the window & so it never appears to play the loading screen.
 					// Just hope that this will work OK...
-					[SlateCocoaWindow performDeferredOrderFront];
+					[CocoaWindow performDeferredOrderFront];
 				}
 			}
 		}
