Fixed crash at startup in Clang-compiled Windows apps
- SetNamedThread() was raising an exception which is not supported with the Clang compiler on Windows yet

diff --git a/Engine/Source/Runtime/Core/Private/Windows/WindowsRunnableThread.h b/Engine/Source/Runtime/Core/Private/Windows/WindowsRunnableThread.h
--- a/Engine/Source/Runtime/Core/Private/Windows/WindowsRunnableThread.h
+++ b/Engine/Source/Runtime/Core/Private/Windows/WindowsRunnableThread.h
@@ -38,6 +38,7 @@ class FRunnableThreadWin
 	 */
 	static void SetThreadName( uint32 ThreadID, LPCSTR ThreadName )
 	{
+#if !PLATFORM_SEH_EXCEPTIONS_DISABLED
 		/**
 		 * Code setting the thread name for use in the debugger.
 		 *
@@ -62,13 +63,10 @@ class FRunnableThreadWin
 		ThreadNameInfo.dwThreadID	= ThreadID;
 		ThreadNameInfo.dwFlags		= 0;
 
-#if !PLATFORM_SEH_EXCEPTIONS_DISABLED
 		__try
-#endif
 		{
 			RaiseException( MS_VC_EXCEPTION, 0, sizeof(ThreadNameInfo)/sizeof(ULONG_PTR), (ULONG_PTR*)&ThreadNameInfo );
 		}
-#if !PLATFORM_SEH_EXCEPTIONS_DISABLED
 		__except( EXCEPTION_EXECUTE_HANDLER )
 		CA_SUPPRESS(6322)
 		{	
