Fixed miscellaneous engine compilation bugs for Clang on Windows
- No functional change

diff --git a/Engine/Source/Runtime/Core/Public/Containers/Array.h b/Engine/Source/Runtime/Core/Public/Containers/Array.h
--- a/Engine/Source/Runtime/Core/Public/Containers/Array.h
+++ b/Engine/Source/Runtime/Core/Public/Containers/Array.h
@@ -461,7 +461,7 @@ public:
 	{
 		DestructItems(GetTypedData(), ArrayNum);
 
-		#if defined(_MSC_VER)
+		#if defined(_MSC_VER) && !defined(__clang__)	// Relies on MSVC-specific lazy template instantiation to support arrays of incomplete types
 			// ensure that DebugGet gets instantiated.
 			//@todo it would be nice if we had a cleaner solution for DebugGet
 			volatile const ElementType* Dummy = &DebugGet(0);
@@ -1594,7 +1594,7 @@ public:
 		::StableSort(GetTypedData(), Num(), Predicate);
 	}
 
-#if defined(_MSC_VER)
+#if defined(_MSC_VER) && !defined(__clang__)	// Relies on MSVC-specific lazy template instantiation to support arrays of incomplete types
 private:
 	/**
 	* Helper function that can be used inside the debuggers watch window to debug TArrays. E.g. "*Class->Defaults.DebugGet(5)". 
diff --git a/Engine/Source/Runtime/Core/Public/Misc/AssertionMacros.h b/Engine/Source/Runtime/Core/Public/Misc/AssertionMacros.h
--- a/Engine/Source/Runtime/Core/Public/Misc/AssertionMacros.h
+++ b/Engine/Source/Runtime/Core/Public/Misc/AssertionMacros.h
@@ -141,8 +141,8 @@ bool JunkFunc(const volatile T&);
 	((void)sizeof(JunkFunc(((ClassName*)0)->MemberName)), TEXT(#MemberName))
 
 // Returns FName(TEXT("FunctionName")), while statically verifying that the function exists in ClassName
-// &((ClassName*)0)->FunctionName will generate 'error: cannot create a non-constant pointer to member function' on the Mac
-#if PLATFORM_WINDOWS
+// &((ClassName*)0)->FunctionName will generate 'error: cannot create a non-constant pointer to member function' on non-MSVC compilers
+#if PLATFORM_WINDOWS && !defined(__clang__)
 	#define GET_FUNCTION_NAME_CHECKED(ClassName, FunctionName) \
 		((void)sizeof(&((ClassName*)0)->FunctionName), FName(TEXT(#FunctionName)))
 #else
diff --git a/Engine/Source/Runtime/Core/Public/Misc/CoreMisc.h b/Engine/Source/Runtime/Core/Public/Misc/CoreMisc.h
--- a/Engine/Source/Runtime/Core/Public/Misc/CoreMisc.h
+++ b/Engine/Source/Runtime/Core/Public/Misc/CoreMisc.h
@@ -262,7 +262,7 @@ FORCEINLINE bool IsRunningDedicatedServer()
 	}
 
 #if UE_EDITOR
-	extern int32 StaticDedicatedServerCheck();
+	extern CORE_API int32 StaticDedicatedServerCheck();
 	return (StaticDedicatedServerCheck() == 1);
 #else
 	return false;
@@ -287,7 +287,7 @@ FORCEINLINE bool IsRunningGame()
 	}
 
 #if UE_EDITOR
-	extern int32 StaticGameCheck();
+	extern CORE_API int32 StaticGameCheck();
 	return (StaticGameCheck() == 1);
 #else
 	return false;
@@ -307,7 +307,7 @@ FORCEINLINE bool IsRunningClientOnly()
 	}
 
 #if UE_EDITOR
-	extern int32 StaticClientOnlyCheck();
+	extern CORE_API int32 StaticClientOnlyCheck();
 	return (StaticClientOnlyCheck() == 1);
 #else
 	return false;
diff --git a/Engine/Source/Runtime/Core/Public/Misc/Database.h b/Engine/Source/Runtime/Core/Public/Misc/Database.h
--- a/Engine/Source/Runtime/Core/Public/Misc/Database.h
+++ b/Engine/Source/Runtime/Core/Public/Misc/Database.h
@@ -15,8 +15,9 @@
 
 
 // Only use ADO on windows, if support is enabled and not for shipping games.
-#define USE_ADO_INTEGRATION (PLATFORM_WINDOWS && !(UE_BUILD_SHIPPING && WITH_EDITOR) && WITH_DATABASE_SUPPORT)
-#define USE_REMOTE_INTEGRATION (!USE_ADO_INTEGRATION && !(UE_BUILD_SHIPPING && WITH_EDITOR) && WITH_DATABASE_SUPPORT)
+// @todo clang: #import is not supported by Clang on Windows platform, but we currently need this for ADO symbol importing.  For now we disable ADO support in Clang builds.
+#define USE_ADO_INTEGRATION (PLATFORM_WINDOWS && !(UE_BUILD_SHIPPING && WITH_EDITOR) && WITH_DATABASE_SUPPORT && !defined(__clang__))
+#define USE_REMOTE_INTEGRATION (!USE_ADO_INTEGRATION && !(UE_BUILD_SHIPPING && WITH_EDITOR) && WITH_DATABASE_SUPPORT && !defined(__clang__))
 
 
 /**  
diff --git a/Engine/Source/Runtime/Engine/Classes/Engine/ObjectLibrary.h b/Engine/Source/Runtime/Engine/Classes/Engine/ObjectLibrary.h
--- a/Engine/Source/Runtime/Engine/Classes/Engine/ObjectLibrary.h
+++ b/Engine/Source/Runtime/Engine/Classes/Engine/ObjectLibrary.h
@@ -1,6 +1,7 @@
 // Copyright 1998-2014 Epic Games, Inc. All Rights Reserved.
 
 #pragma once
+#include "AssetData.h"
 #include "ObjectLibrary.generated.h"
 
 /** Class that holds a library of Objects */
@@ -36,7 +37,7 @@ protected:
 	bool bIsFullyLoaded;
 
 	/** Asset data of objects that will belong in library, possibly not loaded yet */
-	TArray<class FAssetData>	AssetDataList;
+	TArray<FAssetData>	AssetDataList;
 
 public:
 	// Begin UObject Interface
diff --git a/Engine/Source/Runtime/Engine/Public/InstancedFoliage.h b/Engine/Source/Runtime/Engine/Public/InstancedFoliage.h
--- a/Engine/Source/Runtime/Engine/Public/InstancedFoliage.h
+++ b/Engine/Source/Runtime/Engine/Public/InstancedFoliage.h
@@ -251,11 +251,11 @@ struct FFoliageMeshInfo
 	friend FArchive& operator<<(FArchive& Ar, FFoliageMeshInfo& MeshInfo);
 
 #if PLATFORM_COMPILER_HAS_DEFAULTED_FUNCTIONS
-	const FFoliageMeshInfo(const FFoliageMeshInfo& Other) = delete;
+	FFoliageMeshInfo(const FFoliageMeshInfo& Other) = delete;
 	const FFoliageMeshInfo& operator=(const FFoliageMeshInfo& Other) = delete;
 #else
 private:
-	const FFoliageMeshInfo(const FFoliageMeshInfo& Other);
+	FFoliageMeshInfo(const FFoliageMeshInfo& Other);
 	const FFoliageMeshInfo& operator=(const FFoliageMeshInfo& Other);
 #endif
 };
