Add automation tests opt-in for console.
TTP#341142: Add automation tests opt-in for console.

diff --git a/Engine/Config/BaseEngine.ini b/Engine/Config/BaseEngine.ini
--- a/Engine/Config/BaseEngine.ini
+++ b/Engine/Config/BaseEngine.ini
@@ -2016,6 +2016,7 @@ AspectRatioAxisConstraint=AspectRatio_MaintainXFOV
 [AutomationTesting]
 ImportTestPath=../../Content/EditorAutomation/
 ImportTestPackagePath=/Engine/Content/EditorAutomation
+bForceSmokeTests=false
 
 [AutomationTesting.Blueprint]
 TestAllBlueprints=false
diff --git a/Engine/Source/Runtime/Core/Private/Misc/AutomationTest.cpp b/Engine/Source/Runtime/Core/Private/Misc/AutomationTest.cpp
--- a/Engine/Source/Runtime/Core/Private/Misc/AutomationTest.cpp
+++ b/Engine/Source/Runtime/Core/Private/Misc/AutomationTest.cpp
@@ -122,7 +122,7 @@ bool FAutomationTestFramework::RunSmokeTests()
 	//@todo - better determination of whether to run than requires cooked data
 	// Ensure there isn't another slow task in progress when trying to run unit tests
 	const bool bRequiresCookedData = FPlatformProperties::RequiresCookedData();
-	if( !bRequiresCookedData && !GIsSlowTask && !GIsPlayInEditorWorld )
+	if ((!bRequiresCookedData && !GIsSlowTask && !GIsPlayInEditorWorld) || bForceSmokeTests)
 	{
 		TArray<FAutomationTestInfo> TestInfo;
 
@@ -631,6 +631,7 @@ FAutomationTestFramework::FAutomationTestFramework()
 ,	bVisualCommandletFilterOn(false)
 ,	bUseFullSizeScreenShots(false)
 ,	NetworkRoleIndex(0)
+,	bForceSmokeTests(false)
 {
 }
 
diff --git a/Engine/Source/Runtime/Core/Public/Misc/AutomationTest.h b/Engine/Source/Runtime/Core/Public/Misc/AutomationTest.h
--- a/Engine/Source/Runtime/Core/Public/Misc/AutomationTest.h
+++ b/Engine/Source/Runtime/Core/Public/Misc/AutomationTest.h
@@ -520,6 +520,13 @@ public:
 	 */
 	bool ShouldUseFullSizeScreenshots() const;
 
+	/**
+	 * Sets forcing smoke tests.
+	 */
+	void SetForceSmokeTests(bool bForceSmokeTests)
+	{
+		this->bForceSmokeTests = bForceSmokeTests;
+	}
 private:
 
 	/** Special feedback context used exclusively while automation testing */
@@ -671,6 +678,9 @@ private:
 
 	/** Delegate called at the end of the frame when a screenshot is captured and a .png is requested */
 	FOnTestScreenshotCaptured TestScreenshotCapturedDelegate;
+
+	/** Forces running smoke tests */
+	bool bForceSmokeTests;
 };
 
 
diff --git a/Engine/Source/Runtime/Launch/Private/LaunchEngineLoop.cpp b/Engine/Source/Runtime/Launch/Private/LaunchEngineLoop.cpp
--- a/Engine/Source/Runtime/Launch/Private/LaunchEngineLoop.cpp
+++ b/Engine/Source/Runtime/Launch/Private/LaunchEngineLoop.cpp
@@ -2538,6 +2538,11 @@ bool FEngineLoop::AppInit( )
 	// Init list of common colors.
 	GColorList.CreateColorMap();
 
+	bool bForceSmokeTests = false;
+	GConfig->GetBool(TEXT("AutomationTesting"), TEXT("bForceSmokeTests"), bForceSmokeTests, GEngineIni);
+	bForceSmokeTests |= FParse::Param(FCommandLine::Get(), TEXT("bForceSmokeTests"));
+	FAutomationTestFramework::GetInstance().SetForceSmokeTests(bForceSmokeTests);
+
 	// Init other systems.
 	FCoreDelegates::OnInit.Broadcast();
 	return true;
