Don't allow 128-bit atomics on Windows platform unless compiling with Windows 8 as min spec
- These intrinsics are not available on Windows 7
- This fixes link errors when compiling using Clang on Windows platform

diff --git a/Engine/Source/Runtime/Core/Public/Windows/WIndowsPlatform.h b/Engine/Source/Runtime/Core/Public/Windows/WIndowsPlatform.h
--- a/Engine/Source/Runtime/Core/Public/Windows/WIndowsPlatform.h
+++ b/Engine/Source/Runtime/Core/Public/Windows/WIndowsPlatform.h
@@ -69,7 +69,9 @@ typedef FWindowsPlatformTypes FPlatformTypes;
 #define PLATFORM_COMPILER_HAS_EXPLICIT_OPERATORS			0
 #define PLATFORM_COMPILER_HAS_TCHAR_WMAIN					1
 
-#define PLATFORM_HAS_128BIT_ATOMICS							(!HACK_HEADER_GENERATOR && PLATFORM_64BITS)
+// Intrinsics for 128-bit atomics on Windows platform requires Windows 8 or higher (WINVER>0x0602)
+// http://msdn.microsoft.com/en-us/library/windows/desktop/hh972640.aspx
+#define PLATFORM_HAS_128BIT_ATOMICS							(!HACK_HEADER_GENERATOR && PLATFORM_64BITS && (WINVER > 0x602))
 #define PLATFORM_USES_ANSI_STRING_FOR_EXTERNAL_PROFILING	0
 
 #if _MSC_VER < 1700
