Core - FCString::Sprintf usage fix

diff --git a/Engine/Source/Runtime/Core/Private/Windows/WindowsPlatformMisc.cpp b/Engine/Source/Runtime/Core/Private/Windows/WindowsPlatformMisc.cpp
--- a/Engine/Source/Runtime/Core/Private/Windows/WindowsPlatformMisc.cpp
+++ b/Engine/Source/Runtime/Core/Private/Windows/WindowsPlatformMisc.cpp
@@ -1871,7 +1871,7 @@ void FWindowsPlatformMisc::CoUninitialize()
 }
 
 #if !UE_BUILD_SHIPPING
-static TCHAR GErrorRemoteDebugPromptMessage[4096];
+static TCHAR GErrorRemoteDebugPromptMessage[MAX_SPRINTF];
 
 void FWindowsPlatformMisc::PromptForRemoteDebugging(bool bIsEnsure)
 {
