Core - Removed unused variable

diff --git a/Engine/Source/Runtime/Core/Private/Windows/WindowsPlatformOutputDevices.cpp b/Engine/Source/Runtime/Core/Private/Windows/WindowsPlatformOutputDevices.cpp
--- a/Engine/Source/Runtime/Core/Private/Windows/WindowsPlatformOutputDevices.cpp
+++ b/Engine/Source/Runtime/Core/Private/Windows/WindowsPlatformOutputDevices.cpp
@@ -60,7 +60,6 @@ class FFeedbackContext*				FWindowsPlatformOutputDevices::GetWarn()
 //////////////////////////////////
 
 FOutputDeviceWindowsError::FOutputDeviceWindowsError()
-	: ErrorPos(0)
 {
 }
 
@@ -94,7 +93,6 @@ void FOutputDeviceWindowsError::Serialize( const TCHAR* Msg, ELogVerbosity::Type
 
 		FCString::Strncpy( GErrorHist, Msg, ARRAY_COUNT(GErrorHist) - 5 );
 		FCString::Strncat( GErrorHist, TEXT("\r\n\r\n"), ARRAY_COUNT(GErrorHist) - 1  );
-		ErrorPos = FCString::Strlen(GErrorHist);
 	}
 	else
 	{
diff --git a/Engine/Source/Runtime/Core/Private/Windows/WindowsPlatformOutputDevicesPrivate.h b/Engine/Source/Runtime/Core/Private/Windows/WindowsPlatformOutputDevicesPrivate.h
--- a/Engine/Source/Runtime/Core/Private/Windows/WindowsPlatformOutputDevicesPrivate.h
+++ b/Engine/Source/Runtime/Core/Private/Windows/WindowsPlatformOutputDevicesPrivate.h
@@ -134,9 +134,6 @@ public:
 	 * error handler, e.g. using structured exception handling on the PC.
 	 */
 	void HandleError();
-
-private:
-	int32	ErrorPos;
 };
 
 /**
