Core - Misc code clean-up

diff --git a/Engine/Source/Runtime/Core/Private/GenericPlatform/GenericPlatformStackWalk.cpp b/Engine/Source/Runtime/Core/Private/GenericPlatform/GenericPlatformStackWalk.cpp
--- a/Engine/Source/Runtime/Core/Private/GenericPlatform/GenericPlatformStackWalk.cpp
+++ b/Engine/Source/Runtime/Core/Private/GenericPlatform/GenericPlatformStackWalk.cpp
@@ -41,7 +41,7 @@ void FGenericPlatformStackWalk::StackWalkAndDump( ANSICHAR* HumanReadableString,
 	// Temporary memory holding the stack trace.
 	static const int MAX_DEPTH = 100;
 	uint64 StackTrace[MAX_DEPTH];
-	memset( StackTrace, 0, sizeof(StackTrace) );
+	FMemory::MemZero( StackTrace );
 
 	// Capture stack backtrace.
 	FPlatformStackWalk::CaptureStackBackTrace( StackTrace, MAX_DEPTH, Context );
