New checks were a little overzealous, perfectly valid to insert via reference and copy values that are already in the map. Removed checks that were being applied to const references

diff --git a/Engine/Source/Runtime/Core/Public/Containers/Map.h b/Engine/Source/Runtime/Core/Public/Containers/Map.h
--- a/Engine/Source/Runtime/Core/Public/Containers/Map.h
+++ b/Engine/Source/Runtime/Core/Public/Containers/Map.h
@@ -280,10 +280,10 @@ public:
 	 * @param InValue - The value to associate with the key.
 	 * @return A reference to the value as stored in the map.  The reference is only valid until the next change to any key in the map.
 	 */
-	FORCEINLINE ValueType& Add(const KeyType&  InKey, const ValueType&  InValue) { Pairs.CheckAddress(&InKey); Pairs.CheckAddress(&InValue); return Emplace(         InKey ,          InValue ); }
-	FORCEINLINE ValueType& Add(const KeyType&  InKey,       ValueType&& InValue) { Pairs.CheckAddress(&InKey); Pairs.CheckAddress(&InValue); return Emplace(         InKey , MoveTemp(InValue)); }
-	FORCEINLINE ValueType& Add(      KeyType&& InKey, const ValueType&  InValue) { Pairs.CheckAddress(&InKey); Pairs.CheckAddress(&InValue); return Emplace(MoveTemp(InKey),          InValue ); }
-	FORCEINLINE ValueType& Add(      KeyType&& InKey,       ValueType&& InValue) { Pairs.CheckAddress(&InKey); Pairs.CheckAddress(&InValue); return Emplace(MoveTemp(InKey), MoveTemp(InValue)); }
+	FORCEINLINE ValueType& Add(const KeyType&  InKey, const ValueType&  InValue) {																return Emplace(         InKey ,          InValue ); }
+	FORCEINLINE ValueType& Add(const KeyType&  InKey,       ValueType&& InValue) { Pairs.CheckAddress(&InValue);								return Emplace(         InKey , MoveTemp(InValue)); }
+	FORCEINLINE ValueType& Add(      KeyType&& InKey, const ValueType&  InValue) { Pairs.CheckAddress(&InValue);								return Emplace(MoveTemp(InKey),          InValue ); }
+	FORCEINLINE ValueType& Add(      KeyType&& InKey,       ValueType&& InValue) { Pairs.CheckAddress(&InKey); Pairs.CheckAddress(&InValue);	return Emplace(MoveTemp(InKey), MoveTemp(InValue)); }
 
 	/**
 	 * Sets a default value associated with a key.
@@ -292,7 +292,7 @@ public:
 	 * @param InKey - The key to associate the value with.
 	 * @return A reference to the value as stored in the map.  The reference is only valid until the next change to any key in the map.
 	 */
-	FORCEINLINE ValueType& Add(const KeyType&  InKey) { Pairs.CheckAddress(&InKey); return Emplace(         InKey ); }
+	FORCEINLINE ValueType& Add(const KeyType&  InKey) {								return Emplace(         InKey ); }
 	FORCEINLINE ValueType& Add(      KeyType&& InKey) { Pairs.CheckAddress(&InKey); return Emplace(MoveTemp(InKey)); }
 
 	/**
