Removed the wrong check. Caught by [REDACTED_PERSON] code review.

diff --git a/Engine/Source/Runtime/Core/Public/Containers/Map.h b/Engine/Source/Runtime/Core/Public/Containers/Map.h
--- a/Engine/Source/Runtime/Core/Public/Containers/Map.h
+++ b/Engine/Source/Runtime/Core/Public/Containers/Map.h
@@ -282,7 +282,7 @@ public:
 	 */
 	FORCEINLINE ValueType& Add(const KeyType&  InKey, const ValueType&  InValue) {																return Emplace(         InKey ,          InValue ); }
 	FORCEINLINE ValueType& Add(const KeyType&  InKey,       ValueType&& InValue) { Pairs.CheckAddress(&InValue);								return Emplace(         InKey , MoveTemp(InValue)); }
-	FORCEINLINE ValueType& Add(      KeyType&& InKey, const ValueType&  InValue) { Pairs.CheckAddress(&InValue);								return Emplace(MoveTemp(InKey),          InValue ); }
+	FORCEINLINE ValueType& Add(      KeyType&& InKey, const ValueType&  InValue) { Pairs.CheckAddress(&InKey);									return Emplace(MoveTemp(InKey),          InValue ); }
 	FORCEINLINE ValueType& Add(      KeyType&& InKey,       ValueType&& InValue) { Pairs.CheckAddress(&InKey); Pairs.CheckAddress(&InValue);	return Emplace(MoveTemp(InKey), MoveTemp(InValue)); }
 
 	/**
