Changes in platform initialization/teardown.

- New PlatformTearDown() function to correspond to PlatformInit(), called at AppExit() time (unguarded code).
- Linux-specific PlatformInitMultimedia() function (not part of generic platform abstraction) which handles multiple initialization of SDL and/or avoids it for certain targets.
- Fixes problem with SDL message box destroying SDL and ruining things later down the road.

diff --git a/Engine/Source/Developer/StandaloneRenderer/Private/Linux/OpenGL/SlateOpenGLContext.cpp b/Engine/Source/Developer/StandaloneRenderer/Private/Linux/OpenGL/SlateOpenGLContext.cpp
--- a/Engine/Source/Developer/StandaloneRenderer/Private/Linux/OpenGL/SlateOpenGLContext.cpp
+++ b/Engine/Source/Developer/StandaloneRenderer/Private/Linux/OpenGL/SlateOpenGLContext.cpp
@@ -5,7 +5,12 @@
 
 static SDL_Window* CreateDummyGLWindow()
 {
-	SDL_Init( SDL_INIT_VIDEO );
+	FPlatformMisc::PlatformInitMultimedia(); //	will not initialize more than once
+
+#if DO_CHECK
+	uint32 InitializedSubsystems = SDL_WasInit(SDL_INIT_EVERYTHING);
+	check(InitializedSubsystems & SDL_INIT_VIDEO);
+#endif // DO_CHECK
 
 	// Create a dummy window.
 	SDL_Window *h_wnd = SDL_CreateWindow(	NULL,
@@ -66,7 +71,7 @@ void FSlateOpenGLContext::Destroy()
 		if	( bReleaseWindowOnDestroy )
 		{
 			SDL_DestroyWindow( WindowHandle );
-			SDL_Quit();
+			// we will tear down SDL in PlatformTearDown()
 		}
 		WindowHandle = NULL;
 
diff --git a/Engine/Source/Runtime/Core/Private/Linux/LinuxApplication.cpp b/Engine/Source/Runtime/Core/Private/Linux/LinuxApplication.cpp
--- a/Engine/Source/Runtime/Core/Private/Linux/LinuxApplication.cpp
+++ b/Engine/Source/Runtime/Core/Private/Linux/LinuxApplication.cpp
@@ -32,28 +32,14 @@ FLinuxApplication* FLinuxApplication::CreateLinuxApplication()
 {
 	LinuxApplication = new FLinuxApplication();
 	
-	//	init the sdl here
-	if	( SDL_WasInit( 0 ) == 0 )
-	{
-		SDL_Init( SDL_INIT_EVENTS | SDL_INIT_JOYSTICK | SDL_INIT_GAMECONTROLLER );
-	}
-	else
-	{
-		Uint32 subsystem_init = SDL_WasInit( SDL_INIT_EVERYTHING );
-
-		if	( !(subsystem_init & SDL_INIT_EVENTS) )
-		{
-			SDL_InitSubSystem( SDL_INIT_EVENTS );
-		}
-		if	( !(subsystem_init & SDL_INIT_JOYSTICK) )
-		{
-			SDL_InitSubSystem( SDL_INIT_JOYSTICK );
-		}
-		if	( !(subsystem_init & SDL_INIT_GAMECONTROLLER) )
-		{
-			SDL_InitSubSystem( SDL_INIT_GAMECONTROLLER );
-		}
-	}
+	FPlatformMisc::PlatformInitMultimedia(); //	will not initialize more than once
+
+#if DO_CHECK
+	uint32 InitializedSubsystems = SDL_WasInit(SDL_INIT_EVERYTHING);
+	check(InitializedSubsystems & SDL_INIT_EVENTS);
+	check(InitializedSubsystems & SDL_INIT_JOYSTICK);
+	check(InitializedSubsystems & SDL_INIT_GAMECONTROLLER);
+#endif // DO_CHECK
 
 	SDLControllerState *controllerState = LinuxApplication->ControllerStates;
 	for (int i = 0; i < SDL_NumJoysticks(); ++i) {
diff --git a/Engine/Source/Runtime/Core/Private/Linux/LinuxCursor.cpp b/Engine/Source/Runtime/Core/Private/Linux/LinuxCursor.cpp
--- a/Engine/Source/Runtime/Core/Private/Linux/LinuxCursor.cpp
+++ b/Engine/Source/Runtime/Core/Private/Linux/LinuxCursor.cpp
@@ -7,20 +7,12 @@
 
 FLinuxCursor::FLinuxCursor() : bHidden(false)
 {
-	//	init the sdl here
-	if	( SDL_WasInit( 0 ) == 0 )
-	{
-		SDL_Init( SDL_INIT_VIDEO );
-	}
-	else
-	{
-		Uint32 subsystem_init = SDL_WasInit( SDL_INIT_EVERYTHING );
+	FPlatformMisc::PlatformInitMultimedia(); //	will not initialize more than once
 
-		if	( !(subsystem_init & SDL_INIT_VIDEO) )
-		{
-			SDL_InitSubSystem( SDL_INIT_VIDEO );
-		}
-	}
+#if DO_CHECK
+	uint32 InitializedSubsystems = SDL_WasInit(SDL_INIT_EVERYTHING);
+	check(InitializedSubsystems & SDL_INIT_VIDEO);
+#endif // DO_CHECK
 
 	// Load up cursors that we'll be using
 	for( int32 CurCursorIndex = 0; CurCursorIndex < EMouseCursor::TotalCursorCount; ++CurCursorIndex )
diff --git a/Engine/Source/Runtime/Core/Private/Linux/LinuxPlatformMisc.cpp b/Engine/Source/Runtime/Core/Private/Linux/LinuxPlatformMisc.cpp
--- a/Engine/Source/Runtime/Core/Private/Linux/LinuxPlatformMisc.cpp
+++ b/Engine/Source/Runtime/Core/Private/Linux/LinuxPlatformMisc.cpp
@@ -122,6 +122,11 @@ void FLinuxPlatformMisc::NormalizePath(FString& InPath)
 	}
 }
 
+namespace
+{
+	bool GInitializedSDL = false;
+}
+
 void FLinuxPlatformMisc::PlatformInit()
 {
 	// install a platform-specific signal handler
@@ -146,6 +151,39 @@ void FLinuxPlatformMisc::PlatformInit()
 	UE_LOG(LogInit, Log, TEXT(" -httpproxy=ADDRESS:PORT - redirects HTTP requests to a proxy (only supported if compiled with libcurl)"));
 	UE_LOG(LogInit, Log, TEXT(" -reuseconn - allow libcurl to reuse HTTP connections (only matters if compiled with libcurl)"));
 	UE_LOG(LogInit, Log, TEXT(" -virtmemkb=NUMBER - sets process virtual memory (address space) limit (overrides VirtualMemoryLimitInKB value from .ini)"));
+
+	// skip for servers and programs, unless they request later
+	if (!UE_SERVER && !IS_PROGRAM)
+	{
+		PlatformInitMultimedia();
+	}
+}
+
+void FLinuxPlatformMisc::PlatformInitMultimedia()
+{
+	if (!GInitializedSDL)
+	{
+		UE_LOG(LogInit, Log, TEXT("Initializing SDL."));
+		if (SDL_Init(SDL_INIT_EVERYTHING | SDL_INIT_NOPARACHUTE) != 0)
+		{
+			const char * SDLError = SDL_GetError();
+			UE_LOG(LogInit, Fatal, TEXT("Could not initialize SDL: %s"), SDLError);
+			// unreachable
+			return;
+		}
+
+		GInitializedSDL = true;
+	}
+}
+
+void FLinuxPlatformMisc::PlatformTearDown()
+{
+	if (GInitializedSDL)
+	{
+		UE_LOG(LogInit, Log, TEXT("Tearing down SDL."));
+		SDL_Quit();
+		GInitializedSDL = false;
+	}
 }
 
 GenericApplication* FLinuxPlatformMisc::CreateApplication()
@@ -240,10 +278,12 @@ EAppReturnType::Type FLinuxPlatformMisc::MessageBoxExt(EAppMsgType::Type MsgType
 {
 	int NumberOfButtons = 0;
 
-	if (SDL_WasInit(SDL_INIT_VIDEO) == 0)
-	{
-		SDL_InitSubSystem(SDL_INIT_VIDEO);
-	}
+	FPlatformMisc::PlatformInitMultimedia(); //	will not initialize more than once
+
+#if DO_CHECK
+	uint32 InitializedSubsystems = SDL_WasInit(SDL_INIT_EVERYTHING);
+	check(InitializedSubsystems & SDL_INIT_VIDEO);
+#endif // DO_CHECK
 
 	SDL_MessageBoxButtonData *Buttons = nullptr;
 
@@ -375,12 +415,10 @@ EAppReturnType::Type FLinuxPlatformMisc::MessageBoxExt(EAppMsgType::Type MsgType
 	{
 		UE_LOG(LogInit, Fatal, TEXT("Error Presenting MessageBox: %s\n"), ANSI_TO_TCHAR(SDL_GetError()));
 		// unreachable
-		SDL_Quit();
 		return EAppReturnType::Cancel;
 	}
 
 	delete[] Buttons;
-	SDL_Quit();
 
 	return ButtonPressed == -1 ? EAppReturnType::Cancel : static_cast<EAppReturnType::Type>(ButtonPressed);
 }
diff --git a/Engine/Source/Runtime/Core/Private/Linux/LinuxWindow.cpp b/Engine/Source/Runtime/Core/Private/Linux/LinuxWindow.cpp
--- a/Engine/Source/Runtime/Core/Private/Linux/LinuxWindow.cpp
+++ b/Engine/Source/Runtime/Core/Private/Linux/LinuxWindow.cpp
@@ -44,20 +44,12 @@ void FLinuxWindow::Initialize( FLinuxApplication* const Application, const TShar
 	Definition = InDefinition;
 	OwningApplication = Application;
 
-	//	init the sdl here
-	if	( SDL_WasInit( 0 ) == 0 )
-	{
-		SDL_Init( SDL_INIT_VIDEO );
-	}
-	else
-	{
-		Uint32 subsystem_init = SDL_WasInit( SDL_INIT_EVERYTHING );
+	FPlatformMisc::PlatformInitMultimedia(); //	will not initialize more than once
 
-		if	( !(subsystem_init & SDL_INIT_VIDEO) )
-		{
-			SDL_InitSubSystem( SDL_INIT_VIDEO );
-		}
-	}
+#if DO_CHECK
+	uint32 InitializedSubsystems = SDL_WasInit(SDL_INIT_EVERYTHING);
+	check(InitializedSubsystems & SDL_INIT_VIDEO);
+#endif // DO_CHECK
 
 	// Finally, let's initialize the new native window object.  Calling this function will often cause OS
 	// window messages to be sent! (such as activation messages)
diff --git a/Engine/Source/Runtime/Core/Public/GenericPlatform/GenericPlatformMisc.h b/Engine/Source/Runtime/Core/Public/GenericPlatform/GenericPlatformMisc.h
--- a/Engine/Source/Runtime/Core/Public/GenericPlatform/GenericPlatformMisc.h
+++ b/Engine/Source/Runtime/Core/Public/GenericPlatform/GenericPlatformMisc.h
@@ -197,8 +197,8 @@ struct CORE_API FGenericMemoryWarningContext
 struct CORE_API FGenericPlatformMisc
 {
 	/**
-	* Called during appInit() after cmd line setup
-	*/
+	 * Called during appInit() after cmd line setup
+	 */
 	static void PlatformPreInit()
 	{
 	}
@@ -209,6 +209,13 @@ struct CORE_API FGenericPlatformMisc
 	{
 	}
 
+	/**
+	* Called during AppExit(). Log, Config still exist at this point, but not much else does.
+	*/
+	static void PlatformTearDown()
+	{
+	}
+
 	static class GenericApplication* CreateApplication();
 
 	static void* GetHardwareWindow()
diff --git a/Engine/Source/Runtime/Core/Public/Linux/LinuxPlatformMisc.h b/Engine/Source/Runtime/Core/Public/Linux/LinuxPlatformMisc.h
--- a/Engine/Source/Runtime/Core/Public/Linux/LinuxPlatformMisc.h
+++ b/Engine/Source/Runtime/Core/Public/Linux/LinuxPlatformMisc.h
@@ -86,6 +86,7 @@ struct CORE_API FLinuxCrashContext : public FGenericCrashContext
 struct CORE_API FLinuxPlatformMisc : public FGenericPlatformMisc
 {
 	static void PlatformInit();
+	static void PlatformTearDown();
 	static void SetGracefulTerminationHandler();
 	static void SetCrashHandler(void (* CrashHandler)(const FGenericCrashContext & Context));
 	static class GenericApplication* CreateApplication();
@@ -162,6 +163,11 @@ struct CORE_API FLinuxPlatformMisc : public FGenericPlatformMisc
 	{
 		return TEXT("../../../Engine/");
 	}
+
+	/**
+	 * Linux-specific function initializing video (and not only) subsystem.
+	 */
+	static void PlatformInitMultimedia();
 };
 
 typedef FLinuxPlatformMisc FPlatformMisc;
diff --git a/Engine/Source/Runtime/Launch/Private/LaunchEngineLoop.cpp b/Engine/Source/Runtime/Launch/Private/LaunchEngineLoop.cpp
--- a/Engine/Source/Runtime/Launch/Private/LaunchEngineLoop.cpp
+++ b/Engine/Source/Runtime/Launch/Private/LaunchEngineLoop.cpp
@@ -2589,6 +2589,8 @@ void FEngineLoop::AppExit( )
 {
 	UE_LOG(LogExit, Log, TEXT("Exiting."));
 
+	FPlatformMisc::PlatformTearDown();
+
 	if (GConfig)
 	{
 		GConfig->Exit();
diff --git a/Engine/Source/Runtime/OpenGLDrv/Private/Linux/OpenGLLinux.cpp b/Engine/Source/Runtime/OpenGLDrv/Private/Linux/OpenGLLinux.cpp
--- a/Engine/Source/Runtime/OpenGLDrv/Private/Linux/OpenGLLinux.cpp
+++ b/Engine/Source/Runtime/OpenGLDrv/Private/Linux/OpenGLLinux.cpp
@@ -689,20 +689,12 @@ bool PlatformInitOpenGL()
 	static bool bInitialized = false;
 	static bool bOpenGLSupported = false;
 
-	//	init the sdl here
-	if	( SDL_WasInit( 0 ) == 0 )
-	{
-		SDL_Init( SDL_INIT_VIDEO );
-	}
-	else
-	{
-		Uint32 subsystem_init = SDL_WasInit( SDL_INIT_EVERYTHING );
+	FPlatformMisc::PlatformInitMultimedia(); //	will not initialize more than once
 
-		if	( !(subsystem_init & SDL_INIT_VIDEO) )
-		{
-			SDL_InitSubSystem( SDL_INIT_VIDEO );
-		}
-	}
+#if DO_CHECK
+	uint32 InitializedSubsystems = SDL_WasInit(SDL_INIT_EVERYTHING);
+	check(InitializedSubsystems & SDL_INIT_VIDEO);
+#endif // DO_CHECK
 
 	if (!bInitialized)
 	{
