Reverted TMap checks.  Integrated Dan's fixes to K2Node_CastByteToEnum.cpp.

diff --git a/Engine/Source/Editor/BlueprintGraph/Private/K2Node_CastByteToEnum.cpp b/Engine/Source/Editor/BlueprintGraph/Private/K2Node_CastByteToEnum.cpp
--- a/Engine/Source/Editor/BlueprintGraph/Private/K2Node_CastByteToEnum.cpp
+++ b/Engine/Source/Editor/BlueprintGraph/Private/K2Node_CastByteToEnum.cpp
@@ -124,7 +124,10 @@ public:
 		UEdGraphPin* OutPin = Node->FindPinChecked(Schema->PN_ReturnValue);
 		if (ensure(Context.NetMap.Find(OutPin) == NULL))
 		{
-			Context.NetMap.Add(OutPin, *ValueSource);
+			// We need to copy here to avoid passing in a reference to an element inside the map. The array
+			// that owns the map members could be reallocated, causing the reference to become stale.
+			FBPTerminal* ValueSourceCopy = *ValueSource;
+			Context.NetMap.Add(OutPin, ValueSourceCopy);
 		}
 	}
 };
diff --git a/Engine/Source/Runtime/Core/Public/Containers/Map.h b/Engine/Source/Runtime/Core/Public/Containers/Map.h
--- a/Engine/Source/Runtime/Core/Public/Containers/Map.h
+++ b/Engine/Source/Runtime/Core/Public/Containers/Map.h
@@ -280,10 +280,10 @@ public:
 	 * @param InValue - The value to associate with the key.
 	 * @return A reference to the value as stored in the map.  The reference is only valid until the next change to any key in the map.
 	 */
-	FORCEINLINE ValueType& Add(const KeyType&  InKey, const ValueType&  InValue) {																return Emplace(         InKey ,          InValue ); }
-	FORCEINLINE ValueType& Add(const KeyType&  InKey,       ValueType&& InValue) { Pairs.CheckAddress(&InValue);								return Emplace(         InKey , MoveTemp(InValue)); }
-	FORCEINLINE ValueType& Add(      KeyType&& InKey, const ValueType&  InValue) { Pairs.CheckAddress(&InKey);									return Emplace(MoveTemp(InKey),          InValue ); }
-	FORCEINLINE ValueType& Add(      KeyType&& InKey,       ValueType&& InValue) { Pairs.CheckAddress(&InKey); Pairs.CheckAddress(&InValue);	return Emplace(MoveTemp(InKey), MoveTemp(InValue)); }
+	FORCEINLINE ValueType& Add(const KeyType&  InKey, const ValueType&  InValue) { Pairs.CheckAddress(&InKey); Pairs.CheckAddress(&InValue); return Emplace(         InKey ,          InValue ); }
+	FORCEINLINE ValueType& Add(const KeyType&  InKey,       ValueType&& InValue) { Pairs.CheckAddress(&InKey); Pairs.CheckAddress(&InValue); return Emplace(         InKey , MoveTemp(InValue)); }
+	FORCEINLINE ValueType& Add(      KeyType&& InKey, const ValueType&  InValue) { Pairs.CheckAddress(&InKey); Pairs.CheckAddress(&InValue); return Emplace(MoveTemp(InKey),          InValue ); }
+	FORCEINLINE ValueType& Add(      KeyType&& InKey,       ValueType&& InValue) { Pairs.CheckAddress(&InKey); Pairs.CheckAddress(&InValue); return Emplace(MoveTemp(InKey), MoveTemp(InValue)); }
 
 	/**
 	 * Sets a default value associated with a key.
@@ -292,7 +292,7 @@ public:
 	 * @param InKey - The key to associate the value with.
 	 * @return A reference to the value as stored in the map.  The reference is only valid until the next change to any key in the map.
 	 */
-	FORCEINLINE ValueType& Add(const KeyType&  InKey) {								return Emplace(         InKey ); }
+	FORCEINLINE ValueType& Add(const KeyType&  InKey) { Pairs.CheckAddress(&InKey); return Emplace(         InKey ); }
 	FORCEINLINE ValueType& Add(      KeyType&& InKey) { Pairs.CheckAddress(&InKey); return Emplace(MoveTemp(InKey)); }
 
 	/**
