Fix for mac build :(

diff --git a/Engine/Source/Runtime/Core/Private/Mac/MacApplication.cpp b/Engine/Source/Runtime/Core/Private/Mac/MacApplication.cpp
--- a/Engine/Source/Runtime/Core/Private/Mac/MacApplication.cpp
+++ b/Engine/Source/Runtime/Core/Private/Mac/MacApplication.cpp
@@ -70,7 +70,7 @@ void FMacApplication::OnDisplayReconfiguration(CGDirectDisplayID Display, CGDisp
 	{
 		// Slate needs to know when desktop size changes.
 		FDisplayMetrics DisplayMetrics;
-		App->GetDisplayMetrics( DisplayMetrics );
+		FDisplayMetrics::GetDisplayMetrics(DisplayMetrics);
 		App->BroadcastDisplayMetricsChanged(DisplayMetrics);
 	}
 	
