Remove the PumpMessages call when setting mouse capture to a window on OS X - the changes to separate threads should make this redundant & it causes mouse handling problems without (See: https://answers.unrealengine.com/questions/84480/version-440-problem-mouse.html)

diff --git a/Engine/Source/Runtime/Core/Private/Mac/MacApplication.cpp b/Engine/Source/Runtime/Core/Private/Mac/MacApplication.cpp
--- a/Engine/Source/Runtime/Core/Private/Mac/MacApplication.cpp
+++ b/Engine/Source/Runtime/Core/Private/Mac/MacApplication.cpp
@@ -890,7 +890,6 @@ void FMacApplication::PollGameDeviceState( const float TimeDelta )
 
 void FMacApplication::SetCapture( const TSharedPtr< FGenericWindow >& InWindow )
 {
-	FPlatformMisc::PumpMessages(true);
 	bIsMouseCaptureEnabled = InWindow.IsValid();
 	UpdateMouseCaptureWindow( bIsMouseCaptureEnabled ? ((FMacWindow*)InWindow.Get())->GetWindowHandle() : NULL );
 }
