Speed up ShouldUseThreading for non-desktop platforms (TTP #344335).

- Desktop remains switchable, HTML5 is always false, the rest is always true.
- Should help compiler avoid runtime branches etc.

diff --git a/Engine/Source/Runtime/Core/Private/Misc/App.cpp b/Engine/Source/Runtime/Core/Private/Misc/App.cpp
--- a/Engine/Source/Runtime/Core/Private/Misc/App.cpp
+++ b/Engine/Source/Runtime/Core/Private/Misc/App.cpp
@@ -118,8 +118,10 @@ bool FApp::IsEngineInstalled()
 	return bIsInstalledEngine;
 }
 
+#if HAVE_RUNTIME_THREADING_SWITCHES
 bool FApp::ShouldUseThreadingForPerformance( )
 {
 	static bool OnlyOneThread = FParse::Param(FCommandLine::Get(), TEXT("ONETHREAD")) || IsRunningDedicatedServer() || !FPlatformProcess::SupportsMultithreading() || FPlatformMisc::NumberOfCores() < 2;
 	return !OnlyOneThread;
 }
+#endif // HAVE_RUNTIME_THREADING_SWITCHES
diff --git a/Engine/Source/Runtime/Core/Public/Misc/App.h b/Engine/Source/Runtime/Core/Public/Misc/App.h
--- a/Engine/Source/Runtime/Core/Public/Misc/App.h
+++ b/Engine/Source/Runtime/Core/Public/Misc/App.h
@@ -6,6 +6,8 @@
 
 #pragma once
 
+// platforms which can have runtime threading switches
+#define HAVE_RUNTIME_THREADING_SWITCHES			(PLATFORM_DESKTOP)
 
 /**
  * Provides information about the application.
@@ -104,7 +106,7 @@ public:
 		{
 			return ExecutableName.Left(ChopIndex);
 		}
-		
+
 		return ExecutableName;
 	}
 
@@ -142,7 +144,7 @@ public:
 	 * will return the name of the user that launched the session. If this application is
 	 * not part of a session, this function will return the name of the local user account
 	 * under which the application is running.
-	 * 
+	 *
 	 * @return Name of session owner.
 	 */
 	FORCEINLINE static FString GetSessionOwner()
@@ -165,11 +167,11 @@ public:
 	 */
 	static void InitializeSession();
 
-	/** 
+	/**
 	 * Checks whether this application is a game.
 	 *
-	 * Returns true if a normal or PIE game is active (basically !GIsEdit or || GIsPlayInEditorWorld) 
-	 * This must NOT be accessed on threads other than the game thread!  
+	 * Returns true if a normal or PIE game is active (basically !GIsEdit or || GIsPlayInEditorWorld)
+	 * This must NOT be accessed on threads other than the game thread!
 	 * Use View->Family->EnginShowFlags.Game on the rendering thread.
 	 *
 	 * @return true if the application is a game, false otherwise.
@@ -204,7 +206,7 @@ public:
 	 *
 	 * In binary UE4 releases, the engine can be installed while the game is not. The game IsInstalled()
 	 * setting will take precedence over this flag.
-	 * 
+	 *
 	 * To override, pass -engineinstalled or -enginenotinstalled on the command line.
 	 *
 	 * @return true if the engine is installed, false otherwise.
@@ -248,7 +250,18 @@ public:
 	 *
 	 * @return true if this isn't a server, has more than one core, does not have a -onethread command line options, etc.
 	 */
+#if HAVE_RUNTIME_THREADING_SWITCHES
 	static bool ShouldUseThreadingForPerformance();
+#else
+	FORCEINLINE static bool ShouldUseThreadingForPerformance()
+	{
+	#if PLATFORM_HTML5
+		return false;
+	#else
+		return true;
+	#endif // PLATFORM_HTML5
+	}
+#endif // HAVE_RUNTIME_THREADING_SWITCHES
 
 	/**
 	 * Checks whether application is in benchmark mode.
