Unify turning preprocessor tokens to strings using PREPROCESSOR_TO_STRING macro.

diff --git a/Engine/Source/Runtime/Core/Public/HTML5/HTML5AssertionMacros.h b/Engine/Source/Runtime/Core/Public/HTML5/HTML5AssertionMacros.h
--- a/Engine/Source/Runtime/Core/Public/HTML5/HTML5AssertionMacros.h
+++ b/Engine/Source/Runtime/Core/Public/HTML5/HTML5AssertionMacros.h
@@ -10,9 +10,6 @@ void emscripten_log(int flags, ...);
 }
 #endif
 
-#define STRINGIZE_HELPER(x) #x
-#define STRINGIZE_HTML5(x) STRINGIZE_HELPER(x)
-
 #if DO_CHECK && !PLATFORM_HTML5_WIN32
 
 // For the asm.js builds, use emscripten-specific versions of these macros
@@ -31,10 +28,10 @@ void emscripten_log(int flags, ...);
 #define checkNoReentry(...)
 #define checkNoRecursion(...)
 
-#define check(expr)			{ if (!(expr)) { emscripten_log(255, "Expression '" #expr "' failed in " __FILE__ ":" STRINGIZE_HTML5(__LINE__) "!\n"); } }
-#define checkf(expr, ...)	{ if (!(expr)) { emscripten_log(255, "Expression '" #expr "' failed in " __FILE__ ":" STRINGIZE_HTML5(__LINE__) "!\n"); emscripten_log(255, ##__VA_ARGS__); FDebug::AssertFailed( #expr, __FILE__, __LINE__, ##__VA_ARGS__ ); } CA_ASSUME(expr); }
-#define verify(expr)		{ if(!(expr)) { emscripten_log(255, "Expression '" #expr "' failed in " __FILE__ ":" STRINGIZE_HTML5(__LINE__) "!\n"); } }
-#define verifyf(expr, ...)	{ if(!(expr)) { emscripten_log(255, "Expression '" #expr "' failed in " __FILE__ ":" STRINGIZE_HTML5(__LINE__) "!\n"); emscripten_log(255, ##__VA_ARGS__); } }
+#define check(expr)			{ if (!(expr)) { emscripten_log(255, "Expression '" #expr "' failed in " __FILE__ ":" PREPROCESSOR_TO_STRING(__LINE__) "!\n"); } }
+#define checkf(expr, ...)	{ if (!(expr)) { emscripten_log(255, "Expression '" #expr "' failed in " __FILE__ ":" PREPROCESSOR_TO_STRING(__LINE__) "!\n"); emscripten_log(255, ##__VA_ARGS__); FDebug::AssertFailed( #expr, __FILE__, __LINE__, ##__VA_ARGS__ ); } CA_ASSUME(expr); }
+#define verify(expr)		{ if(!(expr)) { emscripten_log(255, "Expression '" #expr "' failed in " __FILE__ ":" PREPROCESSOR_TO_STRING(__LINE__) "!\n"); } }
+#define verifyf(expr, ...)	{ if(!(expr)) { emscripten_log(255, "Expression '" #expr "' failed in " __FILE__ ":" PREPROCESSOR_TO_STRING(__LINE__) "!\n"); emscripten_log(255, ##__VA_ARGS__); } }
 
 #endif
 
@@ -44,8 +41,8 @@ void emscripten_log(int flags, ...);
 #undef checkfSlow
 #undef verifySlow
 
-#define checkSlow(expr, ...)   {if(!(expr)) { emscripten_log(255, "Expression '" #expr "' failed in " __FILE__ ":" STRINGIZE_HTML5(__LINE__) "!\n"); emscripten_log(255, ##__VA_ARGS__); FDebug::AssertFailed( #expr, __FILE__, __LINE__ ); } CA_ASSUME(expr); }
-#define checkfSlow(expr, ...)	{ if(!(expr)) { emscripten_log(255, "Expression '" #expr "' failed in " __FILE__ ":" STRINGIZE_HTML5(__LINE__) "!\n"); emscripten_log(255, ##__VA_ARGS__); FDebug::AssertFailed( #expr, __FILE__, __LINE__, __VA_ARGS__ ); } CA_ASSUME(expr); }
-#define verifySlow(expr)  {if(!(expr)) { emscripten_log(255, "Expression '" #expr "' failed in " __FILE__ ":" STRINGIZE_HTML5(__LINE__) "!\n"); FDebug::AssertFailed( #expr, __FILE__, __LINE__ ); } }
+#define checkSlow(expr, ...)   {if(!(expr)) { emscripten_log(255, "Expression '" #expr "' failed in " __FILE__ ":" PREPROCESSOR_TO_STRING(__LINE__) "!\n"); emscripten_log(255, ##__VA_ARGS__); FDebug::AssertFailed( #expr, __FILE__, __LINE__ ); } CA_ASSUME(expr); }
+#define checkfSlow(expr, ...)	{ if(!(expr)) { emscripten_log(255, "Expression '" #expr "' failed in " __FILE__ ":" PREPROCESSOR_TO_STRING(__LINE__) "!\n"); emscripten_log(255, ##__VA_ARGS__); FDebug::AssertFailed( #expr, __FILE__, __LINE__, __VA_ARGS__ ); } CA_ASSUME(expr); }
+#define verifySlow(expr)  {if(!(expr)) { emscripten_log(255, "Expression '" #expr "' failed in " __FILE__ ":" PREPROCESSOR_TO_STRING(__LINE__) "!\n"); FDebug::AssertFailed( #expr, __FILE__, __LINE__ ); } }
 
 #endif
diff --git a/Engine/Source/Runtime/Core/Public/Misc/CoreMiscDefines.h b/Engine/Source/Runtime/Core/Public/Misc/CoreMiscDefines.h
--- a/Engine/Source/Runtime/Core/Public/Misc/CoreMiscDefines.h
+++ b/Engine/Source/Runtime/Core/Public/Misc/CoreMiscDefines.h
@@ -6,6 +6,15 @@
 
 #pragma once
 
+// This file is included in some resource files, which issue a warning:
+//
+// warning RC4011: identifier truncated to 'PLATFORM_CAN_SUPPORT_EDITORONLY'
+//
+// due to limitations of resource compiler. The only thing needed from this file
+// for resource compilation is PREPROCESSOR_TO_STRING macro at the end, so we take
+// rest of code out for resource compilation.
+#ifndef RC_INVOKED
+
 #define LOCALIZED_SEEKFREE_SUFFIX	TEXT("_LOC")
 #define PLAYWORLD_PACKAGE_PREFIX TEXT("UEDPIE")
 #define PLAYWORLD_CONSOLE_BASE_PACKAGE_PREFIX TEXT("UED")
@@ -84,13 +93,15 @@
 enum {INDEX_NONE	= -1				};
 enum {UNICODE_BOM   = 0xfeff			};
 
-// helpers to turn an preprocessor token into a real string (see UBT_COMPILED_PLATFORM)
-#define PREPROCESSOR_TO_STRING_INNER(x) #x
-#define PREPROCESSOR_TO_STRING(x) PREPROCESSOR_TO_STRING_INNER(x)
-
 enum EForceInit 
 {
 	ForceInit,
 	ForceInitToZero
 };
 enum ENoInit {NoInit};
+
+#endif // RC_INVOKED
+
+// helpers to turn an preprocessor token into a real string (see UBT_COMPILED_PLATFORM)
+#define PREPROCESSOR_TO_STRING_INNER(x) #x
+#define PREPROCESSOR_TO_STRING(x) PREPROCESSOR_TO_STRING_INNER(x)
diff --git a/Engine/Source/Runtime/Core/Resources/Windows/ModuleVersionResource.rc.inl b/Engine/Source/Runtime/Core/Resources/Windows/ModuleVersionResource.rc.inl
--- a/Engine/Source/Runtime/Core/Resources/Windows/ModuleVersionResource.rc.inl
+++ b/Engine/Source/Runtime/Core/Resources/Windows/ModuleVersionResource.rc.inl
@@ -3,9 +3,7 @@
 #include <windows.h>
 #include "../../Public/Modules/ModuleVersion.h"
 #include "ModuleVersionResource.h"
-
-#define XSTR(s) #s
-#define STRINGIFY(s) XSTR(s)
+#include "CoreMiscDefines.h"
 
 // Netral language
 LANGUAGE LANG_NEUTRAL,SUBLANG_NEUTRAL
@@ -14,6 +12,6 @@ LANGUAGE LANG_NEUTRAL,SUBLANG_NEUTRAL
 #if !IS_MONOLITHIC
 ID_MODULE_API_VERSION_RESOURCE RCDATA
 {
-	STRINGIFY(MODULE_API_VERSION) 0
+	PREPROCESSOR_TO_STRING(MODULE_API_VERSION) 0
 }
 #endif
\ No newline at end of file
diff --git a/Engine/Source/Runtime/Engine/Private/ShowFlags.cpp b/Engine/Source/Runtime/Engine/Private/ShowFlags.cpp
--- a/Engine/Source/Runtime/Engine/Private/ShowFlags.cpp
+++ b/Engine/Source/Runtime/Engine/Private/ShowFlags.cpp
@@ -7,9 +7,6 @@
 #include "EnginePrivate.h"
 #include "ShowFlags.h"
 
-#define STRINGIFY(s) JOINSTRING(s)
-#define JOINSTRING(s) #s
-
 static bool IsValidNameChar(TCHAR c)
 {
 	return (c >= (TCHAR)'a' && c <= (TCHAR)'z')
@@ -166,7 +163,7 @@ int32 FEngineShowFlags::FindIndexByName(const TCHAR* Name, const TCHAR *CommaSep
 		// search through all defined showflags.
 		FString Search = Name;
 
-		#define SHOWFLAG_ALWAYS_ACCESSIBLE(a,...) if(Search == STRINGIFY(a)) { return (int32)SF_##a; }
+		#define SHOWFLAG_ALWAYS_ACCESSIBLE(a,...) if(Search == PREPROCESSOR_TO_STRING(a)) { return (int32)SF_##a; }
 
 		#include "ShowFlagsValues.inl"
 
@@ -205,7 +202,7 @@ FString FEngineShowFlags::FindNameByIndex(uint32 InIndex)
 {
 	FString Name;
 
-	#define SHOWFLAG_ALWAYS_ACCESSIBLE(a,...) case SF_##a: Name = STRINGIFY(a); break;
+	#define SHOWFLAG_ALWAYS_ACCESSIBLE(a,...) case SF_##a: Name = PREPROCESSOR_TO_STRING(a); break;
 
 	switch (InIndex)
 	{
@@ -219,7 +216,7 @@ FString FEngineShowFlags::FindNameByIndex(uint32 InIndex)
 
 void FEngineShowFlags::AddNameByIndex(uint32 InIndex, FString& Out)
 {
-	#define SHOWFLAG_ALWAYS_ACCESSIBLE(a,...) case SF_##a: Out += STRINGIFY(a); break;
+	#define SHOWFLAG_ALWAYS_ACCESSIBLE(a,...) case SF_##a: Out += PREPROCESSOR_TO_STRING(a); break;
 	switch (InIndex)
 	{
 		#include "ShowFlagsValues.inl"
diff --git a/Engine/Source/Runtime/Engine/Public/ShowFlags.h b/Engine/Source/Runtime/Engine/Public/ShowFlags.h
--- a/Engine/Source/Runtime/Engine/Public/ShowFlags.h
+++ b/Engine/Source/Runtime/Engine/Public/ShowFlags.h
@@ -66,9 +66,6 @@ struct FEngineShowFlags
 	};
 	// ---------------------------------------------------------
 
-	#define STRINGIFY(s) JOINSTRING(s)
-	#define JOINSTRING(s) #s
-
 	// namespace to match the original in the old loc system
 	#define LOCTEXT_NAMESPACE "UnrealEd"
 
@@ -79,8 +76,8 @@ struct FEngineShowFlags
 
 		if (LocNames.Num() == 0)
 		{
-			#define SHOWFLAG_ALWAYS_ACCESSIBLE(a,b,c) LocNames.Add( TEXT(STRINGIFY(a)), c);
-			#define SHOWFLAG_FIXED_IN_SHIPPING(a,b,c,d) LocNames.Add( TEXT(STRINGIFY(a)), d);
+			#define SHOWFLAG_ALWAYS_ACCESSIBLE(a,b,c) LocNames.Add( TEXT(PREPROCESSOR_TO_STRING(a)), c);
+			#define SHOWFLAG_FIXED_IN_SHIPPING(a,b,c,d) LocNames.Add( TEXT(PREPROCESSOR_TO_STRING(a)), d);
 			#include "ShowFlagsValues.inl"
 
 			// Additional strings that don't correspond to a showflag variable:
@@ -319,7 +316,7 @@ struct FEngineShowFlags
 	 */
 	template <class T> static void IterateAllFlags(T& Sink)
 	{
-		#define SHOWFLAG_ALWAYS_ACCESSIBLE(a,...) if( !Sink.OnEngineShowFlag((uint32)SF_##a, TEXT(STRINGIFY(a))) ) { return; }
+		#define SHOWFLAG_ALWAYS_ACCESSIBLE(a,...) if( !Sink.OnEngineShowFlag((uint32)SF_##a, TEXT(PREPROCESSOR_TO_STRING(a))) ) { return; }
 		#include "ShowFlagsValues.inl"
 	}
 
@@ -388,6 +385,3 @@ ENGINE_API EViewModeIndex FindViewMode(const FEngineShowFlags& EngineShowFlags);
 
 /** @return 0 if outside of the range, e.g. TEXT("LightingOnly") for EVM_LightingOnly */
 const TCHAR* GetViewModeName(EViewModeIndex ViewModeIndex);
-
-#undef STRINGIFY
-#undef JOINSTRING
diff --git a/Engine/Source/Runtime/Launch/Resources/Windows/PCLaunch.rc b/Engine/Source/Runtime/Launch/Resources/Windows/PCLaunch.rc
--- a/Engine/Source/Runtime/Launch/Resources/Windows/PCLaunch.rc
+++ b/Engine/Source/Runtime/Launch/Resources/Windows/PCLaunch.rc
@@ -1,9 +1,7 @@
 // Copyright 1998-2014 Epic Games, Inc. All Rights Reserved.
 
 #include "resource.h"
-
-#define XSTR(s) #s
-#define STRINGIFY(s) XSTR(s)
+#include "CoreMiscDefines.h"
 
 #define APSTUDIO_READONLY_SYMBOLS
 /////////////////////////////////////////////////////////////////////////////
@@ -135,7 +133,7 @@ BEGIN
 			VALUE "ProductVersion", ENGINE_VERSION_STRING
 			VALUE "FileDescription", EPIC_PRODUCT_NAME
 			VALUE "InternalName", EPIC_PRODUCT_IDENTIFIER
-			VALUE "OriginalFilename", STRINGIFY(ORIGINAL_FILE_NAME)
+			VALUE "OriginalFilename", PREPROCESSOR_TO_STRING(ORIGINAL_FILE_NAME)
 		END
 	END
 	BLOCK "VarFileInfo"
