Core - Fixed FCString usage

diff --git a/Engine/Source/Runtime/Core/Private/HAL/OutputDevices.cpp b/Engine/Source/Runtime/Core/Private/HAL/OutputDevices.cpp
--- a/Engine/Source/Runtime/Core/Private/HAL/OutputDevices.cpp
+++ b/Engine/Source/Runtime/Core/Private/HAL/OutputDevices.cpp
@@ -1138,17 +1138,17 @@ FOutputDeviceAnsiError::FOutputDeviceAnsiError()
 void FOutputDeviceAnsiError::Serialize( const TCHAR* Msg, ELogVerbosity::Type Verbosity, const class FName& Category )
 {
 	// Display the error and exit.
-  	FPlatformMisc::LocalPrint( TEXT("\nappError called: \n") );
+	FPlatformMisc::LocalPrint( TEXT("\nappError called: \n") );
 	FPlatformMisc::LocalPrint( Msg );
-  	FPlatformMisc::LocalPrint( TEXT("\n") );
+	FPlatformMisc::LocalPrint( TEXT("\n") );
 
 	if( !GIsCriticalError )
 	{
 		// First appError.
 		GIsCriticalError = 1;
 		UE_LOG(LogHAL, Error, TEXT("appError called: %s"), Msg );
-		FCString::Strncpy( GErrorHist, Msg, ARRAY_COUNT(GErrorHist) - 5 );
-		FCString::Strncat( GErrorHist, TEXT("\r\n\r\n"), ARRAY_COUNT(GErrorHist) - 1 );
+		FCString::Strncpy( GErrorHist, Msg, ARRAY_COUNT(GErrorHist) );
+		FCString::Strncat( GErrorHist, TEXT("\r\n\r\n"), ARRAY_COUNT(GErrorHist) );
 		ErrorPos = FCString::Strlen(GErrorHist);
 	}
 	else
