Stats - Fixed minor memory leak

diff --git a/Engine/Source/Runtime/Core/Private/Windows/WindowsRunnableThread.cpp b/Engine/Source/Runtime/Core/Private/Windows/WindowsRunnableThread.cpp
--- a/Engine/Source/Runtime/Core/Private/Windows/WindowsRunnableThread.cpp
+++ b/Engine/Source/Runtime/Core/Private/Windows/WindowsRunnableThread.cpp
@@ -32,8 +32,8 @@ uint32 FRunnableThreadWin::GuardedRun()
 			GWarn->Flush();
 
 			// Append the thread name at the end of the error report.
-			FCString::Strncat(GErrorHist, TEXT("\r\nCrash in runnable thread "), ARRAY_COUNT(GErrorHist) - 1);
-			FCString::Strncat(GErrorHist, *ThreadName, ARRAY_COUNT(GErrorHist) - 1);
+			FCString::Strncat( GErrorHist, LINE_TERMINATOR TEXT( "Crash in runnable thread " ), ARRAY_COUNT( GErrorHist ) );
+			FCString::Strncat( GErrorHist, *ThreadName, ARRAY_COUNT( GErrorHist ) );
 
 			ExitCode = 1;
 			// Generate status report.				
@@ -48,6 +48,10 @@ uint32 FRunnableThreadWin::GuardedRun()
 		ExitCode = Run();
 	}
 
+#if STATS
+	FThreadStats::Shutdown();
+#endif
+
 	return ExitCode;
 }
 
