Optimize TArray::RemoveAll<Predicate> to use the implementation from TArray::Remove
(and made TArray::Remove call RemoveAll using a lambda predicate instead of having its own copy of the complex run-removal code)

diff --git a/Engine/Source/Runtime/Core/Public/Containers/Array.h b/Engine/Source/Runtime/Core/Public/Containers/Array.h
--- a/Engine/Source/Runtime/Core/Public/Containers/Array.h
+++ b/Engine/Source/Runtime/Core/Public/Containers/Array.h
@@ -1355,10 +1355,24 @@ public:
 	}
 
 	/** Removes as many instances of Item as there are in the array, maintaining order but not indices. */
-	int32 Remove( const ElementType& Item )
+	int32 Remove(const ElementType& Item)
 	{
 		CheckAddress(&Item);
 
+		// Element is non-const to preserve compatibility with existing code with a non-const operator==() member function
+		return RemoveAll([&Item](ElementType& Element) { return Element == Item; });
+	}
+
+
+	/**
+	 * Remove all instances that match the predicate, maintaining order but not indices
+	 * Optimized to work with runs of matches/non-matches
+	 *
+	 * @param Predicate Predicate class instance
+	 */
+	template <class PREDICATE_CLASS>
+	int32 RemoveAll(const PREDICATE_CLASS& Predicate)
+	{
 		const int32 OriginalNum = ArrayNum;
 		if (!OriginalNum)
 		{
@@ -1367,11 +1381,11 @@ public:
 
 		int32 WriteIndex = 0;
 		int32 ReadIndex = 0;
-		bool NotMatch = !(GetTypedData()[ReadIndex] == Item); // use a ! to guarantee it can't be anything other than zero or one
+		bool NotMatch = !Predicate(GetTypedData()[ReadIndex]); // use a ! to guarantee it can't be anything other than zero or one
 		do
 		{
 			int32 RunStartIndex = ReadIndex++;
-			while (ReadIndex < OriginalNum && NotMatch == !(GetTypedData()[ReadIndex] == Item))
+			while (ReadIndex < OriginalNum && NotMatch == !Predicate(GetTypedData()[ReadIndex]))
 			{
 				ReadIndex++;
 			}
@@ -1382,7 +1396,7 @@ public:
 				// this was a non-matching run, we need to move it
 				if (WriteIndex != RunStartIndex)
 				{
-					FMemory::Memmove( &GetTypedData()[ WriteIndex ], &GetTypedData()[ RunStartIndex ], sizeof(ElementType) * RunLength );
+					FMemory::Memmove(&GetTypedData()[WriteIndex], &GetTypedData()[RunStartIndex], sizeof(ElementType) * RunLength);
 				}
 				WriteIndex += RunLength;
 			}
@@ -1398,28 +1412,6 @@ public:
 		return OriginalNum - ArrayNum;
 	}
 
-
-	/**
-	 * Remove all instances that match the predicate
-	 *
-	 * @param Predicate Predicate class instance
-	 */
-	template <class PREDICATE_CLASS>
-	void RemoveAll( const PREDICATE_CLASS& Predicate )
-	{
-		for ( int32 ItemIndex=0; ItemIndex < Num(); )
-		{
-			if ( Predicate( (*this)[ItemIndex] ) )
-			{
-				RemoveAt(ItemIndex);
-			}
-			else
-			{
-				++ItemIndex;
-			}
-		}
-	}
-
 	/**
 	 * Remove all instances that match the predicate
 	 *
