Temporarily disable MAC_SEPARATE_GAME_THREAD since there are a number of crashes caused by it that need to be fixed.

diff --git a/Engine/Source/Runtime/Core/Private/Mac/CocoaThread.cpp b/Engine/Source/Runtime/Core/Private/Mac/CocoaThread.cpp
--- a/Engine/Source/Runtime/Core/Private/Mac/CocoaThread.cpp
+++ b/Engine/Source/Runtime/Core/Private/Mac/CocoaThread.cpp
@@ -3,7 +3,7 @@
 #include "Core.h"
 #include "CocoaThread.h"
 
-#define MAC_SEPARATE_GAME_THREAD 1 // Separate the main & game threads so that we better handle the interaction between the Cocoa's event delegates and UE4's event polling.
+#define MAC_SEPARATE_GAME_THREAD 0 // Separate the main & game threads so that we better handle the interaction between the Cocoa's event delegates and UE4's event polling.
 
 static FCocoaGameThread* GCocoaGameThread = nil;
 static NSRunLoop* GCocoaGameRunLoop = nil;
