Get rid of duplicate log messages on Android
#Android

diff --git a/Engine/Source/Runtime/Core/Private/Android/AndroidOutputDevices.cpp b/Engine/Source/Runtime/Core/Private/Android/AndroidOutputDevices.cpp
--- a/Engine/Source/Runtime/Core/Private/Android/AndroidOutputDevices.cpp
+++ b/Engine/Source/Runtime/Core/Private/Android/AndroidOutputDevices.cpp
@@ -20,8 +20,9 @@ class FOutputDevice* FAndroidOutputDevices::GetLog()
 	return NULL;
 #else
 	// Always enable logging via ADB
-	static FOutputDeviceAndroidDebug Singleton;
-	return &Singleton;
+//	static FOutputDeviceAndroidDebug Singleton;
+//	return &Singleton;
+	return NULL;
 #endif
 }
 
