Changed GetCurrentThreadId on Android to use gettid() instead of pthread_self() to handle 64-bit
#Android

diff --git a/Engine/Source/Runtime/Core/Public/Android/AndroidTLS.h b/Engine/Source/Runtime/Core/Public/Android/AndroidTLS.h
--- a/Engine/Source/Runtime/Core/Public/Android/AndroidTLS.h
+++ b/Engine/Source/Runtime/Core/Public/Android/AndroidTLS.h
@@ -17,7 +17,9 @@ struct CORE_API FAndroidTLS : public FGenericPlatformTLS
 	 */
 	static FORCEINLINE uint32 GetCurrentThreadId(void)
 	{
-		return pthread_self();
+		static_assert(sizeof(pid_t) <= sizeof(uint32), "pid_t is larger than uint32, reconsider implementation of GetCurrentThreadId()");
+		return static_cast<uint32>(gettid());
+//		return pthread_self();
 	}
 
 	/**
