Remove UBT invocation functions from FModuleManager that are no longer used (and shouldn't have been there in the first place).

diff --git a/Engine/Source/Runtime/Core/Private/Modules/ModuleManager.cpp b/Engine/Source/Runtime/Core/Private/Modules/ModuleManager.cpp
--- a/Engine/Source/Runtime/Core/Private/Modules/ModuleManager.cpp
+++ b/Engine/Source/Runtime/Core/Private/Modules/ModuleManager.cpp
@@ -1284,124 +1284,6 @@ bool FModuleManager::RecompileModulesAsync( const TArray< FName > ModuleNames, c
 }
 
 
-bool FModuleManager::CompileGameProject( const FString& ProjectFilename, FOutputDevice &Ar )
-{
-	bool bCompileSucceeded = false;
-
-#if !IS_MONOLITHIC
-	ModuleCompilerStartedEvent.Broadcast();
-
-	// Leave an empty list of modules to recompile to indicate that all modules in the specified project file should be compiled
-	TArray< FModuleToRecompile > ModulesToRecompile;	
-	const FString GameName = FPaths::GetBaseFilename(ProjectFilename);
-	FString AdditionalCmdLineArgs = FString::Printf(TEXT("\"%s\""), *ProjectFilename);
-	if (FRocketSupport::IsRocket())
-	{
-		AdditionalCmdLineArgs += TEXT(" -rocket");
-	}
-
-	if (StartCompilingModuleDLLs(GameName, ModulesToRecompile, FRecompileModulesCallback(), Ar, false, AdditionalCmdLineArgs))
-	{
-		const bool bWaitForCompletion = true;	// Always wait
-		bool bCompileStillInProgress = false;
-		FOutputDeviceNull NullOutput;
-		CheckForFinishedModuleDLLCompile(bWaitForCompletion, bCompileStillInProgress, bCompileSucceeded, NullOutput);
-	}
-#endif // !IS_MONOLITHIC
-	return bCompileSucceeded;
-}
-
-bool FModuleManager::CompileGameProjectEditor( const FString& GameProjectFilename, FOutputDevice &Ar )
-{
-	bool bCompileSucceeded = false;
-#if !IS_MONOLITHIC
-	ModuleCompilerStartedEvent.Broadcast();
-
-	// Leave an empty list of modules to recompile to indicate that all modules in the specified project file should be compiled
-	TArray< FModuleToRecompile > ModulesToRecompile;
-
-	FString GameTargetName = FPaths::GetBaseFilename(GameProjectFilename);
-	//@NOTE: We assume 'MyProject'Editor is the target name for the editor...
-	// Since we generate the targets, this should be fine - but worth noting.
-	GameTargetName += TEXT("Editor");
-	FString AdditionalCmdLineArgs = FString::Printf(TEXT("\"%s\""), *GameProjectFilename);
-	if (FRocketSupport::IsRocket())
-	{
-		AdditionalCmdLineArgs += TEXT(" -rocket");
-	}
-
-	if (StartCompilingModuleDLLs(GameTargetName, ModulesToRecompile, FRecompileModulesCallback(), Ar, false, AdditionalCmdLineArgs))
-	{
-		const bool bWaitForCompletion = true;	// Always wait
-		bool bCompileStillInProgress = false;
-		FOutputDeviceNull NullOutput;
-		CheckForFinishedModuleDLLCompile(bWaitForCompletion, bCompileStillInProgress, bCompileSucceeded, NullOutput);
-	}
-
-#endif // !IS_MONOLITHIC
-	return bCompileSucceeded;
-}
-
-
-bool FModuleManager::GenerateCodeProjectFiles( const FString& ProjectFilename, FOutputDevice &Ar )
-{
-	bool bCompileSucceeded = false;
-#if !IS_MONOLITHIC
-	// Leave an empty list of modules to recompile to indicate that all modules in the specified project file should be compiled
-	TArray< FModuleToRecompile > ModulesToRecompile;
-
-#if PLATFORM_MAC
-	FString CmdLineParams = TEXT("-xcodeprojectfile");
-#elif PLATFORM_LINUX
-	FString CmdLineParams = TEXT("-makefile");
-#else
-	FString CmdLineParams = TEXT("-projectfiles");
-#endif
-	if ( !ProjectFilename.IsEmpty() )
-	{
-		// Normalize the project filename
-		FString NormalizedProjectFilename = ProjectFilename;
-		FPaths::NormalizeFilename(NormalizedProjectFilename);
-
-		// Figure out if this is a foreign project, and specify the project filename on the command line if it is. Don't use the cached project dictionary; the project may have been newly added.
-		FUProjectDictionary Dictionary(FPaths::RootDir());
-		if(Dictionary.IsForeignProject(NormalizedProjectFilename))
-		{
-			CmdLineParams += FString::Printf(TEXT(" \"%s\" -game"), *IFileManager::Get().ConvertToAbsolutePathForExternalAppForRead(*NormalizedProjectFilename));
-
-			// Determine whether we want engine source in the generated project files
-			if(FEngineBuildSettings::IsSourceDistribution())
-			{
-				CmdLineParams += TEXT(" -engine");
-			}
-		}
-
-		// Add the Rocket parameter
-		if ( FRocketSupport::IsRocket() )
-		{
-            CmdLineParams += TEXT(" -rocket");
-        }
-	}
-
-	if( InvokeUnrealBuildToolForCompile( CmdLineParams, Ar ) )
-	{
-		const bool bWaitForCompletion = true;	// Always wait
-		bool bCompileStillInProgress = false;
-		const FText SlowTaskOverrideText = NSLOCTEXT("FModuleManager", "GenerateCodeProjectsStatusMessage", "Generating code projects...");
-		const bool bFireEvents = false;
-		CheckForFinishedModuleDLLCompile( bWaitForCompletion, bCompileStillInProgress, bCompileSucceeded, Ar, SlowTaskOverrideText, bFireEvents );
-	}
-	else
-	{
-		Ar.Logf(TEXT("Failed invoke UnrealBuildTool when generating code project files."));
-	}
-#else
-	Ar.Logf(TEXT("Cannot generate code project files from monolithic builds."));
-#endif
-
-	return bCompileSucceeded;
-}
-
 void FModuleManager::UnloadOrAbandonModuleWithCallback( const FName InModuleName, FOutputDevice &Ar )
 {
 	TSharedRef< FModuleInfo > Module = Modules.FindChecked( InModuleName );
diff --git a/Engine/Source/Runtime/Core/Public/Modules/ModuleManager.h b/Engine/Source/Runtime/Core/Public/Modules/ModuleManager.h
--- a/Engine/Source/Runtime/Core/Public/Modules/ModuleManager.h
+++ b/Engine/Source/Runtime/Core/Public/Modules/ModuleManager.h
@@ -462,33 +462,6 @@ public:
 		bRequestCancelCompilation = true;
 	}
 
-	/**
-	 * Tries to compile the specified game project. Not used for recompiling modules that are already loaded.
-	 *
-	 * @param GameProjectFilename The filename (including path) of the game project to compile.
-	 * @param Ar Output device for logging compilation status.
-	 * @return Returns true if the project was successfully compiled.
-	 */
-	bool CompileGameProject( const FString& GameProjectFilename, FOutputDevice &Ar );
-
-	/**
-	 * Tries to compile the specified game projects editor. Not used for recompiling modules that are already loaded.
-	 *
-	 * @param GameProjectFilename The filename (including path) of the game project to compile.
-	 * @param Ar Output device for logging compilation status.
-	 * @return	Returns true if the project was successfully compiled
-	 */
-	bool CompileGameProjectEditor( const FString& GameProjectFilename, FOutputDevice &Ar );
-
-	/**
-	 * Tries to compile the specified game project. Not used for recompiling modules that are already loaded.
-	 *
-	 * @param GameProjectFilename The filename (including path) of the game project for which to generate code projects.
-	 * @param Ar Output device for logging compilation status.
-	 * @return	Returns true if the project was successfully compiled.
-	 */
-	bool GenerateCodeProjectFiles( const FString& GameProjectFilename, FOutputDevice &Ar );	
-
 	/** Delegate that's used by the module manager to initialize a registered module that we statically linked with (monolithic only) */
 	DECLARE_DELEGATE_RetVal( IModuleInterface*, FInitializeStaticallyLinkedModule )
 
