Better fix for duplicated log messages on Android
#Android

diff --git a/Engine/Source/Runtime/Core/Private/Android/AndroidOutputDevices.cpp b/Engine/Source/Runtime/Core/Private/Android/AndroidOutputDevices.cpp
--- a/Engine/Source/Runtime/Core/Private/Android/AndroidOutputDevices.cpp
+++ b/Engine/Source/Runtime/Core/Private/Android/AndroidOutputDevices.cpp
@@ -14,17 +14,17 @@ class FOutputDeviceError* FAndroidOutputDevices::GetError()
 	return &Singleton;
 }
 
-class FOutputDevice* FAndroidOutputDevices::GetLog()
+/*class FOutputDevice* FAndroidOutputDevices::GetLog()
 {
 #if UE_BUILD_SHIPPING
 	return NULL;
 #else
 	// Always enable logging via ADB
-//	static FOutputDeviceAndroidDebug Singleton;
-//	return &Singleton;
+	static FOutputDeviceAndroidDebug Singleton;
+	return &Singleton;
 	return NULL;
 #endif
-}
+}*/
 
 FOutputDeviceAndroidDebug::FOutputDeviceAndroidDebug()
 {}
diff --git a/Engine/Source/Runtime/Core/Public/Android/AndroidOutputDevices.h b/Engine/Source/Runtime/Core/Public/Android/AndroidOutputDevices.h
--- a/Engine/Source/Runtime/Core/Public/Android/AndroidOutputDevices.h
+++ b/Engine/Source/Runtime/Core/Public/Android/AndroidOutputDevices.h
@@ -8,7 +8,7 @@
 #pragma once
 struct CORE_API FAndroidOutputDevices : public FGenericPlatformOutputDevices
 {
-	static class FOutputDevice*			GetLog();
+//	static class FOutputDevice*			GetLog();
 	static class FOutputDeviceError*	GetError();
 };
 
