Added GetDefaultLocale to AndroidMisc.
#Android

diff --git a/Engine/Source/Runtime/Core/Private/Android/AndroidMisc.cpp b/Engine/Source/Runtime/Core/Private/Android/AndroidMisc.cpp
--- a/Engine/Source/Runtime/Core/Private/Android/AndroidMisc.cpp
+++ b/Engine/Source/Runtime/Core/Private/Android/AndroidMisc.cpp
@@ -273,6 +273,11 @@ const FString FAndroidMisc::GetOSLanguage()
 	return OSLanguage;
 }
 
+FString FAndroidMisc::GetDefaultLocale()
+{
+	return OSLanguage;
+}
+
 uint32 FAndroidMisc::GetKeyMap( uint16* KeyCodes, FString* KeyNames, uint32 MaxMappings )
 {
 #define ADDKEYMAP(KeyCode, KeyName)		if (NumMappings<MaxMappings) { KeyCodes[NumMappings]=KeyCode; if(KeyNames) { KeyNames[NumMappings]=KeyName; } ++NumMappings; };
diff --git a/Engine/Source/Runtime/Core/Public/Android/AndroidMisc.h b/Engine/Source/Runtime/Core/Public/Android/AndroidMisc.h
--- a/Engine/Source/Runtime/Core/Public/Android/AndroidMisc.h
+++ b/Engine/Source/Runtime/Core/Public/Android/AndroidMisc.h
@@ -83,6 +83,7 @@ struct CORE_API FAndroidMisc : public FGenericPlatformMisc
 	static const FString GetDeviceMake();
 	static const FString GetDeviceModel();
 	static const FString GetOSLanguage();
+	static FString GetDefaultLocale();
 	static FString GetGPUFamily();
 	static FString GetGLVersion();
 	static bool SupportsFloatingPointRenderTargets();
