Add support for TCHARs in FOutputDeviceFile.
TTP #343014 CORE: Add support for TCHARs in FOutputDeviceFile.

diff --git a/Engine/Source/Runtime/Core/Private/HAL/OutputDevices.cpp b/Engine/Source/Runtime/Core/Private/HAL/OutputDevices.cpp
--- a/Engine/Source/Runtime/Core/Private/HAL/OutputDevices.cpp
+++ b/Engine/Source/Runtime/Core/Private/HAL/OutputDevices.cpp
@@ -7,6 +7,8 @@
 #include "Core.h"
 
 #include <stdio.h>
+#include "Templates/UniquePtr.h"
+
 // #if _MSC_VER
 // #pragma warning (push)
 // #pragma warning (disable : 4548) // needed as xlocale does not compile cleanly
@@ -50,6 +52,9 @@ static FLogCategoryBase GlobalVerbosity(TEXT("Global"), ELogVerbosity::All, ELog
 /** A "fake" logging category that is used as a proxy for changing the default of all categories at boot time. **/
 static FLogCategoryBase BootGlobalVerbosity(TEXT("BootGlobal"), ELogVerbosity::All, ELogVerbosity::All);
 
+typedef uint8 UTF8BOMType[3];
+static UTF8BOMType UTF8BOM = { 0xEF, 0xBB, 0xBF };
+
 /** Log suppression system implementation **/
 class FLogSuppressionImplementation: public FLogSuppressionInterface, private FSelfRegisteringExec
 {
@@ -950,7 +955,7 @@ void FOutputDeviceFile::Flush()
 /** if the passed in file exists, makes a timestamped backup copy
  * @param Filename the name of the file to check
  */
-static void CreateBackupCopy(TCHAR* Filename)
+static void CreateBackupCopy(const TCHAR* Filename)
 {
 	if (IFileManager::Get().FileSize(Filename) > 0)
 	{
@@ -962,6 +967,79 @@ static void CreateBackupCopy(TCHAR* Filename)
 	}
 }
 
+void FOutputDeviceFile::WriteByteOrderMarkToArchive(EByteOrderMark ByteOrderMark)
+{
+	switch (ByteOrderMark)
+	{
+	case EByteOrderMark::UTF8:
+		LogAr->Serialize(UTF8BOM, sizeof(UTF8BOM));
+		break;
+
+	case EByteOrderMark::Unspecified:
+	default:
+		check(false);
+		break;
+	}
+}
+
+FArchive* FOutputDeviceFile::CreateArchive(uint32 MaxAttempts)
+{
+	uint32 WriteFlags = FILEWRITE_AllowRead | (Opened ? FILEWRITE_Append : 0);
+
+	// Open log file.
+	auto Result = IFileManager::Get().CreateFileWriter(Filename, WriteFlags);
+
+	// If that failed, append an _2 and try again (unless we don't want extra copies). This 
+	// happens in the case of running a server and client on same computer for example.
+	if (!bDisableBackup && !Result)
+	{
+		FString FilenamePart = FPaths::GetBaseFilename(Filename, false) + "_";
+		FString ExtensionPart = FPaths::GetExtension(Filename, true);
+		FString FinalFilename;
+		uint32 FileIndex = 2;
+		do
+		{
+			// Continue to increment indices until a valid filename is found
+			FinalFilename = FilenamePart + FString::FromInt(FileIndex) + ExtensionPart;
+			if (!Opened)
+			{
+				CreateBackupCopy(*FinalFilename);
+			}
+			Result = IFileManager::Get().CreateFileWriter(*FinalFilename, WriteFlags);
+		} while (!Result && FileIndex < MaxAttempts);
+	}
+
+	return Result;
+}
+
+void FOutputDeviceFile::CastAndSerializeData(const TCHAR* Data)
+{
+	auto ConvertedData = FTCHARToUTF8(Data);
+	LogAr->Serialize((ANSICHAR*)(ConvertedData.Get()), ConvertedData.Length());
+}
+
+void FOutputDeviceFile::WriteDataToArchive(const TCHAR* Data, ELogVerbosity::Type Verbosity, const class FName& Category)
+{
+	if (!bSuppressEventTag)
+	{
+		FString Prefix = FOutputDevice::FormatLogLine(Verbosity, Category, NULL, GPrintLogTimes);
+		CastAndSerializeData(*Prefix);
+	}
+
+	CastAndSerializeData(Data);
+
+	if (bAutoEmitLineTerminator)
+	{
+#if PLATFORM_LINUX
+		// on Linux, we still want to have logs with Windows line endings so they can be opened with Windows tools like infamous notepad.exe
+		ANSICHAR WindowsTerminator[] = { '\r', '\n' };
+		LogAr->Serialize(WindowsTerminator, sizeof(WindowsTerminator));
+#else
+		CastAndSerializeData(LINE_TERMINATOR);
+#endif // PLATFORM_LINUX
+	}
+}
+
 /**
  * Serializes the passed in data unless the current event is suppressed.
  *
@@ -989,33 +1067,14 @@ void FOutputDeviceFile::Serialize( const TCHAR* Data, ELogVerbosity::Type Verbos
 			}
 
 			// Open log file.
-			LogAr = IFileManager::Get().CreateFileWriter( Filename, FILEWRITE_AllowRead | (Opened?FILEWRITE_Append:0));
-
-			// If that failed, append an _2 and try again (unless we don't want extra copies). This 
-			// happens in the case of running a server and client on same computer for example.
-			if(!bDisableBackup && !LogAr)
-			{				  
-				int32 FileIndex = 2;
-				TCHAR ExtBuffer[MAX_SPRINTF];
-				FCString::Strcpy(ExtBuffer, TEXT(".log"));
-				do 
-				{
-					// Continue to increment indices until a valid filename is found
-					Filename[ FCString::Strlen(Filename) - FCString::Strlen(ExtBuffer) ] = 0;
-					FCString::Sprintf(ExtBuffer, TEXT("_%d.log"), FileIndex++);
-					FCString::Strcat( Filename, ExtBuffer );
-					if (!Opened)
-					{
-						CreateBackupCopy(Filename);
-					}
-					LogAr = IFileManager::Get().CreateFileWriter( Filename, FILEWRITE_AllowRead | (Opened?FILEWRITE_Append:0));
-				}
-				while(!LogAr && FileIndex < 32);
-			}
+			LogAr = CreateArchive();
 
 			if( LogAr )
 			{
 				Opened = 1;
+
+				WriteByteOrderMarkToArchive(EByteOrderMark::UTF8);
+
 				if (!bSuppressEventTag)
 				{
 					Logf( TEXT("Log file open, %s"), FPlatformTime::StrTimestamp() );
@@ -1029,43 +1088,8 @@ void FOutputDeviceFile::Serialize( const TCHAR* Data, ELogVerbosity::Type Verbos
 
 		if( LogAr && Verbosity != ELogVerbosity::SetColor )
 		{
-			int32 i = 0;
-			ANSICHAR ACh[MAX_SPRINTF];
-			if (!bSuppressEventTag)
-			{
-				FString Prefix = FOutputDevice::FormatLogLine(Verbosity, Category, NULL, GPrintLogTimes);
-				const TCHAR *Ch = *Prefix;
-				for( i=0; Ch[i]; i++ )
-				{
-					ACh[i] = CharCast<ANSICHAR>(Ch[i] );
-				}
-				LogAr->Serialize( ACh, i );
-			}
+			WriteDataToArchive(Data, Verbosity, Category);
 
-			int32 DataOffset = 0;
-			while(Data[DataOffset])
-			{
-				for(i = 0;i < ARRAY_COUNT(ACh) && Data[DataOffset];i++,DataOffset++)
-				{
-					ACh[i] = Data[DataOffset];
-				}
-				LogAr->Serialize(ACh,i);
-			};
-
-			if (bAutoEmitLineTerminator)
-			{
-#if PLATFORM_LINUX
-				// on Linux, we still want to have logs with Windows line endings so they can be opened with Windows tools like infamous notepad.exe
-				ANSICHAR WindowsTerminator[] = { '\r', '\n' };
-				LogAr->Serialize(WindowsTerminator, ARRAY_COUNT(WindowsTerminator));
-#else
-				for(i = 0;LINE_TERMINATOR[i];i++)
-				{
-					ACh[i] = LINE_TERMINATOR[i];
-				}
-				LogAr->Serialize(ACh,i);
-#endif // PLATFORM_LINUX
-			}
 			static bool GForceLogFlush = false;
 			static bool GTestedCmdLine = false;
 			if (!GTestedCmdLine)
diff --git a/Engine/Source/Runtime/Core/Public/HAL/OutputDevices.h b/Engine/Source/Runtime/Core/Public/HAL/OutputDevices.h
--- a/Engine/Source/Runtime/Core/Public/HAL/OutputDevices.h
+++ b/Engine/Source/Runtime/Core/Public/HAL/OutputDevices.h
@@ -9,6 +9,12 @@
 	#define CLEAR_WARN_COLOR() 
 #else
 
+enum class EByteOrderMark : int8
+{
+	UTF8,
+	Unspecified,
+};
+
 /*-----------------------------------------------------------------------------
 	Colorized text.
 
@@ -419,6 +425,14 @@ private:
 	bool		bDisableBackup;
 	
 	void WriteRaw( const TCHAR* C );
+
+	FArchive* CreateArchive(uint32 MaxAttempts = 32);
+
+	void WriteByteOrderMarkToArchive(EByteOrderMark ByteOrderMark);
+
+	void CastAndSerializeData(const TCHAR* Data);
+
+	void WriteDataToArchive(const TCHAR* Data, ELogVerbosity::Type Verbosity, const class FName& Category);
 };
 
 // Null output device.
