Moved Base64 from HTTP module for Core/Misc to have better visibility from other modules (without any dependency to HTTP module)
Changed EQS debug replication to use FBase64 class instead of custom conversion to/from base64

diff --git a/Engine/Source/Runtime/Core/Private/Misc/Base64.cpp b/Engine/Source/Runtime/Core/Private/Misc/Base64.cpp
new file mode 100644
--- /dev/null
+++ b/Engine/Source/Runtime/Core/Private/Misc/Base64.cpp
@@ -0,0 +1,248 @@
+// Copyright 1998-2014 Epic Games, Inc. All Rights Reserved.
+
+#include "Core.h"
+#include "Base64.h"
+
+
+/** The table used to encode a 6 bit value as an ascii character */
+uint8 FBase64::EncodingAlphabet[64] = { 'A', 'B', 'C', 'D', 'E', 'F', 'G', 'H', 'I', 'J', 'K', 'L',
+	'M', 'N', 'O', 'P', 'Q', 'R', 'S', 'T', 'U', 'V', 'W', 'X', 'Y', 'Z', 'a', 'b', 'c', 'd', 'e',
+	'f', 'g', 'h', 'i', 'j', 'k', 'l', 'm', 'n', 'o', 'p', 'q', 'r', 's', 't', 'u', 'v', 'w', 'x',
+	'y', 'z', '0', '1', '2', '3', '4', '5', '6', '7', '8', '9', '+', '/' };
+
+/** The table used to convert an ascii character into a 6 bit value */
+uint8 FBase64::DecodingAlphabet[256] = { 0xFF, 0xFF, 0xFF, 0xFF, 0xFF, 0xFF, 0xFF, 0xFF, 0xFF, 0xFF,
+	0xFF, 0xFF, 0xFF, 0xFF, 0xFF, 0xFF, 0xFF, 0xFF, 0xFF, 0xFF, 0xFF, 0xFF, 0xFF, 0xFF, 0xFF, 0xFF,
+	0xFF, 0xFF, 0xFF, 0xFF, 0xFF, 0xFF, 0xFF, 0xFF, 0xFF, 0xFF, 0xFF, 0xFF, 0xFF, 0xFF, 0xFF, 0xFF,
+	0xFF, 0x3E, 0xFF, 0xFF, 0xFF, 0x3F, 0x34, 0x35, 0x36, 0x37, 0x38, 0x39, 0x3A, 0x3B, 0x3C, 0x3D,
+	0xFF, 0xFF, 0xFF, 0x00, 0xFF, 0xFF, 0xFF, 0x00, 0x01, 0x02, 0x03, 0x04, 0x05, 0x06, 0x07, 0x08,
+	0x09, 0x0A, 0x0B, 0x0C, 0x0D, 0x0E, 0x0F, 0x10, 0x11, 0x12, 0x13, 0x14, 0x15, 0x16, 0x17, 0x18,
+	0x19, 0xFF, 0xFF, 0xFF, 0xFF, 0xFF, 0xFF, 0x1A, 0x1B, 0x1C, 0x1D, 0x1E, 0x1F, 0x20, 0x21, 0x22,
+	0x23, 0x24, 0x25, 0x26, 0x27, 0x28, 0x29, 0x2A, 0x2B, 0x2C, 0x2D, 0x2E, 0x2F, 0x30, 0x31, 0x32,
+	0x33, 0xFF, 0xFF, 0xFF, 0xFF, 0xFF, 0xFF, 0xFF, 0xFF, 0xFF, 0xFF, 0xFF, 0xFF, 0xFF, 0xFF, 0xFF,
+	0xFF, 0xFF, 0xFF, 0xFF, 0xFF, 0xFF, 0xFF, 0xFF, 0xFF, 0xFF, 0xFF, 0xFF, 0xFF, 0xFF, 0xFF, 0xFF,
+	0xFF, 0xFF, 0xFF, 0xFF, 0xFF, 0xFF, 0xFF, 0xFF, 0xFF, 0xFF, 0xFF, 0xFF, 0xFF, 0xFF, 0xFF, 0xFF,
+	0xFF, 0xFF, 0xFF, 0xFF, 0xFF, 0xFF, 0xFF, 0xFF, 0xFF, 0xFF, 0xFF, 0xFF, 0xFF, 0xFF, 0xFF, 0xFF,
+	0xFF, 0xFF, 0xFF, 0xFF, 0xFF, 0xFF, 0xFF, 0xFF, 0xFF, 0xFF, 0xFF, 0xFF, 0xFF, 0xFF, 0xFF, 0xFF,
+	0xFF, 0xFF, 0xFF, 0xFF, 0xFF, 0xFF, 0xFF, 0xFF, 0xFF, 0xFF, 0xFF, 0xFF, 0xFF, 0xFF, 0xFF, 0xFF,
+	0xFF, 0xFF, 0xFF, 0xFF, 0xFF, 0xFF, 0xFF, 0xFF, 0xFF, 0xFF, 0xFF, 0xFF, 0xFF, 0xFF, 0xFF, 0xFF,
+	0xFF, 0xFF, 0xFF, 0xFF, 0xFF, 0xFF, 0xFF, 0xFF, 0xFF, 0xFF, 0xFF, 0xFF, 0xFF, 0xFF, 0xFF, 0xFF,
+	0xFF, 0xFF, 0xFF, 0xFF, 0xFF, 0xFF };
+
+/**
+ * Encodes a binary uint8 array into a Base64 string
+ *
+ * @param Source the binary data to convert
+ *
+ * @return a string that encodes the binary data in a way that can be safely transmitted via various Internet protocols
+ */
+FString FBase64::Encode(const TArray<uint8>& Source)
+{
+	return Encode((uint8*)Source.GetTypedData(), Source.Num());
+}
+
+/**
+ * Decodes a Base64 string into an array of bytes
+ *
+ * @param Source the stringified data to convert
+ * @param Dest the out buffer that will be filled with the decoded data
+ *
+ * @return true if the buffer was decoded, false if it failed to decode
+ */
+bool FBase64::Decode(const FString& Source, TArray<uint8>& Dest)
+{
+	uint32 Length = Source.Len();
+	// Size must be a multiple of 4
+	if (Length % 4)
+	{
+		return false;
+	}
+	// Each 4 uint8 chunk of characters is 3 bytes of data
+	uint32 ExpectedLength = Length / 4 * 3;
+	// Add the number we need for output
+	Dest.AddZeroed(ExpectedLength);
+	uint8* Buffer = Dest.GetTypedData();
+	uint32 PadCount = 0;
+	bool bWasSuccessful = Decode(TCHAR_TO_ANSI(*Source), Length, Dest.GetTypedData(), PadCount);
+	if (bWasSuccessful)
+	{
+		if (PadCount > 0)
+		{
+			Dest.RemoveAt(ExpectedLength - PadCount, PadCount);
+		}
+	}
+	return bWasSuccessful;
+}
+
+/**
+ * Encodes a FString into a Base64 string
+ *
+ * @param Source the string data to convert
+ *
+ * @return a string that encodes the binary data in a way that can be safely transmitted via various Internet protocols
+ */
+FString FBase64::Encode(const FString& Source)
+{
+	return Encode((uint8*)TCHAR_TO_ANSI(*Source), Source.Len());
+}
+
+/**
+ * Decodes a Base64 string into a FString
+ *
+ * @param Source the stringified data to convert
+ * @param Dest the out buffer that will be filled with the decoded data
+ */
+bool FBase64::Decode(const FString& Source, FString& Dest)
+{
+	uint32 Length = Source.Len();
+	// Size must be a multiple of 4
+	if (Length % 4)
+	{
+		return false;
+	}
+	// Each 4 uint8 chunk of characters is 3 bytes of data
+	uint32 ExpectedLength = Length / 4 * 3;
+	TArray<ANSICHAR> TempDest;
+	TempDest.AddZeroed(ExpectedLength);
+	uint8* Buffer = (uint8*)TempDest.GetTypedData();
+	uint32 PadCount = 0;
+
+	bool bWasSuccessful = Decode(TCHAR_TO_ANSI(*Source), Length, Buffer, PadCount);
+	if (bWasSuccessful)
+	{
+		if (PadCount > 0)
+		{
+			Buffer[ExpectedLength - PadCount] = 0;
+		}
+		else
+		{
+			TempDest.Add('\0');
+		}
+		Dest = ANSI_TO_TCHAR(TempDest.GetTypedData());
+	}
+	return bWasSuccessful;
+}
+
+/**
+ * Encodes the source into a Base64 string
+ *
+ * @param Source the binary payload to stringify
+ * @param Length the length of the payload that needs encoding
+ *
+ * @return the stringified form of the binary data
+ */
+FString FBase64::Encode(uint8* Source, uint32 Length)
+{
+	// Each 3 uint8 set of data expands to 4 bytes and must be padded to 4 bytes
+	uint32 ExpectedLength = (Length + 2) / 3 * 4;
+	FString OutBuffer;
+	OutBuffer.Empty(ExpectedLength);
+
+	ANSICHAR EncodedBytes[5];
+	// Null terminate this so we can append it easily as a string
+	EncodedBytes[4] = 0;
+	// Loop through the buffer converting 3 bytes of binary data at a time
+	while (Length >= 3)
+	{
+		uint8 A = *Source++;
+		uint8 B = *Source++;
+		uint8 C = *Source++;
+		Length -= 3;
+		// The algorithm takes 24 bits of data (3 bytes) and breaks it into 4 6bit chunks represented as ascii
+		uint32 ByteTriplet = A << 16 | B << 8 | C;
+		// Use the 6bit block to find the representation ascii character for it
+		EncodedBytes[3] = EncodingAlphabet[ByteTriplet & 0x3F];
+		// Shift each time to get the next 6 bit chunk
+		ByteTriplet >>= 6;
+		EncodedBytes[2] = EncodingAlphabet[ByteTriplet & 0x3F];
+		ByteTriplet >>= 6;
+		EncodedBytes[1] = EncodingAlphabet[ByteTriplet & 0x3F];
+		ByteTriplet >>= 6;
+		EncodedBytes[0] = EncodingAlphabet[ByteTriplet & 0x3F];
+		// Now we can append this buffer to our destination string
+		OutBuffer += EncodedBytes;
+	}
+	// Since this algorithm operates on blocks, we may need to pad the last chunks
+	if (Length > 0)
+	{
+		uint8 A = *Source++;
+		uint8 B = 0;
+		uint8 C = 0;
+		// Grab the second character if it is a 2 uint8 finish
+		if (Length == 2)
+		{
+			B = *Source;
+		}
+		uint32 ByteTriplet = A << 16 | B << 8 | C;
+		// Pad with = to make a 4 uint8 chunk
+		EncodedBytes[3] = '=';
+		ByteTriplet >>= 6;
+		// If there's only one 1 uint8 left in the source, then you need 2 pad chars
+		if (Length == 1)
+		{
+			EncodedBytes[2] = '=';
+		}
+		else
+		{
+			EncodedBytes[2] = EncodingAlphabet[ByteTriplet & 0x3F];
+		}
+		// Now encode the remaining bits the same way
+		ByteTriplet >>= 6;
+		EncodedBytes[1] = EncodingAlphabet[ByteTriplet & 0x3F];
+		ByteTriplet >>= 6;
+		EncodedBytes[0] = EncodingAlphabet[ByteTriplet & 0x3F];
+		OutBuffer += EncodedBytes;
+	}
+	return OutBuffer;
+}
+
+/**
+ * Decodes a Base64 string into an array of bytes
+ *
+ * @param Source the stringified data to convert
+ * @param Length the length of the buffer being converted
+ * @param Dest the out buffer that will be filled with the decoded data
+ * @param PadCount the out count of the padding on the orignal buffer (0 to 2)
+ *
+ * @return true if the buffer was decoded, false if it failed to decode
+ */
+bool FBase64::Decode(const ANSICHAR* Source, uint32 Length, uint8* Dest, uint32& PadCount)
+{
+	PadCount = 0;
+	uint8 DecodedValues[4];
+	while (Length)
+	{
+		// Decode the next 4 BYTEs
+		for (int32 Index = 0; Index < 4; Index++)
+		{
+			// Tell the caller if there were any pad bytes
+			if (*Source == '=')
+			{
+				PadCount++;
+			}
+			DecodedValues[Index] = DecodingAlphabet[(int32)(*Source++)];
+			// Abort on values that we don't understand
+			if (DecodedValues[Index] == -1)
+			{
+				return false;
+			}
+		}
+		Length -= 4;
+		// Rebuild the original 3 bytes from the 4 chunks of 6 bits
+		uint32 OriginalTriplet = DecodedValues[0];
+		OriginalTriplet <<= 6;
+		OriginalTriplet |= DecodedValues[1];
+		OriginalTriplet <<= 6;
+		OriginalTriplet |= DecodedValues[2];
+		OriginalTriplet <<= 6;
+		OriginalTriplet |= DecodedValues[3];
+		// Now we can tear the uint32 into bytes
+		Dest[2] = OriginalTriplet & 0xFF;
+		OriginalTriplet >>= 8;
+		Dest[1] = OriginalTriplet & 0xFF;
+		OriginalTriplet >>= 8;
+		Dest[0] = OriginalTriplet & 0xFF;
+		Dest += 3;
+	}
+	return true;
+}
diff --git a/Engine/Source/Runtime/Core/Public/Misc/Base64.h b/Engine/Source/Runtime/Core/Public/Misc/Base64.h
new file mode 100644
--- /dev/null
+++ b/Engine/Source/Runtime/Core/Public/Misc/Base64.h
@@ -0,0 +1,80 @@
+// Copyright 1998-2014 Epic Games, Inc. All Rights Reserved.
+
+#pragma once
+
+/**
+ * Class for encoding/decoding Base64 data (RFC 4648)
+ */
+class CORE_API FBase64
+{
+	/**
+	 * Base64 supports encoding any 3 uint8 chunk of binary data into a 4 uint8 set of characters from the alphabet below
+	 */
+	static uint8 EncodingAlphabet[64];
+	/**
+	 * Used to do a reverse look up from the encoded alphabet to the original 6 bit value
+	 */
+	static uint8 DecodingAlphabet[256];
+	
+	/** Ctor hidden on purpose, use static methods only */
+	FBase64()
+	{
+	}
+
+	/**
+	 * Encodes the source into a Base64 string
+	 *
+	 * @param Source the binary payload to stringify
+	 * @param Length the length of the payload that needs encoding
+	 *
+	 * @return the stringified form of the binary data
+	 */
+	static FString Encode(uint8* Source, uint32 Length);
+
+	/**
+	 * Decodes a Base64 string into an array of bytes
+	 *
+	 * @param Source the stringified data to convert
+	 * @param Length the length of the buffer being converted
+	 * @param Dest the out buffer that will be filled with the decoded data
+	 * @param PadCount the out count of the padding on the orignal buffer (0 to 2)
+	 *
+	 * @return TRUE if the buffer was decoded, FALSE if it failed to decode
+	 */
+	static bool Decode(const ANSICHAR* Source, uint32 Length, uint8* Dest, uint32& PadCount);
+
+public:
+	/**
+	 * Encodes a binary uint8 array into a Base64 string
+	 *
+	 * @param Source the binary data to convert
+	 *
+	 * @return a string that encodes the binary data in a way that can be safely transmitted via various Internet protocols
+	 */
+	static FString Encode(const TArray<uint8>& Source);
+
+	/**
+	 * Decodes a Base64 string into an array of bytes
+	 *
+	 * @param Source the stringified data to convert
+	 * @param Dest the out buffer that will be filled with the decoded data
+	 */
+	static bool Decode(const FString& Source, TArray<uint8>& Dest);
+
+	/**
+	 * Encodes a FString into a Base64 string
+	 *
+	 * @param Source the string data to convert
+	 *
+	 * @return a string that encodes the binary data in a way that can be safely transmitted via various Internet protocols
+	 */
+	static FString Encode(const FString& Source);
+
+	/**
+	 * Decodes a Base64 string into a FString
+	 *
+	 * @param Source the stringified data to convert
+	 * @param Dest the out buffer that will be filled with the decoded data
+	 */
+	static bool Decode(const FString& Source, FString& Dest);
+};
diff --git a/Engine/Source/Runtime/Engine/Private/VisualLog.cpp b/Engine/Source/Runtime/Engine/Private/VisualLog.cpp
--- a/Engine/Source/Runtime/Engine/Private/VisualLog.cpp
+++ b/Engine/Source/Runtime/Engine/Private/VisualLog.cpp
@@ -1,6 +1,7 @@
 // Copyright 1998-2014 Epic Games, Inc. All Rights Reserved.
 
 #include "EnginePrivate.h"
+#include "Misc/Base64.h"
 #include "VisualLog.h"
 #include "Json.h"
 
@@ -10,162 +11,6 @@ DEFINE_LOG_CATEGORY(LogVisual);
 
 DEFINE_STAT(STAT_VisualLog);
 
-/** The table used to encode a 6 bit value as an ascii character */
-uint8 EncodingAlphabet[64] = { 'A', 'B', 'C', 'D', 'E', 'F', 'G', 'H', 'I', 'J', 'K', 'L',
-'M', 'N', 'O', 'P', 'Q', 'R', 'S', 'T', 'U', 'V', 'W', 'X', 'Y', 'Z', 'a', 'b', 'c', 'd', 'e',
-'f', 'g', 'h', 'i', 'j', 'k', 'l', 'm', 'n', 'o', 'p', 'q', 'r', 's', 't', 'u', 'v', 'w', 'x',
-'y', 'z', '0', '1', '2', '3', '4', '5', '6', '7', '8', '9', '+', '/' };
-
-/** The table used to convert an ascii character into a 6 bit value */
-uint8 DecodingAlphabet[256] = { 0xFF, 0xFF, 0xFF, 0xFF, 0xFF, 0xFF, 0xFF, 0xFF, 0xFF, 0xFF,
-0xFF, 0xFF, 0xFF, 0xFF, 0xFF, 0xFF, 0xFF, 0xFF, 0xFF, 0xFF, 0xFF, 0xFF, 0xFF, 0xFF, 0xFF, 0xFF,
-0xFF, 0xFF, 0xFF, 0xFF, 0xFF, 0xFF, 0xFF, 0xFF, 0xFF, 0xFF, 0xFF, 0xFF, 0xFF, 0xFF, 0xFF, 0xFF,
-0xFF, 0x3E, 0xFF, 0xFF, 0xFF, 0x3F, 0x34, 0x35, 0x36, 0x37, 0x38, 0x39, 0x3A, 0x3B, 0x3C, 0x3D,
-0xFF, 0xFF, 0xFF, 0x00, 0xFF, 0xFF, 0xFF, 0x00, 0x01, 0x02, 0x03, 0x04, 0x05, 0x06, 0x07, 0x08,
-0x09, 0x0A, 0x0B, 0x0C, 0x0D, 0x0E, 0x0F, 0x10, 0x11, 0x12, 0x13, 0x14, 0x15, 0x16, 0x17, 0x18,
-0x19, 0xFF, 0xFF, 0xFF, 0xFF, 0xFF, 0xFF, 0x1A, 0x1B, 0x1C, 0x1D, 0x1E, 0x1F, 0x20, 0x21, 0x22,
-0x23, 0x24, 0x25, 0x26, 0x27, 0x28, 0x29, 0x2A, 0x2B, 0x2C, 0x2D, 0x2E, 0x2F, 0x30, 0x31, 0x32,
-0x33, 0xFF, 0xFF, 0xFF, 0xFF, 0xFF, 0xFF, 0xFF, 0xFF, 0xFF, 0xFF, 0xFF, 0xFF, 0xFF, 0xFF, 0xFF,
-0xFF, 0xFF, 0xFF, 0xFF, 0xFF, 0xFF, 0xFF, 0xFF, 0xFF, 0xFF, 0xFF, 0xFF, 0xFF, 0xFF, 0xFF, 0xFF,
-0xFF, 0xFF, 0xFF, 0xFF, 0xFF, 0xFF, 0xFF, 0xFF, 0xFF, 0xFF, 0xFF, 0xFF, 0xFF, 0xFF, 0xFF, 0xFF,
-0xFF, 0xFF, 0xFF, 0xFF, 0xFF, 0xFF, 0xFF, 0xFF, 0xFF, 0xFF, 0xFF, 0xFF, 0xFF, 0xFF, 0xFF, 0xFF,
-0xFF, 0xFF, 0xFF, 0xFF, 0xFF, 0xFF, 0xFF, 0xFF, 0xFF, 0xFF, 0xFF, 0xFF, 0xFF, 0xFF, 0xFF, 0xFF,
-0xFF, 0xFF, 0xFF, 0xFF, 0xFF, 0xFF, 0xFF, 0xFF, 0xFF, 0xFF, 0xFF, 0xFF, 0xFF, 0xFF, 0xFF, 0xFF,
-0xFF, 0xFF, 0xFF, 0xFF, 0xFF, 0xFF, 0xFF, 0xFF, 0xFF, 0xFF, 0xFF, 0xFF, 0xFF, 0xFF, 0xFF, 0xFF,
-0xFF, 0xFF, 0xFF, 0xFF, 0xFF, 0xFF, 0xFF, 0xFF, 0xFF, 0xFF, 0xFF, 0xFF, 0xFF, 0xFF, 0xFF, 0xFF,
-0xFF, 0xFF, 0xFF, 0xFF, 0xFF, 0xFF };
-
-FString EncodeToBase64(uint8* Source, uint32 Length)
-{
-	// Each 3 uint8 set of data expands to 4 bytes and must be padded to 4 bytes
-	uint32 ExpectedLength = (Length + 2) / 3 * 4;
-	FString OutBuffer;
-	OutBuffer.Empty(ExpectedLength);
-
-	ANSICHAR EncodedBytes[5];
-	// Null terminate this so we can append it easily as a string
-	EncodedBytes[4] = 0;
-	// Loop through the buffer converting 3 bytes of binary data at a time
-	while (Length >= 3)
-	{
-		uint8 A = *Source++;
-		uint8 B = *Source++;
-		uint8 C = *Source++;
-		Length -= 3;
-		// The algorithm takes 24 bits of data (3 bytes) and breaks it into 4 6bit chunks represented as ascii
-		uint32 ByteTriplet = A << 16 | B << 8 | C;
-		// Use the 6bit block to find the representation ascii character for it
-		EncodedBytes[3] = EncodingAlphabet[ByteTriplet & 0x3F];
-		// Shift each time to get the next 6 bit chunk
-		ByteTriplet >>= 6;
-		EncodedBytes[2] = EncodingAlphabet[ByteTriplet & 0x3F];
-		ByteTriplet >>= 6;
-		EncodedBytes[1] = EncodingAlphabet[ByteTriplet & 0x3F];
-		ByteTriplet >>= 6;
-		EncodedBytes[0] = EncodingAlphabet[ByteTriplet & 0x3F];
-		// Now we can append this buffer to our destination string
-		OutBuffer += EncodedBytes;
-	}
-	// Since this algorithm operates on blocks, we may need to pad the last chunks
-	if (Length > 0)
-	{
-		uint8 A = *Source++;
-		uint8 B = 0;
-		uint8 C = 0;
-		// Grab the second character if it is a 2 uint8 finish
-		if (Length == 2)
-		{
-			B = *Source;
-		}
-		uint32 ByteTriplet = A << 16 | B << 8 | C;
-		// Pad with = to make a 4 uint8 chunk
-		EncodedBytes[3] = '=';
-		ByteTriplet >>= 6;
-		// If there's only one 1 uint8 left in the source, then you need 2 pad chars
-		if (Length == 1)
-		{
-			EncodedBytes[2] = '=';
-		}
-		else
-		{
-			EncodedBytes[2] = EncodingAlphabet[ByteTriplet & 0x3F];
-		}
-		// Now encode the remaining bits the same way
-		ByteTriplet >>= 6;
-		EncodedBytes[1] = EncodingAlphabet[ByteTriplet & 0x3F];
-		ByteTriplet >>= 6;
-		EncodedBytes[0] = EncodingAlphabet[ByteTriplet & 0x3F];
-		OutBuffer += EncodedBytes;
-	}
-	return OutBuffer;
-}
-
-bool DecodeFromBase64(const ANSICHAR* Source, uint32 Length, uint8* Dest, uint32& PadCount)
-{
-	PadCount = 0;
-	uint8 DecodedValues[4];
-	while (Length)
-	{
-		// Decode the next 4 BYTEs
-		for (int32 Index = 0; Index < 4; Index++)
-		{
-			// Tell the caller if there were any pad bytes
-			if (*Source == '=')
-			{
-				PadCount++;
-			}
-			DecodedValues[Index] = DecodingAlphabet[(int32)(*Source++)];
-			// Abort on values that we don't understand
-			if (DecodedValues[Index] == -1)
-			{
-				return false;
-			}
-		}
-		Length -= 4;
-		// Rebuild the original 3 bytes from the 4 chunks of 6 bits
-		uint32 OriginalTriplet = DecodedValues[0];
-		OriginalTriplet <<= 6;
-		OriginalTriplet |= DecodedValues[1];
-		OriginalTriplet <<= 6;
-		OriginalTriplet |= DecodedValues[2];
-		OriginalTriplet <<= 6;
-		OriginalTriplet |= DecodedValues[3];
-		// Now we can tear the uint32 into bytes
-		Dest[2] = OriginalTriplet & 0xFF;
-		OriginalTriplet >>= 8;
-		Dest[1] = OriginalTriplet & 0xFF;
-		OriginalTriplet >>= 8;
-		Dest[0] = OriginalTriplet & 0xFF;
-		Dest += 3;
-	}
-	return true;
-}
-
-bool DecodeFromBase64(const FString& Source, TArray<uint8>& Dest)
-{
-	uint32 Length = Source.Len();
-	// Size must be a multiple of 4
-	if (Length % 4)
-	{
-		return false;
-	}
-	// Each 4 uint8 chunk of characters is 3 bytes of data
-	uint32 ExpectedLength = Length / 4 * 3;
-	// Add the number we need for output
-	Dest.AddZeroed(ExpectedLength);
-	uint8* Buffer = Dest.GetTypedData();
-	uint32 PadCount = 0;
-	bool bWasSuccessful = DecodeFromBase64(TCHAR_TO_ANSI(*Source), Length, Dest.GetTypedData(), PadCount);
-	if (bWasSuccessful)
-	{
-		if (PadCount > 0)
-		{
-			Dest.RemoveAt(ExpectedLength - PadCount, PadCount);
-		}
-	}
-	return bWasSuccessful;
-}
-
 //----------------------------------------------------------------------//
 // FVisLogEntry
 //----------------------------------------------------------------------//
@@ -317,7 +162,7 @@ FVisLogEntry::FVisLogEntry(TSharedPtr<FJsonValue> FromJson)
 				// decode data from Base64 string
 				const FString DataBlockAsCompressedString = JsonSampleObject->GetStringField(VisualLogJson::TAG_DATABLOCK_DATA);
 				TArray<uint8> CompressedDataBlock;
-				DecodeFromBase64(DataBlockAsCompressedString, CompressedDataBlock);
+				FBase64::Decode(DataBlockAsCompressedString, CompressedDataBlock);
 
 				// uncompress decoded data to get final TArray<uint8>
 				int32 UncompressedSize = 0;
@@ -450,7 +295,7 @@ TSharedPtr<FJsonValue> FVisLogEntry::ToJson() const
 
 		FCompression::CompressMemory((ECompressionFlags)(COMPRESS_ZLIB | COMPRESS_BiasMemory), (void*)DestBuffer, CompressedSize, (void*)CurrentData.Data.GetData(), UncompressedSize);
 		CompressedData.SetNum(CompressedSize + HeaderSize, true);
-		const FString CurrentDataAsString = EncodeToBase64((uint8*)CompressedData.GetTypedData(), CompressedData.Num());
+		const FString CurrentDataAsString = FBase64::Encode(CompressedData);
 
 		JsonSampleObject->SetStringField(VisualLogJson::TAG_CATEGORY, CurrentData.Category.ToString());
 		JsonSampleObject->SetStringField(VisualLogJson::TAG_DATABLOCK_TAGNAME, CurrentData.TagName.ToString());
diff --git a/Engine/Source/Runtime/Online/HTTP/Private/Base64.cpp b/Engine/Source/Runtime/Online/HTTP/Private/Base64.cpp
deleted file mode 100644
diff --git a/Engine/Source/Runtime/Online/HTTP/Public/Base64.h b/Engine/Source/Runtime/Online/HTTP/Public/Base64.h
deleted file mode 100644
