Fixed loading of custom cursor images on Mac

diff --git a/Engine/Source/Runtime/Core/Private/Mac/MacCursor.cpp b/Engine/Source/Runtime/Core/Private/Mac/MacCursor.cpp
--- a/Engine/Source/Runtime/Core/Private/Mac/MacCursor.cpp
+++ b/Engine/Source/Runtime/Core/Private/Mac/MacCursor.cpp
@@ -41,8 +41,8 @@ FMacCursor::FMacCursor()
 
 			case EMouseCursor::ResizeSouthEast:
 			{
-				FString Path = FString::Printf(TEXT("%s/%sEditor/Slate/Cursor/SouthEastCursor.png"), FPlatformProcess::BaseDir(), *FPaths::EngineContentDir());
-				NSImage* CursorImage = [[NSImage alloc] initWithContentsOfFile: [[NSBundle mainBundle] pathForImageResource: Path.GetNSString()]];
+				FString Path = FString::Printf(TEXT("%s%sEditor/Slate/Cursor/SouthEastCursor.png"), FPlatformProcess::BaseDir(), *FPaths::EngineContentDir());
+				NSImage* CursorImage = [[NSImage alloc] initWithContentsOfFile:Path.GetNSString()];
 				CursorHandle = [[NSCursor alloc] initWithImage: CursorImage hotSpot: NSMakePoint(8, 8)];
 				[CursorImage release];
 				break;
@@ -50,8 +50,8 @@ FMacCursor::FMacCursor()
 
 			case EMouseCursor::ResizeSouthWest:
 			{
-				FString Path = FString::Printf(TEXT("%s/%sEditor/Slate/Cursor/SouthWestCursor.png"), FPlatformProcess::BaseDir(), *FPaths::EngineContentDir());
-				NSImage* CursorImage = [[NSImage alloc] initWithContentsOfFile: [[NSBundle mainBundle] pathForImageResource: Path.GetNSString()]];
+				FString Path = FString::Printf(TEXT("%s%sEditor/Slate/Cursor/SouthWestCursor.png"), FPlatformProcess::BaseDir(), *FPaths::EngineContentDir());
+				NSImage* CursorImage = [[NSImage alloc] initWithContentsOfFile:Path.GetNSString()];
 				CursorHandle = [[NSCursor alloc] initWithImage: CursorImage hotSpot: NSMakePoint(8, 8)];
 				[CursorImage release];
 				break;
@@ -59,8 +59,8 @@ FMacCursor::FMacCursor()
 
 			case EMouseCursor::CardinalCross:
 			{
-				FString Path = FString::Printf(TEXT("%s/%sEditor/Slate/Cursor/CardinalCrossCursor.png"), FPlatformProcess::BaseDir(), *FPaths::EngineContentDir());
-				NSImage* CursorImage = [[NSImage alloc] initWithContentsOfFile: [[NSBundle mainBundle] pathForImageResource: Path.GetNSString()]];
+				FString Path = FString::Printf(TEXT("%s%sEditor/Slate/Cursor/CardinalCrossCursor.png"), FPlatformProcess::BaseDir(), *FPaths::EngineContentDir());
+				NSImage* CursorImage = [[NSImage alloc] initWithContentsOfFile:Path.GetNSString()];
 				CursorHandle = [[NSCursor alloc] initWithImage: CursorImage hotSpot: NSMakePoint(8, 8)];
 				[CursorImage release];
 				break;
@@ -88,8 +88,8 @@ FMacCursor::FMacCursor()
 
 			case EMouseCursor::EyeDropper:
 			{
-				FString Path = FString::Printf(TEXT("%s/%sEditor/Slate/Cursor/EyeDropperCursor.png"), FPlatformProcess::BaseDir(), *FPaths::EngineContentDir());
-				NSImage* CursorImage = [[NSImage alloc] initWithContentsOfFile: [[NSBundle mainBundle] pathForImageResource: Path.GetNSString()]];
+				FString Path = FString::Printf(TEXT("%s%sEditor/Slate/Cursor/EyeDropperCursor.png"), FPlatformProcess::BaseDir(), *FPaths::EngineContentDir());
+				NSImage* CursorImage = [[NSImage alloc] initWithContentsOfFile:Path.GetNSString()];
 				CursorHandle = [[NSCursor alloc] initWithImage: CursorImage hotSpot: NSMakePoint(1, 17)];
 				[CursorImage release];
 				break;
