#ue4 - Clean up comments in FTransform and FQuat, and deprecate FTransform::GetTranslationV, GetScaleV, GetRotationV.

- FTransform empty constructor actually initializes to identity (comment was wrong and said otherwise).

diff --git a/Engine/Source/Runtime/Core/Public/Math/Quat.h b/Engine/Source/Runtime/Core/Public/Math/Quat.h
--- a/Engine/Source/Runtime/Core/Public/Math/Quat.h
+++ b/Engine/Source/Runtime/Core/Public/Math/Quat.h
@@ -41,21 +41,21 @@ public:
 	FORCEINLINE FQuat( ) { }
 
 	/**
-	 * Creates and initializes a new quaternion to zero.
+	 * Creates and initializes a new quaternion, with the W component either 0 or 1.
 	 *
-	 * @param EForceInit Force Init Enum.
+	 * @param EForceInit Force init enum: if equal to ForceInitToZero then W is 0, otherwise W = 1 (creating an identity transform)
 	 */
 	explicit FORCEINLINE FQuat( EForceInit );
 
 	/**
 	 * Constructor.
 	 *
-	 * @param InX x component of the axis.
-	 * @param InY y component of the axis.
-	 * @param InZ z component of the axis.
-	 * @param InA angle to rotate around the given axis.
+	 * @param InX X component of the quaternion
+	 * @param InY Y component of the quaternion
+	 * @param InZ Z component of the quaternion
+	 * @param InW W component of the quaternion
 	 */
-	FORCEINLINE FQuat( float InX, float InY, float InZ, float InA );
+	FORCEINLINE FQuat( float InX, float InY, float InZ, float InW );
 
 	/**
 	 * Copy constructor.
@@ -619,11 +619,11 @@ FORCEINLINE FQuat::FQuat( EForceInit ZeroOrNot )
 { }
 
 
-FORCEINLINE FQuat::FQuat( float InX, float InY, float InZ, float InA )
+FORCEINLINE FQuat::FQuat( float InX, float InY, float InZ, float InW )
 	: X(InX)
 	, Y(InY)
 	, Z(InZ)
-	, W(InA)
+	, W(InW)
 {
 	DiagnosticCheckNaN();
 }
diff --git a/Engine/Source/Runtime/Core/Public/Math/Transform.h b/Engine/Source/Runtime/Core/Public/Math/Transform.h
--- a/Engine/Source/Runtime/Core/Public/Math/Transform.h
+++ b/Engine/Source/Runtime/Core/Public/Math/Transform.h
@@ -58,11 +58,11 @@ public:
 #endif
 
 	/**
-	 * Constructor with no components initialized (garbage values)
+	 * Constructor with initialization to the identity transform.
 	 */
 	FORCEINLINE FTransform()
-		: Rotation(ForceInit)
-		, Translation(ForceInit)
+		: Rotation(0.f,0.f,0.f,1.f)
+		, Translation(0.f)
 		, Scale3D(1.f)
 	{
 		// Note: This can be used to track down initialization issues with bone atom arrays; but it will
@@ -77,7 +77,7 @@ public:
 	}
 
 	/**
-	 * Constructor with all components initialized
+	 * Constructor with an initial translation
 	 *
 	 * @param InTranslation The value to use for the translation component
 	 */
@@ -90,7 +90,7 @@ public:
 	}
 
 	/**
-	 * Constructor with all components initialized
+	 * Constructor with an initial rotation
 	 *
 	 * @param InRotation The value to use for rotation component
 	 */
@@ -103,7 +103,7 @@ public:
 	}
 
 	/**
-	 * Constructor with all components initialized
+	 * Constructor with an initial rotation
 	 *
 	 * @param InRotation The value to use for rotation component  (after being converted to a quaternion)
 	 */
@@ -131,7 +131,7 @@ public:
 	}
 
 	/**
-	 * Constructor taking a FRotator as the rotation component
+	 * Constructor with all components initialized, taking a FRotator as the rotation component
 	 *
 	 * @param InRotation The value to use for rotation component (after being converted to a quaternion)
 	 * @param InTranslation The value to use for the translation component
@@ -148,7 +148,7 @@ public:
 	/**
 	 * Copy-constructor
 	 *
-	 * @param InTransform The source atom from which all components will be copied
+	 * @param InTransform The source transform from which all components will be copied
 	 */
 	FORCEINLINE FTransform(const FTransform& InTransform) : 
 		Rotation(InTransform.Rotation), 
@@ -767,7 +767,7 @@ public:
 	}
 
 	/**
-	 * Set the translation and Scale3D components of this atom to a linearly interpolated combination of two other atoms
+	 * Set the translation and Scale3D components of this transform to a linearly interpolated combination of two other transforms
 	 *
 	 * Translation = FMath::Lerp(SourceAtom1.Translation, SourceAtom2.Translation, Alpha)
 	 * Scale3D = FMath::Lerp(SourceAtom1.Scale3D, SourceAtom2.Scale3D, Alpha)
@@ -827,7 +827,7 @@ public:
 	}
 
 	/**
-	 * Blends the Identity atom with a weighted source atom and accumulates that into a destination atom
+	 * Blends the Identity transform with a weighted source transform and accumulates that into a destination transform
 	 *
 	 * SourceAtom = Blend(Identity, SourceAtom, BlendWeight)
 	 * FinalAtom.Rotation = SourceAtom.Rotation * FinalAtom.Rotation
@@ -901,6 +901,7 @@ public:
 	 *
 	 * @return The rotation component
 	 */
+	DEPRECATED(4.5, "FTransform::GetRotationV() is deprecated, use FTransform::GetRotation() instead.")
 	FORCEINLINE FQuat GetRotationV() const
 	{
 		DiagnosticCheckNaN_Rotate();
@@ -913,6 +914,7 @@ public:
 	 *
 	 * @return The translation component
 	 */
+	DEPRECATED(4.5, "FTransform::GetTranslationV() is deprecated, use FTransform::GetTranslation() instead.")
 	FORCEINLINE FVector GetTranslationV() const
 	{
 		DiagnosticCheckNaN_Translate();
@@ -925,6 +927,7 @@ public:
 	 *
 	 * @return The Scale3D component
 	 */
+	DEPRECATED(4.5, "FTransform::GetScale3DV() is deprecated, use FTransform::GetScale3D() instead.")
 	FORCEINLINE FVector GetScale3DV() const
 	{
 		DiagnosticCheckNaN_Scale3D();
@@ -932,9 +935,9 @@ public:
 	}
 
 	/**
-	 * Sets the Rotation and Scale3D of this transformation from another atom
+	 * Sets the Rotation and Scale3D of this transformation from another transform
 	 *
-	 * @param SrcBA The atom to copy rotation and Scale3D from
+	 * @param SrcBA The transform to copy rotation and Scale3D from
 	 */
 	FORCEINLINE void CopyRotationPart(const FTransform& SrcBA)
 	{
@@ -946,9 +949,9 @@ public:
 	}
 
 	/**
-	 * Sets the Translation and Scale3D of this transformation from another atom
+	 * Sets the Translation and Scale3D of this transformation from another transform
 	 *
-	 * @param SrcBA The atom to copy translation and Scale3D from
+	 * @param SrcBA The transform to copy translation and Scale3D from
 	 */
 	FORCEINLINE void CopyTranslationAndScale3D(const FTransform& SrcBA)
 	{
diff --git a/Engine/Source/Runtime/Core/Public/Math/TransformVectorized.h b/Engine/Source/Runtime/Core/Public/Math/TransformVectorized.h
--- a/Engine/Source/Runtime/Core/Public/Math/TransformVectorized.h
+++ b/Engine/Source/Runtime/Core/Public/Math/TransformVectorized.h
@@ -61,7 +61,7 @@ public:
 #endif
 
 	/**
-	 * Constructor with no components initialized (garbage values)
+	 * Constructor with initialization to the identity transform.
 	 */
 	FORCEINLINE FTransform()
 	{
@@ -72,7 +72,7 @@ public:
 		// Scale3D = {1,1,1,0);
 		Scale3D = VectorSet_W0( VectorOne() );
 
-		// Note: This can be used to track down initialization issues with bone atom arrays; but it will
+		// Note: This can be used to track down initialization issues with bone transform arrays; but it will
 		// cause issues with transient fields such as RootMotionDelta that get initialized to 0 by default
 #if 0
 		float qnan = FMath::Log2(-5.3f);
@@ -84,7 +84,7 @@ public:
 	}
 
 	/**
-	 * Constructor with all components initialized
+	 * Constructor with an initial translation
 	 *
 	 * @param InTranslation The value to use for the translation component
 	 */
@@ -101,7 +101,7 @@ public:
 	}
 
 	/**
-	 * Constructor with all components initialized
+	 * Constructor with an initial rotation
 	 *
 	 * @param InRotation The value to use for rotation component
 	 */
@@ -118,7 +118,7 @@ public:
 	}
 
 	/**
-	 * Constructor with all components initialized
+	 * Constructor with an initial rotation
 	 *
 	 * @param InRotation The value to use for rotation component  (after being converted to a quaternion)
 	 */
@@ -170,7 +170,7 @@ public:
 	}
 
 	/**
-	 * Constructor taking a FRotator as the rotation component
+	 * Constructor with all components initialized, taking a FRotator as the rotation component
 	 *
 	 * @param InRotation The value to use for rotation component (after being converted to a quaternion)
 	 * @param InTranslation The value to use for the translation component
@@ -192,7 +192,7 @@ public:
 	/**
 	 * Copy-constructor
 	 *
-	 * @param InTransform The source atom from which all components will be copied
+	 * @param InTransform The source transform from which all components will be copied
 	 */
 	FORCEINLINE FTransform(const FTransform& InTransform) : 
 		Rotation(InTransform.Rotation), 
@@ -865,7 +865,7 @@ public:
 
 
 	/**
-	 * Set the translation and Scale3D components of this atom to a linearly interpolated combination of two other atoms
+	 * Set the translation and Scale3D components of this transform to a linearly interpolated combination of two other transforms
 	 *
 	 * Translation = FMath::Lerp(SourceAtom1.Translation, SourceAtom2.Translation, Alpha)
 	 * Scale3D = FMath::Lerp(SourceAtom1.Scale3D, SourceAtom2.Scale3D, Alpha)
@@ -931,7 +931,7 @@ public:
 	}
 
 	/**
-	 * Blends the Identity atom with a weighted source atom and accumulates that into a destination atom
+	 * Blends the Identity transform with a weighted source transform and accumulates that into a destination transform
 	 *
 	 * SourceAtom = Blend(Identity, SourceAtom, BlendWeight)
 	 * FinalAtom.Rotation = SourceAtom.Rotation * FinalAtom.Rotation
@@ -1028,6 +1028,7 @@ public:
 	 *
 	 * @return The rotation component
 	 */
+	DEPRECATED(4.5, "FTransform::GetRotationV() is deprecated, use FTransform::GetRotation() instead.")
 	FORCEINLINE FQuat GetRotationV() const
 	{
 		DiagnosticCheckNaN_Rotate();
@@ -1042,6 +1043,7 @@ public:
 	 *
 	 * @return The translation component
 	 */
+	DEPRECATED(4.5, "FTransform::GetTranslationV() is deprecated, use FTransform::GetTranslation() instead.")
 	FORCEINLINE FVector GetTranslationV() const
 	{
 		DiagnosticCheckNaN_Translate();
@@ -1056,6 +1058,7 @@ public:
 	 *
 	 * @return The Scale3D component
 	 */
+	DEPRECATED(4.5, "FTransform::GetScale3DV() is deprecated, use FTransform::GetScale3D() instead.")
 	FORCEINLINE FVector GetScale3DV() const
 	{
 		DiagnosticCheckNaN_Scale3D();
@@ -1065,9 +1068,9 @@ public:
 	}
 
 	/**
-	 * Sets the Rotation and Scale3D of this transformation from another atom
+	 * Sets the Rotation and Scale3D of this transformation from another transform
 	 *
-	 * @param SrcBA The atom to copy rotation and Scale3D from
+	 * @param SrcBA The transform to copy rotation and Scale3D from
 	 */
 	FORCEINLINE void CopyRotationPart(const FTransform& SrcBA)
 	{
@@ -1079,9 +1082,9 @@ public:
 	}
 
 	/**
-	 * Sets the Translation and Scale3D of this transformation from another atom
+	 * Sets the Translation and Scale3D of this transformation from another transform
 	 *
-	 * @param SrcBA The atom to copy translation and Scale3D from
+	 * @param SrcBA The transform to copy translation and Scale3D from
 	 */
 	FORCEINLINE void CopyTranslationAndScale3D(const FTransform& SrcBA)
 	{
