- Added "virtual" IOS file handles to handle more file handles open than is allowed in iOS
- Added early support for thread affinity handling, although it is not currently working right

diff --git a/Engine/Source/Runtime/Core/Private/IOS/IOSPlatformFile.cpp b/Engine/Source/Runtime/Core/Private/IOS/IOSPlatformFile.cpp
--- a/Engine/Source/Runtime/Core/Private/IOS/IOSPlatformFile.cpp
+++ b/Engine/Source/Runtime/Core/Private/IOS/IOSPlatformFile.cpp
@@ -11,75 +11,327 @@
 // make an FTimeSpan object that represents the "epoch" for time_t (from a stat struct)
 const FDateTime IOSEpoch(1970, 1, 1);
 
-/** 
- * iOS file handle implementation
-**/
-FFileHandleIOS::FFileHandleIOS(int32 InFileHandle)
-	: FileHandle(InFileHandle)
-{
-}
 
-FFileHandleIOS::~FFileHandleIOS()
-{
-	close(FileHandle);
-	FileHandle = -1;
-}
 
-int64 FFileHandleIOS::Tell()
-{
-	check(IsValid());
-	return lseek(FileHandle, 0, SEEK_CUR);
-}
 
-bool FFileHandleIOS::Seek(int64 NewPosition)
-{
-	check(IsValid());
-	check(NewPosition >= 0);
-	return lseek(FileHandle, NewPosition, SEEK_SET) != -1;
-}
+/* FIOSFileHandle class
+ *****************************************************************************/
 
-bool FFileHandleIOS::SeekFromEnd(int64 NewPositionRelativeToEnd)
+/** 
+ * Managed IOS file handle implementation which limits number of open files. This
+ * is to prevent running out of system file handles (700). Should not be neccessary when 
+ * using pak file (e.g., SHIPPING?) so not particularly optimized. Only manages 
+ * files which are opened READ_ONLY.
+ **/
+#define MANAGE_FILE_HANDLES 1 // !UE_BUILD_SHIPPING
+
+struct FManagedFile
 {
-	check(IsValid());
-	check(NewPositionRelativeToEnd <= 0);
-	return lseek(FileHandle, NewPositionRelativeToEnd, SEEK_END) != -1;
-}
+	int32 Handle;
+	uint32 ID;
+	double AccessTime;
+};
 
-bool FFileHandleIOS::Read(uint8* Destination, int64 BytesToRead)
+class FIOSFileHandle : public IFileHandle
 {
-	check(IsValid());
-	while (BytesToRead)
+	static const int32 READWRITE_SIZE = 1024 * 1024;
+	static const int32 ACTIVE_HANDLE_COUNT_PER_THREAD = 100;
+
+public:
+
+	FIOSFileHandle( int32 InFileHandle, const FString& InFilename, bool bIsForRead )
+		: FileHandle(InFileHandle)
+#if !UE_BUILD_SHIPPING || MANAGE_FILE_HANDLES
+		, Filename(InFilename)
+#endif
+#if MANAGE_FILE_HANDLES
+        , HandleSlot(-1)
+        , FileSize(0)
+#endif
 	{
-		check(BytesToRead >= 0);
-		int64 ThisSize = FMath::Min<int64>(READWRITE_SIZE, BytesToRead);
-		check(Destination);
-		if (read(FileHandle, Destination, ThisSize) != ThisSize)
+		check(FileHandle != 0);
+
+#if MANAGE_FILE_HANDLES
+
+		static uint32 NextID = 1;
+		FileID = NextID++;
+
+		// get the per-thread buffers
+		ManagedFiles = (FManagedFile*)FPlatformTLS::GetTlsValue(ManagedFilesTlsSlot);
+
+		// not made yet on this thread, make the buffers now
+		if (ManagedFiles == NULL)
 		{
-			return false;
+			ManagedFiles = new FManagedFile[ACTIVE_HANDLE_COUNT_PER_THREAD];
+			FMemory::Memzero(ManagedFiles, sizeof(FManagedFile) * ACTIVE_HANDLE_COUNT_PER_THREAD);
+			FPlatformTLS::SetTlsValue(ManagedFilesTlsSlot, ManagedFiles);
 		}
-		Destination += ThisSize;
-		BytesToRead -= ThisSize;
+
+        // Only files opened for read will be managed
+        if( bIsForRead )
+        {
+            ReserveSlot();
+            ManagedFiles[HandleSlot].Handle = FileHandle;
+
+			struct stat FileInfo;
+			FileInfo.st_size = -1;
+			// check the read path
+			fstat(FileHandle, &FileInfo);
+			FileSize = FileInfo.st_size;
+        }
+#endif
+
+		Seek(0);
 	}
-	return true;
-}
 
-bool FFileHandleIOS::Write(const uint8* Source, int64 BytesToWrite)
-{
-	check(IsValid());
-	while (BytesToWrite)
+	/**
+	 * Destructor.
+	 */
+	virtual ~FIOSFileHandle( )
 	{
-		check(BytesToWrite >= 0);
-		int64 ThisSize = FMath::Min<int64>(READWRITE_SIZE, BytesToWrite);
-		check(Source);
-		if (write(FileHandle, Source, ThisSize) != ThisSize)
+#if MANAGE_FILE_HANDLES
+        if( IsManaged() )
+        {
+            if( ManagedFiles[HandleSlot].ID == FileID )
+            {
+                close(FileHandle);
+				ManagedFiles[HandleSlot].ID = 0;
+            }
+        }
+        else
+#endif
+        {
+		    close(FileHandle);
+        }
+
+		FileHandle = -1;
+	}
+
+	bool InternalRead(uint8* Destination, int64 BytesToRead)
+	{
+		while (BytesToRead)
 		{
+			check(BytesToRead >= 0);
+			int64 ThisSize = FMath::Min<int64>(READWRITE_SIZE, BytesToRead);
+			check(Destination);
+			if (read(FileHandle, Destination, ThisSize) != ThisSize)
+			{
+				return false;
+			}
+			Destination += ThisSize;
+			BytesToRead -= ThisSize;
+		}
+		return true;
+	}
+
+public:
+
+	// Begin IFileHandle interface
+
+	virtual bool Read( uint8* Destination, int64 BytesToRead ) override
+	{
+ #if MANAGE_FILE_HANDLES
+       if( IsManaged() )
+        {
+            ActivateSlot();
+			lseek(FileHandle, FileOffset, SEEK_SET);
+			// read into the buffer, and make sure it worked
+			if (InternalRead(Destination, BytesToRead))
+			{
+				FileOffset += BytesToRead;
+				return true;
+			}
 			return false;
+        }
+        else
+#endif
+        {
+		    return InternalRead(Destination, BytesToRead);
+        }
+	}
+
+	virtual bool Seek( int64 NewPosition ) override
+	{
+		check(NewPosition >= 0);
+
+#if MANAGE_FILE_HANDLES
+        if( IsManaged() )
+        {
+            FileOffset = NewPosition >= FileSize ? FileSize - 1 : NewPosition;
+            return true;
+        }
+        else
+#endif
+        {
+		    return (lseek(FileHandle, NewPosition, SEEK_SET) != -1);
+        }
+	}
+
+	virtual bool SeekFromEnd( int64 NewPositionRelativeToEnd = 0 ) override
+	{
+		check(NewPositionRelativeToEnd <= 0);
+
+#if MANAGE_FILE_HANDLES
+        if( IsManaged() )
+        {
+            FileOffset = (NewPositionRelativeToEnd >= FileSize) ? 0 : ( FileSize + NewPositionRelativeToEnd - 1 );
+            return true;
+        }
+        else
+#endif
+        {
+		    return lseek(FileHandle, NewPositionRelativeToEnd, SEEK_END) != -1;
+        }
+	}
+
+	virtual int64 Size( ) override
+	{
+#if MANAGE_FILE_HANDLES
+        if( IsManaged() )
+        {
+            return FileSize;
+        }
+        else
+#endif
+        {
+		    return IFileHandle::Size();
+        }
+	}
+
+	virtual int64 Tell( ) override
+	{
+#if MANAGE_FILE_HANDLES
+        if( IsManaged() )
+        {
+            return FileOffset;
+        }
+        else
+#endif
+        {
+		    return lseek(FileHandle, 0, SEEK_CUR);
+        }
+	}
+
+	virtual bool Write( const uint8* Source, int64 BytesToWrite ) override
+	{
+		while (BytesToWrite)
+		{
+			check(BytesToWrite >= 0);
+			int64 ThisSize = FMath::Min<int64>(READWRITE_SIZE, BytesToWrite);
+			check(Source);
+			if (write(FileHandle, Source, ThisSize) != ThisSize)
+			{
+				return false;
+			}
+			Source += ThisSize;
+			BytesToWrite -= ThisSize;
 		}
-		Source += ThisSize;
-		BytesToWrite -= ThisSize;
+		return true;
 	}
-	return true;
-}
+
+	// End IFileHandle interface
+
+private:
+
+#if MANAGE_FILE_HANDLES
+    FORCEINLINE bool IsManaged()
+    {
+        return HandleSlot != -1;
+    }
+
+    void ActivateSlot()
+    {
+        if( IsManaged() )
+        {
+            if( ManagedFiles[HandleSlot].ID != FileID )
+            {
+				ReserveSlot();
+                
+				FileHandle = open(TCHAR_TO_UTF8(*Filename), O_RDONLY);
+
+				if (FileHandle != -1)
+                {
+                    ManagedFiles[HandleSlot].Handle = FileHandle;
+                }
+            }
+            else
+            {
+                ManagedFiles[HandleSlot].AccessTime = FPlatformTime::Seconds();
+            }
+        }
+    }
+
+    void ReserveSlot()
+    {
+        HandleSlot = -1;
+
+        // Look for non-reserved slot
+        for( int32 i = 0; i < ACTIVE_HANDLE_COUNT_PER_THREAD; ++i )
+        {
+            if( ManagedFiles[i].ID == 0 )
+            {
+                HandleSlot = i;
+                break;
+            }
+        }
+
+        // Take the oldest handle
+        if( HandleSlot == -1 )
+        {
+            int32 Oldest = 0;
+            for( int32 i = 1; i < ACTIVE_HANDLE_COUNT_PER_THREAD; ++i )
+            {
+                if( ManagedFiles[Oldest].AccessTime > ManagedFiles[i].AccessTime )
+                {
+                    Oldest = i;
+                }
+            }
+
+            close( ManagedFiles[Oldest].Handle );
+            HandleSlot = Oldest;
+        }
+
+        ManagedFiles[HandleSlot].ID = FileID;
+        ManagedFiles[HandleSlot].AccessTime = FPlatformTime::Seconds();
+    }
+#endif
+
+private:
+
+	// Holds the internal file handle.
+	int32 FileHandle;
+
+#if !UE_BUILD_SHIPPING || MANAGE_FILE_HANDLES
+	// Holds the name of the file that this handle represents. Kept around for possible reopen of file.
+	FString Filename;
+#endif
+
+#if MANAGE_FILE_HANDLES
+    // Most recent valid slot index for this handle; >=0 for handles which are managed.
+    int32 HandleSlot;
+
+    // Current file offset; valid iff a managed handle.
+    int64 FileOffset;
+
+    // Cached file size; valid iff a managed handle.
+    int64 FileSize;
+
+    // Each thread keeps a collection of active handles with access times.
+    FManagedFile* ManagedFiles;
+
+	// Unique FileID for this file (since handles aren't unique)
+	uint32 FileID;
+
+	static int32 ManagedFilesTlsSlot;
+#endif
+};
+
+#if MANAGE_FILE_HANDLES
+int32 FIOSFileHandle::ManagedFilesTlsSlot = FPlatformTLS::AllocTlsSlot();
+#endif
+
+
+
 
 
 /**
@@ -263,17 +515,20 @@ FDateTime FIOSPlatformFile::GetAccessTimeStamp(const TCHAR* Filename)
 IFileHandle* FIOSPlatformFile::OpenRead(const TCHAR* Filename)
 {
 	FString NormalizedFilename = NormalizeFilename(Filename);
+
 	// check the read path
-	int32 Handle = open(TCHAR_TO_UTF8(*ConvertToIOSPath(NormalizedFilename, false)), O_RDONLY);
+	FString FinalPath = ConvertToIOSPath(NormalizedFilename, false);
+	int32 Handle = open(TCHAR_TO_UTF8(*FinalPath), O_RDONLY);
 	if(Handle == -1)
 	{
 		// if not in the read path, check the write path
-		Handle = open(TCHAR_TO_UTF8(*ConvertToIOSPath(NormalizedFilename, true)), O_RDONLY);
+		FinalPath = ConvertToIOSPath(NormalizedFilename, true);
+		Handle = open(TCHAR_TO_UTF8(*FinalPath), O_RDONLY);
 	}
 
 	if (Handle != -1)
 	{
-		return new FFileHandleIOS(Handle);
+		return new FIOSFileHandle(Handle, FinalPath, true);
 	}
 	return NULL;
 }
@@ -301,7 +556,7 @@ IFileHandle* FIOSPlatformFile::OpenWrite(const TCHAR* Filename, bool bAppend, bo
 	int32 Handle = open(TCHAR_TO_UTF8(*IOSFilename), Flags, S_IRUSR | S_IWUSR);
 	if (Handle != -1)
 	{
-		FFileHandleIOS* FileHandleIOS = new FFileHandleIOS(Handle);
+		FIOSFileHandle* FileHandleIOS = new FIOSFileHandle(Handle, IOSFilename, false);
 		if (bAppend)
 		{
 			FileHandleIOS->SeekFromEnd(0);
diff --git a/Engine/Source/Runtime/Core/Private/IOS/IOSPlatformFile.h b/Engine/Source/Runtime/Core/Private/IOS/IOSPlatformFile.h
--- a/Engine/Source/Runtime/Core/Private/IOS/IOSPlatformFile.h
+++ b/Engine/Source/Runtime/Core/Private/IOS/IOSPlatformFile.h
@@ -5,29 +5,6 @@
 class IFileHandle;
 class FApplePlatformFile;
 
-/** 
- * iOS file handle implementation
-**/
-class CORE_API FFileHandleIOS : public IFileHandle
-{
-	enum {READWRITE_SIZE = 1024 * 1024};
-	int32 FileHandle;
-
-	FORCEINLINE bool IsValid()
-	{
-		return FileHandle != -1;
-	}
-
-public:
-	FFileHandleIOS(int32 InFileHandle = -1);
-	virtual ~FFileHandleIOS();
-
-	virtual int64 Tell() override;
-	virtual bool Seek(int64 NewPosition) override;
-	virtual bool SeekFromEnd(int64 NewPositionRelativeToEnd = 0) override;
-	virtual bool Read(uint8* Destination, int64 BytesToRead) override;
-	virtual bool Write(const uint8* Source, int64 BytesToWrite) override;
-};
 
 /**
  * iOS File I/O implementation
diff --git a/Engine/Source/Runtime/Core/Private/IOS/IOSPlatformProcess.cpp b/Engine/Source/Runtime/Core/Private/IOS/IOSPlatformProcess.cpp
--- a/Engine/Source/Runtime/Core/Private/IOS/IOSPlatformProcess.cpp
+++ b/Engine/Source/Runtime/Core/Private/IOS/IOSPlatformProcess.cpp
@@ -1,7 +1,7 @@
 // Copyright 1998-2014 Epic Games, Inc. All Rights Reserved.
 
 /*=============================================================================
-	IOSPlatformProcess.mm: iOS implementations of Process functions
+IOSPlatformProcess.cpp: iOS implementations of Process functions
 =============================================================================*/
 
 #include "Core.h"
@@ -86,3 +86,150 @@ void FIOSPlatformProcess::SetupGameOrRenderThread(bool bIsRenderThread)
 	Sched.sched_priority = bIsRenderThread ? RENDER_THREAD_PRIORITY : GAME_THREAD_PRIORITY;
 	pthread_setschedparam(pthread_self(), Policy, &Sched);
 }
+
+
+void FIOSPlatformProcess::SetThreadAffinityMask(uint64 AffinityMask)
+{
+	if ([IOSAppDelegate GetDelegate].OSVersion >= 8 && FPlatformMisc::NumberOfCores() > 1)
+	{
+		thread_affinity_policy AP;
+		AP.affinity_tag = AffinityMask;
+		thread_policy_set(pthread_mach_thread_np(pthread_self()), THREAD_AFFINITY_POLICY, (integer_t*)&AP, THREAD_AFFINITY_POLICY_COUNT);
+	}
+	else
+	{
+		FGenericPlatformProcess::SetThreadAffinityMask(AffinityMask);
+	}
+}
+
+const uint64 FIOSPlatformAffinity::GetMainGameMask()
+{
+	return MAKEAFFINITYMASK1(0);
+}
+
+const uint64 FIOSPlatformAffinity::GetRenderingThreadMask()
+{
+	static int Mask = 0;
+	if (Mask == 0)
+	{
+		switch (FPlatformMisc::NumberOfCores())
+		{
+		case 2:
+		case 3:
+			Mask = MAKEAFFINITYMASK1(1);
+			break;
+		default:
+			Mask = FGenericPlatformAffinity::GetRenderingThreadMask();
+			break;
+		}
+	}
+
+	return Mask;
+}
+
+const uint64 FIOSPlatformAffinity::GetRTHeartBeatMask()
+{
+	static int Mask = 0;
+	if (Mask == 0)
+	{
+		switch (FPlatformMisc::NumberOfCores())
+		{
+		case 2:
+		case 3:
+			Mask = MAKEAFFINITYMASK1(0);
+			break;
+		default:
+			Mask = FGenericPlatformAffinity::GetRTHeartBeatMask();
+			break;
+		}
+	}
+
+	return Mask;
+}
+
+const uint64 FIOSPlatformAffinity::GetPoolThreadMask()
+{
+	static int Mask = 0;
+	if (Mask == 0)
+	{
+		switch (FPlatformMisc::NumberOfCores())
+		{
+		case 2:
+			Mask = MAKEAFFINITYMASK1(1);
+			break;
+		case 3:
+			Mask = MAKEAFFINITYMASK1(2);
+			break;
+		default:
+			Mask = FGenericPlatformAffinity::GetPoolThreadMask();
+			break;
+		}
+	}
+
+	return Mask;
+}
+
+const uint64 FIOSPlatformAffinity::GetTaskGraphThreadMask()
+{
+	static int Mask = 0;
+	if (Mask == 0)
+	{
+		switch (FPlatformMisc::NumberOfCores())
+		{
+		case 2:
+			Mask = MAKEAFFINITYMASK1(1);
+			break;
+		case 3:
+			Mask = MAKEAFFINITYMASK1(2);
+			break;
+		default:
+			Mask = FGenericPlatformAffinity::GetTaskGraphThreadMask();
+			break;
+		}
+	}
+
+	return Mask;
+}
+
+const uint64 FIOSPlatformAffinity::GetStatsThreadMask()
+{
+	static int Mask = 0;
+	if (Mask == 0)
+	{
+		switch (FPlatformMisc::NumberOfCores())
+		{
+		case 2:
+			Mask = MAKEAFFINITYMASK1(0);
+			break;
+		case 3:
+			Mask = MAKEAFFINITYMASK1(2);
+			break;
+		default:
+			Mask = FGenericPlatformAffinity::GetStatsThreadMask();
+			break;
+		}
+	}
+
+	return Mask;
+}
+
+const uint64 FIOSPlatformAffinity::GetNoAffinityMask()
+{
+	static int Mask = 0;
+	if (Mask == 0)
+	{
+		switch (FPlatformMisc::NumberOfCores())
+		{
+		case 2:
+		case 3:
+			Mask = (1 << FPlatformMisc::NumberOfCores()) - 1;
+			break;
+		default:
+			Mask = FGenericPlatformAffinity::GetNoAffinityMask();
+			break;
+		}
+	}
+
+	return Mask;
+}
+
diff --git a/Engine/Source/Runtime/Core/Public/IOS/IOSPlatformIncludes.h b/Engine/Source/Runtime/Core/Public/IOS/IOSPlatformIncludes.h
--- a/Engine/Source/Runtime/Core/Public/IOS/IOSPlatformIncludes.h
+++ b/Engine/Source/Runtime/Core/Public/IOS/IOSPlatformIncludes.h
@@ -99,8 +99,6 @@ struct RECT
 #include "IOS/IOSPlatformHttp.h"
 #import "IOSPlatformString.h"
 
-typedef FGenericPlatformAffinity FPlatformAffinity;
-
 // include platform properties and typedef it for the runtime
 
 #include "IOS/IOSPlatformProperties.h"
diff --git a/Engine/Source/Runtime/Core/Public/IOS/IOSPlatformProcess.h b/Engine/Source/Runtime/Core/Public/IOS/IOSPlatformProcess.h
--- a/Engine/Source/Runtime/Core/Public/IOS/IOSPlatformProcess.h
+++ b/Engine/Source/Runtime/Core/Public/IOS/IOSPlatformProcess.h
@@ -30,8 +30,24 @@ struct CORE_API FIOSPlatformProcess : public FGenericPlatformProcess
 	static const TCHAR* BaseDir();
 	static FRunnableThread* CreateRunnableThread();
 	static void LaunchURL( const TCHAR* URL, const TCHAR* Parms, FString* Error );
-    static void SetRealTimeMode();
+	static void SetRealTimeMode();
 	static void SetupGameOrRenderThread(bool bIsRenderThread);
+	static void SetThreadAffinityMask(uint64 AffinityMask);
 };
 
 typedef FIOSPlatformProcess FPlatformProcess;
+
+class CORE_API FIOSPlatformAffinity : public FGenericPlatformAffinity
+{
+public:
+	static const uint64 GetMainGameMask();
+	static const uint64 GetRenderingThreadMask();
+	static const uint64 GetRTHeartBeatMask();
+	static const uint64 GetPoolThreadMask();
+	static const uint64 GetTaskGraphThreadMask();
+	static const uint64 GetStatsThreadMask();
+	static const uint64 GetNoAffinityMask();
+};
+
+typedef FIOSPlatformAffinity FPlatformAffinity;
+
