Internationalization: More discriptive function parameter name

diff --git a/Engine/Source/Runtime/Core/Private/Internationalization/ICUText.cpp b/Engine/Source/Runtime/Core/Private/Internationalization/ICUText.cpp
--- a/Engine/Source/Runtime/Core/Private/Internationalization/ICUText.cpp
+++ b/Engine/Source/Runtime/Core/Private/Internationalization/ICUText.cpp
@@ -67,7 +67,7 @@ FText FText::AsTime(const FDateTime& DateTime, const EDateTimeStyle::Type TimeSt
 	return ResultText;
 }
 
-FText FText::AsTimespan(const FTimespan& Time, const FCulturePtr& TargetCulture)
+FText FText::AsTimespan(const FTimespan& Timespan, const FCulturePtr& TargetCulture)
 {
 	FInternationalization& I18N = FInternationalization::Get();
 	checkf(I18N.IsInitialized() == true, TEXT("FInternationalization is not initialized. An FText formatting method was likely used in static object initialization - this is not supported."));
@@ -75,10 +75,10 @@ FText FText::AsTimespan(const FTimespan& Time, const FCulturePtr& TargetCulture)
 
 	FText TimespanFormatPattern = NSLOCTEXT("Timespan", "FormatPattern", "{Hours}:{Minutes}:{Seconds}");
 
-	double TotalHours = Time.GetTotalHours();
+	double TotalHours = Timespan.GetTotalHours();
 	int32 Hours = static_cast<int32>(TotalHours);
-	int32 Minutes = Time.GetMinutes();
-	int32 Seconds = Time.GetSeconds();
+	int32 Minutes = Timespan.GetMinutes();
+	int32 Seconds = Timespan.GetSeconds();
 
 	FNumberFormattingOptions NumberFormattingOptions;
 	NumberFormattingOptions.MinimumIntegralDigits = 2;
diff --git a/Engine/Source/Runtime/Core/Private/Internationalization/LegacyText.cpp b/Engine/Source/Runtime/Core/Private/Internationalization/LegacyText.cpp
--- a/Engine/Source/Runtime/Core/Private/Internationalization/LegacyText.cpp
+++ b/Engine/Source/Runtime/Core/Private/Internationalization/LegacyText.cpp
@@ -22,10 +22,10 @@ FText FText::AsTime( const FDateTime& DateTime, const EDateTimeStyle::Type TimeS
 	return FText::FromString( DateTime.ToString( TEXT("%H.%M.%S") ) );
 }
 
-FText FText::AsTimespan( const FTimespan& Time, const FCulturePtr& TargetCulture)
+FText FText::AsTimespan( const FTimespan& Timespan, const FCulturePtr& TargetCulture)
 {
 	checkf(FInternationalization::Get().IsInitialized() == true, TEXT("FInternationalization is not initialized. An FText formatting method was likely used in static object initialization - this is not supported."));
-	FDateTime DateTime(Time.GetTicks());
+	FDateTime DateTime(Timespan.GetTicks());
 	return FText::FromString( DateTime.ToString( TEXT("%H.%M.%S") ) );
 }
 
diff --git a/Engine/Source/Runtime/Core/Public/Internationalization/Text.h b/Engine/Source/Runtime/Core/Public/Internationalization/Text.h
--- a/Engine/Source/Runtime/Core/Public/Internationalization/Text.h
+++ b/Engine/Source/Runtime/Core/Public/Internationalization/Text.h
@@ -163,7 +163,7 @@ public:
 	static FText AsDate(const FDateTime& DateTime, const EDateTimeStyle::Type DateStyle = EDateTimeStyle::Default, const FCulturePtr& TargetCulture = NULL);
 	static FText AsDateTime(const FDateTime& DateTime, const EDateTimeStyle::Type DateStyle = EDateTimeStyle::Default, const EDateTimeStyle::Type TimeStyle = EDateTimeStyle::Default, const FString& TimeZone = TEXT(""), const FCulturePtr& TargetCulture = NULL);
 	static FText AsTime(const FDateTime& DateTime, const EDateTimeStyle::Type TimeStyle = EDateTimeStyle::Default, const FString& TimeZone = TEXT(""), const FCulturePtr& TargetCulture = NULL);
-	static FText AsTimespan(const FTimespan& Time, const FCulturePtr& TargetCulture = NULL);
+	static FText AsTimespan(const FTimespan& Timespan, const FCulturePtr& TargetCulture = NULL);
 
 	/**
 	 * Generate an FText that represents the passed number as a memory size in the current culture
