Fixed Android startup crash

diff --git a/Engine/Source/Runtime/Core/Public/Containers/Array.h b/Engine/Source/Runtime/Core/Public/Containers/Array.h
--- a/Engine/Source/Runtime/Core/Public/Containers/Array.h
+++ b/Engine/Source/Runtime/Core/Public/Containers/Array.h
@@ -976,9 +976,9 @@ public:
 	 *
 	 * @param Addr The address to check.
 	 */
-	FORCEINLINE void CheckAddress(const void* Addr)
+	FORCEINLINE void CheckAddress(const ElementType* Addr) const
 	{
-		checkf(Addr < GetTypedData() || Addr >= GetTypedData() + ArrayMax, TEXT("Attempting to add a container element which already comes from the container!"));
+		checkf(Addr < GetTypedData() || Addr >= (GetTypedData() + ArrayMax), TEXT("Attempting to add a container element (0x%08x) which already comes from the container (0x%08x, ArrayMax: %d)!"), Addr, GetTypedData(), ArrayMax);
 	}
 
 	int32 Insert( ElementType&& Item, int32 Index )
diff --git a/Engine/Source/Runtime/Core/Public/Containers/Map.h b/Engine/Source/Runtime/Core/Public/Containers/Map.h
--- a/Engine/Source/Runtime/Core/Public/Containers/Map.h
+++ b/Engine/Source/Runtime/Core/Public/Containers/Map.h
@@ -280,10 +280,10 @@ public:
 	 * @param InValue - The value to associate with the key.
 	 * @return A reference to the value as stored in the map.  The reference is only valid until the next change to any key in the map.
 	 */
-	FORCEINLINE ValueType& Add(const KeyType&  InKey, const ValueType&  InValue) { Pairs.CheckAddress(&InKey); Pairs.CheckAddress(&InValue); return Emplace(         InKey ,          InValue ); }
-	FORCEINLINE ValueType& Add(const KeyType&  InKey,       ValueType&& InValue) { Pairs.CheckAddress(&InKey); Pairs.CheckAddress(&InValue); return Emplace(         InKey , MoveTemp(InValue)); }
-	FORCEINLINE ValueType& Add(      KeyType&& InKey, const ValueType&  InValue) { Pairs.CheckAddress(&InKey); Pairs.CheckAddress(&InValue); return Emplace(MoveTemp(InKey),          InValue ); }
-	FORCEINLINE ValueType& Add(      KeyType&& InKey,       ValueType&& InValue) { Pairs.CheckAddress(&InKey); Pairs.CheckAddress(&InValue); return Emplace(MoveTemp(InKey), MoveTemp(InValue)); }
+	FORCEINLINE ValueType& Add(const KeyType&  InKey, const ValueType&  InValue) { return Emplace(         InKey ,          InValue ); }
+	FORCEINLINE ValueType& Add(const KeyType&  InKey,       ValueType&& InValue) { return Emplace(         InKey , MoveTemp(InValue)); }
+	FORCEINLINE ValueType& Add(      KeyType&& InKey, const ValueType&  InValue) { return Emplace(MoveTemp(InKey),          InValue ); }
+	FORCEINLINE ValueType& Add(      KeyType&& InKey,       ValueType&& InValue) { return Emplace(MoveTemp(InKey), MoveTemp(InValue)); }
 
 	/**
 	 * Sets a default value associated with a key.
@@ -292,8 +292,8 @@ public:
 	 * @param InKey - The key to associate the value with.
 	 * @return A reference to the value as stored in the map.  The reference is only valid until the next change to any key in the map.
 	 */
-	FORCEINLINE ValueType& Add(const KeyType&  InKey) { Pairs.CheckAddress(&InKey); return Emplace(         InKey ); }
-	FORCEINLINE ValueType& Add(      KeyType&& InKey) { Pairs.CheckAddress(&InKey); return Emplace(MoveTemp(InKey)); }
+	FORCEINLINE ValueType& Add(const KeyType&  InKey) { return Emplace(         InKey ); }
+	FORCEINLINE ValueType& Add(      KeyType&& InKey) { return Emplace(MoveTemp(InKey)); }
 
 	/**
 	 * Sets the value associated with a key.
diff --git a/Engine/Source/Runtime/Core/Public/Containers/Set.h b/Engine/Source/Runtime/Core/Public/Containers/Set.h
--- a/Engine/Source/Runtime/Core/Public/Containers/Set.h
+++ b/Engine/Source/Runtime/Core/Public/Containers/Set.h
@@ -782,7 +782,7 @@ public:
 	 *
 	 * @param Addr The address to check.
 	 */
-	FORCEINLINE void CheckAddress(const void* Addr)
+	FORCEINLINE void CheckAddress(const ElementType* Addr) const
 	{
 		Elements.CheckAddress(Addr);
 	}
diff --git a/Engine/Source/Runtime/Core/Public/Containers/SparseArray.h b/Engine/Source/Runtime/Core/Public/Containers/SparseArray.h
--- a/Engine/Source/Runtime/Core/Public/Containers/SparseArray.h
+++ b/Engine/Source/Runtime/Core/Public/Containers/SparseArray.h
@@ -574,7 +574,7 @@ public:
 	 *
 	 * @param Addr The address to check.
 	 */
-	FORCEINLINE void CheckAddress(const void* Addr)
+	FORCEINLINE void CheckAddress(const ElementType* Addr) const
 	{
 		Data.CheckAddress(Addr);
 	}
