[GitHub] Add better comments for LinePlaneIntersection
https://github.com/EpicGames/UnrealEngine/pull/352 (error454)

diff --git a/Engine/Source/Runtime/Core/Public/Math/UnrealMath.h b/Engine/Source/Runtime/Core/Public/Math/UnrealMath.h
--- a/Engine/Source/Runtime/Core/Public/Math/UnrealMath.h
+++ b/Engine/Source/Runtime/Core/Public/Math/UnrealMath.h
@@ -111,7 +111,18 @@ inline FPlane FPlane::TransformByUsingAdjointT( const FMatrix& M, float DetM, co
 	return FPlane(M.TransformPosition(*this * W), newNorm);
 }
 
-// FMath Implementation
+/**
+ * Find the intersection of a line and an offset plane. Assumes that the
+ * line and plane do indeed intersect; you must make sure they're not
+ * parallel before calling.
+ *
+ * @param Point1 the first point defining the line
+ * @param Point2 the second point defining the line
+ * @param PlaneOrigin the origin of the plane
+ * @param PlaneNormal the normal of the plane
+ *
+ * @return The point of intersection between the line and the plane.
+ */
 inline FVector FMath::LinePlaneIntersection
 	(
 	const FVector &Point1,
@@ -125,6 +136,19 @@ inline FVector FMath::LinePlaneIntersection
 		+	(Point2-Point1)
 		*	(((PlaneOrigin - Point1)|PlaneNormal) / ((Point2 - Point1)|PlaneNormal));
 }
+
+
+/**
+ * Find the intersection of a line and a plane. Assumes that the line and
+ * plane do indeed intersect; you must make sure they're not parallel before
+ * calling.
+ *
+ * @param Point1 the first point defining the line
+ * @param Point2 the second point defining the line
+ * @param Plane the plane
+ *
+ * @return The point of intersection between the line and the plane.
+ */
 inline FVector FMath::LinePlaneIntersection
 	(
 	const FVector &Point1,
