[GitHub] Add CInterpTo for linear colors, exposed in BP
https://github.com/EpicGames/UnrealEngine/pull/331 (EverNewJoy)

diff --git a/Engine/Source/Runtime/Core/Private/Math/UnrealMath.cpp b/Engine/Source/Runtime/Core/Private/Math/UnrealMath.cpp
--- a/Engine/Source/Runtime/Core/Private/Math/UnrealMath.cpp
+++ b/Engine/Source/Runtime/Core/Private/Math/UnrealMath.cpp
@@ -2043,6 +2043,30 @@ CORE_API float FMath::FInterpConstantTo( float Current, float Target, float Delt
 	return Current + FMath::Clamp<float>(Dist, -Step, Step);
 }
 
+/** Interpolate Linear Color from Current to Target. Scaled by distance to Target, so it has a strong start speed and ease out. */
+CORE_API FLinearColor FMath::CInterpTo(const FLinearColor& Current, const FLinearColor& Target, float DeltaTime, float InterpSpeed)
+{
+	// If no interp speed, jump to target value
+	if (InterpSpeed <= 0.f)
+	{
+		return Target;
+	}
+
+	// Difference between colors
+	const float Dist = FLinearColor::Dist(Target, Current);
+
+	// If distance is too small, just set the desired color
+	if (Dist < KINDA_SMALL_NUMBER)
+	{
+		return Target;
+	}
+
+	// Delta change, Clamp so we do not over shoot.
+	const FLinearColor DeltaMove = (Target - Current) * FMath::Clamp<float>(DeltaTime * InterpSpeed, 0.f, 1.f);
+
+	return Current + DeltaMove;
+}
+
 CORE_API float ClampFloatTangent( float PrevPointVal, float PrevTime, float CurPointVal, float CurTime, float NextPointVal, float NextTime )
 {
 	const float PrevToNextTimeDiff = FMath::Max< double >( KINDA_SMALL_NUMBER, NextTime - PrevTime );
diff --git a/Engine/Source/Runtime/Core/Public/Math/UnrealMathUtility.h b/Engine/Source/Runtime/Core/Public/Math/UnrealMathUtility.h
--- a/Engine/Source/Runtime/Core/Public/Math/UnrealMathUtility.h
+++ b/Engine/Source/Runtime/Core/Public/Math/UnrealMathUtility.h
@@ -605,6 +605,9 @@ struct FMath : public FPlatformMath
 	/** Interpolate float from Current to Target. Scaled by distance to Target, so it has a strong start speed and ease out. */
 	static CORE_API float FInterpTo( float Current, float Target, float DeltaTime, float InterpSpeed );
 
+	/** Interpolate Linear Color from Current to Target. Scaled by distance to Target, so it has a strong start speed and ease out. */
+	static CORE_API FLinearColor CInterpTo(const FLinearColor& Current, const FLinearColor& Target, float DeltaTime, float InterpSpeed);
+
 	/**
 	 * Simple function to create a pulsating scalar value
 	 *
diff --git a/Engine/Source/Runtime/Engine/Classes/Kismet/KismetMathLibrary.h b/Engine/Source/Runtime/Engine/Classes/Kismet/KismetMathLibrary.h
--- a/Engine/Source/Runtime/Engine/Classes/Kismet/KismetMathLibrary.h
+++ b/Engine/Source/Runtime/Engine/Classes/Kismet/KismetMathLibrary.h
@@ -1239,6 +1239,17 @@ class UKismetMathLibrary : public UBlueprintFunctionLibrary
 	UFUNCTION(BlueprintPure, Category="Math|Interpolation", meta=(Keywords="rotation rotate"))
 	static FRotator RInterpTo(FRotator Current, FRotator Target, float DeltaTime, float InterpSpeed);
 
+	/**
+	 * Interpolates towards a varying target color smoothly.
+	 *
+	 * @param		Current			Current Color
+	 * @param		Target			Target Color
+	 * @param		DeltaTime		Time since last tick
+	 * @param		InterpSpeed		Interpolation speed
+	 * @return		New interpolated Color
+	 */
+	UFUNCTION(BlueprintPure, Category = "Math|Interpolation", meta = (Keywords = "color"))
+	static FLinearColor CInterpTo(FLinearColor Current, FLinearColor Target, float DeltaTime, float InterpSpeed);
 
 	//
 	// Random stream functions
diff --git a/Engine/Source/Runtime/Engine/Private/KismetMathLibrary.cpp b/Engine/Source/Runtime/Engine/Private/KismetMathLibrary.cpp
--- a/Engine/Source/Runtime/Engine/Private/KismetMathLibrary.cpp
+++ b/Engine/Source/Runtime/Engine/Private/KismetMathLibrary.cpp
@@ -863,6 +863,11 @@ FRotator UKismetMathLibrary::RInterpTo(FRotator Current, FRotator Target, float
 	return FMath::RInterpTo( Current, Target, DeltaTime, InterpSpeed);
 }
 
+FLinearColor UKismetMathLibrary::CInterpTo(FLinearColor Current, FLinearColor Target, float DeltaTime, float InterpSpeed)
+{
+	return FMath::CInterpTo(Current, Target, DeltaTime, InterpSpeed);
+}
+
 FLinearColor UKismetMathLibrary::LinearColorLerp(FLinearColor A, FLinearColor B, float Alpha)
 {
 	return A + Alpha * (B - A);
