- Moved Metal .ini setting into the main IOS project settings section

diff --git a/Engine/Source/Developer/iOS/IOSTargetPlatform/Private/IOSTargetPlatform.cpp b/Engine/Source/Developer/iOS/IOSTargetPlatform/Private/IOSTargetPlatform.cpp
--- a/Engine/Source/Developer/iOS/IOSTargetPlatform/Private/IOSTargetPlatform.cpp
+++ b/Engine/Source/Developer/iOS/IOSTargetPlatform/Private/IOSTargetPlatform.cpp
@@ -286,17 +286,17 @@ void FIOSTargetPlatform::GetAllPossibleShaderFormats( TArray<FName>& OutFormats
 	static FName NAME_SF_METAL(TEXT("SF_METAL"));
 
 	// default to supporting ES2
-	bool bSupportOpenGLES2 = true;
-	GConfig->GetBool(TEXT("/Script/UnrealEd.CookerSettings"), TEXT("bSupportOpenGLES2"), bSupportOpenGLES2, GEngineIni);
-	if (bSupportOpenGLES2)
+	bool bSupportsOpenGLES2 = true;
+	GConfig->GetBool(TEXT("/Script/IOSRuntimeSettings.IOSRuntimeSettings"), TEXT("bSupportsOpenGLES2"), bSupportsOpenGLES2, GEngineIni);
+	if (bSupportsOpenGLES2)
 	{
 		OutFormats.AddUnique(NAME_OPENGL_ES2_IOS);
 	}
 
 	// default to NOT supporting metal
-	bool bSupportMetal = false;
-	GConfig->GetBool(TEXT("/Script/UnrealEd.CookerSettings"), TEXT("bSupportMetal"), bSupportMetal, GEngineIni);
-	if (bSupportMetal)
+	bool bSupportsMetal = false;
+	GConfig->GetBool(TEXT("/Script/IOSRuntimeSettings.IOSRuntimeSettings"), TEXT("bSupportsMetal"), bSupportsMetal, GEngineIni);
+	if (bSupportsMetal)
 	{
 		OutFormats.AddUnique(NAME_SF_METAL);
 	}
diff --git a/Engine/Source/Editor/UnrealEd/Private/Commandlets/CookCommandlet.cpp b/Engine/Source/Editor/UnrealEd/Private/Commandlets/CookCommandlet.cpp
--- a/Engine/Source/Editor/UnrealEd/Private/Commandlets/CookCommandlet.cpp
+++ b/Engine/Source/Editor/UnrealEd/Private/Commandlets/CookCommandlet.cpp
@@ -31,8 +31,6 @@ UCookerSettings::UCookerSettings(const FPostConstructInitializeProperties& PCIP)
 	: Super(PCIP)
 {
 	DefaultPVRTCQuality = 0;
-	bSupportMetal = false;
-	bSupportOpenGLES2 = true;
 	AndroidAudio = EAndroidAudio::Default;
 }
 
diff --git a/Engine/Source/Editor/UnrealEd/Public/CookerSettings.h b/Engine/Source/Editor/UnrealEd/Public/CookerSettings.h
--- a/Engine/Source/Editor/UnrealEd/Public/CookerSettings.h
+++ b/Engine/Source/Editor/UnrealEd/Public/CookerSettings.h
@@ -35,14 +35,6 @@ public:
 	UPROPERTY(config, EditAnywhere, Category=Textures)
 	int32 DefaultPVRTCQuality;
 
-	/** Whether or not to add support for Metal API (requires IOS8 and A7 processors) */
-	UPROPERTY(config, EditAnywhere, Category=IOS)
-	bool bSupportMetal;
-
-	/** Whether or not to add support for OpenGL ES2 (if this is false, then your game should specify minimum IOS8 version and A7, in the plist) */
-	UPROPERTY(config, EditAnywhere, Category = IOS)
-	bool bSupportOpenGLES2;
-
 	/** Android Audio encoding options */
 	UPROPERTY(config, EditAnywhere, Category = Android, meta = (DisplayName = "Audio encoding"))
 	TEnumAsByte<EAndroidAudio::Type> AndroidAudio;
diff --git a/Engine/Source/Runtime/Core/Private/IOS/IOSView.cpp b/Engine/Source/Runtime/Core/Private/IOS/IOSView.cpp
--- a/Engine/Source/Runtime/Core/Private/IOS/IOSView.cpp
+++ b/Engine/Source/Runtime/Core/Private/IOS/IOSView.cpp
@@ -25,8 +25,8 @@ id<MTLDevice> GMetalDevice = nil;
 {
 #if HAS_METAL
 	// make sure the project setting has enabled Metal support (per-project user settings in the editor)
-	bool bSupportMetal = false;
-	GConfig->GetBool(TEXT("/Script/UnrealEd.CookerSettings"), TEXT("bSupportMetal"), bSupportMetal, GEngineIni);
+	bool bSupportsMetal = false;
+	GConfig->GetBool(TEXT("/Script/IOSRuntimeSettings.IOSRuntimeSettings"), TEXT("bSupportsMetal"), bSupportsMetal, GEngineIni);
 
 	// does commandline override?
 	bool bForceES2 = FParse::Param(FCommandLine::Get(), TEXT("ES2"));
@@ -34,7 +34,7 @@ id<MTLDevice> GMetalDevice = nil;
 	bool bTriedToInit = false;
 
 	// the check for the function pointer itself is to determine if the Metal framework exists, before calling it
-	if (bSupportMetal && !bForceES2 && MTLCreateSystemDefaultDevice != NULL)
+	if (bSupportsMetal && !bForceES2 && MTLCreateSystemDefaultDevice != NULL)
 	{
 		// if the device is unable to run with Metal (pre-A7), this will return nil
 		GMetalDevice = MTLCreateSystemDefaultDevice();
@@ -47,7 +47,7 @@ id<MTLDevice> GMetalDevice = nil;
 	if (GMetalDevice == nil)
 	{
 		FPlatformMisc::LowLevelOutputDebugStringf(TEXT("Not using Metal because: [Project Settings Disabled Metal? %s :: Commandline Forced ES2? %s :: Older OS? %s :: Pre-A7 Device? %s]"),
-			bSupportMetal ? TEXT("No") : TEXT("Yes"),
+			bSupportsMetal ? TEXT("No") : TEXT("Yes"),
 			bForceES2? TEXT("Yes") : TEXT("No"),
 			(MTLCreateSystemDefaultDevice == NULL) ? TEXT("Yes") : TEXT("No"),
 			bTriedToInit ? TEXT("Yes") : TEXT("Unknown (didn't test)"));
diff --git a/Engine/Source/Runtime/IOS/IOSRuntimeSettings/Classes/IOSRuntimeSettings.h b/Engine/Source/Runtime/IOS/IOSRuntimeSettings/Classes/IOSRuntimeSettings.h
--- a/Engine/Source/Runtime/IOS/IOSRuntimeSettings/Classes/IOSRuntimeSettings.h
+++ b/Engine/Source/Runtime/IOS/IOSRuntimeSettings/Classes/IOSRuntimeSettings.h
@@ -18,6 +18,14 @@ public:
 	UPROPERTY(GlobalConfig, EditAnywhere, Category = Online)
 	uint32 bEnableGameCenterSupport : 1;
 	
+	// Whether or not to add support for Metal API (requires IOS8 and A7 processors)
+	UPROPERTY(GlobalConfig, EditAnywhere, Category = Rendering)
+	bool bSupportsMetal;
+
+	// Whether or not to add support for OpenGL ES2 (if this is false, then your game should specify minimum IOS8 version and use "metal" instead of "opengles-2" in UIRequiredDeviceCapabilities)
+	UPROPERTY(GlobalConfig, EditAnywhere, Category = Rendering)
+	bool bSupportsOpenGLES2;
+
 	// Does the application support portrait orientation?
 	UPROPERTY(GlobalConfig, EditAnywhere, Category = DeviceOrientations)
 	uint32 bSupportsPortraitOrientation : 1;
diff --git a/Engine/Source/Runtime/IOS/IOSRuntimeSettings/Private/IOSRuntimeSettings.cpp b/Engine/Source/Runtime/IOS/IOSRuntimeSettings/Private/IOSRuntimeSettings.cpp
--- a/Engine/Source/Runtime/IOS/IOSRuntimeSettings/Private/IOSRuntimeSettings.cpp
+++ b/Engine/Source/Runtime/IOS/IOSRuntimeSettings/Private/IOSRuntimeSettings.cpp
@@ -25,5 +25,11 @@ void UIOSRuntimeSettings::PostEditChangeProperty(struct FPropertyChangedEvent& P
 	{
 		bSupportsPortraitOrientation = true;
 	}
+
+	// Ensure that at least one API is supported
+	if (!bSupportsMetal && !bSupportsOpenGLES2)
+	{
+		bSupportsOpenGLES2 = true;
+	}
 }
 #endif
