- Fixed CIS error

diff --git a/Engine/Source/Runtime/Core/Private/IOS/IOSPlatformFile.cpp b/Engine/Source/Runtime/Core/Private/IOS/IOSPlatformFile.cpp
--- a/Engine/Source/Runtime/Core/Private/IOS/IOSPlatformFile.cpp
+++ b/Engine/Source/Runtime/Core/Private/IOS/IOSPlatformFile.cpp
@@ -23,7 +23,8 @@ const FDateTime IOSEpoch(1970, 1, 1);
  * using pak file (e.g., SHIPPING?) so not particularly optimized. Only manages 
  * files which are opened READ_ONLY.
  **/
-#define MANAGE_FILE_HANDLES 1 // !UE_BUILD_SHIPPING
+// @todo: Merge all of the managed file handles into one class!
+#define MANAGE_FILE_HANDLES_IOS 1 // !UE_BUILD_SHIPPING
 
 struct FManagedFile
 {
@@ -41,17 +42,17 @@ public:
 
 	FIOSFileHandle( int32 InFileHandle, const FString& InFilename, bool bIsForRead )
 		: FileHandle(InFileHandle)
-#if !UE_BUILD_SHIPPING || MANAGE_FILE_HANDLES
+#if !UE_BUILD_SHIPPING || MANAGE_FILE_HANDLES_IOS
 		, Filename(InFilename)
 #endif
-#if MANAGE_FILE_HANDLES
+#if MANAGE_FILE_HANDLES_IOS
         , HandleSlot(-1)
         , FileSize(0)
 #endif
 	{
 		check(FileHandle != 0);
 
-#if MANAGE_FILE_HANDLES
+#if MANAGE_FILE_HANDLES_IOS
 
 		static uint32 NextID = 1;
 		FileID = NextID++;
@@ -89,7 +90,7 @@ public:
 	 */
 	virtual ~FIOSFileHandle( )
 	{
-#if MANAGE_FILE_HANDLES
+#if MANAGE_FILE_HANDLES_IOS
         if( IsManaged() )
         {
             if( ManagedFiles[HandleSlot].ID == FileID )
@@ -130,7 +131,7 @@ public:
 
 	virtual bool Read( uint8* Destination, int64 BytesToRead ) override
 	{
- #if MANAGE_FILE_HANDLES
+ #if MANAGE_FILE_HANDLES_IOS
        if( IsManaged() )
         {
             ActivateSlot();
@@ -154,7 +155,7 @@ public:
 	{
 		check(NewPosition >= 0);
 
-#if MANAGE_FILE_HANDLES
+#if MANAGE_FILE_HANDLES_IOS
         if( IsManaged() )
         {
             FileOffset = NewPosition >= FileSize ? FileSize - 1 : NewPosition;
@@ -171,7 +172,7 @@ public:
 	{
 		check(NewPositionRelativeToEnd <= 0);
 
-#if MANAGE_FILE_HANDLES
+#if MANAGE_FILE_HANDLES_IOS
         if( IsManaged() )
         {
             FileOffset = (NewPositionRelativeToEnd >= FileSize) ? 0 : ( FileSize + NewPositionRelativeToEnd - 1 );
@@ -186,7 +187,7 @@ public:
 
 	virtual int64 Size( ) override
 	{
-#if MANAGE_FILE_HANDLES
+#if MANAGE_FILE_HANDLES_IOS
         if( IsManaged() )
         {
             return FileSize;
@@ -200,7 +201,7 @@ public:
 
 	virtual int64 Tell( ) override
 	{
-#if MANAGE_FILE_HANDLES
+#if MANAGE_FILE_HANDLES_IOS
         if( IsManaged() )
         {
             return FileOffset;
@@ -233,7 +234,7 @@ public:
 
 private:
 
-#if MANAGE_FILE_HANDLES
+#if MANAGE_FILE_HANDLES_IOS
     FORCEINLINE bool IsManaged()
     {
         return HandleSlot != -1;
@@ -301,12 +302,12 @@ private:
 	// Holds the internal file handle.
 	int32 FileHandle;
 
-#if !UE_BUILD_SHIPPING || MANAGE_FILE_HANDLES
+#if !UE_BUILD_SHIPPING || MANAGE_FILE_HANDLES_IOS
 	// Holds the name of the file that this handle represents. Kept around for possible reopen of file.
 	FString Filename;
 #endif
 
-#if MANAGE_FILE_HANDLES
+#if MANAGE_FILE_HANDLES_IOS
     // Most recent valid slot index for this handle; >=0 for handles which are managed.
     int32 HandleSlot;
 
@@ -326,7 +327,7 @@ private:
 #endif
 };
 
-#if MANAGE_FILE_HANDLES
+#if MANAGE_FILE_HANDLES_IOS
 int32 FIOSFileHandle::ManagedFilesTlsSlot = FPlatformTLS::AllocTlsSlot();
 #endif
 
