Correctly increment loop counter to avoid infinite loop on failure.

diff --git a/Engine/Source/Runtime/Core/Private/HAL/OutputDevices.cpp b/Engine/Source/Runtime/Core/Private/HAL/OutputDevices.cpp
--- a/Engine/Source/Runtime/Core/Private/HAL/OutputDevices.cpp
+++ b/Engine/Source/Runtime/Core/Private/HAL/OutputDevices.cpp
@@ -1000,7 +1000,7 @@ FArchive* FOutputDeviceFile::CreateArchive(uint32 MaxAttempts)
 		do
 		{
 			// Continue to increment indices until a valid filename is found
-			FinalFilename = FilenamePart + FString::FromInt(FileIndex) + ExtensionPart;
+			FinalFilename = FilenamePart + FString::FromInt(FileIndex++) + ExtensionPart;
 			if (!Opened)
 			{
 				CreateBackupCopy(*FinalFilename);
