Prevent infinite loop inside a defunct input-method context - they need to be deactivated when they are unregistered as the lower-level Cocoa implementation retains the pointer.
Fixes AnswerHub: https://answers.unrealengine.com/questions/87671/mac-editor-infrequent-permanent-lockup-fslatetextv.html

diff --git a/Engine/Source/Runtime/Core/Private/Mac/MacTextInputMethodSystem.cpp b/Engine/Source/Runtime/Core/Private/Mac/MacTextInputMethodSystem.cpp
--- a/Engine/Source/Runtime/Core/Private/Mac/MacTextInputMethodSystem.cpp
+++ b/Engine/Source/Runtime/Core/Private/Mac/MacTextInputMethodSystem.cpp
@@ -136,6 +136,14 @@ void FMacTextInputMethodSystem::UnregisterContext(const TSharedRef<ITextInputMet
 		return;
 	}
 	
+	TSharedPtr<ITextInputMethodChangeNotifier> Notifier(NotifierRef.Pin());
+	FTextInputMethodChangeNotifier* MacNotifier = (FTextInputMethodChangeNotifier*)Notifier.Get();
+	TSharedPtr<FGenericWindow> GenericWindow = MacNotifier->GetContextWindow();
+	if(GenericWindow.IsValid())
+	{
+		DeactivateContext(Context);
+	}
+	
 	ContextMap.Remove(Context);
 }
 
