Added support for "mailto:" url links on Mac and Windows

diff --git a/Engine/Source/Runtime/Core/Private/Mac/MacPlatformProcess.cpp b/Engine/Source/Runtime/Core/Private/Mac/MacPlatformProcess.cpp
--- a/Engine/Source/Runtime/Core/Private/Mac/MacPlatformProcess.cpp
+++ b/Engine/Source/Runtime/Core/Private/Mac/MacPlatformProcess.cpp
@@ -181,7 +181,7 @@ void FMacPlatformProcess::LaunchURL( const TCHAR* URL, const TCHAR* Parms, FStri
 
 	UE_LOG(LogMac, Log,  TEXT("LaunchURL %s %s"), URL, Parms?Parms:TEXT("") );
 	NSString* Url = (NSString*)FPlatformString::TCHARToCFString( URL );
-	NSURL* UrlToOpen = [NSURL URLWithString: ([Url hasPrefix: @"http://"] || [Url hasPrefix: @"https://"] || [Url hasPrefix: @"file://"]) ? Url : [NSString stringWithFormat: @"http://%@", Url]];
+	NSURL* UrlToOpen = [NSURL URLWithString: ([Url hasPrefix: @"http://"] || [Url hasPrefix: @"https://"] || [Url hasPrefix: @"file://"] || [Url hasPrefix: @"mailto:"]) ? Url : [NSString stringWithFormat: @"http://%@", Url]];
 	[[NSWorkspace sharedWorkspace] openURL: UrlToOpen];
 	CFRelease( (CFStringRef)Url );
 	if( Error )
diff --git a/Engine/Source/Runtime/Core/Private/Windows/WindowsPlatformProcess.cpp b/Engine/Source/Runtime/Core/Private/Windows/WindowsPlatformProcess.cpp
--- a/Engine/Source/Runtime/Core/Private/Windows/WindowsPlatformProcess.cpp
+++ b/Engine/Source/Runtime/Core/Private/Windows/WindowsPlatformProcess.cpp
@@ -136,24 +136,15 @@ void FWindowsPlatformProcess::PopDllDirectory(const TCHAR* Directory)
 	}
 }
 
-void FWindowsPlatformProcess::LaunchURL( const TCHAR* URL, const TCHAR* Parms, FString* Error )
+static void LaunchWebURL( const FString& URLParams, FString* Error )
 {
-	// Initialize the error to empty string.
-	if (Error)
-	{
-		*Error = TEXT("");
-	}
-
-	check( URL );
-	FString URLParams = FString::Printf(TEXT("%s %s"), URL, Parms ? Parms : TEXT("")).TrimTrailing();
-		
 	UE_LOG(LogWindows, Log, TEXT("LaunchURL %s"), *URLParams);
 
 	FString BrowserOpenCommand;
 
 	// First lookup the program Id for the default browser.
 	FString ProgId;
-	if ( FWindowsPlatformMisc::QueryRegKey(HKEY_CURRENT_USER, TEXT("Software\\Microsoft\\Windows\\Shell\\Associations\\UrlAssociations\\http\\UserChoice"), TEXT("Progid"), ProgId) )
+	if (FWindowsPlatformMisc::QueryRegKey(HKEY_CURRENT_USER, TEXT("Software\\Microsoft\\Windows\\Shell\\Associations\\UrlAssociations\\http\\UserChoice"), TEXT("Progid"), ProgId))
 	{
 		// If we found it, then lookup it's open shell command in the classes registry.
 		FString BrowserRegPath = ProgId + TEXT("\\shell\\open\\command");
@@ -161,14 +152,14 @@ void FWindowsPlatformProcess::LaunchURL( const TCHAR* URL, const TCHAR* Parms, F
 	}
 
 	// If we failed to find a default browser using the newer location, revert to using shell open command for the HTTP file association.
-	if ( BrowserOpenCommand.IsEmpty() )
+	if (BrowserOpenCommand.IsEmpty())
 	{
 		FWindowsPlatformMisc::QueryRegKey(HKEY_CLASSES_ROOT, TEXT("http\\shell\\open\\command"), NULL, BrowserOpenCommand);
 	}
 
 	// If we have successfully looked up the correct shell command, then we can create a new process using that command
 	// we do this instead of shell execute due to security concerns.  By starting the browser directly we avoid most issues.
-	if ( !BrowserOpenCommand.IsEmpty() )
+	if (!BrowserOpenCommand.IsEmpty())
 	{
 		FString ExePath, ExeArgs;
 
@@ -179,30 +170,30 @@ void FWindowsPlatformProcess::LaunchURL( const TCHAR* URL, const TCHAR* Parms, F
 
 		// Extract the exe and any arguments to the executable.
 		const int32 FirstQuote = BrowserOpenCommand.Find(TEXT("\""));
-		if ( FirstQuote != INDEX_NONE )
+		if (FirstQuote != INDEX_NONE)
 		{
 			const int32 SecondQuote = BrowserOpenCommand.Find(TEXT("\""), ESearchCase::IgnoreCase, ESearchDir::FromStart, FirstQuote + 1);
-			if ( SecondQuote != INDEX_NONE )
+			if (SecondQuote != INDEX_NONE)
 			{
-				ExePath = BrowserOpenCommand.Mid(FirstQuote + 1, ( SecondQuote - 1 ) - FirstQuote);
+				ExePath = BrowserOpenCommand.Mid(FirstQuote + 1, (SecondQuote - 1) - FirstQuote);
 				ExeArgs = BrowserOpenCommand.Mid(SecondQuote + 1);
 			}
 		}
 
 		// If anything failed to parse right, don't continue down this path, just use shell execute.
-		if ( !ExePath.IsEmpty() )
+		if (!ExePath.IsEmpty())
 		{
 			ExeArgs = ExeArgs.Replace(TEXT("%1"), *URLParams);
 
 			// Now that we have the shell open command to use, run the shell command in the open process with any and all parameters.
-			if ( FPlatformProcess::CreateProc(*ExePath, *ExeArgs, true, false, false, NULL, 0, NULL, NULL).IsValid() )
+			if (FPlatformProcess::CreateProc(*ExePath, *ExeArgs, true, false, false, NULL, 0, NULL, NULL).IsValid())
 			{
 				// Success!
 				return;
 			}
 			else
 			{
-				if ( Error )
+				if (Error)
 				{
 					*Error = NSLOCTEXT("Core", "UrlFailed", "Failed launching URL").ToString();
 				}
@@ -212,16 +203,48 @@ void FWindowsPlatformProcess::LaunchURL( const TCHAR* URL, const TCHAR* Parms, F
 
 	// If all else fails just do a shell execute and let windows sort it out.  But only do it if it's an
 	// HTTP or HTTPS address.  A malicious address could be problematic if just passed directly to shell execute.
-	if ( URLParams.StartsWith(TEXT("http://")) || URLParams.StartsWith(TEXT("https://")) )
+	if (URLParams.StartsWith(TEXT("http://")) || URLParams.StartsWith(TEXT("https://")))
 	{
 		const HINSTANCE Code = ::ShellExecuteW(NULL, TEXT("open"), *URLParams, NULL, NULL, SW_SHOWNORMAL);
-		if ( Error )
+		if (Error)
 		{
-			*Error = ( (PTRINT)Code <= 32 ) ? NSLOCTEXT("Core", "UrlFailed", "Failed launching URL").ToString() : TEXT("");
+			*Error = ((PTRINT)Code <= 32) ? NSLOCTEXT("Core", "UrlFailed", "Failed launching URL").ToString() : TEXT("");
 		}
 	}
 }
 
+static void LaunchMailToURL( const TCHAR* MailTo, FString* Error )
+{
+	// ShellExecute will open the default mail app with a "mailto:" link
+	const HINSTANCE Code = ::ShellExecuteW(NULL, TEXT("open"), MailTo, NULL, NULL, SW_SHOWNORMAL);
+	if (Error)
+	{
+		*Error = ((PTRINT)Code <= 32) ? NSLOCTEXT("Core", "UrlFailed", "Failed launching URL").ToString() : TEXT("");
+	}
+}
+
+void FWindowsPlatformProcess::LaunchURL( const TCHAR* URL, const TCHAR* Parms, FString* Error )
+{
+	check(URL);
+
+	// Initialize the error to empty string.
+	if (Error)
+	{
+		*Error = TEXT("");
+	}
+
+	if( FString(URL).StartsWith(TEXT("mailto:")) )
+	{
+		LaunchMailToURL( URL, Error );
+	}
+	else
+	{
+		FString URLParams = FString::Printf(TEXT("%s %s"), URL, Parms ? Parms : TEXT("")).TrimTrailing();
+		LaunchWebURL( URLParams, Error );
+	}
+
+}
+
 FProcHandle FWindowsPlatformProcess::CreateProc( const TCHAR* URL, const TCHAR* Parms, bool bLaunchDetached, bool bLaunchHidden, bool bLaunchReallyHidden, uint32* OutProcessID, int32 PriorityModifier, const TCHAR* OptionalWorkingDirectory, void* PipeWrite )
 {
 	//UE_LOG(LogWindows, Log,  TEXT("CreateProc %s %s"), URL, Parms );
