Make Mouse Lock / FPS style mouse control work on HTML5 browser.

            #HTML5 #UE4

            TTP#324327 HTML5: TEMPLATES: Mouse remains on screen when it controls the camera
            TTP#324325 HTML5: TEMPLATES: Vertical mouse control is inverted

diff --git a/Engine/Source/Runtime/Core/Private/HTML5/HTML5Application.cpp b/Engine/Source/Runtime/Core/Private/HTML5/HTML5Application.cpp
--- a/Engine/Source/Runtime/Core/Private/HTML5/HTML5Application.cpp
+++ b/Engine/Source/Runtime/Core/Private/HTML5/HTML5Application.cpp
@@ -5,6 +5,20 @@
 #include "HTML5Cursor.h"
 #include "HTML5InputInterface.h"
 
+
+#if PLATFORM_HTML5_BROWSER
+#include "emscripten.h"
+#include "html5.h"
+
+EM_BOOL mouse_callback(int eventType, const EmscriptenMouseEvent *e, void *userData)
+{
+	return 1;
+}
+
+#endif 
+
+static const uint32 MaxWarmUpTicks = 10; 
+
 FHTML5Application* FHTML5Application::CreateHTML5Application()
 {
 	return new FHTML5Application();
@@ -12,15 +26,56 @@ FHTML5Application* FHTML5Application::CreateHTML5Application()
 
 FHTML5Application::FHTML5Application()
 	: GenericApplication( MakeShareable( new FHTML5Cursor() ) )
-	, InputInterface( FHTML5InputInterface::Create( MessageHandler, Cursor ) )
 	, ApplicationWindow(FHTML5Window::Make())
+	, InputInterface( FHTML5InputInterface::Create(MessageHandler, Cursor ) )
+	, WarmUpTicks(-1)
 {
+
+#if PLATFORM_HTML5_BROWSER
+    // full screen will only be requested after the first click after the window gains focus. 
+    // the problem is that because of security/UX reasons browsers don't allow pointer lock in main loop
+    // but only through callbacks generated by browser. 
+
+ 	// work around emscripten bug where deffered browser requests are not called if there are no callbacks.
+	emscripten_set_mousedown_callback("#canvas",0,1,mouse_callback);
+#endif 
+
+#if PLATFORM_HTML5_WIN32
+	// these events don't happen on startup. 
+	// To warm up the app correctly push these events on startup. 
+
+	SDL_ActiveEvent Event2; 
+	Event2.gain = 1; 
+	Event2.state = SDL_APPMOUSEFOCUS;
+	Event2.type = SDL_ACTIVEEVENT; 
+	SDL_PushEvent((SDL_Event*)&Event2);
+
+	SDL_ActiveEvent Event; 
+	Event.gain = 1; 
+	Event.state = SDL_APPINPUTFOCUS;
+	Event.type = SDL_ACTIVEEVENT; 
+	SDL_PushEvent((SDL_Event*)&Event);
+
+#endif 
+
+#if PLATFORM_HTML5_BROWSER
+
+	SDL_WindowEvent Event ;
+	Event.type = SDL_WINDOWEVENT; 
+	Event.event = SDL_WINDOWEVENT_FOCUS_GAINED;
+    SDL_PushEvent((SDL_Event*)&Event);
+
+#endif 
+
+
 }
 
+
 void FHTML5Application::SetMessageHandler( const TSharedRef< FGenericApplicationMessageHandler >& InMessageHandler )
 {
 	GenericApplication::SetMessageHandler(InMessageHandler);
 	InputInterface->SetMessageHandler( MessageHandler );
+
 }
 
 void FHTML5Application::PollGameDeviceState( const float TimeDelta )
@@ -29,31 +84,142 @@ void FHTML5Application::PollGameDeviceState( const float TimeDelta )
 	while (SDL_PollEvent(&Event)) 
 	{
 		// Tick Input Interface. 
-		if ( Event.type == SDL_VIDEORESIZE )
+		switch (Event.type)
 		{
-			SDL_ResizeEvent *r = (SDL_ResizeEvent*)&Event;
+			case SDL_VIDEORESIZE:
+				{
+					SDL_ResizeEvent *r = (SDL_ResizeEvent*)&Event;
 
-			if ( r->h > 0 )
-			{
-				MessageHandler->OnSizeChanged(  ApplicationWindow, r->w,r->h); 
-			} 
-			else 
-			{
-				// Ask the game to reshape the window. 
-				int w = r->w;  
-				int h = -r->h;
-				ANSICHAR command[1048]; 
-				FCStringAnsi::Sprintf(command, "r.setRes %dx%d ", w, h);
+					if ( r->h > 0 )
+					{
+						MessageHandler->OnSizeChanged(  ApplicationWindow, r->w,r->h); 
+					} 
+					else 
+					{
+						// Ask the game to reshape the window. 
+						int w = r->w;  
+						int h = -r->h;
+						ANSICHAR command[1048]; 
+						FCStringAnsi::Sprintf(command, "r.setRes %dx%d ", w, h);
+
+						IConsoleManager::Get().ProcessUserConsoleInput(ANSI_TO_TCHAR(command), *GWarn, NULL );
+					}
+				}
+				break;
+#if PLATFORM_HTML5_WIN32
+			case SDL_ACTIVEEVENT:
+				{
+					EWindowActivation::Type ActivationType;
+					static bool DontSkip = true; 
+				
+					if (Event.active.gain && Event.active.state & (SDL_APPINPUTFOCUS)) {
+
+						ActivationType = EWindowActivation::Activate;
+						MessageHandler->OnWindowActivationChanged( ApplicationWindow, ActivationType );
+						WarmUpTicks = 0; 
 
-				IConsoleManager::Get().ProcessUserConsoleInput(ANSI_TO_TCHAR(command), *GWarn, NULL );
+					} 
+					else if (Event.active.gain && Event.active.state & (SDL_APPMOUSEFOCUS)) 
+					{
+						ActivationType = EWindowActivation::ActivateByMouse;
+						MessageHandler->OnWindowActivationChanged( ApplicationWindow, ActivationType );
+					}
+					else if ( !Event.active.gain && Event.active.state & (SDL_APPINPUTFOCUS ) )
+					{
+							ActivationType = EWindowActivation::Deactivate;
+							MessageHandler->OnWindowActivationChanged( ApplicationWindow, ActivationType );
+							WarmUpTicks = 0; 
+					}
+				}
+
+			break;
+#endif 
+
+#if PLATFORM_HTML5_BROWSER
+
+			case SDL_WINDOWEVENT: 
+				{
+					SDL_WindowEvent windowEvent = Event.window;
+
+					switch (windowEvent.event)
+					{
+                         // first two events don't exist in the SDK yet, see open pull request [REDACTED_PROFILE_URL]/emscripten/pull/2704
+						case SDL_WINDOWEVENT_ENTER:
+							{
+									MessageHandler->OnCursorSet();
+									MessageHandler->OnWindowActivationChanged(ApplicationWindow, EWindowActivation::ActivateByMouse);
+							}
+							break;
+
+						case SDL_WINDOWEVENT_LEAVE:
+							{
+									MessageHandler->OnWindowActivationChanged(ApplicationWindow, EWindowActivation::Deactivate);
+							}
+							break;
+
+						case SDL_WINDOWEVENT_FOCUS_GAINED:
+							{
+									MessageHandler->OnWindowActivationChanged(ApplicationWindow, EWindowActivation::Activate);
+									WarmUpTicks = 0;
+							}
+							break;
+
+						case SDL_WINDOWEVENT_FOCUS_LOST:
+							{
+									MessageHandler->OnWindowActivationChanged(ApplicationWindow, EWindowActivation::Deactivate);
+									WarmUpTicks = 0;
+							}
+							break;
+						default:
+							break;
+					}
+				}
+#endif
+			default:
+			{
+				InputInterface->Tick( TimeDelta,Event, ApplicationWindow);
 			}
 		}
-		else 
+	}
+	InputInterface->SendControllerEvents();
+
+
+	if ( WarmUpTicks >= 0)
+		WarmUpTicks ++; 
+
+
+	if ( WarmUpTicks == MaxWarmUpTicks  )
+	{
+        // browsers don't allow locking and hiding to work independently. use warmup ticks after the application has settled
+        // on its mouse lock/visibility status.  This is necessary even in cases where the game doesn't want to locking because 
+        // the lock status oscillates for few ticks before settling down. This causes a Browser UI pop even when we don't intend to lock.
+        // see http://www.w3.org/TR/pointerlock more for information. 
+#if PLATFORM_HTML5_WIN32
+		if (((FHTML5Cursor*)Cursor.Get())->LockStatus && !((FHTML5Cursor*)Cursor.Get())->CursorStatus)
+		{
+			SDL_WM_GrabInput(SDL_GRAB_ON);
+			SDL_ShowCursor(SDL_DISABLE);
+		}
+		else
+		{
+			SDL_ShowCursor(SDL_ENABLE);
+			SDL_WM_GrabInput(SDL_GRAB_OFF);
+		}
+#endif 
+
+#if PLATFORM_HTML5_BROWSER
+		if (((FHTML5Cursor*)Cursor.Get())->LockStatus && !((FHTML5Cursor*)Cursor.Get())->CursorStatus)
+		{
+			emscripten_request_pointerlock ( "#canvas" , true);
+		}
+		else
 		{
-			InputInterface->Tick( TimeDelta, Event );
+			emscripten_exit_pointerlock(); 
 		}
+#endif 
+
+		WarmUpTicks = -1; 
 	}
-	InputInterface->SendControllerEvents();
 }
 
 FPlatformRect FHTML5Application::GetWorkArea( const FPlatformRect& CurrentWindow ) const
diff --git a/Engine/Source/Runtime/Core/Private/HTML5/HTML5Cursor.cpp b/Engine/Source/Runtime/Core/Private/HTML5/HTML5Cursor.cpp
--- a/Engine/Source/Runtime/Core/Private/HTML5/HTML5Cursor.cpp
+++ b/Engine/Source/Runtime/Core/Private/HTML5/HTML5Cursor.cpp
@@ -7,8 +7,12 @@ FHTML5Cursor::FHTML5Cursor()
 {
 	// Set the default cursor
 	SetType( EMouseCursor::Default );
+
+	CursorStatus = true; 
+	LockStatus = false; 
 }
 
+
 FHTML5Cursor::~FHTML5Cursor()
 {
 }
@@ -25,6 +29,12 @@ void FHTML5Cursor::SetPosition( const int32 X, const int32 Y )
 
 void FHTML5Cursor::SetType( const EMouseCursor::Type InNewCursor )
 {
+	// TODO: use .css to correctly change cursor types in javascript. 
+	if ( InNewCursor == EMouseCursor::None)
+		CursorStatus = false; 
+	else 
+		CursorStatus  = true; 
+
 }
 
 void FHTML5Cursor::GetSize( int32& Width, int32& Height ) const
@@ -35,8 +45,17 @@ void FHTML5Cursor::GetSize( int32& Width, int32& Height ) const
 
 void FHTML5Cursor::Show( bool bShow )
 {
+	CursorStatus = bShow;
 }
 
 void FHTML5Cursor::Lock( const RECT* const Bounds )
 {
+	if(Bounds == NULL)
+	{
+		LockStatus = false;
+	}
+	else
+	{
+		LockStatus = true; 
+	}
 }
diff --git a/Engine/Source/Runtime/Core/Private/HTML5/HTML5InputInterface.cpp b/Engine/Source/Runtime/Core/Private/HTML5/HTML5InputInterface.cpp
--- a/Engine/Source/Runtime/Core/Private/HTML5/HTML5InputInterface.cpp
+++ b/Engine/Source/Runtime/Core/Private/HTML5/HTML5InputInterface.cpp
@@ -8,19 +8,6 @@
 DECLARE_LOG_CATEGORY_EXTERN(LogHTML5, Log, All);
 DEFINE_LOG_CATEGORY(LogHTML5);
 
-#if PLATFORM_HTML5_WIN32 && SDL_MAJOR_VERSION < 2
-static void GrabMouse(bool grab)
-{
-	if (grab) {
-		SDL_WM_GrabInput(SDL_GRAB_ON);
-		SDL_ShowCursor(SDL_DISABLE);
-	} else {
-		SDL_ShowCursor(SDL_ENABLE);
-		SDL_WM_GrabInput(SDL_GRAB_OFF);
-	}
-}
-#endif
-
 TSharedRef< FHTML5InputInterface > FHTML5InputInterface::Create(  const TSharedRef< FGenericApplicationMessageHandler >& InMessageHandler, const TSharedPtr< ICursor >& InCursor )
 {
 	return MakeShareable( new FHTML5InputInterface( InMessageHandler, InCursor ) );
@@ -43,7 +30,7 @@ void FHTML5InputInterface::SetMessageHandler( const TSharedRef< FGenericApplicat
 	MessageHandler = InMessageHandler;
 }
 
-void FHTML5InputInterface::Tick(float DeltaTime, const SDL_Event& Event)
+void FHTML5InputInterface::Tick(float DeltaTime, const SDL_Event& Event,TSharedRef < FGenericWindow>& ApplicationWindow )
 {
 #if PLATFORM_HTML5_WIN32 && SDL_MAJOR_VERSION < 2
 	// we have to emulate repeat
@@ -58,14 +45,6 @@ void FHTML5InputInterface::Tick(float DeltaTime, const SDL_Event& Event)
 	{
 		if (Event.type == SDL_KEYDOWN)
 		{
-#if PLATFORM_HTML5_WIN32 && SDL_MAJOR_VERSION < 2
-			isRepeat = (lastKey == Event.key.keysym.sym);
-			lastKey = Event.key.keysym.sym;
-
-			if (!isRepeat && Event.key.keysym.sym == SDLK_ESCAPE) {
-				GrabMouse(false);
-			}
-#endif
 			MessageHandler->OnKeyDown(Event.key.keysym.sym, Event.key.keysym.sym, EVENT_KEY_REPEAT);
 			if (!(Event.key.keysym.mod & KMOD_CTRL) && !(Event.key.keysym.mod & KMOD_ALT) && !(Event.key.keysym.mod & KMOD_GUI))
 			{
@@ -82,7 +61,7 @@ void FHTML5InputInterface::Tick(float DeltaTime, const SDL_Event& Event)
 		else if (Event.type == SDL_MOUSEBUTTONDOWN)
 		{
 			EMouseButtons::Type MouseButton = Event.button.button == 1 ? EMouseButtons::Left : Event.button.button == 2 ? EMouseButtons::Middle : EMouseButtons::Right;
-			MessageHandler->OnMouseDown(NULL, MouseButton );
+			MessageHandler->OnMouseDown(ApplicationWindow, MouseButton );
 		}
 		else if (Event.type == SDL_MOUSEBUTTONUP)
 		{
@@ -92,7 +71,7 @@ void FHTML5InputInterface::Tick(float DeltaTime, const SDL_Event& Event)
 		else if (Event.type == SDL_MOUSEMOTION)
 		{
 			Cursor->SetPosition(Event.motion.x, Event.motion.y);
-			MessageHandler->OnRawMouseMove(Event.motion.xrel, -Event.motion.yrel);
+			MessageHandler->OnRawMouseMove(Event.motion.xrel, Event.motion.yrel);
 			MessageHandler->OnMouseMove(); 
 		} 
 #if PLATFORM_HTML5_BROWSER 
@@ -103,16 +82,7 @@ void FHTML5InputInterface::Tick(float DeltaTime, const SDL_Event& Event)
 			MessageHandler->OnMouseWheel(w->y * SpinFactor);
 		}
 #endif 
-#if PLATFORM_HTML5_WIN32 && SDL_MAJOR_VERSION < 2
-		else if (Event.type == SDL_ACTIVEEVENT)
-		{
-			if (Event.active.gain && Event.active.state & (SDL_APPINPUTFOCUS | SDL_APPACTIVE)) {
-				GrabMouse(true);
-			} else {
-				GrabMouse(false);
-			}
-		}
-#endif
+
 	}
 }
 
diff --git a/Engine/Source/Runtime/Core/Public/HTML5/HTML5Cursor.h b/Engine/Source/Runtime/Core/Public/HTML5/HTML5Cursor.h
--- a/Engine/Source/Runtime/Core/Public/HTML5/HTML5Cursor.h
+++ b/Engine/Source/Runtime/Core/Public/HTML5/HTML5Cursor.h
@@ -27,4 +27,8 @@ public:
 
 private:
 	FVector2D Position;
+	bool CursorStatus; 
+	bool LockStatus; 
+
+	friend class FHTML5Application;
 };
\ No newline at end of file
diff --git a/Engine/Source/Runtime/Core/Public/HTML5/HTML5InputInterface.h b/Engine/Source/Runtime/Core/Public/HTML5/HTML5InputInterface.h
--- a/Engine/Source/Runtime/Core/Public/HTML5/HTML5InputInterface.h
+++ b/Engine/Source/Runtime/Core/Public/HTML5/HTML5InputInterface.h
@@ -17,10 +17,10 @@ public:
 
 	~FHTML5InputInterface() {}
 
-	void SetMessageHandler( const TSharedRef< FGenericApplicationMessageHandler >& InMessageHandler );
+	void SetMessageHandler( const TSharedRef< FGenericApplicationMessageHandler >& InMessageHandler);
 
 	/** Tick the interface (i.e check for new controllers) */
-	void Tick( float DeltaTime, const SDL_Event& Event );
+	void Tick( float DeltaTime, const SDL_Event& Event,TSharedRef < FGenericWindow>& ApplicationWindow );
 	/**
 	 * Poll for controller state and send events if needed
 	 */
diff --git a/Engine/Source/Runtime/Core/Public/HTML5/HTML5PlatformProperties.h b/Engine/Source/Runtime/Core/Public/HTML5/HTML5PlatformProperties.h
--- a/Engine/Source/Runtime/Core/Public/HTML5/HTML5PlatformProperties.h
+++ b/Engine/Source/Runtime/Core/Public/HTML5/HTML5PlatformProperties.h
@@ -77,4 +77,10 @@ struct FHTML5PlatformProperties
 	{
 		return true;
 	}
+
+	static FORCEINLINE bool SupportsWindowedMode()
+	{
+		return true;
+	}
+
 };
