CIS Fix.  Missing file from CL 2274997.

diff --git a/Engine/Source/Runtime/Core/Public/HTML5/HTML5Application.h b/Engine/Source/Runtime/Core/Public/HTML5/HTML5Application.h
--- a/Engine/Source/Runtime/Core/Public/HTML5/HTML5Application.h
+++ b/Engine/Source/Runtime/Core/Public/HTML5/HTML5Application.h
@@ -35,4 +35,6 @@ private:
 
 	TSharedPtr< class FHTML5InputInterface > InputInterface;
 	TSharedRef< class FGenericWindow > ApplicationWindow;
+
+	int32 WarmUpTicks; 
 };
\ No newline at end of file
