better cvar help

diff --git a/Engine/Source/Runtime/Core/Private/HAL/ConsoleManager.cpp b/Engine/Source/Runtime/Core/Private/HAL/ConsoleManager.cpp
--- a/Engine/Source/Runtime/Core/Private/HAL/ConsoleManager.cpp
+++ b/Engine/Source/Runtime/Core/Private/HAL/ConsoleManager.cpp
@@ -1825,7 +1825,8 @@ static TAutoConsoleVariable<int32> CVarDBuffer(
 	TEXT("r.DBuffer"),
 	0,
 	TEXT("Experimental DBuffer feature: Generating deferred decals before the BasePass.\n")
-	TEXT("Allows decals to be correctly lit by baked lighting.\n")
+	TEXT("Allows decals to be correctly lit by baked lighting. Receivers need to be rendered in the early zpass.\n")
+	TEXT("At the moment only can be ensures by full enablng this pass: r.EarlyZPassMovable=1 r.EarlyZPass=2\n")
 	TEXT(" 0: off\n")
 	TEXT(" 1: on (needs early pass rendering on all decal receivers and base pass lookups into the DBuffer, costs GPU memory, allows GBuffer compression)"),
 	ECVF_Scalability | ECVF_RenderThreadSafe);
