Correcting spelling error in a comment.
https://github.com/EpicGames/UnrealEngine/pull/401

diff --git a/Engine/Source/Runtime/Core/Public/Containers/Map.h b/Engine/Source/Runtime/Core/Public/Containers/Map.h
--- a/Engine/Source/Runtime/Core/Public/Containers/Map.h
+++ b/Engine/Source/Runtime/Core/Public/Containers/Map.h
@@ -456,7 +456,7 @@ public:
 	}
 
 	/**
-	 * Returns a reference to the alue associated with a specified key.
+	 * Returns a reference to the value associated with a specified key.
 	 * @param	Key - The key to search for.
 	 * @return	The value associated with the specified key, or triggers an assertion if the key does not exist.
 	 */
