Added FMacPlatformMisc::IsRunningOnMavericks() and used it to launch atos tool with different params depending on whether the editor is running on Mavericks or Yosemite. Fixes TTP #345031.

diff --git a/Engine/Source/Editor/UnrealEd/Private/SourceCodeNavigation.cpp b/Engine/Source/Editor/UnrealEd/Private/SourceCodeNavigation.cpp
--- a/Engine/Source/Editor/UnrealEd/Private/SourceCodeNavigation.cpp
+++ b/Engine/Source/Editor/UnrealEd/Private/SourceCodeNavigation.cpp
@@ -509,7 +509,8 @@ void FSourceCodeNavigationImpl::NavigateToFunctionSource( const FString& Functio
 							if(FunctionSymbolName == SymbolName)
 							{
 								uint64 Address = SymbolEntry.n_value;
-								FString AtoSCommand = FString::Printf(TEXT("-nowarning -arch x86_64 -d -o \"%s\" 0x%x"), *FullModulePath, Address);
+								// Mavericks requires additional params to atos to silence a deprecation warning. On Yosemite same params cause "unrecognized option" error
+								FString AtoSCommand = FMacPlatformMisc::IsRunningOnMavericks() ? FString::Printf(TEXT("-nowarning -arch x86_64 -d -o \"%s\" 0x%x"), *FullModulePath, Address) : FString::Printf(TEXT("-arch x86_64 -o \"%s\" 0x%x"), *FullModulePath, Address);
 								int32 ReturnCode = 0;
 								FString Results;
 								FPlatformProcess::ExecProcess( TEXT("/usr/bin/atos"), *AtoSCommand, &ReturnCode, &Results, NULL );
diff --git a/Engine/Source/Runtime/Core/Private/Apple/ApplePlatformStackWalk.cpp b/Engine/Source/Runtime/Core/Private/Apple/ApplePlatformStackWalk.cpp
--- a/Engine/Source/Runtime/Core/Private/Apple/ApplePlatformStackWalk.cpp
+++ b/Engine/Source/Runtime/Core/Private/Apple/ApplePlatformStackWalk.cpp
@@ -178,7 +178,15 @@ void FApplePlatformStackWalk::ProgramCounterToSymbolInfo( uint64 ProgramCounter,
 			FCStringAnsi::Sprintf(FBase, "%p", (void*)DylibInfo.dli_fbase);
 			FCStringAnsi::Sprintf(Address, "%p", (void*)ProgramCounter);
 			
-			execl("/usr/bin/atos", "-nowarning", "-arch", "x86_64", "-d", "-l", FBase, "-o", DylibInfo.dli_fname, Address, NULL);
+			// Mavericks requires additional params to atos to silence a deprecation warning. On Yosemite same params cause "unrecognized option" error
+			if (FPlatformMisc::IsRunningOnMavericks())
+			{
+				execl("/usr/bin/atos", "-nowarning", "-arch", "x86_64", "-d", "-l", FBase, "-o", DylibInfo.dli_fname, Address, NULL);
+			}
+			else
+			{
+				execl("/usr/bin/atos", "-arch", "x86_64", "-l", FBase, "-o", DylibInfo.dli_fname, Address, NULL);
+			}
 			close(FileDesc[1]);
 		}
 		else
diff --git a/Engine/Source/Runtime/Core/Private/Mac/MacPlatformMisc.cpp b/Engine/Source/Runtime/Core/Private/Mac/MacPlatformMisc.cpp
--- a/Engine/Source/Runtime/Core/Private/Mac/MacPlatformMisc.cpp
+++ b/Engine/Source/Runtime/Core/Private/Mac/MacPlatformMisc.cpp
@@ -63,6 +63,15 @@ struct MacApplicationInfo
 		FCStringAnsi::Strcpy(OSVersionUTF8, PATH_MAX+1, [[SystemVersion objectForKey: @"ProductVersion"] UTF8String]);
 		OSBuild = FString((NSString*)[SystemVersion objectForKey: @"ProductBuildVersion"]);
 		
+		TArray<FString> Components;
+		OSVersion.ParseIntoArray(&Components, TEXT("."), true);
+		uint8 ComponentValues[2] = {0};
+		for(uint32 i = 0; i < Components.Num() && i < 2; i++)
+		{
+			TTypeFromString<uint8>::FromString(ComponentValues[i], *Components[i]);
+		}
+		RunningOnMavericks == ComponentValues[0] == 10 && ComponentValues[1] == 9;
+		
 		char TempSysCtlBuffer[PATH_MAX] = {};
 		size_t TempSysCtlBufferSize = PATH_MAX;
 		
@@ -177,6 +186,7 @@ struct MacApplicationInfo
 	}
 	
 	bool RunningOnBattery;
+	bool RunningOnMavericks;
 	int32 PowerSourceNotification;
 	char AppNameUTF8[PATH_MAX+1];
 	char AppLogPath[PATH_MAX+1];
@@ -1039,6 +1049,11 @@ bool FMacPlatformMisc::IsRunningOnBattery()
 	return GMacAppInfo.RunningOnBattery;
 }
 
+bool FMacPlatformMisc::IsRunningOnMavericks()
+{
+	return GMacAppInfo.RunningOnMavericks;
+}
+
 /** Global pointer to crash handler */
 void (* GCrashHandlerPointer)(const FGenericCrashContext & Context) = NULL;
 
diff --git a/Engine/Source/Runtime/Core/Public/Mac/MacPlatformMisc.h b/Engine/Source/Runtime/Core/Public/Mac/MacPlatformMisc.h
--- a/Engine/Source/Runtime/Core/Public/Mac/MacPlatformMisc.h
+++ b/Engine/Source/Runtime/Core/Public/Mac/MacPlatformMisc.h
@@ -140,6 +140,11 @@ struct CORE_API FMacPlatformMisc : public FGenericPlatformMisc
 	 * Returns whether the platform is running on battery power or not.
 	 */
 	static bool IsRunningOnBattery();
+
+	/**
+	 * Returns whether the Mac OS X version is 10.9.x or not.
+	 */
+	static bool IsRunningOnMavericks();
 };
 
 #ifdef __OBJC__
