Change 2122846 on 2014/07/01 00:51:24 by [REDACTED_PERSON][REDACTED_HANDLE]

diff --git a/Engine/Source/Runtime/Core/Public/Math/TransformCalculus2D.h b/Engine/Source/Runtime/Core/Public/Math/TransformCalculus2D.h
new file mode 100644
--- /dev/null
+++ b/Engine/Source/Runtime/Core/Public/Math/TransformCalculus2D.h
@@ -0,0 +1,690 @@
+// Copyright 1998-2014 Epic Games, Inc. All Rights Reserved.
+
+#pragma once
+
+#include "TransformCalculus.h"
+
+//////////////////////////////////////////////////////////////////////////
+// Transform calculus for 2D types. UE4 already has a 2D Vector class that we
+// will adapt to be interpreted as a translate transform. The rest we create
+// new classes for.
+//
+// The following types are supported
+// * float           -> represents a uniform scale.
+// * FScale2D        -> represents a 2D non-uniform scale.
+// * FVector2D       -> represents a 2D translation.
+// * FShear2D        -> represents a "2D shear", interpreted as a shear parallel to the X axis followed by a shear parallel to the Y axis.
+// * FQuat2D         -> represents a pure 2D rotation.
+// * FMatrix2x2      -> represents a general 2D transform.
+//
+//////////////////////////////////////////////////////////////////////////
+
+class FMatrix2x2;
+
+//////////////////////////////////////////////////////////////////////////
+// Adapters for FVector2D. 
+// 
+// Since it is an existing UE4 types, we cannot rely on the default
+// template that calls member functions. Instead, we provide direct overloads.
+//////////////////////////////////////////////////////////////////////////
+
+/** Specialization for concatenating two 2D Translations. */
+inline FVector2D Concatenate(const FVector2D& LHS, const FVector2D& RHS)
+{
+	return LHS + RHS;
+}
+
+/** Specialization for inverting a 2D translation. */
+inline FVector2D Inverse(const FVector2D& Transform)
+{
+	return -Transform;
+}
+
+/** Specialization for FVector2D Translation. */
+inline FVector2D TransformPoint(const FVector2D& Transform, const FVector2D& Point)
+{
+	return Transform + Point;
+}
+
+/** Specialization for FVector2D Translation (does nothing). */
+inline const FVector2D& TransformVector(const FVector2D& Transform, const FVector2D& Vector)
+{
+	return Vector;
+}
+
+//////////////////////////////////////////////////////////////////////////
+// Adapters for 2D uniform scale.
+// 
+// Since it is a fundamental type, we cannot rely on the default
+// template that calls member functions. Instead, we provide direct overloads.
+//////////////////////////////////////////////////////////////////////////
+
+/**
+ * Specialization for uniform Scale.
+ */
+inline FVector2D TransformPoint(float Transform, const FVector2D& Point)
+{
+	return Transform * Point;
+}
+
+/**
+ * Specialization for uniform Scale.
+ */
+inline FVector2D TransformVector(float Transform, const FVector2D& Vector)
+{
+	return Transform * Vector;
+}
+
+/** Represents a 2D non-uniform scale (to disambiguate from an FVector2D, which is used for translation) */
+class FScale2D
+{
+public:
+	/** Ctor. initialize to an identity scale, 1.0. */
+	FScale2D() :Scale(1.0f, 1.0f) {}
+	/** Ctor. initialize from a uniform scale. */
+	explicit FScale2D(float InScale) :Scale(InScale, InScale) {}
+	/** Ctor. initialize from a non-uniform scale. */
+	explicit FScale2D(float InScaleX, float InScaleY) :Scale(InScaleX, InScaleY) {}
+	/** Ctor. initialize from an FVector defining the 3D scale. */
+	explicit FScale2D(const FVector2D& InScale) :Scale(InScale) {}
+	
+	/** Transform 2D Point */
+	FVector2D TransformPoint(const FVector2D& Point) const
+	{
+		return Scale * Point;
+	}
+	/** Transform 2D Vector*/
+	FVector2D TransformVector(const FVector2D& Vector) const
+	{
+		return TransformPoint(Vector);
+	}
+
+	/** Concatenate two scales. */
+	FScale2D Concatenate(const FScale2D& RHS) const
+	{
+		return FScale2D(Scale * RHS.Scale);
+	}
+	/** Invert the scale. */
+	FScale2D Inverse() const
+	{
+		return FScale2D(FVector2D(1.0f / Scale.X, 1.0f / Scale.Y));
+	}
+
+	/** Equality. */
+	bool operator==(const FScale2D& Other) const
+	{
+		return Scale == Other.Scale;
+	}
+	
+	/** Inequality. */
+	bool operator!=(const FScale2D& Other) const
+	{
+		return !operator==(Other);
+	}
+
+	/** Access to the underlying FVector2D that stores the scale. */
+	const FVector2D& GetVector() const { return Scale; }
+private:
+	/** Underlying storage of the 2D scale. */
+	FVector2D Scale;
+};
+
+/** concatenation rules for 2D scales. */
+template<> struct ConcatenateRules<float          , FScale2D  > { typedef FScale2D ResultType; };
+/** concatenation rules for 2D scales. */
+template<> struct ConcatenateRules<FScale2D  , float          > { typedef FScale2D ResultType; };
+
+/** 
+ * Represents a 2D shear:
+ *   [1 YY]
+ *   [XX 1]
+ * XX represents a shear parallel to the X axis. YY represents a shear parallel to the Y axis.
+ */
+class FShear2D
+{
+public:
+	/** Ctor. initialize to an identity. */
+	FShear2D() :Shear(0,0) {}
+	/** Ctor. initialize from a set of shears parallel to the X and Y axis, respectively. */
+	explicit FShear2D(float ShearX, float ShearY) :Shear(ShearX, ShearY) {}
+	/** Ctor. initialize from a 2D vector representing a set of shears parallel to the X and Y axis, respectively. */
+	explicit FShear2D(const FVector2D& InShear) :Shear(InShear) {}
+	
+	/** 
+	 * Transform 2D Point 
+	 * [X Y] * [1 YY] == [X+Y*XX Y+X*YY]
+	 *         [XX 1]
+	 */
+	FVector2D TransformPoint(const FVector2D& Point) const
+	{
+		return Point + FVector2D(Point.Y, Point.X) * Shear;
+	}
+	/** Transform 2D Vector*/
+	FVector2D TransformVector(const FVector2D& Vector) const
+	{
+		return TransformPoint(Vector);
+	}
+
+	/**
+	 * Concatenate two shears. The result is NOT a shear, but must be represented by a generalized 2x2 transform.
+	 * Defer the implementation until we can declare a 2x2 matrix.
+	 * [1 YYA] * [1 YYB] == [1+YYA*XXB YYB*YYA]
+	 * [XXA 1]   [XXB 1]    [XXA+XXB XXA*XXB+1]
+	 */
+	inline FMatrix2x2 Concatenate(const FShear2D& RHS) const;
+
+	/**
+	 * Invert the shear. The result is NOT a shear, but must be represented by a generalized 2x2 transform.
+	 * Defer the implementation until we can declare a 2x2 matrix.
+	 * [1 YY]^-1  == 1/(1-YY*XX) * [1 -YY]
+	 * [XX 1]                      [-XX 1]
+	 */
+	FMatrix2x2 Inverse() const;
+
+
+	/** Equality. */
+	bool operator==(const FShear2D& Other) const
+	{
+		return Shear == Other.Shear;
+	}
+	
+	/** Inequality. */
+	bool operator!=(const FShear2D& Other) const
+	{
+		return !operator==(Other);
+	}
+
+	/** Access to the underlying FVector2D that stores the scale. */
+	const FVector2D& GetVector() const { return Shear; }
+private:
+	/** Underlying storage of the 2D shear. */
+	FVector2D Shear;
+};
+
+/** 
+ * Represents a 2D rotation as a complex number (analagous to quaternions). 
+ *   Rot(theta) == cos(theta) + i * sin(theta)
+ *   General transformation follows complex number algebra from there.
+ * Does not use "spinor" notation using theta/2 as we don't need that decomposition for our purposes.
+ * This makes the implementation for straightforward and efficient for 2D.
+ */
+class FQuat2D
+{
+public:
+	/** Ctor. initialize to an identity rotation. */
+	FQuat2D() :Rot(1.0f, 0.0f) {}
+	/** Ctor. initialize from a rotation in radians. */
+	explicit FQuat2D(float RotRadians) :Rot(FMath::Cos(RotRadians), FMath::Sin(RotRadians)) {}
+	/** Ctor. initialize from an FVector2D, representing a complex number. */
+	explicit FQuat2D(const FVector2D& InRot) :Rot(InRot) {}
+
+	/**
+	 * Transform a 2D point by the 2D complex number representing the rotation:
+	 * In imaginary land: (x + yi) * (u + vi) == (xu - yv) + (xv + yu)i
+	 * 
+	 * Looking at this as a matrix, x == cos(A), y == sin(A)
+	 * 
+	 * [x y] * [ cosA  sinA] == [x y] * [ u v] == [xu-yv xv+yu]
+	 *         [-sinA  cosA]            [-v u]
+	 *         
+	 * Looking at the above results, we see the equivalence with matrix multiplication.
+	 */
+	FVector2D TransformPoint(const FVector2D& Point) const
+	{
+		return FVector2D(
+			Point.X * Rot.X - Point.Y * Rot.Y,
+			Point.X * Rot.Y + Point.Y * Rot.X);
+	}
+	/**
+	 * Vector rotation is equivalent to rotating a point.
+	 */
+	FVector2D TransformVector(const FVector2D& Vector) const
+	{
+		return TransformPoint(Vector);
+	}
+	/**
+	 * Transform 2 rotations defined by complex numbers:
+	 * In imaginary land: (A + Bi) * (C + Di) == (AC - BD) + (AD + BC)i
+	 * 
+	 * Looking at this as a matrix, A == cos(theta), B == sin(theta), C == cos(sigma), D == sin(sigma):
+	 * 
+	 * [ A B] * [ C D] == [  AC-BD  AD+BC]
+	 * [-B A]   [-D C]    [-(AD+BC) AC-BD]
+	 * 
+	 * If you look at how the vector multiply works out: [X(AC-BD)+Y(-BC-AD)  X(AD+BC)+Y(-BD+AC)]
+	 * you can see it follows the same form of the imaginary form. Indeed, check out how the matrix nicely works
+	 * out to [ A B] for a visual proof of the results.
+	 *        [-B A]
+	 */
+	FQuat2D Concatenate(const FQuat2D& RHS) const
+	{
+		return FQuat2D(TransformPoint(RHS.Rot));
+	}
+	/**
+	 * Invert the rotation  defined by complex numbers:
+	 * In imaginary land, an inverse is a complex conjugate, which is equivalent to reflecting about the X axis:
+	 * Conj(A + Bi) == A - Bi
+	 */
+	FQuat2D Inverse() const
+	{
+		return FQuat2D(FVector2D(Rot.X, -Rot.Y));
+	}
+
+	/** Equality. */
+	bool operator==(const FQuat2D& Other) const
+	{
+		return Rot == Other.Rot;
+	}
+	
+	/** Inequality. */
+	bool operator!=(const FQuat2D& Other) const
+	{
+		return !operator==(Other);
+	}
+
+	/** Access to the underlying FVector2D that stores the complex number. */
+	const FVector2D& GetVector() const { return Rot; }
+private:
+	/** Underlying storage of the rotation (X = cos(theta), Y = sin(theta). */
+	FVector2D Rot;
+};
+
+/**
+ * 2x2 generalized matrix. As FMatrix, we assume row vectors, row major storage:
+ *    [X Y] * [m00 m01]
+ *            [m10 m11]
+ */
+class FMatrix2x2
+{
+public:
+	/** Ctor. initialize to an identity. */
+	FMatrix2x2()
+	{
+		M[0][0] = 1; M[0][1] = 0;
+		M[1][0] = 0; M[1][1] = 1;
+	}
+
+	FMatrix2x2(float m00, float m01, float m10, float m11)
+	{
+		M[0][0] = m00; M[0][1] = m01;
+		M[1][0] = m10; M[1][1] = m11;
+	}
+
+
+	/** Ctor. initialize from a scale. */
+	explicit FMatrix2x2(float UniformScale)
+	{
+		M[0][0] = UniformScale; M[0][1] = 0;
+		M[1][0] = 0; M[1][1] = UniformScale;
+	}
+
+	/** Ctor. initialize from a scale. */
+	explicit FMatrix2x2(const FScale2D& Scale)
+	{
+		float ScaleX = Scale.GetVector().X;
+		float ScaleY = Scale.GetVector().Y;
+		M[0][0] = ScaleX; M[0][1] = 0;
+		M[1][0] = 0; M[1][1] = ScaleY;
+	}
+
+	/** Factory function. initialize from a 2D shear. */
+	explicit FMatrix2x2(const FShear2D& Shear)
+	{
+		float XX = Shear.GetVector().X;
+		float YY = Shear.GetVector().Y;
+		M[0][0] = 1; M[0][1] =YY;
+		M[1][0] =XX; M[1][1] = 1;
+	}
+
+	/** Ctor. initialize from a rotation. */
+	explicit FMatrix2x2(const FQuat2D& Rotation)
+	{
+		float CosAngle = Rotation.GetVector().X;
+		float SinAngle = Rotation.GetVector().Y;
+		M[0][0] = CosAngle; M[0][1] = SinAngle;
+		M[1][0] = -SinAngle; M[1][1] = CosAngle;
+	}
+
+	/**
+	 * Transform a 2D point
+	 *    [X Y] * [m00 m01]
+	 *            [m10 m11]
+	 */
+	FVector2D TransformPoint(const FVector2D& Point) const
+	{
+		return FVector2D(
+			Point.X * M[0][0] + Point.Y * M[1][0],
+			Point.X * M[0][1] + Point.Y * M[1][1]);
+	}
+	/**
+	 * Vector transformation is equivalent to point transformation as our matrix is not homogeneous.
+	 */
+	FVector2D TransformVector(const FVector2D& Vector) const
+	{
+		return TransformPoint(Vector);
+	}
+	/**
+	 * Concatenate 2 matrices:
+	 * [A B] * [E F] == [AE+BG AF+BH]
+	 * [C D]   [G H]    [CE+DG CF+DH]
+	 */
+	FMatrix2x2 Concatenate(const FMatrix2x2& RHS) const
+	{
+		float A, B, C, D;
+		GetMatrix(A, B, C, D);
+		float E, F, G, H;
+		RHS.GetMatrix(E, F, G, H);
+		return FMatrix2x2(
+			A*E + B*G, A*F + B*H,
+			C*E + D*G, C*F + D*H);
+	}
+	/**
+	 * Invert the transform.
+	 */
+	FMatrix2x2 Inverse() const
+	{
+		float A, B, C, D;
+		GetMatrix(A, B, C, D);
+		float InvDet = InverseDeterminant();
+		return FMatrix2x2(
+			 D*InvDet, -B*InvDet,
+			-C*InvDet,  A*InvDet);
+	}
+
+	/** Equality. */
+	bool operator==(const FMatrix2x2& RHS) const
+	{
+		float A, B, C, D;
+		GetMatrix(A, B, C, D);
+		float E, F, G, H;
+		RHS.GetMatrix(E, F, G, H);
+		return
+			FMath::IsNearlyEqual(A, E, KINDA_SMALL_NUMBER) &&
+			FMath::IsNearlyEqual(B, F, KINDA_SMALL_NUMBER) &&
+			FMath::IsNearlyEqual(C, G, KINDA_SMALL_NUMBER) &&
+			FMath::IsNearlyEqual(D, H, KINDA_SMALL_NUMBER);
+	}
+
+	/** Inequality. */
+	bool operator!=(const FMatrix2x2& Other) const
+	{
+		return !operator==(Other);
+	}
+
+	void GetMatrix(float &A, float &B, float &C, float &D) const
+	{
+		A = M[0][0]; B = M[0][1];
+		C = M[1][0]; D = M[1][1];
+	}
+
+	float Determinant() const
+	{
+		float A, B, C, D;
+		GetMatrix(A, B, C, D);
+		return (A*D - B*C);
+	}
+
+	float InverseDeterminant() const
+	{
+		float Det = Determinant();
+		checkSlow(Det != 0.0f);
+		return 1.0f / Det;
+	}
+
+	/** Extracts the squared scale from the matrix (avoids sqrt). */
+	FScale2D GetScaleSquared() const
+	{
+		float A, B, C, D;
+		GetMatrix(A, B, C, D);
+		return FScale2D(A*A + B*B, C*C + D*D);
+	}
+
+	/** Gets the squared scale from the matrix (avoids sqrt). */
+	FScale2D GetScale() const
+	{
+		FScale2D ScaleSquared = GetScaleSquared();
+		return FScale2D(FMath::Sqrt(ScaleSquared.GetVector().X), FMath::Sqrt(ScaleSquared.GetVector().Y));
+	}
+
+private:
+	float M[2][2];
+};
+
+inline FMatrix2x2 FShear2D::Concatenate(const FShear2D& RHS) const
+{
+	float XXA = Shear.X;
+	float YYA = Shear.Y;
+	float XXB = RHS.Shear.X;
+	float YYB = RHS.Shear.Y;
+	return FMatrix2x2(
+		1+YYA*XXB, YYB*YYA,
+		XXA+XXB, XXA*XXB+1);
+}
+
+inline FMatrix2x2 FShear2D::Inverse() const
+{
+	float InvDet = 1.0f / (1.0f - Shear.X*Shear.Y);
+	return FMatrix2x2(
+		InvDet, -Shear.Y * InvDet,
+		-Shear.X * InvDet, InvDet);
+}
+
+/** Partial specialization of ConcatenateRules for 2x2 and any other type via Upcast to 2x2 first. Requires a conversion ctor on FMatrix2x2. Funky template logic so we don't hide the default rule for NULL conversions. */
+template<typename TransformType> struct ConcatenateRules<typename TEnableIf<!TAreTypesEqual<FMatrix2x2, TransformType>::Value, FMatrix2x2>::Type, TransformType  > { typedef FMatrix2x2 ResultType; };
+/** Partial specialization of ConcatenateRules for 2x2 and any other type via Upcast to 2x2 first. Requires a conversion ctor on FMatrix2x2. Funky template logic so we don't hide the default rule for NULL conversions. */
+template<typename TransformType> struct ConcatenateRules<TransformType  , typename TEnableIf<!TAreTypesEqual<FMatrix2x2, TransformType>::Value, FMatrix2x2>::Type> { typedef FMatrix2x2 ResultType; };
+
+/** concatenation rules for 2x2 transform types. Convert to 2x2 matrix as the fully decomposed math is not that perf critical right now. */
+template<> struct ConcatenateRules<FScale2D  , FShear2D  > { typedef FMatrix2x2 ResultType; };
+template<> struct ConcatenateRules<FScale2D  , FQuat2D   > { typedef FMatrix2x2 ResultType; };
+template<> struct ConcatenateRules<FShear2D  , FScale2D  > { typedef FMatrix2x2 ResultType; };
+template<> struct ConcatenateRules<FQuat2D   , FScale2D  > { typedef FMatrix2x2 ResultType; };
+template<> struct ConcatenateRules<FShear2D  , FQuat2D   > { typedef FMatrix2x2 ResultType; };
+template<> struct ConcatenateRules<FQuat2D   , FShear2D  > { typedef FMatrix2x2 ResultType; };
+
+ /**
+ * Support for generalized 2D affine transforms. 
+ * Implemented as a 2x2 transform followed by translation. In matrix form:
+ *   [A B 0]
+ *   [C D 0]
+ *   [X Y 1]
+ */
+class FTransform2D
+{
+public:
+	/** Initialize the transform using an identity matrix and a translation. */
+	FTransform2D(const FVector2D& Translation = FVector2D(0.f,0.f))
+		: Trans(Translation)
+	{
+	}
+
+	/** Initialize the transform using a uniform scale and a translation. */
+	explicit FTransform2D(float UniformScale, const FVector2D& Translation = FVector2D(0.f,0.f))
+		: M(FScale2D(UniformScale)), Trans(Translation)
+	{
+	}
+
+	/** Initialize the transform using a 2D scale and a translation. */
+	explicit FTransform2D(const FScale2D& Scale, const FVector2D& Translation = FVector2D(0.f,0.f))
+		: M(Scale), Trans(Translation)
+	{
+	}
+
+	/** Initialize the transform using a 2D shear and a translation. */
+	explicit FTransform2D(const FShear2D& Shear, const FVector2D& Translation = FVector2D(0.f,0.f))
+		: M(Shear), Trans(Translation)
+	{
+	}
+
+	/** Initialize the transform using a 2D rotation and a translation. */
+	explicit FTransform2D(const FQuat2D& Rot, const FVector2D& Translation = FVector2D(0.f,0.f))
+		: M(Rot), Trans(Translation)
+	{
+	}
+
+	/** Initialize the transform using a general 2x2 transform and a translation. */
+	explicit FTransform2D(const FMatrix2x2& Transform, const FVector2D& Translation = FVector2D(0.f,0.f))
+		: M(Transform), Trans(Translation)
+	{
+	}
+
+	/** 2D transformation of a point. */
+	FVector2D TransformPoint(const FVector2D& Point) const
+	{
+		return ::TransformPoint(Trans, ::TransformPoint(M, Point));
+	}
+
+	/** 2D transformation of a vector. */
+	FVector2D TransformVector(const FVector2D& Vector) const
+	{
+		return ::TransformVector(M, Vector);
+	}
+
+	/** 
+	 * Concatenates two transforms. Result is equivalent to transforming first by this, followed by RHS.
+	 *  Concat(A,B) == (P * MA + TA) * MB + TB
+	 *              == (P * MA * MB) + TA*MB + TB
+	 *  NewM == MA * MB
+	 *  NewT == TA * MB + TB
+	 */
+	FTransform2D Concatenate(const FTransform2D& RHS) const
+	{
+		return FTransform2D(
+			::Concatenate(M, RHS.M),
+			::Concatenate(::TransformPoint(RHS.M, Trans), RHS.Trans));
+	}
+
+	/** 
+	 * Inverts a transform. So a transform from space A to space B results in a transform from space B to space A. 
+	 * Since this class applies the 2x2 transform followed by translation, our inversion logic needs to be able to recast
+	 * the result as a M * T. It does it using the following identity:
+	 *   (M * T)^-1 == T^-1 * M^-1
+	 *   
+	 * In homogeneous form, we represent our affine transform like so:
+	 *      M    *    T
+	 *   [A B 0]   [1 0 0]   [A B 0]
+	 *   [C D 0] * [0 1 0] = [C D 0]. This class simply decomposes the 2x2 transform and translation.
+	 *   [0 0 1]   [X Y 1]   [X Y 1]
+	 * 
+	 * But if we were applying the transforms in reverse order (as we need to for the inverse identity above):
+	 *    T^-1   *  M^-1
+	 *   [1 0 0]   [A B 0]   [A  B  0]  where [X' Y'] = [X Y] * [A B]
+	 *   [0 1 0] * [C D 0] = [C  D  0]                          [C D]
+	 *   [X Y 1]   [0 0 1]   [X' Y' 1]
+	 * 
+	 * This can be conceptualized by seeing that a translation effectively defines a new local origin for that 
+	 * frame of reference. Since there is a 2x2 transform AFTER that, the concatenated frame of reference has an origin
+	 * that is the old origin transformed by the 2x2 transform.
+	 * 
+	 * In the last equation:
+	 * We know that [X Y] is the translation induced by inverting T, or -Translate.
+	 * We know that [[A B][C D]] == Inverse(M), so we can represent T^-1 * M^-1 as M'* T' where:
+	 *   M' == Inverse(M)
+	 *   T' == Inverse(Translate) * Inverse(M)
+	 */
+	FTransform2D Inverse() const
+	{
+		FMatrix2x2 InvM = ::Inverse(M);
+		FVector2D InvTrans = ::TransformPoint(InvM, ::Inverse(Trans));
+		return FTransform2D(InvM, InvTrans);
+	}
+	
+	/** Equality. */
+	bool operator==(const FTransform2D& Other) const
+	{
+		return M == Other.M && Trans == Other.Trans;
+	}
+	
+	/** Inequality. */
+	bool operator!=(const FTransform2D& Other) const
+	{
+		return !operator==(Other);
+	}
+
+	/** Access to the 2x2 transform */
+	const FMatrix2x2& GetMatrix() const { return M; }
+	/** Access to the translation */
+	const FVector2D& GetTranslation() const { return Trans; }
+
+private:
+	FMatrix2x2 M;
+	FVector2D Trans;
+};
+
+template<> struct TIsPODType<FTransform2D> { enum { Value = true }; };
+
+//////////////////////////////////////////////////////////////////////////
+// Concatenate overloads. 
+// 
+// Efficient overloads for concatenating 2D affine transforms.
+// Better than simply upcasting both to FTransform2D first.
+//////////////////////////////////////////////////////////////////////////
+
+/** Specialization for concatenating a 2D scale and 2D Translation. */
+inline FTransform2D Concatenate(const FScale2D& Scale, const FVector2D& Translation)
+{
+	return FTransform2D(Scale, Translation);
+}
+
+/** Specialization for concatenating a 2D shear and 2D Translation. */
+inline FTransform2D Concatenate(const FShear2D& Shear, const FVector2D& Translation)
+{
+	return FTransform2D(Shear, Translation);
+}
+
+/** Specialization for concatenating 2D Rotation and 2D Translation. */
+inline FTransform2D Concatenate(const FQuat2D& Rot, const FVector2D& Translation)
+{
+	return FTransform2D(Rot, Translation);
+}
+
+/** Specialization for concatenating 2D generalized transform and 2D Translation. */
+inline FTransform2D Concatenate(const FMatrix2x2& Transform, const FVector2D& Translation)
+{
+	return FTransform2D(Transform, Translation);
+}
+
+/** Specialization for concatenating transform and 2D Translation. */
+inline FTransform2D Concatenate(const FTransform2D& Transform, const FVector2D& Translation)
+{
+	return FTransform2D(Transform.GetMatrix(), Concatenate(Transform.GetTranslation(), Translation));
+}
+
+/** Specialization for concatenating a 2D Translation and 2D scale. */
+inline FTransform2D Concatenate(const FVector2D& Translation, const FScale2D& Scale)
+{
+	return FTransform2D(Scale, ::TransformPoint(Scale, Translation));
+}
+
+/** Specialization for concatenating a 2D Translation and 2D shear. */
+inline FTransform2D Concatenate(const FVector2D& Translation, const FShear2D& Shear)
+{
+	return FTransform2D(Shear, ::TransformPoint(Shear, Translation));
+}
+
+/** Specialization for concatenating 2D Translation and 2D Rotation. */
+inline FTransform2D Concatenate(const FVector2D& Translation, const FQuat2D& Rot)
+{
+	return FTransform2D(Rot, ::TransformPoint(Rot, Translation));
+}
+
+/** Specialization for concatenating 2D Translation and 2D generalized transform. See docs for FTransform2D::Inverse for details on how this math is derived. */
+inline FTransform2D Concatenate(const FVector2D& Translation, const FMatrix2x2& Transform)
+{
+	return FTransform2D(Transform, ::TransformPoint(Transform, Translation));
+}
+
+/** Specialization for concatenating 2D Translation and transform. See docs for FTransform2D::Inverse for details on how this math is derived. */
+inline FTransform2D Concatenate(const FVector2D& Translation, const FTransform2D& Transform)
+{
+	return FTransform2D(Transform.GetMatrix(), Concatenate(::TransformPoint(Transform.GetMatrix(), Translation), Transform.GetTranslation()));
+}
+
+/** Partial specialization of ConcatenateRules for FTransform2D and any other type via Upcast to FTransform2D first. Requires a conversion ctor on FTransform2D. Funky template logic so we don't hide the default rule for NULL conversions. */
+template<typename TransformType> struct ConcatenateRules<typename TEnableIf<!TAreTypesEqual<FTransform2D, TransformType>::Value, FTransform2D>::Type, TransformType  > { typedef FTransform2D ResultType; };
+/** Partial specialization of ConcatenateRules for FTransform2D and any other type via Upcast to FTransform2D first. Requires a conversion ctor on FTransform2D. Funky template logic so we don't hide the default rule for NULL conversions. */
+template<typename TransformType> struct ConcatenateRules<TransformType  , typename TEnableIf<!TAreTypesEqual<FTransform2D, TransformType>::Value, FTransform2D>::Type> { typedef FTransform2D ResultType; };
+
+/** Provide a disambiguating overload for 2x2 and FTransform2D, since both try to provide a generic set of ConcatenateRules for all types. */
+template<> struct ConcatenateRules<FMatrix2x2, FTransform2D> { typedef FTransform2D ResultType; };
+template<> struct ConcatenateRules<FTransform2D, FMatrix2x2> { typedef FTransform2D ResultType; };
diff --git a/Engine/Source/Runtime/Core/Public/Math/TransformCalculus3D.h b/Engine/Source/Runtime/Core/Public/Math/TransformCalculus3D.h
new file mode 100644
--- /dev/null
+++ b/Engine/Source/Runtime/Core/Public/Math/TransformCalculus3D.h
@@ -0,0 +1,383 @@
+// Copyright 1998-2014 Epic Games, Inc. All Rights Reserved.
+#pragma once
+
+#include "TransformCalculus.h"
+
+//////////////////////////////////////////////////////////////////////////
+// Transform calculus for 3D types. Since UE4 already has existing 3D transform
+// types, this is mostly a set of adapter overloads for the primitive operations
+// requires by the transform calculus framework.
+//
+// The following types are adapted.
+// * float           -> represents a uniform scale.
+// * FScale          -> represents a 3D non-uniform scale.
+// * FVector         -> represents a 3D translation.
+// * FRotator        -> represents a pure rotation.
+// * FQuat           -> represents a pure rotation.
+// * FMatrix         -> represents a general 3D homogeneous transform.
+//
+//////////////////////////////////////////////////////////////////////////
+
+/**
+ * Represents a 3D non-uniform scale (to disambiguate from an FVector, which is used for translation).
+ * 
+ * Serves as a good base example of how to write a class that supports the basic transform calculus
+ * operations.
+ */
+class FScale
+{
+public:
+	/** Ctor. initialize to an identity scale, 1.0. */
+	FScale() :Scale(1.0f) {}
+	/** Ctor. initialize from a uniform scale. */
+	explicit FScale(float InScale) :Scale(InScale) {}
+	/** Ctor. initialize from an FVector defining the 3D scale. */
+	explicit FScale(const FVector& InScale) :Scale(InScale) {}
+	/** Access to the underlying FVector that stores the scale. */
+	const FVector& GetVector() const { return Scale; }
+	/** Concatenate two scales. */
+	const FScale Concatenate(const FScale& RHS) const
+	{
+		return FScale(Scale * RHS.GetVector());
+	}
+	/** Invert the scale. */
+	const FScale Inverse() const
+	{
+		return FScale(FVector(1.0f / Scale.X, 1.0f / Scale.Y, 1.0f / Scale.Z));
+	}
+private:
+	/** Underlying storage of the 3D scale. */
+	FVector Scale;
+};
+
+/** Specialization for converting a FMatrix to an FRotator. It uses a non-standard explicit conversion function. */
+template<> template<> inline FRotator TransformConverter<FRotator>::Convert<FMatrix>(const FMatrix& Transform)
+{
+	return Transform.Rotator();
+}
+
+//////////////////////////////////////////////////////////////////////////
+// FMatrix Support
+//////////////////////////////////////////////////////////////////////////
+
+/**
+ * Converts a generic transform to a matrix using a ToMatrix() member function.
+ * Uses decltype to allow some classes to return const-ref types for efficiency.
+ * 
+ * @param Transform
+ * @return the FMatrix stored by the Transform.
+ */
+template<typename TransformType>
+inline auto ToMatrix(const TransformType& Transform) -> decltype(Transform.ToMatrix())
+{
+	return Transform.ToMatrix();
+}
+
+/**
+ * Specialization for the NULL Matrix conversion.
+ * 
+ * @param Scale Uniform Scale
+ * @return Matrix that represents the uniform Scale space.
+ */
+inline const FMatrix& ToMatrix(const FMatrix& Transform)
+{
+	return Transform;
+}
+
+/**
+ * Specialization for floats as a uniform scale.
+ * 
+ * @param Scale Uniform Scale
+ * @return Matrix that represents the uniform Scale space.
+ */
+inline FMatrix ToMatrix(float Scale)
+{
+	return FScaleMatrix(Scale);
+}
+
+/**
+ * Specialization for non-uniform Scale.
+ * 
+ * @param Scale Non-uniform Scale
+ * @return Matrix that represents the non-uniform Scale space.
+ */
+inline FMatrix ToMatrix(const FScale& Scale)
+{
+	return FScaleMatrix(Scale.GetVector());
+}
+
+/**
+ * Specialization for translation.
+ * 
+ * @param Translation Translation
+ * @return Matrix that represents the translated space.
+ */
+inline FMatrix ToMatrix(const FVector& Translation)
+{
+	return FTranslationMatrix(Translation);
+}
+
+/**
+ * Specialization for rotation.
+ * 
+ * @param Rotation Rotation
+ * @return Matrix that represents the rotated space.
+ */
+inline FMatrix ToMatrix(const FRotator& Rotation)
+{
+	return FRotationMatrix(Rotation);
+}
+
+/**
+ * Specialization for rotation.
+ * 
+ * @param Rotation Rotation
+ * @return Matrix that represents the rotated space.
+ */
+inline FMatrix ToMatrix(const FQuat& Rotation)
+{
+	return FRotationMatrix::Make(Rotation);
+}
+
+/**
+ * Specialization of TransformConverter for FMatrix. Calls ToMatrix() by default.
+ * Allows custom types to easily provide support via a ToMatrix() overload or a ToMatrix() member function.
+ * Uses decltype to support efficient passthrough of classes that can convert to a FMatrix without creating
+ * a new instance.
+ */
+template<>
+struct TransformConverter<FMatrix>
+{
+	template<typename OtherTransformType>
+	static auto Convert(const OtherTransformType& Transform) -> decltype(ToMatrix(Transform))
+	{
+		return ToMatrix(Transform);
+	}
+};
+
+/** concatenation rules for basic UE4 types. */
+template<> struct ConcatenateRules<float        , FScale       > { typedef FScale ResultType; };
+template<> struct ConcatenateRules<FScale       , float        > { typedef FScale ResultType; };
+template<> struct ConcatenateRules<float        , FVector      > { typedef FMatrix ResultType; };
+template<> struct ConcatenateRules<FVector      , float        > { typedef FMatrix ResultType; };
+template<> struct ConcatenateRules<float        , FRotator     > { typedef FMatrix ResultType; };
+template<> struct ConcatenateRules<FRotator     , float        > { typedef FMatrix ResultType; };
+template<> struct ConcatenateRules<float        , FQuat        > { typedef FMatrix ResultType; };
+template<> struct ConcatenateRules<FQuat        , float        > { typedef FMatrix ResultType; };
+template<> struct ConcatenateRules<float        , FMatrix      > { typedef FMatrix ResultType; };
+template<> struct ConcatenateRules<FMatrix      , float        > { typedef FMatrix ResultType; };
+template<> struct ConcatenateRules<FScale       , FVector      > { typedef FMatrix ResultType; };
+template<> struct ConcatenateRules<FVector      , FScale       > { typedef FMatrix ResultType; };
+template<> struct ConcatenateRules<FScale       , FRotator     > { typedef FMatrix ResultType; };
+template<> struct ConcatenateRules<FRotator     , FScale       > { typedef FMatrix ResultType; };
+template<> struct ConcatenateRules<FScale       , FQuat        > { typedef FMatrix ResultType; };
+template<> struct ConcatenateRules<FQuat        , FScale       > { typedef FMatrix ResultType; };
+template<> struct ConcatenateRules<FScale       , FMatrix      > { typedef FMatrix ResultType; };
+template<> struct ConcatenateRules<FMatrix      , FScale       > { typedef FMatrix ResultType; };
+template<> struct ConcatenateRules<FVector      , FRotator     > { typedef FMatrix ResultType; };
+template<> struct ConcatenateRules<FRotator     , FVector      > { typedef FMatrix ResultType; };
+template<> struct ConcatenateRules<FVector      , FQuat        > { typedef FMatrix ResultType; };
+template<> struct ConcatenateRules<FQuat        , FVector      > { typedef FMatrix ResultType; };
+template<> struct ConcatenateRules<FVector      , FMatrix      > { typedef FMatrix ResultType; };
+template<> struct ConcatenateRules<FMatrix      , FVector      > { typedef FMatrix ResultType; };
+template<> struct ConcatenateRules<FRotator     , FQuat        > { typedef FQuat ResultType; };
+template<> struct ConcatenateRules<FQuat        , FRotator     > { typedef FQuat ResultType; };
+template<> struct ConcatenateRules<FRotator     , FMatrix      > { typedef FMatrix ResultType; };
+template<> struct ConcatenateRules<FMatrix      , FRotator     > { typedef FMatrix ResultType; };
+template<> struct ConcatenateRules<FQuat        , FMatrix      > { typedef FMatrix ResultType; };
+template<> struct ConcatenateRules<FMatrix      , FQuat        > { typedef FMatrix ResultType; };
+
+//////////////////////////////////////////////////////////////////////////
+// Concatenate overloads. 
+// 
+// Since these are existing UE4 types, we cannot rely on the default
+// template that calls member functions. Instead, we provide direct overloads.
+//////////////////////////////////////////////////////////////////////////
+
+/**
+ * Specialization for concatenating two Matrices.
+ * 
+ * @param LHS rotation that goes from space A to space B
+ * @param RHS rotation that goes from space B to space C.
+ * @return a new rotation representing the transformation from the input space of LHS to the output space of RHS.
+ */
+inline FMatrix Concatenate(const FMatrix& LHS, const FMatrix& RHS)
+{
+	return LHS * RHS;
+}
+
+/**
+ * Specialization for concatenating two rotations.
+ *
+ * NOTE: FQuat concatenates right to left, opposite of how FMatrix implements it.
+ *       Confusing, no? That's why we have these high level functions!
+ * 
+ * @param LHS rotation that goes from space A to space B
+ * @param RHS rotation that goes from space B to space C.
+ * @return a new rotation representing the transformation from the input space of LHS to the output space of RHS.
+ */
+inline FQuat Concatenate(const FQuat& LHS, const FQuat& RHS)
+{
+	return RHS * LHS;
+}
+
+/**
+ * Specialization for concatenating two rotations.
+ *
+ * @param LHS rotation that goes from space A to space B
+ * @param RHS rotation that goes from space B to space C.
+ * @return a new rotation representing the transformation from the input space of LHS to the output space of RHS.
+ */
+inline FRotator Concatenate(const FRotator& LHS, const FRotator& RHS)
+{
+	//@todo implement a more efficient way to do this.
+	return TransformCast<FRotator>(Concatenate(TransformCast<FMatrix>(LHS), TransformCast<FMatrix>(RHS)));
+}
+
+/**
+ * Specialization for concatenating two translations.
+ * 
+ * @param LHS Translation that goes from space A to space B
+ * @param RHS Translation that goes from space B to space C.
+ * @return a new Translation representing the transformation from the input space of LHS to the output space of RHS.
+ */
+inline FVector Concatenate(const FVector& LHS, const FVector& RHS)
+{
+	return LHS + RHS;
+}
+
+//////////////////////////////////////////////////////////////////////////
+// Inverse overloads. 
+// 
+// Since these are existing UE4 types, we cannot rely on the default
+// template that calls member functions. Instead, we provide direct overloads.
+//////////////////////////////////////////////////////////////////////////
+
+/**
+ * Inverts a transform from space A to space B so it transforms from space B to space A.
+ * Specialization for FMatrix.
+ * 
+ * @param Transform Input transform from space A to space B.
+ * @return Inverted transform from space B to space A.
+ */
+inline FMatrix Inverse(const FMatrix& Transform)
+{
+	return Transform.Inverse();
+}
+
+/**
+ * Inverts a transform from space A to space B so it transforms from space B to space A.
+ * Specialization for FRotator.
+ * 
+ * @param Transform Input transform from space A to space B.
+ * @return Inverted transform from space B to space A.
+ */
+inline FRotator Inverse(const FRotator& Transform)
+{
+	FVector EulerAngles = Transform.Euler();
+	return FRotator::MakeFromEuler(FVector(-EulerAngles.Z, -EulerAngles.Y, -EulerAngles.X));
+}
+
+/**
+ * Inverts a transform from space A to space B so it transforms from space B to space A.
+ * Specialization for FQuat.
+ * 
+ * @param Transform Input transform from space A to space B.
+ * @return Inverted transform from space B to space A.
+ */
+inline FQuat Inverse(const FQuat& Transform)
+{
+	return Transform.Inverse();
+}
+
+/**
+ * Inverts a transform from space A to space B so it transforms from space B to space A.
+ * Specialization for translation.
+ * 
+ * @param Transform Input transform from space A to space B.
+ * @return Inverted transform from space B to space A.
+ */
+inline FVector Inverse(const FVector& Transform)
+{
+	return -Transform;
+}
+
+//////////////////////////////////////////////////////////////////////////
+// TransformPoint overloads. 
+// 
+// Since these are existing UE4 types, we cannot rely on the default
+// template that calls member functions. Instead, we provide direct overloads.
+//////////////////////////////////////////////////////////////////////////
+
+/**
+ * Specialization for FMatrix as it's member function is called something slightly different.
+ */
+inline FVector TransformPoint(const FMatrix& Transform, const FVector& Point)
+{
+	return Transform.TransformPosition(Point);
+}
+
+/**
+ * Specialization for FQuat as it's member function is called something slightly different.
+ */
+inline FVector TransformPoint(const FQuat& Transform, const FVector& Point)
+{
+	return Transform.RotateVector(Point);
+}
+
+/**
+ * Specialization for FQuat as it's member function is called something slightly different.
+ */
+inline FVector TransformVector(const FQuat& Transform, const FVector& Vector)
+{
+	return Transform.RotateVector(Vector);
+}
+
+/**
+ * Specialization for FRotator as it's member function is called something slightly different.
+ */
+inline FVector TransformPoint(const FRotator& Transform, const FVector& Point)
+{
+	return Transform.RotateVector(Point);
+}
+
+/**
+ * Specialization for FRotator as it's member function is called something slightly different.
+ */
+inline FVector TransformVector(const FRotator& Transform, const FVector& Vector)
+{
+	return Transform.RotateVector(Vector);
+}
+
+/**
+ * Specialization for FVector Translation.
+ */
+inline FVector TransformPoint(const FVector& Transform, const FVector& Point)
+{
+	return Transform + Point;
+}
+
+/**
+ * Specialization for FVector Translation (does nothing).
+ */
+inline const FVector& TransformVector(const FVector& Transform, const FVector& Vector)
+{
+	return Vector;
+}
+
+/**
+ * Specialization for Scale.
+ */
+inline FVector TransformPoint(const FScale& Transform, const FVector& Point)
+{
+	return Transform.GetVector() * Point;
+}
+
+/**
+ * Specialization for Scale.
+ */
+inline FVector TransformVector(const FScale& Transform, const FVector& Vector)
+{
+	return Transform.GetVector() * Vector;
+}
+
diff --git a/Engine/Source/Runtime/Core/Public/Math/TranslationMatrix.h b/Engine/Source/Runtime/Core/Public/Math/TranslationMatrix.h
--- a/Engine/Source/Runtime/Core/Public/Math/TranslationMatrix.h
+++ b/Engine/Source/Runtime/Core/Public/Math/TranslationMatrix.h
@@ -8,6 +8,12 @@ public:
 
 	/** Constructor translation matrix based on given vector */
 	FTranslationMatrix(const FVector& Delta);
+
+	/** Matrix factory. Return an FMatrix so we don't have type conversion issues in expressions. */
+	static FMatrix Make(FVector const& Delta)
+	{
+		return FTranslationMatrix(Delta);
+	}
 };
 
 FORCEINLINE FTranslationMatrix::FTranslationMatrix(const FVector& Delta) :
diff --git a/Engine/Source/Runtime/Core/Public/Math/UnrealMath.h b/Engine/Source/Runtime/Core/Public/Math/UnrealMath.h
--- a/Engine/Source/Runtime/Core/Public/Math/UnrealMath.h
+++ b/Engine/Source/Runtime/Core/Public/Math/UnrealMath.h
@@ -29,6 +29,7 @@
 #include "Axis.h"
 #include "Matrix.h"
 #include "RotationTranslationMatrix.h"
+#include "RotationAboutPointMatrix.h"
 #include "ScaleRotationTranslationMatrix.h"
 #include "RotationMatrix.h"
 #include "Quat.h"
diff --git a/Engine/Source/Runtime/Engine/Private/GameEngine.cpp b/Engine/Source/Runtime/Engine/Private/GameEngine.cpp
--- a/Engine/Source/Runtime/Engine/Private/GameEngine.cpp
+++ b/Engine/Source/Runtime/Engine/Private/GameEngine.cpp
@@ -348,7 +348,7 @@ void UGameEngine::OnGameWindowClosed( const TSharedRef<SWindow>& WindowBeingClos
 
 void UGameEngine::OnGameWindowMoved( const TSharedRef<SWindow>& WindowBeingMoved )
 {
-	const FSlateRect WindowRect = WindowBeingMoved->GetWindowGeometryInScreen().GetRect();
+	const FSlateRect WindowRect = WindowBeingMoved->GetRectInScreen();
 	GetGameUserSettings()->SetWindowPosition(WindowRect.Left, WindowRect.Top);
 	GetGameUserSettings()->SaveConfig();
 }
diff --git a/Engine/Source/Runtime/Engine/Private/Slate/SceneViewport.cpp b/Engine/Source/Runtime/Engine/Private/Slate/SceneViewport.cpp
--- a/Engine/Source/Runtime/Engine/Private/Slate/SceneViewport.cpp
+++ b/Engine/Source/Runtime/Engine/Private/Slate/SceneViewport.cpp
@@ -264,7 +264,7 @@ void FSceneViewport::OnDrawViewport( const FGeometry& AllottedGeometry, const FS
 	FScopedConditionalWorldSwitcher WorldSwitcher( ViewportClient );
 
 	/** Check to see if the viewport should be resized */
-	FIntPoint DrawSize = FIntPoint( FMath::TruncToInt( AllottedGeometry.GetDrawSize().X ), FMath::TruncToInt( AllottedGeometry.GetDrawSize().Y ) );
+	FIntPoint DrawSize = FIntPoint( FMath::RoundToInt( AllottedGeometry.GetDrawSize().X ), FMath::RoundToInt( AllottedGeometry.GetDrawSize().Y ) );
 	if( GetSizeXY() != DrawSize )
 	{
 		TSharedPtr<SWindow> Window = FSlateApplication::Get().FindWidgetWindow( ViewportWidget.Pin().ToSharedRef() );
@@ -277,10 +277,10 @@ void FSceneViewport::OnDrawViewport( const FGeometry& AllottedGeometry, const FS
 	float CanvasMinX = FMath::Max(0.0f, AllottedGeometry.AbsolutePosition.X);
 	float CanvasMinY = FMath::Max(0.0f, AllottedGeometry.AbsolutePosition.Y);
 	FIntRect CanvasRect(
-		FMath::TruncToInt( CanvasMinX ),
-		FMath::TruncToInt( CanvasMinY ),
-		FMath::TruncToInt( CanvasMinX + AllottedGeometry.Size.X * AllottedGeometry.Scale ), 
-		FMath::TruncToInt( CanvasMinY + AllottedGeometry.Size.Y * AllottedGeometry.Scale ) );
+		FMath::RoundToInt( CanvasMinX ),
+		FMath::RoundToInt( CanvasMinY ),
+		FMath::RoundToInt( CanvasMinX + AllottedGeometry.Size.X * AllottedGeometry.Scale ), 
+		FMath::RoundToInt( CanvasMinY + AllottedGeometry.Size.Y * AllottedGeometry.Scale ) );
 
 
 	DebugCanvasDrawer->BeginRenderingCanvas( CanvasRect );
diff --git a/Engine/Source/Runtime/Slate/Private/Framework/Application/SlateApplication.cpp b/Engine/Source/Runtime/Slate/Private/Framework/Application/SlateApplication.cpp
--- a/Engine/Source/Runtime/Slate/Private/Framework/Application/SlateApplication.cpp
+++ b/Engine/Source/Runtime/Slate/Private/Framework/Application/SlateApplication.cpp
@@ -2334,13 +2334,18 @@ int32 FSlateApplication::DrawKeyboardFocus( const FWidgetPath& FocusPath, FSlate
 		// so the user can see what they are doing.
 		const FArrangedWidget& FocusedWidgetGeomertry = FocusPath.Widgets.Last();
 
+		// The FGeometry we get is from a WidgetPath, so it's rooted in desktop space.
+		// We need to APPEND a transform to the Geometry to essentially undo this root transform
+		// and get us back into Window Space.
+		// This is nonstandard so we have to go through some hoops and a specially exposed method 
+		// in FPaintGeometry to allow appending layout transforms.
+		FPaintGeometry WindowSpaceGeometry = FocusedWidgetGeomertry.Geometry.ToPaintGeometry();
+		WindowSpaceGeometry.AppendTransform(TransformCast<FSlateLayoutTransform>(Inverse(FocusPath.GetWindow()->GetPositionInScreen())));
+
 		FSlateDrawElement::MakeBox(
 			WindowElementList,
 			InLayerId++,
-			FPaintGeometry(
-			FocusedWidgetGeomertry.Geometry.AbsolutePosition - FocusPath.GetWindow()->GetPositionInScreen(),
-			FocusedWidgetGeomertry.Geometry.Size * FocusedWidgetGeomertry.Geometry.Scale,
-			FocusedWidgetGeomertry.Geometry.Scale ),
+			WindowSpaceGeometry,
 			FCoreStyle::Get().GetBrush("FocusRectangle"),
 			FocusPath.GetWindow()->GetClippingRectangleInWindow(),
 			ESlateDrawEffect::None,
@@ -2565,13 +2570,11 @@ FVector2D FSlateApplication::CalculatePopupWindowPosition( const FSlateRect& InA
 		AnchorRect.Top = InAnchor.Top + 1;
 		const FPlatformRect PlatformWorkArea = PlatformApplication->GetWorkArea( AnchorRect );
 
-		FGeometry WorkArea(
-			FVector2D(PlatformWorkArea.Left, PlatformWorkArea.Top),
-			FVector2D::ZeroVector,
-			FVector2D(PlatformWorkArea.Right - PlatformWorkArea.Left, PlatformWorkArea.Bottom - PlatformWorkArea.Top),
-			1 );
-
-		FSlateRect WorkAreaRect( WorkArea.Position.X, WorkArea.Position.Y, WorkArea.Position.X+WorkArea.Size.X, WorkArea.Position.Y+WorkArea.Size.Y );
+		FSlateRect WorkAreaRect( 
+			PlatformWorkArea.Left, 
+			PlatformWorkArea.Top, 
+			PlatformWorkArea.Left+(PlatformWorkArea.Right - PlatformWorkArea.Left), 
+			PlatformWorkArea.Top+(PlatformWorkArea.Bottom - PlatformWorkArea.Top) );
 
 		// In the direction we are opening, see if there is enough room. If there is not, flip the opening direction along the same axis.
 		FVector2D NewPosition = FVector2D::ZeroVector;
diff --git a/Engine/Source/Runtime/Slate/Private/Framework/MultiBox/SClippingHorizontalBox.cpp b/Engine/Source/Runtime/Slate/Private/Framework/MultiBox/SClippingHorizontalBox.cpp
--- a/Engine/Source/Runtime/Slate/Private/Framework/MultiBox/SClippingHorizontalBox.cpp
+++ b/Engine/Source/Runtime/Slate/Private/Framework/MultiBox/SClippingHorizontalBox.cpp
@@ -25,7 +25,7 @@ void SClippingHorizontalBox::OnArrangeChildren( const FGeometry& AllottedGeometr
 	for (int32 ChildIdx = NumChildren - 2; ChildIdx >= 0; --ChildIdx)
 	{
 		const FArrangedWidget& CurWidget = ArrangedChildren[ChildIdx];
-		if (FMath::TruncToInt(CurWidget.Geometry.AbsolutePosition.X + CurWidget.Geometry.Size.X) > FMath::TruncToInt(AllottedGeometry.AbsolutePosition.X + AllottedGeometry.Size.X))
+		if (FMath::TruncToInt(CurWidget.Geometry.AbsolutePosition.X + CurWidget.Geometry.Size.X * CurWidget.Geometry.Scale) > FMath::TruncToInt(AllottedGeometry.AbsolutePosition.X + AllottedGeometry.Size.X * CurWidget.Geometry.Scale))
 		{
 			ArrangedChildren.Remove(ChildIdx);
 			IndexClippedAt = ChildIdx;
@@ -41,13 +41,13 @@ void SClippingHorizontalBox::OnArrangeChildren( const FGeometry& AllottedGeometr
 	{
 		// Right align the wrap button
 		FArrangedWidget& ArrangedButton = ArrangedChildren[ArrangedChildren.Num() - 1];
-		ArrangedButton.Geometry.AbsolutePosition = AllottedGeometry.AbsolutePosition + AllottedGeometry.Size - ArrangedButton.Geometry.Size;
+		ArrangedButton.Geometry = AllottedGeometry.MakeChild(ArrangedButton.Geometry.Size, FSlateLayoutTransform(AllottedGeometry.Size - ArrangedButton.Geometry.Size));
 
 		// Further remove any children that the wrap button overlaps with
 		for (int32 ChildIdx = IndexClippedAt - 1; ChildIdx >= 0; --ChildIdx)
 		{
 			const FArrangedWidget& CurWidget = ArrangedChildren[ChildIdx];
-			if (FMath::TruncToInt(CurWidget.Geometry.AbsolutePosition.X + CurWidget.Geometry.Size.X) > FMath::TruncToInt(ArrangedButton.Geometry.AbsolutePosition.X))
+			if (FMath::TruncToInt(CurWidget.Geometry.AbsolutePosition.X + CurWidget.Geometry.Size.X * CurWidget.Geometry.Scale) > FMath::TruncToInt(ArrangedButton.Geometry.AbsolutePosition.X))
 			{
 				ArrangedChildren.Remove(ChildIdx);
 			}
@@ -72,23 +72,23 @@ int32 SClippingHorizontalBox::OnPaint( const FPaintArgs& Args, const FGeometry&
 		const FArrangedWidget& LastCippedChild = ClippedArrangedChildren[IndexClippedAt];
 		const FArrangedWidget& FirstChild = ArrangedChildren[0];
 		const FArrangedWidget& LastChild = ArrangedChildren[ArrangedChildren.Num() - 1];
-		FGeometry ResizedGeometry = AllottedGeometry;
+		float BorderLocalWidth = AllottedGeometry.Size.X;
 	
 		// If only the last child/block, which is the wrap button, is being clipped
 		if (IndexClippedAt == ArrangedChildren.Num() - 2)
 		{
 			// Only recalculate the alloted geometry size if said size is fitted to the toolbar/menubar
-			if (FMath::TruncToInt(AllottedGeometry.AbsolutePosition.X + AllottedGeometry.Size.X) <= FMath::TruncToInt(LastChild.Geometry.AbsolutePosition.X + LastChild.Geometry.Size.X))
+			if (FMath::TruncToInt(AllottedGeometry.AbsolutePosition.X + AllottedGeometry.Size.X * AllottedGeometry.Scale) <= FMath::TruncToInt(LastChild.Geometry.AbsolutePosition.X + LastChild.Geometry.Size.X * LastChild.Geometry.Scale))
 			{
 				// Calculate the size of the custom border
-				ResizedGeometry.Size.X = LastCippedChild.Geometry.AbsolutePosition.X + LastCippedChild.Geometry.Size.X - FirstChild.Geometry.AbsolutePosition.X;
+				BorderLocalWidth = (LastCippedChild.Geometry.AbsolutePosition.X + LastCippedChild.Geometry.Size.X * LastCippedChild.Geometry.Scale - FirstChild.Geometry.AbsolutePosition.X) / AllottedGeometry.Scale;
 			}
 		}
 		else
 		{
 			// Children/blocks are being clipped, calculate the size of the custom border
 			const FArrangedWidget& NextChild = (IndexClippedAt + 1 < ClippedArrangedChildren.Num())? ClippedArrangedChildren[IndexClippedAt + 1]: LastCippedChild;
-			ResizedGeometry.Size.X = NextChild.Geometry.AbsolutePosition.X + NextChild.Geometry.Size.X - FirstChild.Geometry.AbsolutePosition.X;
+			BorderLocalWidth = (NextChild.Geometry.AbsolutePosition.X + NextChild.Geometry.Size.X * NextChild.Geometry.Scale - FirstChild.Geometry.AbsolutePosition.X) / AllottedGeometry.Scale;
 		}
 	
 		bool bEnabled = ShouldBeEnabled( bParentEnabled );
@@ -99,7 +99,7 @@ int32 SClippingHorizontalBox::OnPaint( const FPaintArgs& Args, const FGeometry&
 		FSlateDrawElement::MakeBox(
 			OutDrawElements,
 			LayerId,
-			ResizedGeometry.ToPaintGeometry(),
+			AllottedGeometry.ToPaintGeometry(FVector2D(BorderLocalWidth, AllottedGeometry.Size.Y), FSlateLayoutTransform()),
 			BackgroundBrush,
 			MyClippingRect,
 			DrawEffects,
diff --git a/Engine/Source/Runtime/Slate/Private/Widgets/Images/SThrobber.cpp b/Engine/Source/Runtime/Slate/Private/Widgets/Images/SThrobber.cpp
--- a/Engine/Source/Runtime/Slate/Private/Widgets/Images/SThrobber.cpp
+++ b/Engine/Source/Runtime/Slate/Private/Widgets/Images/SThrobber.cpp
@@ -139,25 +139,18 @@ void SCircularThrobber::ConstructSequence()
 int32 SCircularThrobber::OnPaint( const FPaintArgs& Args, const FGeometry& AllottedGeometry, const FSlateRect& MyClippingRect, FSlateWindowElementList& OutDrawElements, int32 LayerId, const FWidgetStyle& InWidgetStyle, bool bParentEnabled ) const
 {
 	const FColor FinalColorAndOpacity( InWidgetStyle.GetColorAndOpacityTint() * PieceImage->GetTint( InWidgetStyle ) );
-	
-	const float Scale = AllottedGeometry.Scale; 
-	const float OffsetX = (AllottedGeometry.Size.X * Scale * 0.5f) - (PieceImage->ImageSize.X * Scale * 0.5f);
-	const float OffsetY = (AllottedGeometry.Size.Y * Scale * 0.5f) - (PieceImage->ImageSize.Y * Scale * 0.5f);
-
-	FPaintGeometry PaintGeom = AllottedGeometry.ToPaintGeometry();
-	FVector2D Origin = PaintGeom.DrawPosition;
-	Origin.X += OffsetX;
-	Origin.Y += OffsetY;
-
+	const FVector2D LocalOffset = (AllottedGeometry.Size - PieceImage->ImageSize) * 0.5f;
 	const float DeltaAngle = NumPieces > 0 ? 2 * PI / NumPieces : 0;
 	const float Phase = Curve.GetLerpLooping() * 2 * PI;
 
 	for (int32 PieceIdx = 0; PieceIdx < NumPieces; ++PieceIdx)
 	{
-		PaintGeom.DrawPosition.X = Origin.X + FMath::Sin(DeltaAngle * PieceIdx + Phase) * OffsetX;
-		PaintGeom.DrawPosition.Y = Origin.Y + FMath::Cos(DeltaAngle * PieceIdx + Phase) * OffsetY;
-		PaintGeom.DrawSize = PieceImage->ImageSize * Scale * (PieceIdx + 1) / NumPieces;
-
+		const float Angle = DeltaAngle * PieceIdx + Phase;
+		// scale each piece linearly until the last piece is full size
+		FSlateLayoutTransform PieceLocalTransform(
+			(PieceIdx + 1) / (float)NumPieces,
+			LocalOffset + LocalOffset * FVector2D(FMath::Sin(Angle), FMath::Cos(Angle)));
+		FPaintGeometry PaintGeom = AllottedGeometry.ToPaintGeometry(PieceImage->ImageSize, PieceLocalTransform);
 		FSlateDrawElement::MakeBox(OutDrawElements, LayerId, PaintGeom, PieceImage, MyClippingRect, ESlateDrawEffect::None, InWidgetStyle.GetColorAndOpacityTint() );
 	}
 	
diff --git a/Engine/Source/Runtime/Slate/Private/Widgets/Input/SSlider.cpp b/Engine/Source/Runtime/Slate/Private/Widgets/Input/SSlider.cpp
--- a/Engine/Source/Runtime/Slate/Private/Widgets/Input/SSlider.cpp
+++ b/Engine/Source/Runtime/Slate/Private/Widgets/Input/SSlider.cpp
@@ -28,40 +28,56 @@ int32 SSlider::OnPaint( const FPaintArgs& Args, const FGeometry& AllottedGeometr
 	SCOPE_CYCLE_COUNTER(STAT_SlateOnPaint_SSlider);
 #endif
 
+	// we draw the slider like a horizontal slider regardless of the orientation, and apply a render transform to make it display correctly.
+	// However, the AllottedGeometry is computed as it will be rendered, so we have to use the "horizontal orientation" when doing drawing computations.
+	const float AllottedWidth = Orientation == Orient_Horizontal ? AllottedGeometry.GetLocalSize().X : AllottedGeometry.GetLocalSize().Y;
+	const float AllottedHeight = Orientation == Orient_Horizontal ? AllottedGeometry.GetLocalSize().Y : AllottedGeometry.GetLocalSize().X;
+
 	float HandleRotation;
 	FVector2D HandleTopLeftPoint;
 	FVector2D SliderStartPoint;
 	FVector2D SliderEndPoint;
 
-	// calculate slider geometry
+	// calculate slider geometry as if it's a horizontal slider (we'll rotate it later if it's vertical)
 	const FVector2D HalfHandleSize = 0.5f * NormalHandleImage->ImageSize;
 	const float Indentation = IndentHandle.Get() ? NormalHandleImage->ImageSize.X : 0.0f;
 
-	if (Orientation == Orient_Horizontal)
-	{
-		const float SliderLength = AllottedGeometry.Size.X - Indentation;
-		const float SliderHandleOffset = ValueAttribute.Get() * SliderLength;
-		const float SliderY = 0.5f * AllottedGeometry.Size.Y;
+	const float SliderLength = AllottedWidth - Indentation;
+	const float SliderHandleOffset = ValueAttribute.Get() * SliderLength;
+	const float SliderY = 0.5f * AllottedHeight;
 
-		HandleRotation = 0.0f;
-		HandleTopLeftPoint = FVector2D(SliderHandleOffset - HalfHandleSize.X + 0.5f * Indentation, SliderY - HalfHandleSize.Y);
+	HandleRotation = 0.0f;
+	HandleTopLeftPoint = FVector2D(SliderHandleOffset - HalfHandleSize.X + 0.5f * Indentation, SliderY - HalfHandleSize.Y);
 
-		SliderStartPoint = FVector2D(HalfHandleSize.X, SliderY);
-		SliderEndPoint = FVector2D(AllottedGeometry.Size.X - HalfHandleSize.X, SliderY);
-	}
-	else
-	{
-		const float SliderLength = AllottedGeometry.Size.Y - Indentation;
-		const float SliderHandleOffset = ValueAttribute.Get() * SliderLength;
-		const float SliderX = 0.5f * AllottedGeometry.Size.X;
+	SliderStartPoint = FVector2D(HalfHandleSize.X, SliderY);
+	SliderEndPoint = FVector2D(AllottedWidth - HalfHandleSize.X, SliderY);
 
-		HandleRotation = -0.5f * PI;
-		HandleTopLeftPoint = FVector2D(SliderX - HalfHandleSize.X, AllottedGeometry.Size.Y - SliderHandleOffset - HalfHandleSize.Y - 0.5f * Indentation);
-
-		SliderStartPoint = FVector2D(SliderX, AllottedGeometry.Size.Y - HalfHandleSize.X);
-		SliderEndPoint = FVector2D(SliderX, HalfHandleSize.X);
+	FSlateRect RotatedClippingRect = MyClippingRect;
+	FGeometry SliderGeometry = AllottedGeometry;
+	
+	// rotate the slider 90deg if it's vertical. The 0 side goes on the bottom, the 1 side on the top.
+	if (Orientation == Orient_Vertical)
+	{
+		// Do this by translating along -X by the width of the geometry, then rotating 90 degreess CCW (left-hand coords)
+		FSlateRenderTransform RenderTransform = TransformCast<FSlateRenderTransform>(Concatenate(Inverse(FVector2D(AllottedWidth, 0)), FQuat2D(FMath::DegreesToRadians(-90.0f))));
+		// create a child geometry matching this one, but with the render transform.
+		SliderGeometry = AllottedGeometry.MakeChild(
+			FVector2D(AllottedWidth, AllottedHeight), 
+			FSlateLayoutTransform(), 
+			RenderTransform);
+		// The clipping rect is already given properly in window space. But we do not support layout rotations, so our local space rendering cannot
+		// get the clipping rect into local space properly for the local space clipping we do in the shader.
+		// Thus, we transform the clip coords into local space manually, UNDO the render transform so it will clip properly,
+		// and then bring the clip coords back into window space where DrawElements expect them.
+		RotatedClippingRect = TransformRect(
+			Concatenate(
+				Inverse(SliderGeometry.GetAccumulatedLayoutTransform()), 
+				Inverse(RenderTransform), 
+				SliderGeometry.GetAccumulatedLayoutTransform()), 
+			MyClippingRect);
 	}
 
+
 	// draw slider bar
 	TArray<FVector2D> LinePoints;
 	LinePoints.Add(SliderStartPoint);
@@ -73,9 +89,9 @@ int32 SSlider::OnPaint( const FPaintArgs& Args, const FGeometry& AllottedGeometr
 	FSlateDrawElement::MakeLines( 
 		OutDrawElements,
 		LayerId,
-		AllottedGeometry.ToPaintGeometry(),
+		SliderGeometry.ToPaintGeometry(),
 		LinePoints,
-		MyClippingRect,
+		RotatedClippingRect,
 		DrawEffects,
 		SliderBarColor.Get().GetColor(InWidgetStyle) * InWidgetStyle.GetColorAndOpacityTint()
 	);
@@ -83,16 +99,13 @@ int32 SSlider::OnPaint( const FPaintArgs& Args, const FGeometry& AllottedGeometr
 	++LayerId;
 
 	// draw slider handle
-	FSlateDrawElement::MakeRotatedBox( 
+	FSlateDrawElement::MakeBox( 
 		OutDrawElements,
 		LayerId,
-		AllottedGeometry.ToPaintGeometry(HandleTopLeftPoint, NormalHandleImage->ImageSize),
+		SliderGeometry.ToPaintGeometry(HandleTopLeftPoint, NormalHandleImage->ImageSize),
 		LockedAttribute.Get() ? DisabledHandleImage : NormalHandleImage,
-		MyClippingRect,
+		RotatedClippingRect,
 		DrawEffects,
-		HandleRotation,
-		TOptional<FVector2D>(),
-		FSlateDrawElement::RelativeToElement,
 		SliderHandleColor.Get().GetColor(InWidgetStyle) * InWidgetStyle.GetColorAndOpacityTint()
 	);
 
diff --git a/Engine/Source/Runtime/Slate/Private/Widgets/Layout/SFxWidget.cpp b/Engine/Source/Runtime/Slate/Private/Widgets/Layout/SFxWidget.cpp
--- a/Engine/Source/Runtime/Slate/Private/Widgets/Layout/SFxWidget.cpp
+++ b/Engine/Source/Runtime/Slate/Private/Widgets/Layout/SFxWidget.cpp
@@ -51,28 +51,19 @@ void SFxWidget::SetColorAndOpacity( FLinearColor InColorAndOpacity )
 }
 
 /**
- * Transform the geometry for rendering only:
- * 
- * @param RenderScale     Zoom by this amount.
- * @param ScaleOrigin     Where the origin of the transforms is within the object (range 0..1)
- * @param Offset          Offset in either direction as multiple of unscaled object size.
+ * This widget was created before render transforms existed for each widget, and it chose to apply the render transform AFTER the layout transform.
+ * This means leveraging the render transform of FGeometry would be expensive, as we would need to use Concat(LayoutTransform, RenderTransform, Inverse(LayoutTransform).
+ * Instead, we maintain the old way of doing it by modifying the AllottedGeometry only during rendering to append the widget's implied RenderTransform to the existing LayoutTransform.
  */
-static FGeometry MakeScaledGeometry( const FGeometry& Geometry, const float RenderScale, const FVector2D& ScaleOrigin, const FVector2D& Offset )
-{
-	const FVector2D CenteringAdjustment = Geometry.Size*ScaleOrigin*Geometry.Scale - Geometry.Size*ScaleOrigin*Geometry.Scale*RenderScale;
-	const FVector2D OffsetAdjustment = Geometry.Size*Offset*Geometry.Scale;
-	
-	return FGeometry(
-		FVector2D::ZeroVector,
-		Geometry.AbsolutePosition + CenteringAdjustment + OffsetAdjustment,
-		Geometry.Size,
-		Geometry.Scale * RenderScale );
-}
-
 int32 SFxWidget::OnPaint( const FPaintArgs& Args, const FGeometry& AllottedGeometry, const FSlateRect& MyClippingRect, FSlateWindowElementList& OutDrawElements, int32 LayerId, const FWidgetStyle& InWidgetStyle, bool bParentEnabled ) const
 {
-	// Transformt the geometry for rendering only. Layout is unaffected.
-	const FGeometry ModifiedGeometry = MakeScaledGeometry( AllottedGeometry, RenderScale.Get(), RenderScaleOrigin.Get(), VisualOffset.Get() );
+	// Convert the 0..1 origin into local space extents.
+	const FVector2D ScaleOrigin = RenderScaleOrigin.Get() * AllottedGeometry.Size;
+	const FVector2D Offset = VisualOffset.Get() * AllottedGeometry.Size;
+	// create the render transform as a scale around ScaleOrigin and offset it by Offset.
+	const auto RenderTransform = Concatenate(Inverse(ScaleOrigin), RenderScale.Get(), ScaleOrigin, Offset);
+	// This will append the render transform to the layout transform, and we only use it for rendering.
+	FGeometry ModifiedGeometry = AllottedGeometry.MakeChild(AllottedGeometry.Size, RenderTransform);
 	
 	FArrangedChildren ArrangedChildren(EVisibility::Visible);
 	this->ArrangeChildren(ModifiedGeometry, ArrangedChildren);
@@ -99,27 +90,23 @@ int32 SFxWidget::OnPaint( const FPaintArgs& Args, const FGeometry& AllottedGeome
 
 }
 
-
 FVector2D SFxWidget::ComputeDesiredSize() const
 {
 	// Layout scale affects out desired size.
-	return LayoutScale.Get() * ChildSlot.GetWidget()->GetDesiredSize();
+	return TransformVector(LayoutScale.Get(), ChildSlot.GetWidget()->GetDesiredSize());
 }
 
-
 void SFxWidget::OnArrangeChildren( const FGeometry& AllottedGeometry, FArrangedChildren& ArrangedChildren ) const
 {
 	const EVisibility MyVisibility = this->GetVisibility();
 	if ( ArrangedChildren.Accepts( MyVisibility ) )
 	{
-		const float MyDPIScale = LayoutScale.Get();
+		// Only layout scale affects the arranged geometry.
+		const FSlateLayoutTransform LayoutTransform(LayoutScale.Get());
 
-		// Only layout scale affects the arranged geometry. This part is identical to DPI scaling.
 		ArrangedChildren.AddWidget( AllottedGeometry.MakeChild(
 			this->ChildSlot.GetWidget(),
-			FVector2D::ZeroVector,
-			AllottedGeometry.Size / MyDPIScale,
-			MyDPIScale
-		));
+			TransformVector(Inverse(LayoutTransform), AllottedGeometry.Size),
+			LayoutTransform));
 	}
 }
diff --git a/Engine/Source/Runtime/Slate/Private/Widgets/Layout/SSpacer.cpp b/Engine/Source/Runtime/Slate/Private/Widgets/Layout/SSpacer.cpp
--- a/Engine/Source/Runtime/Slate/Private/Widgets/Layout/SSpacer.cpp
+++ b/Engine/Source/Runtime/Slate/Private/Widgets/Layout/SSpacer.cpp
@@ -18,6 +18,7 @@ void SSpacer::Construct( const FArguments& InArgs )
 		InArgs._Cursor, 
 		InArgs._IsEnabled,
 		InArgs._Visibility,
+		InArgs._RenderTransform,
 		InArgs._Tag,
 		InArgs.MetaData);
 }
diff --git a/Engine/Source/Runtime/Slate/Private/Widgets/Layout/SSplitter.cpp b/Engine/Source/Runtime/Slate/Private/Widgets/Layout/SSplitter.cpp
--- a/Engine/Source/Runtime/Slate/Private/Widgets/Layout/SSplitter.cpp
+++ b/Engine/Source/Runtime/Slate/Private/Widgets/Layout/SSplitter.cpp
@@ -1,7 +1,7 @@
 // Copyright 1998-2014 Epic Games, Inc. All Rights Reserved.
 
 #include "SlatePrivatePCH.h"
-
+#include "LayoutGeometry.h"
 
 /** The user is not allowed to make any of the splitter's children smaller than this. */
 const float MinSplitterChildLength = 20.0f;
@@ -63,12 +63,11 @@ void SSplitter::Construct( const SSplitter::FArguments& InArgs )
 	}
 }
 
-/**
-* Panels arrange their children in a space described by the AllottedGeometry parameter. The results of the arrangement
-* should be returned by appending a FArrangedWidget pair for every child widget. See StackPanel for an example
-*/
-void SSplitter::OnArrangeChildren( const FGeometry& AllottedGeometry, FArrangedChildren& ArrangedChildren ) const
+TArray<FLayoutGeometry> SSplitter::ArrangeChildrenForLayout(const FGeometry& AllottedGeometry) const
 {
+	TArray<FLayoutGeometry> Result; 
+	Result.Empty(Children.Num());
+
 	const int32 AxisIndex = (Orientation == Orient_Horizontal) ? 0 : 1;
 
 	// Splitters divide the space between their children proportionately based on size coefficients.
@@ -128,18 +127,9 @@ void SSplitter::OnArrangeChildren( const FGeometry& AllottedGeometry, FArrangedC
 
 		const EVisibility ChildVisibility = CurSlot.GetWidget()->GetVisibility();
 
-		// If the output array wants arranged children
-		if ( ArrangedChildren.Accepts(ChildVisibility) )
-		{
-			if (Orientation == Orient_Horizontal)
-			{
-				ArrangedChildren.AddWidget( ChildVisibility, AllottedGeometry.MakeChild( CurSlot.GetWidget(), FVector2D(XOffset, 0), FVector2D( ChildSpace, AllottedGeometry.Size.Y ) ) );
-			}
-			else
-			{
-				ArrangedChildren.AddWidget( ChildVisibility, AllottedGeometry.MakeChild( CurSlot.GetWidget(), FVector2D(0, XOffset), FVector2D( AllottedGeometry.Size.X, ChildSpace ) ) );
-			}
-		}
+		FVector2D ChildOffset = Orientation == Orient_Horizontal ? FVector2D(XOffset, 0) : FVector2D(0, XOffset);
+		FVector2D ChildSize = Orientation == Orient_Horizontal ? FVector2D(ChildSpace, AllottedGeometry.Size.Y) : FVector2D(AllottedGeometry.Size.X, ChildSpace);
+		Result.Emplace(FSlateLayoutTransform(ChildOffset), ChildSize);
 
 		// Advance to the next slot. If the child is collapsed, it takes up no room and does not need a splitter
 		if ( ChildVisibility != EVisibility::Collapsed)
@@ -147,6 +137,22 @@ void SSplitter::OnArrangeChildren( const FGeometry& AllottedGeometry, FArrangedC
 			XOffset += ChildSpace + PhysicalSplitterHandleSize;
 		}
 	}
+	return Result;
+}
+
+/**
+* Panels arrange their children in a space described by the AllottedGeometry parameter. The results of the arrangement
+* should be returned by appending a FArrangedWidget pair for every child widget. See StackPanel for an example
+*/
+void SSplitter::OnArrangeChildren( const FGeometry& AllottedGeometry, FArrangedChildren& ArrangedChildren ) const
+{
+	TArray<FLayoutGeometry> LayoutChildren = ArrangeChildrenForLayout(AllottedGeometry);
+
+	// Arrange the children horizontally or vertically.
+	for (int32 ChildIndex=0; ChildIndex < Children.Num(); ++ChildIndex)
+	{
+		ArrangedChildren.AddWidget( AllottedGeometry.MakeChild( Children[ChildIndex].GetWidget(), LayoutChildren[ChildIndex] ) );
+	}
 }
 
 
@@ -330,8 +336,7 @@ FReply SSplitter::OnMouseMove( const FGeometry& MyGeometry, const FPointerEvent&
 {
 	const FVector2D LocalMousePosition = MyGeometry.AbsoluteToLocal( MouseEvent.GetScreenSpacePosition() );
 
-	FArrangedChildren ArrangedChildren(EVisibility::All);
-	ArrangeChildren( MyGeometry, ArrangedChildren );
+	TArray<FLayoutGeometry> LayoutChildren = ArrangeChildrenForLayout(MyGeometry);
 
 	if ( bIsResizing )
 	{
@@ -339,11 +344,11 @@ FReply SSplitter::OnMouseMove( const FGeometry& MyGeometry, const FPointerEvent&
 		{
 			if (Orientation == Orient_Horizontal)
 			{
-				HandleResizing<Orient_Horizontal>( PhysicalSplitterHandleSize, ResizeMode, HoveredHandleIndex, LocalMousePosition, Children, ArrangedChildren );
+				HandleResizing<Orient_Horizontal>( PhysicalSplitterHandleSize, ResizeMode, HoveredHandleIndex, LocalMousePosition, Children, LayoutChildren );
 			}
 			else
 			{
-				HandleResizing<Orient_Vertical>( PhysicalSplitterHandleSize, ResizeMode, HoveredHandleIndex, LocalMousePosition, Children, ArrangedChildren );
+				HandleResizing<Orient_Vertical>( PhysicalSplitterHandleSize, ResizeMode, HoveredHandleIndex, LocalMousePosition, Children, LayoutChildren );
 			}
 		}
 
@@ -353,8 +358,8 @@ FReply SSplitter::OnMouseMove( const FGeometry& MyGeometry, const FPointerEvent&
 	{	
 		// Hit test which handle we are hovering over.		
 		HoveredHandleIndex = (Orientation == Orient_Horizontal)
-			? GetHandleBeingResizedFromMousePosition<Orient_Horizontal>( PhysicalSplitterHandleSize, HitDetectionSplitterHandleSize, LocalMousePosition, ArrangedChildren )
-			: GetHandleBeingResizedFromMousePosition<Orient_Vertical>( PhysicalSplitterHandleSize, HitDetectionSplitterHandleSize, LocalMousePosition, ArrangedChildren );
+			? GetHandleBeingResizedFromMousePosition<Orient_Horizontal>( PhysicalSplitterHandleSize, HitDetectionSplitterHandleSize, LocalMousePosition, LayoutChildren )
+			: GetHandleBeingResizedFromMousePosition<Orient_Vertical>( PhysicalSplitterHandleSize, HitDetectionSplitterHandleSize, LocalMousePosition, LayoutChildren );
 
 		if (HoveredHandleIndex != INDEX_NONE)
 		{
@@ -386,13 +391,12 @@ FCursorReply SSplitter::OnCursorQuery( const FGeometry& MyGeometry, const FPoint
 {
 	const FVector2D LocalMousePosition = MyGeometry.AbsoluteToLocal( CursorEvent.GetScreenSpacePosition() );
 
-	FArrangedChildren ArrangedChildren(EVisibility::All);
-	ArrangeChildren( MyGeometry, ArrangedChildren );
+	TArray<FLayoutGeometry> LayoutChildren = ArrangeChildrenForLayout(MyGeometry);
 
 	// Hit test which handle we are hovering over.		
 	const int32 CurrentHoveredHandleIndex = (Orientation == Orient_Horizontal)
-		? GetHandleBeingResizedFromMousePosition<Orient_Horizontal>( PhysicalSplitterHandleSize, HitDetectionSplitterHandleSize, LocalMousePosition, ArrangedChildren )
-		: GetHandleBeingResizedFromMousePosition<Orient_Vertical>( PhysicalSplitterHandleSize, HitDetectionSplitterHandleSize, LocalMousePosition, ArrangedChildren );
+		? GetHandleBeingResizedFromMousePosition<Orient_Horizontal>( PhysicalSplitterHandleSize, HitDetectionSplitterHandleSize, LocalMousePosition, LayoutChildren )
+		: GetHandleBeingResizedFromMousePosition<Orient_Vertical>( PhysicalSplitterHandleSize, HitDetectionSplitterHandleSize, LocalMousePosition, LayoutChildren );
 
 	if (CurrentHoveredHandleIndex != INDEX_NONE)
 	{
@@ -484,7 +488,7 @@ void SSplitter::FindAllResizeableSlotsAfterHandle( int32 DraggedHandle, const TP
 
 
 template<EOrientation SplitterOrientation>
-void SSplitter::HandleResizing( const float PhysicalSplitterHandleSize, const ESplitterResizeMode::Type ResizeMode, int32 DraggedHandle, const FVector2D& LocalMousePos, TPanelChildren<FSlot>& Children, FArrangedChildren& ChildGeometries )
+void SSplitter::HandleResizing( const float PhysicalSplitterHandleSize, const ESplitterResizeMode::Type ResizeMode, int32 DraggedHandle, const FVector2D& LocalMousePos, TPanelChildren<FSlot>& Children, const TArray<FLayoutGeometry>& ChildGeometries )
 {
 	const int32 NumChildren = Children.Num();
 
@@ -494,7 +498,7 @@ void SSplitter::HandleResizing( const float PhysicalSplitterHandleSize, const ES
 	//  - Prev vs. Next refers to the widgets in the order they are laid out (left->right, top->bottom).
 	//  - New vs. Old refers to the Old values for width/height vs. the post-resize values.
 
-	const float HandlePos = ChildGeometries[DraggedHandle+1].Geometry.Position.Component(AxisIndex) - PhysicalSplitterHandleSize / 2;
+	const float HandlePos = ChildGeometries[DraggedHandle+1].GetLocalToParentTransform().GetTranslation().Component(AxisIndex) - PhysicalSplitterHandleSize / 2;
 	float Delta = LocalMousePos.Component(AxisIndex) - HandlePos;
 
 	const int32 SlotBeforeDragHandle = FindResizeableSlotBeforeHandle( DraggedHandle, Children );
@@ -519,7 +523,7 @@ void SSplitter::HandleResizing( const float PhysicalSplitterHandleSize, const ES
 		struct FSlotInfo 
 		{
 			FSlot* Slot;
-			FGeometry* Geometry;
+			const FLayoutGeometry* Geometry;
 			float NewSize;
 		};
 
@@ -529,19 +533,20 @@ void SSplitter::HandleResizing( const float PhysicalSplitterHandleSize, const ES
 			FSlotInfo SlotInfo;
 
 			SlotInfo.Slot = &Children[ SlotsAfterDragHandleIndicies[ SlotIndex ] ];
-			SlotInfo.Geometry = &ChildGeometries[ SlotsAfterDragHandleIndicies[ SlotIndex ] ].Geometry;
-			SlotInfo.NewSize = SlotInfo.Geometry->Size.Component( AxisIndex );
+			SlotInfo.Geometry = &ChildGeometries[ SlotsAfterDragHandleIndicies[ SlotIndex ] ];
+			SlotInfo.NewSize = SlotInfo.Geometry->GetSizeInParentSpace().Component( AxisIndex );
 
 			SlotsAfterDragHandle.Add( SlotInfo );
 		}
 
 		// Get references the prev and next children and their layout settings so that we can modify them.
 		FSlot& PrevChild = Children[SlotBeforeDragHandle];
-		FGeometry& PrevChildGeom = ChildGeometries[SlotBeforeDragHandle].Geometry;
+		const FLayoutGeometry& PrevChildGeom = ChildGeometries[SlotBeforeDragHandle];
 
 		// Compute the new sizes of the children
-		float NewPrevChildLength = ClampChild( PrevChildGeom.Size.Component(AxisIndex) + Delta );
-		Delta = NewPrevChildLength - PrevChildGeom.Size.Component(AxisIndex);
+		const float PrevChildLength = PrevChildGeom.GetSizeInParentSpace().Component(AxisIndex);
+		float NewPrevChildLength = ClampChild( PrevChildLength + Delta );
+		Delta = NewPrevChildLength - PrevChildLength;
 
 		// Distribute the Delta across the affected slots after the drag handle
 		float UnusedDelta = Delta;
@@ -553,7 +558,7 @@ void SSplitter::HandleResizing( const float PhysicalSplitterHandleSize, const ES
 			{
 				FSlotInfo& SlotInfo = SlotsAfterDragHandle[ SlotIndex ];
 
-				float CurrentSize = ClampChild( SlotInfo.Geometry->Size.Component(AxisIndex) );
+				float CurrentSize = ClampChild( SlotInfo.Geometry->GetSizeInParentSpace().Component(AxisIndex) );
 				SlotInfo.NewSize = ClampChild( CurrentSize - DividedDelta );
 
 				// If one of the slots couldn't be fully adjusted by the delta due to min/max constraints then
@@ -565,7 +570,7 @@ void SSplitter::HandleResizing( const float PhysicalSplitterHandleSize, const ES
 		Delta = Delta - UnusedDelta;
 
 		// PrevChildLength needs to be updated: it's value has to take into account the next child's min/max restrictions
-		NewPrevChildLength = ClampChild( PrevChildGeom.Size.Component(AxisIndex) + Delta );
+		NewPrevChildLength = ClampChild( PrevChildLength + Delta );
 
 		// Cells being resized are both stretch values -> redistribute the stretch coefficients proportionately
 		// to match the new child sizes on the screen.
@@ -623,7 +628,7 @@ float SSplitter::ClampChild( float ProposedSize )
 
 
 template<EOrientation SplitterOrientation>
-int32 SSplitter::GetHandleBeingResizedFromMousePosition( float PhysicalSplitterHandleSize, float HitDetectionSplitterHandleSize, FVector2D LocalMousePos, FArrangedChildren& ChildGeometries )
+int32 SSplitter::GetHandleBeingResizedFromMousePosition( float PhysicalSplitterHandleSize, float HitDetectionSplitterHandleSize, FVector2D LocalMousePos, const TArray<FLayoutGeometry>& ChildGeometries )
 {
 	const int32 AxisIndex = (SplitterOrientation == Orient_Horizontal) ? 0 : 1;
 	const float HalfHitDetectionSplitterHandleSize = ( HitDetectionSplitterHandleSize / 2 );
@@ -632,10 +637,10 @@ int32 SSplitter::GetHandleBeingResizedFromMousePosition( float PhysicalSplitterH
 	// Search for the two widgets between which the cursor currently resides.
 	for ( int32 ChildIndex = 1; ChildIndex < ChildGeometries.Num(); ++ChildIndex )
 	{
-		FGeometry& PrevChild = ChildGeometries[ChildIndex - 1].Geometry;
-		FGeometry& NextChild = ChildGeometries[ChildIndex].Geometry;
-		float PrevBound = PrevChild.Position.Component(AxisIndex) + PrevChild.Size.Component(AxisIndex) - HalfHitDetectionSplitterHandleSize + HalfPhysicalSplitterHandleSize;
-		float NextBound = NextChild.Position.Component(AxisIndex) + HalfHitDetectionSplitterHandleSize - HalfPhysicalSplitterHandleSize;
+		FSlateRect PrevChildRect = ChildGeometries[ChildIndex - 1].GetRectInParentSpace();
+		FVector2D NextChildOffset = ChildGeometries[ChildIndex].GetOffsetInParentSpace();
+		float PrevBound = PrevChildRect.GetTopLeft().Component(AxisIndex) + PrevChildRect.GetSize().Component(AxisIndex) - HalfHitDetectionSplitterHandleSize + HalfPhysicalSplitterHandleSize;
+		float NextBound = NextChildOffset.Component(AxisIndex) + HalfHitDetectionSplitterHandleSize - HalfPhysicalSplitterHandleSize;
 
 		if ( LocalMousePos.Component(AxisIndex) > PrevBound && LocalMousePos.Component(AxisIndex) < NextBound )
 		{
@@ -670,14 +675,16 @@ void SSplitter2x2::Construct( const FArguments& InArgs )
 	ResizingAxis = INDEX_NONE;
 }
 
-
-void SSplitter2x2::OnArrangeChildren( const FGeometry& AllottedGeometry, FArrangedChildren& ArrangedChildren ) const
+TArray<FLayoutGeometry> SSplitter2x2::ArrangeChildrenForLayout( const FGeometry& AllottedGeometry ) const
 {
+	check( Children.Num() == 4 );
+
+	TArray<FLayoutGeometry> Result;
+	Result.Empty(Children.Num());
+
 	int32 NumNonCollapsedChildren = 0;
 	FVector2D CoefficientTotal(0,0);
 
-	check( Children.Num() == 4 );
-
 	// The allotted space for our children is our geometry minus a little space to show splitter handles
 	const FVector2D SpaceAllottedForChildren = AllottedGeometry.Size - FVector2D(SplitterHandleSize,SplitterHandleSize);
 
@@ -690,15 +697,10 @@ void SSplitter2x2::OnArrangeChildren( const FGeometry& AllottedGeometry, FArrang
 
 		// Calculate the amount of space that this child should take up.  
 		// It is based on the current percentage of space it should take up which is defined by a user moving the splitters
-		FVector2D ChildSpace = SpaceAllottedForChildren * CurSlot.PercentageAttribute.Get();
+		const FVector2D ChildSpace = SpaceAllottedForChildren * CurSlot.PercentageAttribute.Get();
 
-		const EVisibility ChildVisibility = CurSlot.GetWidget()->GetVisibility();
-
-		// If the child is visible, put them in their spot
-		if ( ArrangedChildren.Accepts( ChildVisibility ) )
-		{
-			ArrangedChildren.AddWidget( ChildVisibility, AllottedGeometry.MakeChild( CurSlot.GetWidget(), Offset, ChildSpace ) );
-		}
+		// put them in their spot
+		Result.Emplace(FSlateLayoutTransform(Offset), ChildSpace);
 
 		// Advance to the next slot. If the child is collapsed, it takes up no room and does not need a splitter
 		if( ChildIndex == 1 )
@@ -712,6 +714,20 @@ void SSplitter2x2::OnArrangeChildren( const FGeometry& AllottedGeometry, FArrang
 			Offset += FVector2D( 0, ChildSpace.Y + SplitterHandleSize );
 		}
 	}
+	return Result;
+}
+
+void SSplitter2x2::OnArrangeChildren( const FGeometry& AllottedGeometry, FArrangedChildren& ArrangedChildren ) const
+{
+	TArray<FLayoutGeometry> LayoutChildren = ArrangeChildrenForLayout(AllottedGeometry);
+
+	for (int32 ChildIndex=0; ChildIndex < Children.Num(); ++ChildIndex)
+	{
+		const FSlot& CurSlot = Children[ChildIndex];
+
+		// put them in their spot
+		ArrangedChildren.AddWidget( AllottedGeometry.MakeChild( Children[ChildIndex].GetWidget(), LayoutChildren[ChildIndex] ) );
+	}
 }
 
 
@@ -763,10 +779,8 @@ FReply SSplitter2x2::OnMouseMove( const FGeometry& MyGeometry, const FPointerEve
 
 	if ( bIsResizing && FSlateApplication::Get().GetMouseCaptor() == AsShared() )
 	{	
-		FArrangedChildren ArrangedChildren(EVisibility::Visible);
-		ArrangeChildren( MyGeometry, ArrangedChildren );
-
-		ResizeChildren( ArrangedChildren, LocalMousePos );
+		TArray<FLayoutGeometry> LayoutChildren = ArrangeChildrenForLayout(MyGeometry);
+		ResizeChildren( MyGeometry, LayoutChildren, LocalMousePos );
 		return FReply::Handled();
 	}
 	else
@@ -795,16 +809,16 @@ FCursorReply SSplitter2x2::OnCursorQuery( const FGeometry& MyGeometry, const FPo
 	return FCursorReply::Unhandled();
 }
 
-void SSplitter2x2::ResizeChildren( FArrangedChildren& ArrangedChildren, FVector2D LocalMousePos )
+void SSplitter2x2::ResizeChildren( const FGeometry& MyGeometry, const TArray<FLayoutGeometry>& ArrangedChildren, const FVector2D& LocalMousePos )
 {
 	// Compute the handle position.  The last child is used because it is always the furthest away from the origin
-	const FVector2D HandlePos = ArrangedChildren[3].Geometry.Position - (FVector2D( SplitterHandleSize , SplitterHandleSize ) * .5f);
+	const FVector2D HandlePos = ArrangedChildren[3].GetOffsetInParentSpace() - (FVector2D( SplitterHandleSize , SplitterHandleSize ) * .5f);
 	FVector2D Delta = LocalMousePos - HandlePos;
 
-	FGeometry& TopLeftGeom = ArrangedChildren[0].Geometry;
-	FGeometry& BotLeftGeom = ArrangedChildren[1].Geometry;
-	FGeometry& TopRightGeom = ArrangedChildren[2].Geometry;
-	FGeometry& BotRightGeom = ArrangedChildren[3].Geometry;
+	FVector2D TopLeftSize = ArrangedChildren[0].GetSizeInParentSpace();
+	FVector2D BotLeftSize = ArrangedChildren[1].GetSizeInParentSpace();
+	FVector2D TopRightSize = ArrangedChildren[2].GetSizeInParentSpace();
+	FVector2D BotRightSize = ArrangedChildren[3].GetSizeInParentSpace();
 
 	FSlot& TopLeft = Children[0];
 	FSlot& BotLeft = Children[1];
@@ -830,31 +844,31 @@ void SSplitter2x2::ResizeChildren( FArrangedChildren& ArrangedChildren, FVector2
 
 	if( ResizingAxis == 0)
 	{
-		NewSizeTL = TopLeftGeom.Size + Delta;
-		NewSizeBL = BotLeftGeom.Size + Delta;
-		NewSizeTR = TopRightGeom.Size - Delta;
-		NewSizeBR = BotRightGeom.Size - Delta;
+		NewSizeTL = TopLeftSize + Delta;
+		NewSizeBL = BotLeftSize + Delta;
+		NewSizeTR = TopRightSize - Delta;
+		NewSizeBR = BotRightSize - Delta;
 	}
 	else if( ResizingAxis == 1 )
 	{
-		NewSizeTL = TopLeftGeom.Size + Delta;
-		NewSizeBL = BotLeftGeom.Size - Delta;
-		NewSizeTR = TopRightGeom.Size + Delta;
-		NewSizeBR = BotRightGeom.Size - Delta;
+		NewSizeTL = TopLeftSize + Delta;
+		NewSizeBL = BotLeftSize - Delta;
+		NewSizeTR = TopRightSize + Delta;
+		NewSizeBR = BotRightSize - Delta;
 	}
 	else
 	{
 		// Resize X and Y independently as they have different rules for X and Y
-		NewSizeTL.X = TopLeftGeom.Size.X + Delta.X;
-		NewSizeBL.X = BotLeftGeom.Size.X + Delta.X;
-		NewSizeTR.X = TopRightGeom.Size.X - Delta.X;
-		NewSizeBR.X = BotRightGeom.Size.X - Delta.X;
+		NewSizeTL.X = TopLeftSize.X + Delta.X;
+		NewSizeBL.X = BotLeftSize.X + Delta.X;
+		NewSizeTR.X = TopRightSize.X - Delta.X;
+		NewSizeBR.X = BotRightSize.X - Delta.X;
 
 
-		NewSizeTL.Y = TopLeftGeom.Size.Y + Delta.Y;
-		NewSizeBL.Y = BotLeftGeom.Size.Y - Delta.Y;
-		NewSizeTR.Y = TopRightGeom.Size.Y + Delta.Y;
-		NewSizeBR.Y = BotRightGeom.Size.Y - Delta.Y;
+		NewSizeTL.Y = TopLeftSize.Y + Delta.Y;
+		NewSizeBL.Y = BotLeftSize.Y - Delta.Y;
+		NewSizeTR.Y = TopRightSize.Y + Delta.Y;
+		NewSizeBR.Y = BotRightSize.Y - Delta.Y;
 	}
 
 	// Clamp all values so they cant be too small
@@ -882,18 +896,17 @@ int32 SSplitter2x2::CalculateResizingAxis( const FGeometry& MyGeometry, const FV
 {
 	int32 Axis = INDEX_NONE;
 
-	FArrangedChildren ArrangedChildren(EVisibility::Visible);
-	ArrangeChildren( MyGeometry, ArrangedChildren );
+	TArray<FLayoutGeometry> LayoutChildren = ArrangeChildrenForLayout(MyGeometry);
 
 	// The axis is in the center if it passes all hit tests
 	bool bInCenter = true;
 	// Search for the two widgets between which the cursor currently resides.
-	for ( int32 ChildIndex = 1; ChildIndex < ArrangedChildren.Num(); ++ChildIndex )
+	for ( int32 ChildIndex = 1; ChildIndex < LayoutChildren.Num(); ++ChildIndex )
 	{
-		FGeometry& PrevChild = ArrangedChildren[ChildIndex - 1].Geometry;
-		FGeometry& NextChild = ArrangedChildren[ChildIndex].Geometry;
-		FVector2D PrevBound = PrevChild.Position + PrevChild.Size;
-		FVector2D NextBound = NextChild.Position;
+		FLayoutGeometry& PrevChild = LayoutChildren[ChildIndex - 1];
+		FLayoutGeometry& NextChild = LayoutChildren[ChildIndex];
+		FVector2D PrevBound = PrevChild.GetOffsetInParentSpace() + PrevChild.GetSizeInParentSpace();
+		FVector2D NextBound = NextChild.GetOffsetInParentSpace();
 
 		if( LocalMousePos.X > PrevBound.X && LocalMousePos.X < NextBound.X )
 		{
diff --git a/Engine/Source/Runtime/Slate/Private/Widgets/Notifications/SProgressBar.cpp b/Engine/Source/Runtime/Slate/Private/Widgets/Notifications/SProgressBar.cpp
--- a/Engine/Source/Runtime/Slate/Private/Widgets/Notifications/SProgressBar.cpp
+++ b/Engine/Source/Runtime/Slate/Private/Widgets/Notifications/SProgressBar.cpp
@@ -95,8 +95,9 @@ int32 SProgressBar::OnPaint( const FPaintArgs& Args, const FGeometry& AllottedGe
 	TOptional<float> ProgressFraction = Percent.Get();	
 
 	// Paint inside the border only. 
-	FPaintGeometry ForegroundPaintGeometry = AllottedGeometry.ToInflatedPaintGeometry( -BorderPadding.Get() );
-	const FSlateRect ForegroundClippingRect = ForegroundPaintGeometry.ToSlateRect().IntersectionWith( MyClippingRect );
+	// Pre-snap the clipping rect to try and reduce common jitter, since the padding is typically only a single pixel.
+	FSlateRect SnappedClippingRect = FSlateRect(FMath::RoundToInt(MyClippingRect.Left), FMath::RoundToInt(MyClippingRect.Top), FMath::RoundToInt(MyClippingRect.Right), FMath::RoundToInt(MyClippingRect.Bottom));
+	const FSlateRect ForegroundClippingRect = SnappedClippingRect.InsetBy(FMargin(BorderPadding.Get().X, BorderPadding.Get().Y));
 	
 	const FSlateBrush* CurrentBackgroundImage = GetBackgroundImage();
 
@@ -105,7 +106,7 @@ int32 SProgressBar::OnPaint( const FPaintArgs& Args, const FGeometry& AllottedGe
 		RetLayerId++,
 		AllottedGeometry.ToPaintGeometry(),
 		CurrentBackgroundImage,
-		MyClippingRect,
+		SnappedClippingRect,
 		DrawEffects,
 		InWidgetStyle.GetColorAndOpacityTint() * CurrentBackgroundImage->GetTint( InWidgetStyle )
 	);	
diff --git a/Engine/Source/Runtime/Slate/Private/Widgets/Text/SMultiLineEditableText.cpp b/Engine/Source/Runtime/Slate/Private/Widgets/Text/SMultiLineEditableText.cpp
--- a/Engine/Source/Runtime/Slate/Private/Widgets/Text/SMultiLineEditableText.cpp
+++ b/Engine/Source/Runtime/Slate/Private/Widgets/Text/SMultiLineEditableText.cpp
@@ -2426,9 +2426,7 @@ bool SMultiLineEditableText::FTextInputMethodContext::GetTextBounds(const uint32
 
 		// Translate the position (which is in local space) into screen (absolute) space
 		// Note: The local positions are pre-scaled, so we don't scale them again here
-		FGeometry CachedGeometryNoScale = CachedGeometry;
-		CachedGeometryNoScale.Scale = 1.0f;
-		Position = CachedGeometryNoScale.LocalToAbsolute(Position);
+		Position += CachedGeometry.AbsolutePosition;
 	}
 	return IsClipped;
 }
diff --git a/Engine/Source/Runtime/Slate/Private/Widgets/Views/SHeaderRow.cpp b/Engine/Source/Runtime/Slate/Private/Widgets/Views/SHeaderRow.cpp
--- a/Engine/Source/Runtime/Slate/Private/Widgets/Views/SHeaderRow.cpp
+++ b/Engine/Source/Runtime/Slate/Private/Widgets/Views/SHeaderRow.cpp
@@ -34,7 +34,7 @@ public:
 	{
 		check(InArgs._Style);
 
-		SWidget::Construct( InArgs._ToolTipText, InArgs._ToolTip, InArgs._Cursor, InArgs._IsEnabled, InArgs._Visibility, InArgs._Tag, InArgs.MetaData );
+		SWidget::Construct( InArgs._ToolTipText, InArgs._ToolTip, InArgs._Cursor, InArgs._IsEnabled, InArgs._Visibility, InArgs._RenderTransform, InArgs._Tag, InArgs.MetaData );
 
 		Style = InArgs._Style;
 		ColumnId = Column.ColumnId;
@@ -384,7 +384,7 @@ void SHeaderRow::Construct( const FArguments& InArgs )
 {
 	check(InArgs._Style);
 
-	SWidget::Construct( InArgs._ToolTipText, InArgs._ToolTip, InArgs._Cursor, InArgs._IsEnabled, InArgs._Visibility, InArgs._Tag, InArgs.MetaData );
+	SWidget::Construct( InArgs._ToolTipText, InArgs._ToolTip, InArgs._Cursor, InArgs._IsEnabled, InArgs._Visibility, InArgs._RenderTransform, InArgs._Tag, InArgs.MetaData );
 
 	ScrollBarThickness = FVector2D::ZeroVector;
 	ScrollBarVisibility = EVisibility::Collapsed;
diff --git a/Engine/Source/Runtime/Slate/Public/Widgets/Layout/SSplitter.h b/Engine/Source/Runtime/Slate/Public/Widgets/Layout/SSplitter.h
--- a/Engine/Source/Runtime/Slate/Public/Widgets/Layout/SSplitter.h
+++ b/Engine/Source/Runtime/Slate/Public/Widgets/Layout/SSplitter.h
@@ -11,6 +11,7 @@ namespace ESplitterResizeMode
 	};
 }
 
+class FLayoutGeometry;
 /**
  * SSPlitter divides its allotted area into N segments, where N is the number of children it has.
  * It allows the users to resize the children along the splitters axis: that is, horizontally or vertically.
@@ -204,6 +205,9 @@ public:
 	EOrientation GetOrientation() const;	
 
 
+private:
+	TArray<FLayoutGeometry> ArrangeChildrenForLayout( const FGeometry& AllottedGeometry ) const;
+
 protected:
 
 	/**
@@ -231,7 +235,7 @@ protected:
 	 * @param ChildGeometries  The arranged children; we need their sizes and positions so that we can perform a resizing.
 	 */
 	template<EOrientation SplitterOrientation>
-	static void HandleResizing( const float PhysicalSplitterHandleSize, const ESplitterResizeMode::Type ResizeMode, int32 DraggedHandle, const FVector2D& LocalMousePos, TPanelChildren<FSlot>& Children, FArrangedChildren& ChildGeometries );
+	static void HandleResizing( const float PhysicalSplitterHandleSize, const ESplitterResizeMode::Type ResizeMode, int32 DraggedHandle, const FVector2D& LocalMousePos, TPanelChildren<FSlot>& Children, const TArray<FLayoutGeometry>& ChildGeometries );
 
 	/**
 	 * @param ProposedSize  A size that a child would like to be
@@ -250,7 +254,7 @@ protected:
 	 * @return The index of the handle being hovered, or INDEX_NONE if we are not hovering a handle.
 	 */
 	template<EOrientation SplitterOrientation>
-	static int32 GetHandleBeingResizedFromMousePosition(  float PhysicalSplitterHandleSize, float HitDetectionSplitterHandleSize, FVector2D LocalMousePos, FArrangedChildren& ChildGeometries );
+	static int32 GetHandleBeingResizedFromMousePosition(  float PhysicalSplitterHandleSize, float HitDetectionSplitterHandleSize, FVector2D LocalMousePos, const TArray<FLayoutGeometry>& ChildGeometries );
 
 	TPanelChildren< FSlot > Children;
 
@@ -385,6 +389,8 @@ public:
 
 private:
 
+	TArray<FLayoutGeometry> ArrangeChildrenForLayout( const FGeometry& AllottedGeometry ) const;
+
 	virtual void OnArrangeChildren( const FGeometry& AllottedGeometry, FArrangedChildren& ArrangedChildren ) const override;
 
 	/**
@@ -451,7 +457,7 @@ private:
 	 * @param ArrangedChildren	The current geometry of all arranged children before the user moved the splitter
 	 * @param LocalMousePos		The current mouse position		
 	 */
-	void ResizeChildren( FArrangedChildren& ArrangedChildren, FVector2D LocalMousePos );
+	void ResizeChildren( const FGeometry& MyGeometry, const TArray<FLayoutGeometry>& ArrangedChildren, const FVector2D& LocalMousePos );
 
 
 private:
diff --git a/Engine/Source/Runtime/SlateCore/Private/Fonts/FontCache.cpp b/Engine/Source/Runtime/SlateCore/Private/Fonts/FontCache.cpp
--- a/Engine/Source/Runtime/SlateCore/Private/Fonts/FontCache.cpp
+++ b/Engine/Source/Runtime/SlateCore/Private/Fonts/FontCache.cpp
@@ -659,10 +659,10 @@ bool FSlateFontCache::AddNewEntry( TCHAR Character, const FSlateFontKey& InKey,
 
 	if( bSuccess )
 	{
-		OutCharacterEntry.StartU = NewSlot->X;
-		OutCharacterEntry.StartV = NewSlot->Y;
-		OutCharacterEntry.USize = NewSlot->Width;
-		OutCharacterEntry.VSize = NewSlot->Height;
+		OutCharacterEntry.StartU = NewSlot->X + NewSlot->Padding;
+		OutCharacterEntry.StartV = NewSlot->Y + NewSlot->Padding;
+		OutCharacterEntry.USize = NewSlot->Width - 2 * NewSlot->Padding;
+		OutCharacterEntry.VSize = NewSlot->Height - 2 * NewSlot->Padding;
 		OutCharacterEntry.TextureIndex = 0;
 		OutCharacterEntry.XAdvance = RenderData.MeasureInfo.XAdvance;
 		OutCharacterEntry.VerticalOffset = RenderData.MeasureInfo.VerticalOffset;
diff --git a/Engine/Source/Runtime/SlateCore/Private/Fonts/FontTypes.cpp b/Engine/Source/Runtime/SlateCore/Private/Fonts/FontTypes.cpp
--- a/Engine/Source/Runtime/SlateCore/Private/Fonts/FontTypes.cpp
+++ b/Engine/Source/Runtime/SlateCore/Private/Fonts/FontTypes.cpp
@@ -4,7 +4,7 @@
 
 
 FSlateFontAtlas::FSlateFontAtlas( uint32 InWidth, uint32 InHeight )
-	: FSlateTextureAtlas( InWidth, InHeight, sizeof(uint8), 0 )
+	: FSlateTextureAtlas( InWidth, InHeight, sizeof(uint8), ESlateTextureAtlasPaddingStyle::PadWithZero )
 {
 }
 
diff --git a/Engine/Source/Runtime/SlateCore/Private/Input/HittestGrid.cpp b/Engine/Source/Runtime/SlateCore/Private/Input/HittestGrid.cpp
--- a/Engine/Source/Runtime/SlateCore/Private/Input/HittestGrid.cpp
+++ b/Engine/Source/Runtime/SlateCore/Private/Input/HittestGrid.cpp
@@ -63,7 +63,14 @@ TArray<FArrangedWidget> FHittestGrid::GetBubblePath( FVector2D DesktopSpaceCoord
 			check( IndexesInCell[i] < WidgetsCachedThisFrame->Num() ); 
 
 			const FCachedWidget& TestCandidate = (*WidgetsCachedThisFrame)[IndexesInCell[i]];
-			if ( TestCandidate.CachedGeometry.IsUnderLocation( DesktopSpaceCoordinate ) && TestCandidate.ClippingRect.ContainsPoint( DesktopSpaceCoordinate ) && TestCandidate.WidgetPtr.IsValid() )
+
+			// Compute the render space clipping rect.
+			FSlateRotatedRect DesktopOrientedClipRect(
+				Inverse(TestCandidate.CachedGeometry.GetAccumulatedLayoutTransform()),
+				TestCandidate.CachedGeometry.GetAccumulatedRenderTransform(),
+				TestCandidate.CachedGeometry.GetClippingRect().IntersectionWith(TestCandidate.ClippingRect));
+
+			if (DesktopOrientedClipRect.IsUnderLocation(DesktopSpaceCoordinate) && TestCandidate.WidgetPtr.IsValid())
 			{
 				HitWidgetIndex = IndexesInCell[i];
 			}
@@ -138,12 +145,10 @@ int32 FHittestGrid::InsertWidget( const int32 ParentHittestIndex, const EVisibil
 {
 	check( ParentHittestIndex < WidgetsCachedThisFrame->Num() );
 
+	// Update the FGeometry to transform into desktop space.
 	FArrangedWidget WindowAdjustedWidget(Widget);
-	WindowAdjustedWidget.Geometry.AbsolutePosition += InWindowOffset;
-
-	const FSlateRect WindowAdjustedRect(
-		InClippingRect.Left + InWindowOffset.X, InClippingRect.Top + InWindowOffset.Y,
-		InClippingRect.Right + InWindowOffset.X, InClippingRect.Bottom + InWindowOffset.Y );
+	WindowAdjustedWidget.Geometry.AppendTransform(FSlateLayoutTransform(InWindowOffset));
+	const FSlateRect WindowAdjustedRect = InClippingRect.OffsetBy(InWindowOffset);
 
 	// Remember this widget, its geometry, and its place in the logical hierarchy.
 	const int32 WidgetIndex = WidgetsCachedThisFrame->Add( FCachedWidget( ParentHittestIndex, WindowAdjustedWidget, WindowAdjustedRect ) );
@@ -156,19 +161,21 @@ int32 FHittestGrid::InsertWidget( const int32 ParentHittestIndex, const EVisibil
 	if (Visibility.IsHitTestVisible())
 	{
 		// Mark any cell that is overlapped by this widget.
-		const FVector2D UpperLeft = WindowAdjustedWidget.Geometry.AbsolutePosition - GridOrigin;
-		const FVector2D LowerRight = UpperLeft + WindowAdjustedWidget.Geometry.Size*WindowAdjustedWidget.Geometry.Scale;
-		const FSlateRect GridSpaceWindowRect = FSlateRect(WindowAdjustedRect.Left - GridOrigin.X, WindowAdjustedRect.Top - GridOrigin.Y, WindowAdjustedRect.Right - GridOrigin.X, WindowAdjustedRect.Bottom - GridOrigin.Y);
-		const FSlateRect CellCoverageRect = FSlateRect(UpperLeft, LowerRight).IntersectionWith(GridSpaceWindowRect);
-		
+
+		// Compute the render space clipping rect, and compute it's aligned bounds so we can insert conservatively into the hit test grid.
+		FSlateRect GridRelativeBoundingClipRect = FSlateRotatedRect(
+			Inverse(WindowAdjustedWidget.Geometry.GetAccumulatedLayoutTransform()),
+			WindowAdjustedWidget.Geometry.GetAccumulatedRenderTransform(),
+			WindowAdjustedWidget.Geometry.GetClippingRect().IntersectionWith(WindowAdjustedRect)).ToBoundingRect().OffsetBy(-GridOrigin);
+
 		// Starting and ending cells covered by this widget.	
 		const FIntPoint UpperLeftCell = FIntPoint(
-			FMath::Max(0, FMath::FloorToInt(CellCoverageRect.Left / CellSize.X)),
-			FMath::Max(0, FMath::FloorToInt(CellCoverageRect.Top / CellSize.Y)));
+			FMath::Max(0, FMath::FloorToInt(GridRelativeBoundingClipRect.Left / CellSize.X)),
+			FMath::Max(0, FMath::FloorToInt(GridRelativeBoundingClipRect.Top / CellSize.Y)));
 
 		const FIntPoint LowerRightCell = FIntPoint(
-			FMath::Min( NumCells.X-1, FMath::FloorToInt(CellCoverageRect.Right / CellSize.X)),
-			FMath::Min( NumCells.Y-1, FMath::FloorToInt(CellCoverageRect.Bottom / CellSize.Y)));
+			FMath::Min( NumCells.X-1, FMath::FloorToInt(GridRelativeBoundingClipRect.Right / CellSize.X)),
+			FMath::Min( NumCells.Y-1, FMath::FloorToInt(GridRelativeBoundingClipRect.Bottom / CellSize.Y)));
 
 		for (int32 XIndex=UpperLeftCell.X; XIndex <= LowerRightCell.X; ++ XIndex )
 		{
