added AutoConsoleVariableSink to make it easier to use

diff --git a/Engine/Source/Runtime/Core/Public/HAL/IConsoleManager.h b/Engine/Source/Runtime/Core/Public/HAL/IConsoleManager.h
--- a/Engine/Source/Runtime/Core/Public/HAL/IConsoleManager.h
+++ b/Engine/Source/Runtime/Core/Public/HAL/IConsoleManager.h
@@ -510,6 +510,28 @@ private:
 };
 
 
+/**
+ * auto registering console variable sinks (register a callback function that is called when ever a cvar is changes by the user, changes are grouped and happen in specific engine spots during the frame/main loop)
+ */
+class CORE_API FAutoConsoleVariableSink
+{
+public:
+	/** Constructor, saves the argument for future removal from the console variable system **/
+	FAutoConsoleVariableSink(const FConsoleCommandDelegate& InCommand)
+		: Command(InCommand)
+	{
+		IConsoleManager::Get().RegisterConsoleVariableSink(Command);
+	}
+	/** Destructor, removes the console variable sink **/
+	virtual ~FAutoConsoleVariableSink()
+	{
+		IConsoleManager::Get().UnregisterConsoleVariableSink(Command);
+	}
+
+	const FConsoleCommandDelegate& Command;
+};
+
+
 /**
  * Base class for autoregistering console commands.
  */
