Removing modules from Linux server for reals.

(File missed CL 2276426 :/ ).

diff --git a/Engine/Source/Runtime/Core/Private/Linux/LinuxPlatformMisc.cpp b/Engine/Source/Runtime/Core/Private/Linux/LinuxPlatformMisc.cpp
--- a/Engine/Source/Runtime/Core/Private/Linux/LinuxPlatformMisc.cpp
+++ b/Engine/Source/Runtime/Core/Private/Linux/LinuxPlatformMisc.cpp
@@ -493,11 +493,14 @@ void FLinuxPlatformMisc::LoadPreInitModules()
 
 void FLinuxPlatformMisc::LoadStartupModules()
 {
-#if !IS_PROGRAM
-	FModuleManager::Get().LoadModule(TEXT("ALAudio"));
-	FModuleManager::Get().LoadModule(TEXT("SteamController"));
+#if !IS_PROGRAM && !UE_SERVER
+	FModuleManager::Get().LoadModule(TEXT("ALAudio"));	// added in Launch.Build.cs for non-server targets
 	FModuleManager::Get().LoadModule(TEXT("HeadMountedDisplay"));
-#endif // !IS_PROGRAM
+#endif // !IS_PROGRAM && !UE_SERVER
+
+#if WITH_STEAMCONTROLLER
+	FModuleManager::Get().LoadModule(TEXT("SteamController"));
+#endif // WITH_STEAMCONTROLLER
 
 #if WITH_EDITOR
 	FModuleManager::Get().LoadModule(TEXT("SourceCodeAccess"));
