Graceful handling of PlatformInitMultimedia failure.

Headless programs (servers) may want to throw message boxes. Without this change, they would be failing early attempting to initialize SDL for that purpose (this actually happened when a plugin was missing from a server build), and the error message about SDL initialization failure is unhelpful to pinpoint the actual problem.

diff --git a/Engine/Source/Runtime/Core/Private/Linux/LinuxApplication.cpp b/Engine/Source/Runtime/Core/Private/Linux/LinuxApplication.cpp
--- a/Engine/Source/Runtime/Core/Private/Linux/LinuxApplication.cpp
+++ b/Engine/Source/Runtime/Core/Private/Linux/LinuxApplication.cpp
@@ -32,7 +32,12 @@ FLinuxApplication* FLinuxApplication::CreateLinuxApplication()
 {
 	LinuxApplication = new FLinuxApplication();
 	
-	FPlatformMisc::PlatformInitMultimedia(); //	will not initialize more than once
+	if (!FPlatformMisc::PlatformInitMultimedia()) //	will not initialize more than once
+	{
+		UE_LOG(LogInit, Fatal, TEXT("FLinuxApplication::CreateLinuxApplication() : PlatformInitMultimedia() failed, cannot create application instance."));
+		// unreachable
+		return nullptr;
+	}
 
 #if DO_CHECK
 	uint32 InitializedSubsystems = SDL_WasInit(SDL_INIT_EVERYTHING);
@@ -960,7 +965,13 @@ bool FLinuxApplication::TryCalculatePopupWindowPosition( const FPlatformRect& In
 
 void FDisplayMetrics::GetDisplayMetrics(FDisplayMetrics& OutDisplayMetrics)
 {
-	FPlatformMisc::PlatformInitMultimedia(); //	will not initialize more than once
+	if (!FPlatformMisc::PlatformInitMultimedia()) //	will not initialize more than once
+	{
+		// consider making non-fatal and just returning bullshit? (which can be checked for and handled more gracefully)
+		UE_LOG(LogInit, Fatal, TEXT("FDisplayMetrics::GetDisplayMetrics: PlatformInitMultimedia() failed, cannot get display metrics"));
+		// unreachable
+		return;
+	}
 
 	SDL_Rect bounds;
 	SDL_GetDisplayBounds( 0, &bounds );
diff --git a/Engine/Source/Runtime/Core/Private/Linux/LinuxCursor.cpp b/Engine/Source/Runtime/Core/Private/Linux/LinuxCursor.cpp
--- a/Engine/Source/Runtime/Core/Private/Linux/LinuxCursor.cpp
+++ b/Engine/Source/Runtime/Core/Private/Linux/LinuxCursor.cpp
@@ -7,7 +7,12 @@
 
 FLinuxCursor::FLinuxCursor() : bHidden(false)
 {
-	FPlatformMisc::PlatformInitMultimedia(); //	will not initialize more than once
+	if (!FPlatformMisc::PlatformInitMultimedia()) //	will not initialize more than once
+	{
+		UE_LOG(LogInit, Fatal, TEXT("FLinuxCursor::FLinuxCursor() : PlatformInitMultimedia() failed, cannot construct cursor."));
+		// unreachable
+		return;
+	}
 
 #if DO_CHECK
 	uint32 InitializedSubsystems = SDL_WasInit(SDL_INIT_EVERYTHING);
diff --git a/Engine/Source/Runtime/Core/Private/Linux/LinuxPlatformMisc.cpp b/Engine/Source/Runtime/Core/Private/Linux/LinuxPlatformMisc.cpp
--- a/Engine/Source/Runtime/Core/Private/Linux/LinuxPlatformMisc.cpp
+++ b/Engine/Source/Runtime/Core/Private/Linux/LinuxPlatformMisc.cpp
@@ -159,7 +159,7 @@ void FLinuxPlatformMisc::PlatformInit()
 	}
 }
 
-void FLinuxPlatformMisc::PlatformInitMultimedia()
+bool FLinuxPlatformMisc::PlatformInitMultimedia()
 {
 	if (!GInitializedSDL)
 	{
@@ -167,13 +167,16 @@ void FLinuxPlatformMisc::PlatformInitMultimedia()
 		if (SDL_Init(SDL_INIT_EVERYTHING | SDL_INIT_NOPARACHUTE) != 0)
 		{
 			const char * SDLError = SDL_GetError();
-			UE_LOG(LogInit, Fatal, TEXT("Could not initialize SDL: %s"), SDLError);
-			// unreachable
-			return;
+
+			// do not fail at this point, allow caller handle failure
+			UE_LOG(LogInit, Warning, TEXT("Could not initialize SDL: %s"), SDLError);
+			return false;
 		}
 
 		GInitializedSDL = true;
 	}
+
+	return true;
 }
 
 void FLinuxPlatformMisc::PlatformTearDown()
@@ -278,7 +281,11 @@ EAppReturnType::Type FLinuxPlatformMisc::MessageBoxExt(EAppMsgType::Type MsgType
 {
 	int NumberOfButtons = 0;
 
-	FPlatformMisc::PlatformInitMultimedia(); //	will not initialize more than once
+	// if multimedia cannot be initialized for messagebox, just fall back to default implementation
+	if (!FPlatformMisc::PlatformInitMultimedia()) //	will not initialize more than once
+	{
+		return FGenericPlatformMisc::MessageBoxExt(MsgType, Text, Caption);
+	}
 
 #if DO_CHECK
 	uint32 InitializedSubsystems = SDL_WasInit(SDL_INIT_EVERYTHING);
diff --git a/Engine/Source/Runtime/Core/Private/Linux/LinuxWindow.cpp b/Engine/Source/Runtime/Core/Private/Linux/LinuxWindow.cpp
--- a/Engine/Source/Runtime/Core/Private/Linux/LinuxWindow.cpp
+++ b/Engine/Source/Runtime/Core/Private/Linux/LinuxWindow.cpp
@@ -44,7 +44,12 @@ void FLinuxWindow::Initialize( FLinuxApplication* const Application, const TShar
 	Definition = InDefinition;
 	OwningApplication = Application;
 
-	FPlatformMisc::PlatformInitMultimedia(); //	will not initialize more than once
+	if (!FPlatformMisc::PlatformInitMultimedia()) //	will not initialize more than once
+	{
+		UE_LOG(LogInit, Fatal, TEXT("FLinuxWindow::Initialize() : PlatformInitMultimedia() failed, cannot initialize window."));
+		// unreachable
+		return;
+	}
 
 #if DO_CHECK
 	uint32 InitializedSubsystems = SDL_WasInit(SDL_INIT_EVERYTHING);
diff --git a/Engine/Source/Runtime/Core/Public/Linux/LinuxPlatformMisc.h b/Engine/Source/Runtime/Core/Public/Linux/LinuxPlatformMisc.h
--- a/Engine/Source/Runtime/Core/Public/Linux/LinuxPlatformMisc.h
+++ b/Engine/Source/Runtime/Core/Public/Linux/LinuxPlatformMisc.h
@@ -167,7 +167,7 @@ struct CORE_API FLinuxPlatformMisc : public FGenericPlatformMisc
 	/**
 	 * Linux-specific function initializing video (and not only) subsystem.
 	 */
-	static void PlatformInitMultimedia();
+	static bool PlatformInitMultimedia();
 };
 
 typedef FLinuxPlatformMisc FPlatformMisc;
diff --git a/Engine/Source/Runtime/OpenGLDrv/Private/Linux/OpenGLLinux.cpp b/Engine/Source/Runtime/OpenGLDrv/Private/Linux/OpenGLLinux.cpp
--- a/Engine/Source/Runtime/OpenGLDrv/Private/Linux/OpenGLLinux.cpp
+++ b/Engine/Source/Runtime/OpenGLDrv/Private/Linux/OpenGLLinux.cpp
@@ -684,12 +684,17 @@ void PlatformRestoreDesktopDisplayMode()
 }
 
 
-bool	PlatformInitOpenGL()
+bool PlatformInitOpenGL()
 {
 	static bool bInitialized = false;
 	static bool bOpenGLSupported = false;
 
-	FPlatformMisc::PlatformInitMultimedia(); //	will not initialize more than once
+	if (!FPlatformMisc::PlatformInitMultimedia()) //	will not initialize more than once
+	{
+		UE_LOG(LogInit, Fatal, TEXT("PlatformInitOpenGL() : PlatformInitMultimedia() failed, cannot initialize OpenGL."));
+		// unreachable
+		return false;
+	}
 
 #if DO_CHECK
 	uint32 InitializedSubsystems = SDL_WasInit(SDL_INIT_EVERYTHING);
