SystemDefaultFont - IOS implementation added. Extended from Nick Atams work to make this available to us

diff --git a/Engine/Source/Runtime/Core/Private/GenericPlatform/GenericPlatformMisc.cpp b/Engine/Source/Runtime/Core/Private/GenericPlatform/GenericPlatformMisc.cpp
--- a/Engine/Source/Runtime/Core/Private/GenericPlatform/GenericPlatformMisc.cpp
+++ b/Engine/Source/Runtime/Core/Private/GenericPlatform/GenericPlatformMisc.cpp
@@ -753,6 +753,11 @@ void FGenericPlatformMisc::GetValidTargetPlatforms(class TArray<class FString>&
 	TargetPlatformNames.Add(FPlatformProperties::PlatformName());
 }
 
+TArray<uint8> FGenericPlatformMisc::GetSystemFontBytes()
+{
+	return TArray<uint8>();
+}
+
 const TCHAR* FGenericPlatformMisc::GetDefaultPathSeparator()
 {
 	return TEXT( "/" );
diff --git a/Engine/Source/Runtime/Core/Private/IOS/IOSPlatformMisc.cpp b/Engine/Source/Runtime/Core/Private/IOS/IOSPlatformMisc.cpp
--- a/Engine/Source/Runtime/Core/Private/IOS/IOSPlatformMisc.cpp
+++ b/Engine/Source/Runtime/Core/Private/IOS/IOSPlatformMisc.cpp
@@ -564,4 +564,160 @@ class IPlatformChunkInstall* FIOSPlatformMisc::GetPlatformChunkInstall()
 
 	return ChunkInstall;
 
-}
\ No newline at end of file
+}
+
+
+struct FFontHeader
+{
+	int32 Version;
+	uint16 NumTables;
+	uint16 SearchRange;
+	uint16 EntrySelector;
+	uint16 RangeShift;
+};
+
+struct FFontTableEntry
+{
+	uint32 Tag;
+	uint32 CheckSum;
+	uint32 Offset;
+	uint32 Length;
+};
+
+
+static uint32 CalcTableCheckSum(const uint32 *Table, uint32 NumberOfBytesInTable)
+{
+	uint32 Sum = 0;
+	uint32 NumLongs = (NumberOfBytesInTable + 3) / 4;
+	while (NumLongs-- > 0)
+	{
+		Sum += CFSwapInt32HostToBig(*Table++);
+	}
+	return Sum;
+}
+
+static uint32 CalcTableDataRefCheckSum(CFDataRef DataRef)
+{
+	const uint32 *DataBuff = (const uint32 *)CFDataGetBytePtr(DataRef);
+	uint32 DataLength = (uint32)CFDataGetLength(DataRef);
+	return CalcTableCheckSum(DataBuff, DataLength);
+}
+
+/**
+ * In order to get a system font from IOS we need to build one from the data we can gather from a CGFontRef
+ * @param InFontName - The name of the system font we are seeking to load.
+ * @param OutBytes - The data we have built for the font.
+ */
+void GetBytesForFont(const NSString* InFontName, OUT TArray<uint8>& OutBytes)
+{
+	CGFontRef cgFont = CGFontCreateWithFontName((CFStringRef)InFontName);
+
+	if (cgFont)
+	{
+		CFRetain(cgFont);
+
+		// Gather information on the font tags
+		CFArrayRef Tags = CGFontCopyTableTags(cgFont);
+		int TableCount = CFArrayGetCount(Tags);
+
+		// Collate the table sizes
+		TArray<size_t> TableSizes;
+
+		bool bContainsCFFTable = false;
+
+		size_t TotalSize = sizeof(FFontHeader)+sizeof(FFontTableEntry)* TableCount;
+		for (int TableIndex = 0; TableIndex < TableCount; ++TableIndex)
+		{
+			size_t TableSize = 0;
+			
+			uint32 aTag = (uint32)CFArrayGetValueAtIndex(Tags, TableIndex);
+			if (aTag == 'CFF ' && !bContainsCFFTable)
+			{
+				bContainsCFFTable = true;
+			}
+
+			CFDataRef TableDataRef = CGFontCopyTableForTag(cgFont, aTag);
+			if (TableDataRef != NULL)
+			{
+				TableSize = CFDataGetLength(TableDataRef);
+				CFRelease(TableDataRef);
+			}
+
+			TotalSize += (TableSize + 3) & ~3;
+			TableSizes.Add( TableSize );
+		}
+
+		OutBytes.Reserve( TotalSize );
+		OutBytes.AddZeroed( TotalSize );
+
+		// Start copying the table data into our buffer
+		uint8* DataStart = OutBytes.GetTypedData();
+		uint8* DataPtr = DataStart;
+
+		// Compute font header entries
+		uint16 EntrySelector = 0;
+		uint16 SearchRange = 1;
+		while (SearchRange < TableCount >> 1)
+		{
+			EntrySelector++;
+			SearchRange <<= 1;
+		}
+		SearchRange <<= 4;
+
+		uint16 RangeShift = (TableCount << 4) - SearchRange;
+
+		// Write font header (also called sfnt header, offset subtable)
+		FFontHeader* OffsetTable = (FFontHeader*)DataPtr;
+
+		// OpenType Font contains CFF Table use 'OTTO' as version, and with .otf extension
+		// otherwise 0001 0000
+		OffsetTable->Version = bContainsCFFTable ? 'OTTO' : CFSwapInt16HostToBig(1);
+		OffsetTable->NumTables = CFSwapInt16HostToBig((uint16)TableCount);
+		OffsetTable->SearchRange = CFSwapInt16HostToBig((uint16)SearchRange);
+		OffsetTable->EntrySelector = CFSwapInt16HostToBig((uint16)EntrySelector);
+		OffsetTable->RangeShift = CFSwapInt16HostToBig((uint16)RangeShift);
+
+		DataPtr += sizeof(FFontHeader);
+
+		// Write tables
+		FFontTableEntry* CurrentTableEntry = (FFontTableEntry*)DataPtr;
+		DataPtr += sizeof(FFontTableEntry) * TableCount;
+
+		for (int TableIndex = 0; TableIndex < TableCount; ++TableIndex)
+		{
+			uint32 aTag = (uint32)CFArrayGetValueAtIndex(Tags, TableIndex);
+			CFDataRef TableDataRef = CGFontCopyTableForTag(cgFont, aTag);
+			uint32 TableSize = CFDataGetLength(TableDataRef);
+
+			FMemory::Memcpy(DataPtr, CFDataGetBytePtr(TableDataRef), TableSize);
+
+			CurrentTableEntry->Tag = CFSwapInt32HostToBig((uint32_t)aTag);
+			CurrentTableEntry->CheckSum = CFSwapInt32HostToBig(CalcTableCheckSum((uint32 *)DataPtr, TableSize));
+
+			uint32 Offset = DataPtr - DataStart;
+			CurrentTableEntry->Offset = CFSwapInt32HostToBig((uint32)Offset);
+			CurrentTableEntry->Length = CFSwapInt32HostToBig((uint32)TableSize);
+
+			DataPtr += (TableSize + 3) & ~3;
+			++CurrentTableEntry;
+
+			CFRelease(TableDataRef);
+		}
+
+		CFRelease(cgFont);
+	}
+}
+
+
+TArray<uint8> FIOSPlatformMisc::GetSystemFontBytes()
+{
+	// Gather some details about the system font
+	uint32 SystemFontSize = [UIFont systemFontSize];
+	NSString* SystemFontName = [UIFont systemFontOfSize:SystemFontSize].fontName;
+
+	TArray<uint8> FontBytes;
+	GetBytesForFont(SystemFontName, FontBytes);
+
+	return FontBytes;
+}
+
diff --git a/Engine/Source/Runtime/Core/Private/Windows/WindowsPlatformMisc.cpp b/Engine/Source/Runtime/Core/Private/Windows/WindowsPlatformMisc.cpp
--- a/Engine/Source/Runtime/Core/Private/Windows/WindowsPlatformMisc.cpp
+++ b/Engine/Source/Runtime/Core/Private/Windows/WindowsPlatformMisc.cpp
@@ -2276,6 +2276,13 @@ FText FWindowsPlatformMisc::GetFileManagerName()
 	return NSLOCTEXT("WindowsPlatform", "FileManagerName", "Explorer");
 }
 
+TArray<uint8> FWindowsPlatformMisc::GetSystemFontBytes()
+{
+	TArray<uint8> FontBytes;
+	FFileHelper::LoadFileToArray( FontBytes, TEXT("C:\\Windows\\Fonts\\comicbd.ttf") );
+	return FontBytes;
+}
+
 bool FWindowsPlatformMisc::IsRunningOnBattery()
 {
 	SYSTEM_POWER_STATUS status;
diff --git a/Engine/Source/Runtime/Core/Public/GenericPlatform/GenericPlatformMisc.h b/Engine/Source/Runtime/Core/Public/GenericPlatform/GenericPlatformMisc.h
--- a/Engine/Source/Runtime/Core/Public/GenericPlatform/GenericPlatformMisc.h
+++ b/Engine/Source/Runtime/Core/Public/GenericPlatform/GenericPlatformMisc.h
@@ -785,6 +785,9 @@ public:
 		return PLATFORM_HAS_TOUCH_MAIN_SCREEN;
 	}
 
+	/** @return Memory representing a true type or open type font provided by the platform as a default font for unreal to consume; empty array if the default font failed to load. */
+	static TArray<uint8> GetSystemFontBytes();
+
 	/**
 	 * Returns whether the given platform feature is currently available (for instance, Metal is only available in IOS8 and with A7 devices)
 	 */
diff --git a/Engine/Source/Runtime/Core/Public/IOS/IOSPlatformMisc.h b/Engine/Source/Runtime/Core/Public/IOS/IOSPlatformMisc.h
--- a/Engine/Source/Runtime/Core/Public/IOS/IOSPlatformMisc.h
+++ b/Engine/Source/Runtime/Core/Public/IOS/IOSPlatformMisc.h
@@ -67,6 +67,7 @@ struct CORE_API FIOSPlatformMisc : public FGenericPlatformMisc
 	static FString GetDefaultLocale();
 	static bool SetStoredValue(const FString& InStoreId, const FString& InSectionName, const FString& InKeyName, const FString& InValue);
 	static bool GetStoredValue(const FString& InStoreId, const FString& InSectionName, const FString& InKeyName, FString& OutValue);
+	static TArray<uint8> GetSystemFontBytes();
 
 	static class IPlatformChunkInstall* GetPlatformChunkInstall();
     
diff --git a/Engine/Source/Runtime/Core/Public/Windows/WindowsPlatformMisc.h b/Engine/Source/Runtime/Core/Public/Windows/WindowsPlatformMisc.h
--- a/Engine/Source/Runtime/Core/Public/Windows/WindowsPlatformMisc.h
+++ b/Engine/Source/Runtime/Core/Public/Windows/WindowsPlatformMisc.h
@@ -207,6 +207,8 @@ struct CORE_API FWindowsPlatformMisc : public FGenericPlatformMisc
 	/** @return Get the name of the platform specific file manager (Explorer) */
 	static FText GetFileManagerName();
 
+	static TArray<uint8> GetSystemFontBytes();
+
 	/**
 	 * Returns whether the platform is running on battery power or not.
 	 */
diff --git a/Engine/Source/Runtime/SlateCore/Private/Fonts/FontCache.cpp b/Engine/Source/Runtime/SlateCore/Private/Fonts/FontCache.cpp
--- a/Engine/Source/Runtime/SlateCore/Private/Fonts/FontCache.cpp
+++ b/Engine/Source/Runtime/SlateCore/Private/Fonts/FontCache.cpp
@@ -381,48 +381,59 @@ private:
 	 */
 	FT_Face GetFontFace( const FName& FontName )
 	{
+		static const FName SpecialName_DefaultSystemFont("DefaultSystemFont");
+
 		FFontFaceAndMemory* FaceAndMemory = FontFaceMap.Find( FontName );
 		if (!FaceAndMemory)
 		{
 			// make a new entry
 			FaceAndMemory = &FontFaceMap.Add(FontName, FFontFaceAndMemory());
 
-			// load the font via UE4 methods, so that it will route through all proper file management (network file loading, etc)
-			FString FontPath = FontName.ToString();
-
 			// default to error condition
 			bool Error = true;
 
-			int64 FileSize = IFileManager::Get().FileSize(*FontPath);
-			if ( FileSize > 0 )
+			if (FontName == SpecialName_DefaultSystemFont)
 			{
-				// allocate space for the font
-				FaceAndMemory->Memory = (uint8*)FMemory::Malloc((uint32)FileSize);
-				FArchive* Ar = IFileManager::Get().CreateFileReader(*FontPath);
-				if (Ar != nullptr)
+				// Ask the platform to load the data for a default system font.
+				const TArray<uint8> FontBytes = FPlatformMisc::GetSystemFontBytes();
+				if ( FontBytes.Num() > 0 )
 				{
-					// read in the file
-					Ar->Serialize( FaceAndMemory->Memory, FileSize );
-					delete Ar;
-
+					FaceAndMemory->Memory = static_cast<uint8*>( FMemory::Malloc(FontBytes.Num()) );
+					FMemory::Memcpy( FaceAndMemory->Memory, FontBytes.GetData(), FontBytes.Num() );
+					
 					// initialize the font, setting the error code
-					Error = FT_New_Memory_Face( FTLibrary, FaceAndMemory->Memory, (FT_Long)FileSize, 0, &FaceAndMemory->Face ) != 0;
+					Error = FT_New_Memory_Face( FTLibrary, FaceAndMemory->Memory, static_cast<FT_Long>(FontBytes.Num()), 0, &FaceAndMemory->Face ) != 0;
 				}
-
-				// if it failed, we don't want to keep the memory around
-				if ( Error )
+			}
+			else
+			{
+				// load the font via UE4 methods, so that it will route through all proper file management (network file loading, etc)
+				const FString FontPath = FontName.ToString();
+				int64 FileSize = IFileManager::Get().FileSize(*FontPath);
+				if ( FileSize > 0 )
 				{
-					FMemory::Free(FaceAndMemory->Memory);
-					FontFaceMap.Remove(FontName);
-					FaceAndMemory = nullptr;
+					// allocate space for the font
+					FaceAndMemory->Memory = (uint8*)FMemory::Malloc((uint32)FileSize);
+					FArchive* Ar = IFileManager::Get().CreateFileReader(*FontPath);
+					if (Ar != nullptr)
+					{
+						// read in the file
+						Ar->Serialize( FaceAndMemory->Memory, FileSize );
+						delete Ar;
+
+						// initialize the font, setting the error code
+						Error = FT_New_Memory_Face( FTLibrary, FaceAndMemory->Memory, (FT_Long)FileSize, 0, &FaceAndMemory->Face ) != 0;
+					}
 				}
 			}
 
+			// if it failed, we don't want to keep the memory around
 			if ( Error )
 			{
-				UE_LOG( LogSlate, Warning, TEXT("GetFontFace failed to load or process '%s'"), *FontPath);
+				FMemory::Free(FaceAndMemory->Memory);
 				FontFaceMap.Remove(FontName);
 				FaceAndMemory = nullptr;
+				UE_LOG( LogSlate, Warning, TEXT("GetFontFace failed to load or process '%s'"), *FontName.ToString());
 			}
 		}
 
