CrashReport - Added ensure message to the crash report, reworked assertion message passing to the crash, reworked FailDebug method and renamed to StaticFailDebug, misc code clean-up, cleaned-up callstack dumping to the log file, now it should just dump the callstack once

diff --git a/Engine/Source/Runtime/Core/Private/HAL/OutputDevices.cpp b/Engine/Source/Runtime/Core/Private/HAL/OutputDevices.cpp
--- a/Engine/Source/Runtime/Core/Private/HAL/OutputDevices.cpp
+++ b/Engine/Source/Runtime/Core/Private/HAL/OutputDevices.cpp
@@ -1221,7 +1221,7 @@ void FOutputDeviceAnsiError::HandleError()
 	if (GLog)
 	{
 		// print to log and flush it
-		UE_LOG(LogHAL, Log, TEXT("=== Critical error: ===") LINE_TERMINATOR TEXT("%s") LINE_TERMINATOR, GErrorExceptionDescription);
+		UE_LOG( LogHAL, Log, TEXT( "=== Critical error: ===" ) LINE_TERMINATOR LINE_TERMINATOR TEXT( "%s" ) LINE_TERMINATOR, GErrorExceptionDescription );
 		UE_LOG(LogHAL, Log, GErrorHist);
 
 		GLog->Flush();
diff --git a/Engine/Source/Runtime/Core/Private/Misc/Core.cpp b/Engine/Source/Runtime/Core/Private/Misc/Core.cpp
--- a/Engine/Source/Runtime/Core/Private/Misc/Core.cpp
+++ b/Engine/Source/Runtime/Core/Private/Misc/Core.cpp
@@ -31,8 +31,16 @@ FOutputDeviceConsole*	GLogConsole						= NULL;						/* Console log hook */
 CORE_API FMalloc*		GMalloc							= NULL;						/* Memory allocator */
 
 class UPropertyWindowManager*	GPropertyWindowManager	= NULL;						/* Manages and tracks property editing windows */
-TCHAR					GErrorHist[16384]				= TEXT("");					/* For building call stack text dump in guard/unguard mechanism */
-TCHAR					GErrorExceptionDescription[1024] = TEXT("");				/* For building exception description text dump in guard/unguard mechanism */
+
+/** For building call stack text dump in guard/unguard mechanism. */
+TCHAR GErrorHist[16384]	= TEXT("");
+
+/** For building exception description text dump in guard/unguard mechanism. */
+TCHAR GErrorExceptionDescription[4096] = TEXT( "" );
+
+/** The error message, can be assertion message, ensure message or message from the fatal error. */
+TCHAR GErrorMessage[4096] = TEXT( "" );
+
 
 CORE_API const FText GYes	= LOCTEXT("Yes",	"Yes");
 CORE_API const FText GNo	= LOCTEXT("No",		"No");
diff --git a/Engine/Source/Runtime/Core/Private/Misc/OutputDevice.cpp b/Engine/Source/Runtime/Core/Private/Misc/OutputDevice.cpp
--- a/Engine/Source/Runtime/Core/Private/Misc/OutputDevice.cpp
+++ b/Engine/Source/Runtime/Core/Private/Misc/OutputDevice.cpp
@@ -145,10 +145,45 @@ VARARG_BODY( void, FOutputDevice::Logf, const TCHAR*, VARARG_NONE )
 	GROWABLE_LOGF(Serialize(Buffer, TLS->Verbosity, TLS->Category))
 }
 
+#define FILE_LINE_DESC TEXT(" [File:%s] [Line: %i] ")
 
-void FailDebug(const TCHAR* Error, const ANSICHAR* File, int32 Line, const TCHAR* Description)
+
+/** Critical errors. */
+CORE_API FOutputDeviceError* GError = NULL;
+
+/** Lock used to synchronize the fail debug calls. */
+static FCriticalSection	FailDebugCriticalSection;
+
+/**
+ *	Prints error to the debug output, 
+ *	prompts for the remote debugging if there is not debugger, breaks into the debugger 
+ *	and copies the error into the global error message.
+ */
+void StaticFailDebug( const TCHAR* Error, const ANSICHAR* File, int32 Line, const TCHAR* Description, bool bIsEnsure = false )
 {
-	FPlatformMisc::LowLevelOutputDebugStringf( TEXT("%s(%i): %s\n%s\n"), ANSI_TO_TCHAR(File), Line, Error, Description );
+	// For ensure log should be flushed in the engine loop.
+	if( !bIsEnsure )
+	{
+		GLog->PanicFlushThreadedLogs();
+	}
+
+	FScopeLock Lock( &FailDebugCriticalSection );
+
+	FPlatformMisc::LowLevelOutputDebugStringf( TEXT( "%s" ) FILE_LINE_DESC TEXT("\n%s\n" ), Error, ANSI_TO_TCHAR( File ), Line, Description );
+
+	// Copy the detailed error into the error message.
+	FCString::Snprintf( GErrorMessage, ARRAY_COUNT( GErrorMessage ), TEXT( "%s" ) FILE_LINE_DESC TEXT( "\n%s\n" ), Error, ANSI_TO_TCHAR( File ), Line, Description );
+
+	// Copy the error message to the error history.
+	FCString::Strncpy( GErrorHist, GErrorMessage, ARRAY_COUNT( GErrorHist ) );
+	FCString::Strncat( GErrorHist, TEXT( "\r\n\r\n" ), ARRAY_COUNT( GErrorHist ) );
+
+	if( !FPlatformMisc::IsDebuggerPresent() )
+	{
+		FPlatformMisc::PromptForRemoteDebugging( false );
+	}
+
+	FPlatformMisc::DebugBreak();
 }
 
 #if DO_CHECK || DO_GUARD_SLOW
@@ -158,51 +193,31 @@ void FailDebug(const TCHAR* Error, const ANSICHAR* File, int32 Line, const TCHAR
 //
 void VARARGS FDebug::AssertFailed( const ANSICHAR* Expr, const ANSICHAR* File, int32 Line, const TCHAR* Format/*=TEXT("")*/, ... )
 {
-	TCHAR DescriptionString[4096];
-	GET_VARARGS( DescriptionString, ARRAY_COUNT(DescriptionString), ARRAY_COUNT(DescriptionString)-1, Format, Format );
-	
-	if(FPlatformMisc::IsDebuggerPresent())
+#if !(UE_BUILD_SHIPPING || UE_BUILD_TEST)
+	// Walk the script stack, if any
+	if( GScriptStack.Num() > 0 )
 	{
-		TCHAR ErrorString[4096];
-		FCString::Sprintf(ErrorString,TEXT("Assertion failed: %s"),ANSI_TO_TCHAR(Expr));
-
-		FailDebug(ErrorString,File,Line,DescriptionString);
+		FString ScriptStack = TEXT( "\n\nScript Stack:\n" );
+		while( GScriptStack.Num() )
+		{
+			ScriptStack += GScriptStack.Pop().GetStackDescription() + TEXT( "\n" );
+		}
 
-		FPlatformMisc::DebugBreak();
-	}
-	else
-	{
-		FPlatformMisc::PromptForRemoteDebugging(false);
+		UE_LOG( LogOutputDevice, Warning, TEXT( "%s" ), *ScriptStack );
 	}
+#endif
 
 	// Ignore this assert if we're already forcibly shutting down because of a critical error.
-	// Note that appFailAssertFuncDebug() is still called.
-	if ( !GIsCriticalError )
+	if( !GIsCriticalError )
 	{
-		const SIZE_T StackTraceSize = 65535;
-		ANSICHAR* StackTrace = (ANSICHAR*) FMemory::SystemMalloc( StackTraceSize );
-		if( StackTrace != NULL )
-		{
-#if !(UE_BUILD_SHIPPING || UE_BUILD_TEST)
-			// Walk the script stack, if any
-			if (GScriptStack.Num() > 0)
-			{
-				FString ScriptStack = TEXT("\n\nScript Stack:\n");
-				while (GScriptStack.Num())
-				{
-					ScriptStack += GScriptStack.Pop().GetStackDescription() + TEXT("\n");
-				}
+		TCHAR DescriptionString[4096];
+		GET_VARARGS( DescriptionString, ARRAY_COUNT( DescriptionString ), ARRAY_COUNT( DescriptionString ) - 1, Format, Format );
 
-				UE_LOG(LogOutputDevice, Warning, TEXT("%s"), *ScriptStack);
-			}
-#endif
+		TCHAR ErrorString[MAX_SPRINTF];
+		FCString::Sprintf( ErrorString, TEXT( "Assertion failed: %s" ), ANSI_TO_TCHAR( Expr ) );
 
-			StackTrace[0] = 0;
-			// Walk the stack and dump it to the allocated memory.
-			FPlatformStackWalk::StackWalkAndDump( StackTrace, StackTraceSize, CALLSTACK_IGNOREDEPTH );
-			UE_LOG(LogOutputDevice, Fatal, TEXT("Assertion failed: %s [File:%s] [Line: %i]\n%s\nStack:\n%s"), ANSI_TO_TCHAR(Expr), ANSI_TO_TCHAR(File), Line, DescriptionString, ANSI_TO_TCHAR(StackTrace) );
-			FMemory::SystemFree( StackTrace );
-		}
+		StaticFailDebug( ErrorString, File, Line, DescriptionString );
+		GError->Logf( TEXT( "Assertion failed: %s" ) FILE_LINE_DESC TEXT( "\n%s\n" ), ErrorString, ANSI_TO_TCHAR( File ), Line, DescriptionString );
 	}
 }
 
@@ -226,9 +241,10 @@ void FDebug::EnsureFailed( const ANSICHAR* Expr, const ANSICHAR* File, int32 Lin
 	else
 	{
 		// Print initial debug message for this error
-		TCHAR ErrorString[4096];
+		TCHAR ErrorString[MAX_SPRINTF];
 		FCString::Sprintf(ErrorString,TEXT("Ensure condition failed: %s"),ANSI_TO_TCHAR(Expr));
-		FailDebug( ErrorString, File, Line, Msg );
+
+		StaticFailDebug( ErrorString, File, Line, Msg, true );
 
 		// Is there a debugger attached?  If so we'll just break, otherwise we'll submit an error report.
 		if( FPlatformMisc::IsDebuggerPresent() )
@@ -256,14 +272,14 @@ void FDebug::EnsureFailed( const ANSICHAR* Expr, const ANSICHAR* File, int32 Lin
 
 				// Create a final string that we'll output to the log (and error history buffer)
 				TCHAR ErrorMsg[16384];
-				FCString::Sprintf( ErrorMsg, TEXT("Ensure condition failed: %s [File:%s] [Line: %i]") LINE_TERMINATOR TEXT("%s") LINE_TERMINATOR TEXT("Stack: "), ANSI_TO_TCHAR(Expr), ANSI_TO_TCHAR(File), Line, Msg );
+				FCString::Snprintf( ErrorMsg, ARRAY_COUNT( ErrorMsg ), TEXT( "Ensure condition failed: %s [File:%s] [Line: %i]" ) LINE_TERMINATOR TEXT( "%s" ) LINE_TERMINATOR TEXT( "Stack: " ) LINE_TERMINATOR, ANSI_TO_TCHAR( Expr ), ANSI_TO_TCHAR( File ), Line, Msg );
 
 				// Also append the stack trace
 				FCString::Strncat( ErrorMsg, ANSI_TO_TCHAR(StackTrace), ARRAY_COUNT(ErrorMsg) - 1 );
 				FMemory::SystemFree( StackTrace );
 
 				// Dump the error and flush the log.
-				UE_LOG(LogOutputDevice, Log, TEXT("=== Handled error: ===") LINE_TERMINATOR TEXT("%s"), ErrorMsg);
+				UE_LOG(LogOutputDevice, Log, TEXT("=== Handled error: ===") LINE_TERMINATOR LINE_TERMINATOR TEXT("%s"), ErrorMsg);
 				GLog->Flush();
 
 				// Submit the error report to the server! (and display a balloon in the system tray)
@@ -319,7 +335,7 @@ void FDebug::EnsureFailed( const ANSICHAR* Expr, const ANSICHAR* File, int32 Lin
 			if ( bShouldSendNewReport )
 			{
 #if PLATFORM_DESKTOP
-				NewReportEnsure( ErrorString );
+				NewReportEnsure( GErrorMessage );
 #endif
 			}
 		}
@@ -328,22 +344,12 @@ void FDebug::EnsureFailed( const ANSICHAR* Expr, const ANSICHAR* File, int32 Lin
 
 #endif // DO_CHECK || DO_GUARD_SLOW
 
-//appLowLevelFatalErrorFunc
 void VARARGS FError::LowLevelFatal(const ANSICHAR* File, int32 Line, const TCHAR* Format, ... )
 {
 	TCHAR DescriptionString[4096];
-
 	GET_VARARGS( DescriptionString, ARRAY_COUNT(DescriptionString), ARRAY_COUNT(DescriptionString)-1, Format, Format );
 
-	FailDebug(TEXT("LowLevelFatalError"),File,Line,DescriptionString);
-
-	if(!FPlatformMisc::IsDebuggerPresent())
-	{
-		FPlatformMisc::PromptForRemoteDebugging(false);
-	}
-
-	FPlatformMisc::DebugBreak();
-
+	StaticFailDebug(TEXT("LowLevelFatalError"),File,Line,DescriptionString);
 	GError->Log(DescriptionString);
 }
 
@@ -459,33 +465,12 @@ VARARG_BODY( void VARARGS, FError::Throwf, const TCHAR*, VARARG_NONE )
 #else
 	UE_LOG(LogOutputDevice, Error, TEXT("THROW: %s"), TempStr);
 #endif
-}
-
-// Null output device.
-FOutputDeviceNull NullOut;
-/** Global output device redirector, can be static as FOutputDeviceRedirector explicitly empties TArray on TearDown */
-static FOutputDeviceRedirector LogRedirector;
-
-// Exception thrower.
-class FThrowOut : public FOutputDevice
-{
-public:
-	void Serialize( const TCHAR* V, ELogVerbosity::Type Verbosity, const class FName& Category ) override
-	{
-#if PLATFORM_EXCEPTIONS_DISABLED
-		FPlatformMisc::DebugBreak();
-#else
-		throw( V );
-#endif
-	}
-} ThrowOut;
+}					
 
-CORE_API FOutputDeviceError*			GError							= NULL;						/* Critical errors */
-CORE_API FOutputDevice*					GThrow							= &ThrowOut;				/* Exception thrower */
-
-// Statics to prevent FMsg::Logf from allocating too much stack memory
+/** Statics to prevent FMsg::Logf from allocating too much stack memory. */
 static FCriticalSection					MsgLogfStaticBufferGuard;
-static TCHAR							MsgLogfStaticBuffer[8192]; // Increased from 4096 to fix crashes in the renderthread without autoreporter
+/** Increased from 4096 to fix crashes in the renderthread without autoreporter. */
+static TCHAR							MsgLogfStaticBuffer[8192];
 
 VARARG_BODY(void, FMsg::Logf, const TCHAR*, VARARG_EXTRA(const ANSICHAR* File) VARARG_EXTRA(int32 Line) VARARG_EXTRA(const class FName& Category) VARARG_EXTRA(ELogVerbosity::Type Verbosity))
 {
@@ -516,9 +501,6 @@ VARARG_BODY(void, FMsg::Logf, const TCHAR*, VARARG_EXTRA(const ANSICHAR* File) V
 	}
 	else
 	{
-		// Flush logs queued by threads when we hit an assert because they will not make it to the log otherwise
-		GLog->PanicFlushThreadedLogs();
-
 		// Keep Message buffer small, in some cases, this code is executed with 16KB stack.
 		TCHAR Message[4096];
 		{
@@ -533,14 +515,7 @@ VARARG_BODY(void, FMsg::Logf, const TCHAR*, VARARG_EXTRA(const ANSICHAR* File) V
 			Message[ARRAY_COUNT(Message) - 1] = '\0';
 		}
 
-		FailDebug(TEXT("Fatal error:"), File, Line, Message);
-		if (!FPlatformMisc::IsDebuggerPresent())
-		{
-			FPlatformMisc::PromptForRemoteDebugging(false);
-		}
-
-		FPlatformMisc::DebugBreak();
-
+		StaticFailDebug(TEXT("Fatal error:"), File, Line, Message);
 		GError->Log(Category, Verbosity, Message);
 	}
 #endif
diff --git a/Engine/Source/Runtime/Core/Private/Windows/WindowsPlatformExceptionHandling.cpp b/Engine/Source/Runtime/Core/Private/Windows/WindowsPlatformExceptionHandling.cpp
--- a/Engine/Source/Runtime/Core/Private/Windows/WindowsPlatformExceptionHandling.cpp
+++ b/Engine/Source/Runtime/Core/Private/Windows/WindowsPlatformExceptionHandling.cpp
@@ -19,10 +19,7 @@
 
 namespace
 {
-static const int LocalBufferSize = 1024;
-volatile LONG ReportCrashCallCount = 0;
-const TCHAR* AssertionMessageStart = nullptr;	// pointer into GErrorHist, if non-null
-int AssertionMessageLength;						// Number of characters to copy after AssertionMessageStart
+static int32 ReportCrashCallCount = 0;
 
 /**
  * Write a Windows minidump to disk
@@ -36,7 +33,9 @@ bool WriteMinidump(const TCHAR* Path, LPEXCEPTION_POINTERS ExceptionInfo)
 	HANDLE FileHandle = CreateFileW(Path, GENERIC_WRITE, 0, NULL, CREATE_ALWAYS, FILE_ATTRIBUTE_NORMAL, NULL);
 	
 	if (FileHandle == INVALID_HANDLE_VALUE)
+	{
 		return false;
+	}
 
 	// Initialise structure required by MiniDumpWriteDumps
 	MINIDUMP_EXCEPTION_INFORMATION DumpExceptionInfo = {};
@@ -67,23 +66,16 @@ const TCHAR* GetEngineMode()
 /** 
  * Get one line of text to describe the crash
  */
-void GetCrashDescription(WER_REPORT_INFORMATION& ReportInformation, const TCHAR* ErrorMessage)
+void GetCrashDescription(WER_REPORT_INFORMATION& ReportInformation)
 {
-	if (ErrorMessage)
-	{
-		StringCchCat(ReportInformation.wzDescription, ARRAYSIZE(ReportInformation.wzDescription), ErrorMessage);
-	}
-	else
-	{
-		StringCchCopy(ReportInformation.wzDescription, ARRAYSIZE(ReportInformation.wzDescription), TEXT("The application crashed while running "));
+	StringCchCopy( ReportInformation.wzDescription, ARRAYSIZE( ReportInformation.wzDescription ), TEXT( "The application crashed while running " ) );
 
-		const TCHAR* Description =	IsRunningCommandlet()?	 	TEXT("a commandlet") :
-									GIsEditor?				 	TEXT("the editor") :
-									IsRunningDedicatedServer()?	TEXT("a server") :
-																TEXT("the game");
+	const TCHAR* Description = IsRunningCommandlet() ? TEXT( "a commandlet" ) :
+		GIsEditor ? TEXT( "the editor" ) :
+		IsRunningDedicatedServer() ? TEXT( "a server" ) :
+		TEXT( "the game" );
 
-		StringCchCat(ReportInformation.wzDescription, ARRAYSIZE(ReportInformation.wzDescription), Description);
-	}
+	StringCchCat( ReportInformation.wzDescription, ARRAYSIZE( ReportInformation.wzDescription ), Description );
 }
 
 #endif	// WINVER
@@ -125,79 +117,72 @@ void GetModuleVersion( TCHAR* ModuleName, TCHAR* StringBuffer, DWORD MaxSize )
  * Parameters 0 through 7 are predefined for Windows
  * Parameters 8 and 9 are user defined
  */
-void SetReportParameters( HREPORT ReportHandle, EXCEPTION_POINTERS* ExceptionInfo )
+void SetReportParameters( HREPORT ReportHandle, EXCEPTION_POINTERS* ExceptionInfo, const TCHAR* ErrorMessage )
 {
 	HRESULT Result;
-	TCHAR StringBuffer[LocalBufferSize] = {0};
-	TCHAR LocalBuffer[LocalBufferSize] = {0};
+	TCHAR StringBuffer[MAX_SPRINTF] = {0};
+	TCHAR LocalBuffer[MAX_SPRINTF] = {0};
 
 	// Set the parameters for the standard problem signature
 	HMODULE ModuleHandle = GetModuleHandle( NULL );
 
-	StringCchPrintf( StringBuffer, LocalBufferSize, TEXT( "UE4-%s" ), FApp::GetGameName() );
+	StringCchPrintf( StringBuffer, MAX_SPRINTF, TEXT( "UE4-%s" ), FApp::GetGameName() );
 	Result = WerReportSetParameter( ReportHandle, WER_P0, TEXT( "Application Name" ), StringBuffer );
 
-	GetModuleFileName( ModuleHandle, LocalBuffer, LocalBufferSize );
+	GetModuleFileName( ModuleHandle, LocalBuffer, MAX_SPRINTF );
 	PathStripPath( LocalBuffer );
-	GetModuleVersion( LocalBuffer, StringBuffer, LocalBufferSize );
+	GetModuleVersion( LocalBuffer, StringBuffer, MAX_SPRINTF );
 	Result = WerReportSetParameter( ReportHandle, WER_P1, TEXT( "Application Version" ), StringBuffer );
 
-	StringCchPrintf( StringBuffer, LocalBufferSize, TEXT( "%08x" ), GetTimestampForLoadedLibrary( ModuleHandle ) );
+	StringCchPrintf( StringBuffer, MAX_SPRINTF, TEXT( "%08x" ), GetTimestampForLoadedLibrary( ModuleHandle ) );
 	Result = WerReportSetParameter( ReportHandle, WER_P2, TEXT( "Application Timestamp" ), StringBuffer );
 
 	HMODULE FaultModuleHandle = NULL;
 	GetModuleHandleEx( GET_MODULE_HANDLE_EX_FLAG_FROM_ADDRESS | GET_MODULE_HANDLE_EX_FLAG_UNCHANGED_REFCOUNT, ( LPCTSTR )ExceptionInfo->ExceptionRecord->ExceptionAddress, &FaultModuleHandle );
 
-	GetModuleFileName( FaultModuleHandle, LocalBuffer, LocalBufferSize );
+	GetModuleFileName( FaultModuleHandle, LocalBuffer, MAX_SPRINTF );
 	PathStripPath( LocalBuffer );
 	Result = WerReportSetParameter( ReportHandle, WER_P3, TEXT( "Fault Module Name" ), LocalBuffer );
 
-	GetModuleVersion( LocalBuffer, StringBuffer, LocalBufferSize );
+	GetModuleVersion( LocalBuffer, StringBuffer, MAX_SPRINTF );
 	Result = WerReportSetParameter( ReportHandle, WER_P4, TEXT( "Fault Module Version" ), StringBuffer );
 
-	StringCchPrintf( StringBuffer, LocalBufferSize, TEXT( "%08x" ), GetTimestampForLoadedLibrary( FaultModuleHandle ) );
+	StringCchPrintf( StringBuffer, MAX_SPRINTF, TEXT( "%08x" ), GetTimestampForLoadedLibrary( FaultModuleHandle ) );
 	Result = WerReportSetParameter( ReportHandle, WER_P5, TEXT( "Fault Module Timestamp" ), StringBuffer );
 
-	StringCchPrintf( StringBuffer, LocalBufferSize, TEXT( "%08x" ), ExceptionInfo->ExceptionRecord->ExceptionCode );
+	StringCchPrintf( StringBuffer, MAX_SPRINTF, TEXT( "%08x" ), ExceptionInfo->ExceptionRecord->ExceptionCode );
 	Result = WerReportSetParameter( ReportHandle, WER_P6, TEXT( "Exception Code" ), StringBuffer );
 
 	INT_PTR ExceptionOffset = ( char* )( ExceptionInfo->ExceptionRecord->ExceptionAddress ) - ( char* )FaultModuleHandle;
 	CA_SUPPRESS(6066) // The format specifier should probably be something like %tX, but VS 2013 doesn't support 't'.
-	StringCchPrintf( StringBuffer, LocalBufferSize, TEXT( "%p" ), ExceptionOffset );
+		StringCchPrintf( StringBuffer, MAX_SPRINTF, TEXT( "%p" ), ExceptionOffset );
 	Result = WerReportSetParameter( ReportHandle, WER_P7, TEXT( "Exception Offset" ), StringBuffer );
 
-	// Look for the assertion fail.
-	if (AssertionMessageStart)
-	{
-		// Use LocalBuffer to store the assertion fail.
-		FCString::Strncpy(LocalBuffer, AssertionMessageStart, AssertionMessageLength);
+	// Use LocalBuffer to store the error message.
+	FCString::Strncpy( LocalBuffer, ErrorMessage, MAX_SPRINTF );
 
-		// Replace " with ' and replace \n with #
-		for (TCHAR& Char: LocalBuffer)
+	// Replace " with ' and replace \n with #
+	for (TCHAR& Char: LocalBuffer)
+	{
+		if (Char == 0)
 		{
-			if (Char == 0)
-			{
-				break;
-			}
+			break;
+		}
 
-			switch (Char)
-			{
-				default: break;
-				case '"':	Char = '\'';	break;
-				case '\r':	Char = '#';		break;
-				case '\n':	Char = '#';		break;
-			}
+		switch (Char)
+		{
+			default: break;
+			case '"':	Char = '\'';	break;
+			case '\r':	Char = '#';		break;
+			case '\n':	Char = '#';		break;
 		}
 	}
-	else
-	{
-		FCString::Strcpy( LocalBuffer, TEXT("") );
-	}
 
-	StringCchPrintf( StringBuffer, LocalBufferSize, TEXT( "!%s!AssertLog=\"%s\"" ), FCommandLine::Get(), LocalBuffer );
+	// AssertLog should be ErrorMessage, but this require crash server changes, so don't change this.
+	StringCchPrintf( StringBuffer, MAX_SPRINTF, TEXT( "!%s!AssertLog=\"%s\"" ), FCommandLine::Get(), LocalBuffer );
 	Result = WerReportSetParameter( ReportHandle, WER_P8, TEXT( "Commandline" ), StringBuffer );
 
-	StringCchPrintf( StringBuffer, LocalBufferSize, TEXT( "%s!%s!%s!%d" ), TEXT( BRANCH_NAME ), FPlatformProcess::BaseDir(), GetEngineMode(), BUILT_FROM_CHANGELIST );
+	StringCchPrintf( StringBuffer, MAX_SPRINTF, TEXT( "%s!%s!%s!%d" ), TEXT( BRANCH_NAME ), FPlatformProcess::BaseDir(), GetEngineMode(), BUILT_FROM_CHANGELIST );
 	Result = WerReportSetParameter( ReportHandle, WER_P9, TEXT( "BranchBaseDir" ), StringBuffer );
 }
 
@@ -208,18 +193,12 @@ void SetReportParameters( HREPORT ReportHandle, EXCEPTION_POINTERS* ExceptionInf
  */
 void AddMiniDump(HREPORT ReportHandle, EXCEPTION_POINTERS* ExceptionInfo)
 {
-#if 0
-	// This doesn't work in x64
-	WER_EXCEPTION_INFORMATION WerExceptionInfo = { ExceptionInfo, FALSE };
-	Result = WerReportAddDump( ReportHandle, GetCurrentProcess(), GetCurrentThread(), WerDumpTypeMiniDump, &WerExceptionInfo, NULL, 0 );
-#else
 	FString MinidumpFileName = FString::Printf(TEXT("%sDump%d.dmp"), *FPaths::GameLogDir(), FDateTime::UtcNow().GetTicks());
 	
 	if (WriteMinidump(*MinidumpFileName, ExceptionInfo))
 	{
 		WerReportAddFile(ReportHandle, *MinidumpFileName, WerFileTypeMinidump, WER_FILE_ANONYMOUS_DATA); 
 	}
-#endif
 }
 
 /** 
@@ -309,19 +288,22 @@ int32 ReportCrashUsingCrashReportClient(EXCEPTION_POINTERS* ExceptionInfo, const
 		WER_REPORT_INFORMATION ReportInformation = {sizeof( WER_REPORT_INFORMATION )};
 
 		StringCchCopy( ReportInformation.wzConsentKey, ARRAYSIZE( ReportInformation.wzConsentKey ), TEXT( "" ) );
-		StringCchCopy( ReportInformation.wzApplicationName, ARRAYSIZE( ReportInformation.wzApplicationName ), FApp::GetGameName() );
+
+		StringCchCopy( ReportInformation.wzApplicationName, ARRAYSIZE( ReportInformation.wzApplicationName ), TEXT("UE4-") );
+		StringCchCat( ReportInformation.wzApplicationName, ARRAYSIZE( ReportInformation.wzApplicationName ), FApp::GetGameName() );
+
 		StringCchCopy( ReportInformation.wzApplicationPath, ARRAYSIZE( ReportInformation.wzApplicationPath ), FPlatformProcess::BaseDir() );
 		StringCchCat( ReportInformation.wzApplicationPath, ARRAYSIZE( ReportInformation.wzApplicationPath ), FPlatformProcess::ExecutableName() );
 		StringCchCat( ReportInformation.wzApplicationPath, ARRAYSIZE( ReportInformation.wzApplicationPath ), TEXT( ".exe" ) );
 
-		GetCrashDescription( ReportInformation, ErrorMessage );
+		GetCrashDescription( ReportInformation );
 
 		// Create a crash event report
 		HREPORT ReportHandle = NULL;
 		if( WerReportCreate( APPCRASH_EVENT, WerReportApplicationCrash, &ReportInformation, &ReportHandle ) == S_OK )
 		{
 			// Set the standard set of a crash parameters
-			SetReportParameters( ReportHandle, ExceptionInfo );
+			SetReportParameters( ReportHandle, ExceptionInfo, ErrorMessage );
 
 			// Add a manually generated minidump
 			AddMiniDump( ReportHandle, ExceptionInfo );
@@ -339,7 +321,9 @@ int32 ReportCrashUsingCrashReportClient(EXCEPTION_POINTERS* ExceptionInfo, const
 
 		// Build machines do not upload these automatically since it is not okay to have lingering processes after the build completes.
 		if( GIsBuildMachine )
+		{
 			return EXCEPTION_CONTINUE_EXECUTION;
+		}
 
 		FString CrashReportClientArguments;
 
@@ -424,6 +408,7 @@ void NewReportEnsure( const TCHAR* ErrorMessage )
 #include "AllowWindowsPlatformTypes.h"
 void CreateExceptionInfoString(EXCEPTION_RECORD* ExceptionRecord)
 {
+	// @TODO yrx 2014-08-18 Fix FString usage?
 	FString ErrorString = TEXT("Unhandled Exception: ");
 
 #define HANDLE_CASE(x) case x: ErrorString += TEXT(#x); break;
@@ -456,7 +441,7 @@ void CreateExceptionInfoString(EXCEPTION_RECORD* ExceptionRecord)
 	}
 
 #if WITH_EDITORONLY_DATA
-	FCString::Strncpy(GErrorExceptionDescription, *ErrorString, FMath::Min(ErrorString.Len() + 1, (int32)ARRAY_COUNT(GErrorExceptionDescription)));
+	FCString::Strncpy(GErrorExceptionDescription, *ErrorString, ARRAY_COUNT(GErrorExceptionDescription));
 #endif
 #undef HANDLE_CASE
 }
@@ -474,41 +459,36 @@ int32 ReportCrash( LPEXCEPTION_POINTERS ExceptionInfo )
 {
 	// Only create a minidump the first time this function is called.
 	// (Can be called the first time from the RenderThread, then a second time from the MainThread.)
-	if (::InterlockedIncrement(&ReportCrashCallCount) != 1)
+	if( FPlatformAtomics::InterlockedIncrement( &ReportCrashCallCount ) != 1 )
+	{
 		return EXCEPTION_EXECUTE_HANDLER;
+	}
 
 	FCoreDelegates::OnHandleSystemError.Broadcast();
 
-	const TCHAR* AssertPos = FCString::Stristr(GErrorHist, TEXT("Assertion failed:"));
-	const TCHAR* StackPos = FCString::Stristr(GErrorHist, TEXT("\nStack:"));
-	if (AssertPos && StackPos && StackPos > AssertPos)
-	{
-		// The assertion code has already written a callstack to the log
-		AssertionMessageStart = AssertPos;
-		AssertionMessageLength = FMath::Min(int(StackPos - AssertPos), LocalBufferSize);
-	}
-	else
+	// Note: ReportCrashUsingCrashReportClient reads the callstack from GErrorHist
+	const SIZE_T StackTraceSize = 65535;
+	ANSICHAR* StackTrace = (ANSICHAR*) FMemory::SystemMalloc( StackTraceSize );
+	StackTrace[0] = 0;
+	// Walk the stack and dump it to the allocated memory.
+	FPlatformStackWalk::StackWalkAndDump( StackTrace, StackTraceSize, 0, ExceptionInfo->ContextRecord );
+
+	// @TODO yrx 2014-08-20 Make this constant.
+	if( ExceptionInfo->ExceptionRecord->ExceptionCode != 1 )
 	{
-		AssertionMessageStart = GErrorHist;
-		AssertionMessageLength = FCString::Strlen(GErrorHist);
-
-		// Note: ReportCrashUsingCrashReportClient reads the callstack from GErrorHist
-		const SIZE_T StackTraceSize = 65535;
-		ANSICHAR* StackTrace = (ANSICHAR*) FMemory::SystemMalloc( StackTraceSize );
-		StackTrace[0] = 0;
-		// Walk the stack and dump it to the allocated memory.
-		FPlatformStackWalk::StackWalkAndDump( StackTrace, StackTraceSize, 0, ExceptionInfo->ContextRecord );
-		FCString::Strncat( GErrorHist, ANSI_TO_TCHAR(StackTrace), ARRAY_COUNT(GErrorHist) - 1 );
-		CreateExceptionInfoString(ExceptionInfo->ExceptionRecord);
-		FMemory::SystemFree( StackTrace );
+		CreateExceptionInfoString( ExceptionInfo->ExceptionRecord );
+		FCString::Strncat( GErrorHist, GErrorExceptionDescription, ARRAY_COUNT( GErrorHist ) );
+		FCString::Strncat( GErrorHist, TEXT( "\r\n\r\n" ), ARRAY_COUNT( GErrorHist ) );
 	}
 
+	FCString::Strncat( GErrorHist, ANSI_TO_TCHAR(StackTrace), ARRAY_COUNT(GErrorHist) );
+
+	FMemory::SystemFree( StackTrace );
 
 #if WINVER > 0x502	// Windows Error Reporting is not supported on Windows XP
-	// 
 	if (GUseCrashReportClient)
 	{
-		ReportCrashUsingCrashReportClient(ExceptionInfo, nullptr, EErrorReportUI::ShowDialog);
+		ReportCrashUsingCrashReportClient(ExceptionInfo, GErrorMessage, EErrorReportUI::ShowDialog);
 	}
 	else
 #endif		// WINVER
diff --git a/Engine/Source/Runtime/Core/Private/Windows/WindowsPlatformOutputDevices.cpp b/Engine/Source/Runtime/Core/Private/Windows/WindowsPlatformOutputDevices.cpp
--- a/Engine/Source/Runtime/Core/Private/Windows/WindowsPlatformOutputDevices.cpp
+++ b/Engine/Source/Runtime/Core/Private/Windows/WindowsPlatformOutputDevices.cpp
@@ -67,32 +67,17 @@ void FOutputDeviceWindowsError::Serialize( const TCHAR* Msg, ELogVerbosity::Type
 {
 	FPlatformMisc::DebugBreak();
    
-	int32 Error = GetLastError();
 	if( !GIsCriticalError )
 	{   
+		const int32 LastError = ::GetLastError();
+
 		// First appError.
 		GIsCriticalError = 1;
 		TCHAR ErrorBuffer[1024];
 		ErrorBuffer[0] = 0;
-		// pop up a crash window if we are not in unattended mode
-		if( FApp::IsUnattended() == false )
-		{ 
-			UE_LOG(LogWindows, Error, TEXT("appError called: %s"), Msg );
-
-			// Windows error.
-			UE_LOG(LogWindows, Error, TEXT("Windows GetLastError: %s (%i)"), FPlatformMisc::GetSystemErrorMessage(ErrorBuffer,1024,Error), Error );
-		} 
-		// log the warnings to the log
-		else
-		{
-			UE_LOG(LogWindows, Error, TEXT("appError called: %s"), Msg );
-
-			// Windows error.
-			UE_LOG(LogWindows, Error, TEXT("Windows GetLastError: %s (%i)"), FPlatformMisc::GetSystemErrorMessage(ErrorBuffer,1024,Error), Error );
-		}
 
-		FCString::Strncpy( GErrorHist, Msg, ARRAY_COUNT(GErrorHist) - 5 );
-		FCString::Strncat( GErrorHist, TEXT("\r\n\r\n"), ARRAY_COUNT(GErrorHist) - 1  );
+		// Windows error.
+		UE_LOG( LogWindows, Error, TEXT( "Windows GetLastError: %s (%i)" ), FPlatformMisc::GetSystemErrorMessage( ErrorBuffer, 1024, LastError ), LastError );
 	}
 	else
 	{
@@ -137,10 +122,9 @@ void FOutputDeviceWindowsError::HandleError()
 	FCoreDelegates::OnHandleSystemError.Broadcast();
 
 	// Dump the error and flush the log.
-	UE_LOG(LogWindows, Log, TEXT("=== Critical error: ===") LINE_TERMINATOR TEXT("%s") LINE_TERMINATOR, GErrorExceptionDescription);
-	UE_LOG(LogWindows, Log, GErrorHist);
+	UE_LOG( LogWindows, Log, TEXT( "=== Critical error: ===" ) LINE_TERMINATOR TEXT( "%s" ) LINE_TERMINATOR, GErrorHist );
 
-	GLog->Flush();
+	GLog->PanicFlushThreadedLogs();
 
 	// Unhide the mouse.
 	while( ::ShowCursor(true)<0 );
diff --git a/Engine/Source/Runtime/Core/Public/Core.h b/Engine/Source/Runtime/Core/Public/Core.h
--- a/Engine/Source/Runtime/Core/Public/Core.h
+++ b/Engine/Source/Runtime/Core/Public/Core.h
@@ -122,7 +122,10 @@ extern CORE_API ITransaction* GUndo;
 extern CORE_API FOutputDeviceConsole* GLogConsole;
 
 extern CORE_API TCHAR GErrorHist[16384];
-extern CORE_API TCHAR GErrorExceptionDescription[1024];
+
+// @TODO yrx 2014-08-19 Combine into one.
+extern CORE_API TCHAR GErrorExceptionDescription[4096];
+extern CORE_API TCHAR GErrorMessage[4096];
 
 extern CORE_API const FText GTrue, GFalse, GYes, GNo, GNone;
 
diff --git a/Engine/Source/Runtime/Core/Public/Misc/OutputDevice.h b/Engine/Source/Runtime/Core/Public/Misc/OutputDevice.h
--- a/Engine/Source/Runtime/Core/Public/Misc/OutputDevice.h
+++ b/Engine/Source/Runtime/Core/Public/Misc/OutputDevice.h
@@ -10,7 +10,6 @@ class FText;
 
 // Globals.
 #define GLog FOutputDeviceRedirector::Get()
-CORE_API extern class FOutputDevice*				GThrow;
 CORE_API extern class FOutputDeviceError*			GError;
 CORE_API extern class FFeedbackContext*				GWarn;
 
