When separating the main & game threads on OS X use custom run-loop modes so that re-entrant event handling only process the required events, more akin to Windows, to prevent problems in Slate.

diff --git a/Engine/Source/Developer/DesktopPlatform/Private/Mac/MacNativeFeedbackContext.cpp b/Engine/Source/Developer/DesktopPlatform/Private/Mac/MacNativeFeedbackContext.cpp
--- a/Engine/Source/Developer/DesktopPlatform/Private/Mac/MacNativeFeedbackContext.cpp
+++ b/Engine/Source/Developer/DesktopPlatform/Private/Mac/MacNativeFeedbackContext.cpp
@@ -261,7 +261,7 @@ FMacNativeFeedbackContext::FMacNativeFeedbackContext()
 {
 	MainThreadCall(^{
 		WindowController = [[FMacNativeFeedbackContextWindowController alloc] init];
-	}, true);
+	}, UE4NilEventMode, true);
 	SetDefaultTextColor();
 }
 
@@ -274,7 +274,7 @@ FMacNativeFeedbackContext::~FMacNativeFeedbackContext()
 	
 	MainThreadCall(^{
 		[WindowController release];
-	}, true);
+	}, UE4NilEventMode, true);
 }
 
 void FMacNativeFeedbackContext::Serialize( const TCHAR* Data, ELogVerbosity::Type Verbosity, const class FName& Category )
@@ -340,7 +340,7 @@ void FMacNativeFeedbackContext::Serialize( const TCHAR* Data, ELogVerbosity::Typ
 						[Colors release];
 						[AttributeKeys release];
 						OutstandingTasks--;
-					}, false);
+					}, NSDefaultRunLoopMode, false);
 				}
 			}
 			else
@@ -360,7 +360,7 @@ void FMacNativeFeedbackContext::Serialize( const TCHAR* Data, ELogVerbosity::Typ
 					[AttributedString release];
 					CFRelease(CocoaText);
 					OutstandingTasks--;
-				}, false);
+				}, NSDefaultRunLoopMode, false);
 				
 				if(!MacApplication)
 				{
@@ -417,7 +417,7 @@ void FMacNativeFeedbackContext::BeginSlowTask( const FText& Task, bool bShowProg
 			[WindowController showWindow];
 			
 			bShowingConsoleForSlowTask = true;
-		}, true);
+		});
 	}
 }
 
@@ -434,7 +434,7 @@ void FMacNativeFeedbackContext::EndSlowTask()
 				[WindowController hideWindow];
 			}
 			bShowingConsoleForSlowTask = false;
-		}, true);
+		});
 	}
 }
 
@@ -444,7 +444,7 @@ bool FMacNativeFeedbackContext::StatusUpdate( int32 Numerator, int32 Denominator
 	{
 		MainThreadCall(^{
 			[WindowController setStatusText:NewStatus.ToString().GetNSString()];
-		}, true);
+		});
 		UpdateProgress(Numerator, Denominator);
 	}
 	return true;
@@ -463,7 +463,7 @@ void FMacNativeFeedbackContext::UpdateProgress(int32 Numerator, int32 Denominato
 			[WindowController setShowProgress:true];
 			[WindowController setProgress:Numerator total:Denominator];
 		}
-	}, true);
+	});
 }
 
 FContextSupplier* FMacNativeFeedbackContext::GetContext() const
@@ -500,5 +500,5 @@ void FMacNativeFeedbackContext::SetDefaultTextColor()
 		[Colors release];
 		[AttributeKeys release];
 		OutstandingTasks--;
-	}, false);
+	}, NSDefaultRunLoopMode, false);
 }
diff --git a/Engine/Source/Runtime/Core/Private/Mac/CocoaTextView.cpp b/Engine/Source/Runtime/Core/Private/Mac/CocoaTextView.cpp
--- a/Engine/Source/Runtime/Core/Private/Mac/CocoaTextView.cpp
+++ b/Engine/Source/Runtime/Core/Private/Mac/CocoaTextView.cpp
@@ -178,7 +178,7 @@
 				GameThreadCall(^{
 					ITextInputMethodContext::ECaretPosition CaretPosition;
 					IMMContext->GetSelectionRange(SelectionLocation, SelectionLength, CaretPosition);
-				}, true);
+				}, InGameRunLoopMode(@[ UE4IMEEventMode ]));
 			}
 		}
 		else
@@ -201,7 +201,7 @@
 			FString TheFString(TheString);
 			IMMContext->SetTextInRange(SelectionLocation, SelectionLength, TheFString);
 			IMMContext->SetSelectionRange(SelectionLocation+TheFString.Len(), 0, ITextInputMethodContext::ECaretPosition::Ending);
-		}, true);
+		}, InGameRunLoopMode(@[ UE4IMEEventMode ]));
 		
 		[self unmarkText];
 		[[self inputContext] invalidateCharacterCoordinates]; // recentering
@@ -239,7 +239,7 @@
 				GameThreadCall(^{
 					ITextInputMethodContext::ECaretPosition CaretPosition;
 					IMMContext->GetSelectionRange(SelectionLocation, SelectionLength, CaretPosition);
-				}, true);
+				}, InGameRunLoopMode(@[ UE4IMEEventMode ]));
 			}
 		}
 		else
@@ -252,7 +252,7 @@
 		{
 			GameThreadCall(^{
 				IMMContext->SetTextInRange(SelectionLocation, SelectionLength, FString());
-			}, true);
+			}, InGameRunLoopMode(@[ UE4IMEEventMode ]));
 			[self unmarkText];
 		}
 		else
@@ -261,7 +261,7 @@
 			{
 				GameThreadCall(^{
 					IMMContext->BeginComposition();
-				}, true);
+				}, InGameRunLoopMode(@[ UE4IMEEventMode ]));
 			}
 			markedRange = NSMakeRange(SelectionLocation, [aString length]);
 			
@@ -302,7 +302,7 @@
 				IMMContext->SetTextInRange(SelectionLocation, SelectionLength, TheFString);
 				IMMContext->UpdateCompositionRange(CompositionRange.location, CompositionRange.length);
 				IMMContext->SetSelectionRange(markedRange.location + selectedRange.location, 0, ITextInputMethodContext::ECaretPosition::Ending);
-			}, true);
+			}, InGameRunLoopMode(@[ UE4IMEEventMode ]));
 		}
 		[[self inputContext] invalidateCharacterCoordinates]; // recentering
 	}
@@ -324,7 +324,7 @@
 			GameThreadCall(^{
 				IMMContext->UpdateCompositionRange(0, 0);
 				IMMContext->EndComposition();
-			}, true);
+			}, InGameRunLoopMode(@[ UE4IMEEventMode ]));
 		}
 	}
 }
@@ -340,7 +340,7 @@
 		__block ITextInputMethodContext::ECaretPosition CaretPosition;
 		GameThreadCall(^{
 			IMMContext->GetSelectionRange(SelectionLocation, SelectionLength, CaretPosition);
-		}, true);
+		}, InGameRunLoopMode(@[ UE4IMEEventMode ]));
 		return {SelectionLocation, SelectionLength};
 	}
 	else
@@ -380,7 +380,7 @@
 		__block FString String;
 		GameThreadCall(^{
 			IMMContext->GetTextInRange(aRange.location, aRange.length, String);
-		}, true);
+		}, InGameRunLoopMode(@[ UE4IMEEventMode ]));
 		CFStringRef CFString = FPlatformString::TCHARToCFString(*String);
 		if(CFString)
 		{
@@ -415,7 +415,7 @@
 		__block FVector2D Size;
 		GameThreadCall(^{
 			IMMContext->GetTextBounds(aRange.location, aRange.length, Position, Size);
-		}, true);
+		}, InGameRunLoopMode(@[ UE4IMEEventMode ]));
 		
 		if(actualRange)
 		{
@@ -443,7 +443,7 @@
 	if (IMMContext.IsValid())
 	{
 		FVector2D Pos(aPoint.x, aPoint.y);
-		int32 Index = GameThreadReturn(^{ return IMMContext->GetCharacterIndexFromPoint(Pos); });
+		int32 Index = GameThreadReturn(^{ return IMMContext->GetCharacterIndexFromPoint(Pos); }, InGameRunLoopMode(@[ UE4IMEEventMode ]));
 		NSUInteger Idx = Index == INDEX_NONE ? NSNotFound : (NSUInteger)Index;
 		return Idx;
 	}
diff --git a/Engine/Source/Runtime/Core/Private/Mac/CocoaThread.cpp b/Engine/Source/Runtime/Core/Private/Mac/CocoaThread.cpp
--- a/Engine/Source/Runtime/Core/Private/Mac/CocoaThread.cpp
+++ b/Engine/Source/Runtime/Core/Private/Mac/CocoaThread.cpp
@@ -5,8 +5,17 @@
 
 #define MAC_SEPARATE_GAME_THREAD 0 // Separate the main & game threads so that we better handle the interaction between the Cocoa's event delegates and UE4's event polling.
 
+NSString* UE4NilEventMode = @"UE4NilEventMode";
+NSString* UE4ShowEventMode = @"UE4ShowEventMode";
+NSString* UE4ResizeEventMode = @"UE4ResizeEventMode";
+NSString* UE4FullscreenEventMode = @"UE4FullscreenEventMode";
+NSString* UE4CloseEventMode = @"UE4CloseEventMode";
+NSString* UE4IMEEventMode = @"UE4IMEEventMode";
+
 static FCocoaGameThread* GCocoaGameThread = nil;
 static NSRunLoop* GCocoaGameRunLoop = nil;
+static TMap<NSRunLoop*, TArray<NSString*>> GRunLoopModes;
+static FCriticalSection GRunLoopModesLock;
 
 @implementation NSThread (FCocoaThread)
 
@@ -36,19 +45,35 @@ static NSRunLoop* GCocoaGameRunLoop = nil;
 	check(RunLoop);
 	if(RunLoop != [NSRunLoop currentRunLoop])
 	{
+		dispatch_block_t CopiedBlock = Block_copy(Block);
+		
 		if(bWait)
 		{
 			__block bool bDone = false;
-			[RunLoop performBlock:^{ Block(); bDone = true; } forMode:RunMode wake:true];
+			{
+				FScopeLock Lock(&GRunLoopModesLock);
+				TArray<NSString*>& Modes = GRunLoopModes.FindOrAdd([NSRunLoop currentRunLoop]);
+				Modes.Add([WaitMode retain]);
+			}
+			[RunLoop performBlock:^{ CopiedBlock(); bDone = true; } forMode:RunMode wake:true];
 			while(!bDone)
 			{
-				CFRunLoopRunInMode((CFStringRef)WaitMode, 0.0, true);
+				CFRunLoopRunInMode((CFStringRef)WaitMode, 0, true);
+			}
+			{
+				FScopeLock Lock(&GRunLoopModesLock);
+				TArray<NSString*>& Modes = GRunLoopModes.FindOrAdd([NSRunLoop currentRunLoop]);
+				check(Modes.Num() > 0);
+				[Modes.Last() release];
+				Modes.RemoveAt(Modes.Num() - 1);
 			}
 		}
 		else
 		{
-			[RunLoop performBlock:Block forMode:RunMode wake:true];
+			[RunLoop performBlock:CopiedBlock forMode:RunMode wake:true];
 		}
+		
+		Block_release(CopiedBlock);
 	}
 	else
 	{
@@ -80,6 +105,20 @@ static NSRunLoop* GCocoaGameRunLoop = nil;
 	}
 }
 
+- (NSString*)intendedMode
+{
+	NSString* Mode = nil;
+	{
+		FScopeLock Lock(&GRunLoopModesLock);
+		TArray<NSString*>& Modes = GRunLoopModes.FindOrAdd(self);
+		if(Modes.Num() > 0)
+		{
+			Mode = Modes.Last();
+		}
+	}
+	return Mode;
+}
+
 @end
 
 @implementation FCocoaGameThread : NSThread
@@ -119,19 +158,29 @@ static NSRunLoop* GCocoaGameRunLoop = nil;
 
 @end
 
-void PerformBlockOnRunLoop(dispatch_block_t Block, NSRunLoop* const RunLoop, bool const bWait)
+NSString* InGameRunLoopMode(NSArray* AllowedModes)
+{
+	NSString* Mode = [[NSRunLoop gameRunLoop] intendedMode];
+	if(Mode == nil || ![AllowedModes containsObject:Mode])
+	{
+		Mode = NSDefaultRunLoopMode;
+	}
+	return Mode;
+}
+
+void PerformBlockOnRunLoop(dispatch_block_t Block, NSRunLoop* const RunLoop, NSString* SendMode, NSString* WaitMode, bool const bWait)
 {
-	[[NSThread currentThread] performBlock:Block onRunLoop:RunLoop forMode:NSDefaultRunLoopMode wait:bWait inMode:NSDefaultRunLoopMode];
+	[[NSThread currentThread] performBlock:Block onRunLoop:RunLoop forMode:SendMode wait:bWait inMode:WaitMode];
 }
 
-void MainThreadCall(dispatch_block_t Block, bool const bWait)
+void MainThreadCall(dispatch_block_t Block, NSString* WaitMode, bool const bWait)
 {
-	[[NSThread currentThread] performBlock:Block onRunLoop:[NSRunLoop mainRunLoop] forMode:NSDefaultRunLoopMode wait:bWait inMode:NSDefaultRunLoopMode];
+	PerformBlockOnRunLoop(Block, [NSRunLoop mainRunLoop], NSDefaultRunLoopMode, WaitMode, bWait);
 }
 
-void GameThreadCall(dispatch_block_t Block, bool const bWait)
+void GameThreadCall(dispatch_block_t Block, NSString* SendMode, bool const bWait)
 {
-	[[NSThread currentThread] performBlock:Block onRunLoop:[NSRunLoop gameRunLoop] forMode:NSDefaultRunLoopMode wait:bWait inMode:NSDefaultRunLoopMode];
+	PerformBlockOnRunLoop(Block, [NSRunLoop gameRunLoop], SendMode, NSDefaultRunLoopMode, bWait);
 }
 
 void RunGameThread(id Target, SEL Selector)
diff --git a/Engine/Source/Runtime/Core/Private/Mac/CocoaWindow.cpp b/Engine/Source/Runtime/Core/Private/Mac/CocoaWindow.cpp
--- a/Engine/Source/Runtime/Core/Private/Mac/CocoaWindow.cpp
+++ b/Engine/Source/Runtime/Core/Private/Mac/CocoaWindow.cpp
@@ -167,7 +167,7 @@ NSString* NSPerformDragOperation = @"NSPerformDragOperation";
 	if(bNeedsRedraw && bForwardEvents && ([self isVisible] && [super alphaValue] > 0.0f))
 	{
 		NSNotification* Notification = [NSNotification notificationWithName:NSWindowRedrawContents object:self];
-		FMacEvent::SendToGameRunLoop(Notification, self, EMacEventSendMethod::Sync);
+		FMacEvent::SendToGameRunLoop(Notification, self, EMacEventSendMethod::Sync, InGameRunLoopMode(@[ UE4NilEventMode, UE4ShowEventMode, UE4ResizeEventMode, UE4FullscreenEventMode, UE4CloseEventMode, UE4IMEEventMode ]));
 	}
 	bNeedsRedraw = false;
 }
@@ -284,7 +284,7 @@ NSString* NSPerformDragOperation = @"NSPerformDragOperation";
 		if(self.bForwardEvents)
 		{
 			NSNotification* Notification = [NSNotification notificationWithName:NSWindowDidResizeNotification object:self];
-			FMacEvent::SendToGameRunLoop(Notification, self, EMacEventSendMethod::Async);
+			FMacEvent::SendToGameRunLoop(Notification, self, EMacEventSendMethod::Async, InGameRunLoopMode(@[ UE4ResizeEventMode, UE4ShowEventMode ]));
 		}
 	}
 }
@@ -296,13 +296,15 @@ NSString* NSPerformDragOperation = @"NSPerformDragOperation";
 	{
 		MainThreadCall(^{
 			[super setFrameOrigin:Point];
-		}, true);
+		});
 		bDeferSetOrigin = false;
 	}
 	else
 	{
 		bDeferSetOrigin = true;
 		DeferFrame.origin = Point;
+		NSNotification* Notification = [NSNotification notificationWithName:NSWindowDidMoveNotification object:self];
+		FMacEvent::SendToGameRunLoop(Notification, self, EMacEventSendMethod::Async, InGameRunLoopMode(@[ UE4ResizeEventMode, UE4ShowEventMode ]));
 	}
 }
 
@@ -337,7 +339,7 @@ NSString* NSPerformDragOperation = @"NSPerformDragOperation";
 	WindowMode = self.TargetWindowMode;
 	if(self.bForwardEvents)
 	{
-		FMacEvent::SendToGameRunLoop(notification, self, EMacEventSendMethod::Async);
+		FMacEvent::SendToGameRunLoop(notification, self, EMacEventSendMethod::Async, InGameRunLoopMode(@[ UE4FullscreenEventMode ]));
 	}
 }
 
@@ -347,7 +349,7 @@ NSString* NSPerformDragOperation = @"NSPerformDragOperation";
 	self.TargetWindowMode = EWindowMode::Windowed;
 	if(self.bForwardEvents)
 	{
-		FMacEvent::SendToGameRunLoop(notification, self, EMacEventSendMethod::Async);
+		FMacEvent::SendToGameRunLoop(notification, self, EMacEventSendMethod::Async, InGameRunLoopMode(@[ UE4FullscreenEventMode ]));
 	}
 }
 
@@ -367,7 +369,7 @@ NSString* NSPerformDragOperation = @"NSPerformDragOperation";
 	
 	if(self.bForwardEvents)
 	{
-		FMacEvent::SendToGameRunLoop(Notification, self, EMacEventSendMethod::Async);
+		FMacEvent::SendToGameRunLoop(Notification, self, EMacEventSendMethod::Async, InGameRunLoopMode(@[ UE4ShowEventMode, UE4CloseEventMode, UE4FullscreenEventMode ]));
 	}
 }
 
@@ -378,7 +380,7 @@ NSString* NSPerformDragOperation = @"NSPerformDragOperation";
 	
 	if(self.bForwardEvents)
 	{
-		FMacEvent::SendToGameRunLoop(Notification, self, EMacEventSendMethod::Async);
+		FMacEvent::SendToGameRunLoop(Notification, self, EMacEventSendMethod::Async, InGameRunLoopMode(@[ UE4ShowEventMode, UE4CloseEventMode, UE4FullscreenEventMode ]));
 	}
 }
 
@@ -386,7 +388,7 @@ NSString* NSPerformDragOperation = @"NSPerformDragOperation";
 {
 	if(self.bForwardEvents)
 	{
-		FMacEvent::SendToGameRunLoop(Notification, self, EMacEventSendMethod::Async);
+		FMacEvent::SendToGameRunLoop(Notification, self, EMacEventSendMethod::Async, InGameRunLoopMode(@[ UE4ResizeEventMode, UE4ShowEventMode, UE4FullscreenEventMode ]));
 	}
 }
 
@@ -399,7 +401,7 @@ NSString* NSPerformDragOperation = @"NSPerformDragOperation";
 	
 	if(self.bForwardEvents)
 	{
-		FMacEvent::SendToGameRunLoop(Notification, self, EMacEventSendMethod::Async);
+		FMacEvent::SendToGameRunLoop(Notification, self, EMacEventSendMethod::Async, InGameRunLoopMode(@[ UE4ResizeEventMode, UE4ShowEventMode, UE4FullscreenEventMode ]));
 	}
 }
 
@@ -471,7 +473,7 @@ NSString* NSPerformDragOperation = @"NSPerformDragOperation";
 	bZoomed = [self isZoomed];
 	if(self.bForwardEvents)
 	{
-		FMacEvent::SendToGameRunLoop(Notification, self, EMacEventSendMethod::Async);
+		FMacEvent::SendToGameRunLoop(Notification, self, EMacEventSendMethod::Async, InGameRunLoopMode(@[ UE4ResizeEventMode, UE4ShowEventMode, UE4FullscreenEventMode ]));
 	}
 	bNeedsRedraw = true;
 }
@@ -480,7 +482,7 @@ NSString* NSPerformDragOperation = @"NSPerformDragOperation";
 {
 	if(self.bForwardEvents && MacApplication)
 	{
-		FMacEvent::SendToGameRunLoop(notification, self, EMacEventSendMethod::Async);
+		FMacEvent::SendToGameRunLoop(notification, self, EMacEventSendMethod::Async, InGameRunLoopMode(@[ UE4CloseEventMode ]));
 	}
 }
 
@@ -588,7 +590,7 @@ NSString* NSPerformDragOperation = @"NSPerformDragOperation";
 	BOOL Movable = [super isMovable];
 	if(Movable && bRenderInitialised && MacApplication)
 	{
-		Movable &= (BOOL)(GameThreadReturn(^{ return MacApplication->IsWindowMovable(self, NULL); }));
+		Movable &= (BOOL)(GameThreadReturn(^{ return MacApplication->IsWindowMovable(self, NULL); }, InGameRunLoopMode(@[ UE4NilEventMode, UE4ShowEventMode, UE4ResizeEventMode, UE4FullscreenEventMode, UE4CloseEventMode, UE4IMEEventMode ])));
 	}
 	return Movable;
 }
diff --git a/Engine/Source/Runtime/Core/Private/Mac/MacEvent.cpp b/Engine/Source/Runtime/Core/Private/Mac/MacEvent.cpp
--- a/Engine/Source/Runtime/Core/Private/Mac/MacEvent.cpp
+++ b/Engine/Source/Runtime/Core/Private/Mac/MacEvent.cpp
@@ -42,7 +42,7 @@ FMacEvent::FMacEvent(NSNotification* const Notification, NSWindow* const Window)
 	CFRelease(Event);
 }
 
-void FMacEvent::SendToGameRunLoop(FMacEvent const* const Event, EMacEventSendMethod SendMethod)
+void FMacEvent::SendToGameRunLoop(FMacEvent const* const Event, EMacEventSendMethod SendMethod, NSString* SendMode)
 {
 	NSThread* GameThread = [NSThread gameThread];
 	check(GameThread);
@@ -53,24 +53,24 @@ void FMacEvent::SendToGameRunLoop(FMacEvent const* const Event, EMacEventSendMet
 	dispatch_block_t Block = ^{ FMacApplication::ProcessEvent(Event); };
 	const bool bWait = (SendMethod == EMacEventSendMethod::Sync);
 	
-	[GameThread performBlock: Block onRunLoop: GameRunLoop forMode: NSDefaultRunLoopMode wait: bWait inMode: NSDefaultRunLoopMode];
+	[GameThread performBlock: Block onRunLoop: GameRunLoop forMode: SendMode wait: bWait inMode: NSDefaultRunLoopMode];
 }
 	
-void FMacEvent::SendToGameRunLoop(NSEvent* const Event, EMacEventSendMethod SendMethod)
+void FMacEvent::SendToGameRunLoop(NSEvent* const Event, EMacEventSendMethod SendMethod, NSString* SendMode)
 {
 	if(MacApplication)
 	{
 		FMacEvent* MacEvent = new FMacEvent(Event);
-		FMacEvent::SendToGameRunLoop(MacEvent, SendMethod);
+		FMacEvent::SendToGameRunLoop(MacEvent, SendMethod, SendMode);
 	}
 }
 
-void FMacEvent::SendToGameRunLoop(NSNotification* const Notification, NSWindow* const Window, EMacEventSendMethod SendMethod)
+void FMacEvent::SendToGameRunLoop(NSNotification* const Notification, NSWindow* const Window, EMacEventSendMethod SendMethod, NSString* SendMode)
 {
 	if(MacApplication)
 	{
 		FMacEvent* MacEvent = new FMacEvent(Notification, Window);
-		FMacEvent::SendToGameRunLoop(MacEvent, SendMethod);
+		FMacEvent::SendToGameRunLoop(MacEvent, SendMethod, SendMode);
 	}
 }
 
diff --git a/Engine/Source/Runtime/Core/Private/Mac/MacPlatformOutputDevices.cpp b/Engine/Source/Runtime/Core/Private/Mac/MacPlatformOutputDevices.cpp
--- a/Engine/Source/Runtime/Core/Private/Mac/MacPlatformOutputDevices.cpp
+++ b/Engine/Source/Runtime/Core/Private/Mac/MacPlatformOutputDevices.cpp
@@ -199,58 +199,58 @@ void FOutputDeviceConsoleMac::CreateConsole()
 		bHasY = GConfig->GetInt(TEXT("DebugMac"), TEXT("ConsoleY"), ConsolePosY, GGameIni);
 	}
 
-	ConsoleHandle = [[NSWindow alloc] initWithContentRect: NSMakeRect(ConsolePosX, ConsolePosY, ConsoleWidth, ConsoleHeight)
-									styleMask: (NSTitledWindowMask | NSClosableWindowMask | NSMiniaturizableWindowMask | NSResizableWindowMask)
-									  backing: NSBackingStoreBuffered
-										defer: NO];
+	MainThreadCall(^{
+		ConsoleHandle = [[NSWindow alloc] initWithContentRect: NSMakeRect(ConsolePosX, ConsolePosY, ConsoleWidth, ConsoleHeight)
+										styleMask: (NSTitledWindowMask | NSClosableWindowMask | NSMiniaturizableWindowMask | NSResizableWindowMask)
+										  backing: NSBackingStoreBuffered
+											defer: NO];
 
-	ScrollView = [[NSScrollView alloc] initWithFrame:[[ConsoleHandle contentView] frame]];
-	NSSize ContentSize = [ScrollView contentSize];
-	
-	[ScrollView setBorderType:NSNoBorder];
-	[ScrollView setHasVerticalScroller:YES];
-	[ScrollView setHasHorizontalScroller:NO];
-	[ScrollView setAutoresizingMask:NSViewWidthSizable | NSViewHeightSizable];
-	
-	TextView = [[NSTextView alloc] initWithFrame:NSMakeRect( 0, 0, ContentSize.width, ContentSize.height )];
-	[TextView setMinSize:NSMakeSize( 0.0, ContentSize.height ) ];
-	[TextView setMaxSize:NSMakeSize( FLT_MAX, FLT_MAX )];
-	[TextView setVerticallyResizable:YES];
-	[TextView setHorizontallyResizable:NO];
-	[TextView setAutoresizingMask:NSViewWidthSizable];
-	[TextView setBackgroundColor: [NSColor blackColor]];
-	
-	[[TextView textContainer] setContainerSize:NSMakeSize( ContentSize.width, FLT_MAX )];
-	[[TextView textContainer] setWidthTracksTextView:YES];
-	
-	[ScrollView setDocumentView:TextView];
-	[ConsoleHandle setContentView:ScrollView];
+		ScrollView = [[NSScrollView alloc] initWithFrame:[[ConsoleHandle contentView] frame]];
+		NSSize ContentSize = [ScrollView contentSize];
+		
+		[ScrollView setBorderType:NSNoBorder];
+		[ScrollView setHasVerticalScroller:YES];
+		[ScrollView setHasHorizontalScroller:NO];
+		[ScrollView setAutoresizingMask:NSViewWidthSizable | NSViewHeightSizable];
+		
+		TextView = [[NSTextView alloc] initWithFrame:NSMakeRect( 0, 0, ContentSize.width, ContentSize.height )];
+		[TextView setMinSize:NSMakeSize( 0.0, ContentSize.height ) ];
+		[TextView setMaxSize:NSMakeSize( FLT_MAX, FLT_MAX )];
+		[TextView setVerticallyResizable:YES];
+		[TextView setHorizontallyResizable:NO];
+		[TextView setAutoresizingMask:NSViewWidthSizable];
+		[TextView setBackgroundColor: [NSColor blackColor]];
+		
+		[[TextView textContainer] setContainerSize:NSMakeSize( ContentSize.width, FLT_MAX )];
+		[[TextView textContainer] setWidthTracksTextView:YES];
+		
+		[ScrollView setDocumentView:TextView];
+		[ConsoleHandle setContentView:ScrollView];
 
-	if (!bHasX || !bHasY)
-	{
-		[ConsoleHandle center];
-	}
+		if (!bHasX || !bHasY)
+		{
+			[ConsoleHandle center];
+		}
 
-	[ConsoleHandle setOpaque: YES];
-	[ConsoleHandle makeKeyAndOrderFront:nil];
-	
-	if(!MacApplication)
-	{
-		do
+		[ConsoleHandle setOpaque: YES];
+		[ConsoleHandle makeKeyAndOrderFront:nil];
+		
+		if(!MacApplication)
 		{
-			FPlatformMisc::PumpMessages( true );
-		} while(ConsoleHandle && ![ConsoleHandle isVisible]);
-	}
-	
-	SetDefaultTextColor();
+			do
+			{
+				FPlatformMisc::PumpMessages( true );
+			} while(ConsoleHandle && ![ConsoleHandle isVisible]);
+		}
+		
+		SetDefaultTextColor();
+	}, UE4NilEventMode, true);
 }
 
 void FOutputDeviceConsoleMac::DestroyConsole()
 {
 	if (ConsoleHandle)
 	{
-		SCOPED_AUTORELEASE_POOL;
-		
 		do
 		{
 			FPlatformMisc::PumpMessages( true );
@@ -258,12 +258,15 @@ void FOutputDeviceConsoleMac::DestroyConsole()
 
 		SaveToINI();
 		
-		if( TextViewTextColor )
-			[TextViewTextColor release];
-		
-		[ConsoleHandle close];
-		ConsoleHandle = NULL;
-		TextViewTextColor = NULL;
+		MainThreadCall(^{
+			SCOPED_AUTORELEASE_POOL;
+			if( TextViewTextColor )
+				[TextViewTextColor release];
+			
+			[ConsoleHandle close];
+			ConsoleHandle = NULL;
+			TextViewTextColor = NULL;
+		}, UE4NilEventMode, true);
 	}
 }
 
@@ -342,7 +345,7 @@ void FOutputDeviceConsoleMac::Serialize( const TCHAR* Data, ELogVerbosity::Type
 						[Colors release];
 						[AttributeKeys release];
 						OutstandingTasks--;
-					}, false);
+					}, NSDefaultRunLoopMode, false);
 				}
 			}
 			else
@@ -362,7 +365,7 @@ void FOutputDeviceConsoleMac::Serialize( const TCHAR* Data, ELogVerbosity::Type
 					[AttributedString release];
 					CFRelease(CocoaText);
 					OutstandingTasks--;
-				}, false);
+				}, NSDefaultRunLoopMode, false);
 				
 				if(!MacApplication)
 				{
@@ -409,5 +412,5 @@ void FOutputDeviceConsoleMac::SetDefaultTextColor()
 		[Colors release];
 		[AttributeKeys release];
 		OutstandingTasks--;
-	}, false);
+	}, NSDefaultRunLoopMode, false);
 }
diff --git a/Engine/Source/Runtime/Core/Private/Mac/MacTextInputMethodSystem.cpp b/Engine/Source/Runtime/Core/Private/Mac/MacTextInputMethodSystem.cpp
--- a/Engine/Source/Runtime/Core/Private/Mac/MacTextInputMethodSystem.cpp
+++ b/Engine/Source/Runtime/Core/Private/Mac/MacTextInputMethodSystem.cpp
@@ -56,7 +56,7 @@ namespace
 					FCocoaTextView* TextView = (FCocoaTextView*)[CocoaWindow openGLView];
 					[[TextView inputContext] invalidateCharacterCoordinates];
 				}
-			}, true);
+			}, UE4IMEEventMode, true);
 		}
 	}
 	
@@ -71,7 +71,7 @@ namespace
 					FCocoaTextView* TextView = (FCocoaTextView*)[CocoaWindow openGLView];
 					[[TextView inputContext] invalidateCharacterCoordinates];
 				}
-			}, true);
+			}, UE4IMEEventMode, true);
 		}
 	}
 	
@@ -86,7 +86,7 @@ namespace
 					FCocoaTextView* TextView = (FCocoaTextView*)[CocoaWindow openGLView];
 					[[TextView inputContext] invalidateCharacterCoordinates];
 				}
-			}, true);
+			}, UE4IMEEventMode, true);
 		}
 	}
 	
@@ -102,7 +102,7 @@ namespace
 					[[TextView inputContext] discardMarkedText];
 					[TextView unmarkText];
 				}
-			}, true);
+			}, UE4IMEEventMode, true);
 		}
 	}
 }
@@ -177,7 +177,7 @@ void FMacTextInputMethodSystem::ActivateContext(const TSharedRef<ITextInputMetho
 				}
 			}
 			return bSuccess;
-		});
+		}, UE4IMEEventMode);
 	}
 	if(!bActivatedContext)
 	{
@@ -214,7 +214,7 @@ void FMacTextInputMethodSystem::DeactivateContext(const TSharedRef<ITextInputMet
 				}
 			}
 			return bSuccess;
-		});
+		}, UE4IMEEventMode);
 	}
 	if(!bDeactivatedContext)
 	{
diff --git a/Engine/Source/Runtime/Core/Private/Mac/MacWindow.cpp b/Engine/Source/Runtime/Core/Private/Mac/MacWindow.cpp
--- a/Engine/Source/Runtime/Core/Private/Mac/MacWindow.cpp
+++ b/Engine/Source/Runtime/Core/Private/Mac/MacWindow.cpp
@@ -85,7 +85,7 @@ void FMacWindow::Initialize( FMacApplication* const Application, const TSharedRe
 		WindowStyle &= ~(NSTexturedBackgroundWindowMask);
 	}
 
-	WindowHandle = MainThreadReturn(^{ return [[FCocoaWindow alloc] initWithContentRect: ViewRect styleMask: WindowStyle backing: NSBackingStoreBuffered defer: NO]; });
+	WindowHandle = MainThreadReturn(^{ return [[FCocoaWindow alloc] initWithContentRect: ViewRect styleMask: WindowStyle backing: NSBackingStoreBuffered defer: NO]; }, UE4NilEventMode);
 
 	if( WindowHandle == NULL )
 	{
@@ -167,7 +167,7 @@ void FMacWindow::Initialize( FMacApplication* const Application, const TSharedRe
 		{
 			SetOpacity( 1.0f );
 		}
-	}, true);
+	}, UE4ResizeEventMode, true);
 }
 
 FMacWindow::FMacWindow()
@@ -214,7 +214,7 @@ void FMacWindow::ReshapeWindow( int32 X, int32 Y, int32 Width, int32 Height )
 			}
 			
 			WindowHandle->bZoomed = [WindowHandle isZoomed];
-		}, true);
+		}, UE4ResizeEventMode, true);
 	}
 	
 	MessageHandler->FinishedReshapingWindow( SharedThis( this ) );
@@ -238,7 +238,7 @@ void FMacWindow::MoveWindowTo( int32 X, int32 Y )
 		SCOPED_AUTORELEASE_POOL;
 		const int32 InvertedY = FPlatformMisc::ConvertSlateYPositionToCocoa(Y) - [WindowHandle frame].size.height + 1;
 		[WindowHandle setFrameOrigin: NSMakePoint(X, InvertedY)];
-	}, true);
+	}, UE4ResizeEventMode, true);
 }
 
 void FMacWindow::BringToFront( bool bForce )
@@ -248,7 +248,7 @@ void FMacWindow::BringToFront( bool bForce )
 		MainThreadCall(^{
 			SCOPED_AUTORELEASE_POOL;
 			[WindowHandle orderFrontAndMakeMain:IsRegularWindow() andKey:IsRegularWindow()];
-		}, true);
+		}, UE4ShowEventMode, true);
 	}
 }
 
@@ -281,7 +281,7 @@ void FMacWindow::Destroy()
 			// Close the window, but don't destruct it...
 			MainThreadCall(^{
 				[Window performClose:nil];
-			}, true);
+			}, UE4CloseEventMode, true);
 			
 			// Since event handling of the change in focus may try to communicate with this window,
 			// dispatch a call to release it which should be handled afterward.
@@ -298,7 +298,7 @@ void FMacWindow::Minimize()
 	MainThreadCall(^{
 		SCOPED_AUTORELEASE_POOL;
 		[WindowHandle miniaturize:nil];
-	}, true);
+	}, UE4ResizeEventMode, true);
 }
 
 void FMacWindow::Maximize()
@@ -310,7 +310,7 @@ void FMacWindow::Maximize()
 			WindowHandle->bZoomed = true;
 			[WindowHandle zoom:nil];
 		}
-	}, true);
+	}, UE4ResizeEventMode, true);
 }
 
 void FMacWindow::Restore()
@@ -327,7 +327,7 @@ void FMacWindow::Restore()
 			WindowHandle->bZoomed = false;
 			[WindowHandle deminiaturize:nil];
 		}
-	}, true);
+	}, UE4ResizeEventMode, true);
 }
 
 void FMacWindow::Show()
@@ -344,7 +344,7 @@ void FMacWindow::Show()
 		
 		MainThreadCall(^{
 			[WindowHandle orderFrontAndMakeMain:bMakeMainAndKey andKey:bMakeMainAndKey];
-		}, true);
+		}, UE4ShowEventMode, true);
 		
 		bIsVisible = [WindowHandle isVisible];
 		static bool bCannotRecurse = false;
@@ -377,7 +377,7 @@ void FMacWindow::Hide()
 		}
 		MainThreadCall(^{
 			[WindowHandle orderOut:nil];
-		}, true);
+		}, UE4CloseEventMode, false);
 	}
 }
 
@@ -411,7 +411,7 @@ void FMacWindow::SetWindowMode( EWindowMode::Type NewWindowMode )
 		MainThreadCall(^{
 			[WindowHandle setCollectionBehavior: Behaviour];
 			[WindowHandle toggleFullScreen:nil];
-		}, true);
+		}, UE4FullscreenEventMode, true);
 		
 		// Ensure that the window has transitioned BEFORE leaving this function
 		// this prevents problems with failure to correctly update mouse locks
@@ -467,7 +467,7 @@ void FMacWindow::SetWindowFocus()
 	MainThreadCall(^{
 		SCOPED_AUTORELEASE_POOL;
 		[WindowHandle orderFrontAndMakeMain:true andKey:true];
-	}, true);
+	}, UE4ShowEventMode, true);
 }
 
 void FMacWindow::SetOpacity( const float InOpacity )
@@ -475,7 +475,7 @@ void FMacWindow::SetOpacity( const float InOpacity )
 	MainThreadCall(^{
 		SCOPED_AUTORELEASE_POOL;
 		[WindowHandle setAlphaValue:InOpacity];
-	}, true);
+	}, UE4NilEventMode, true);
 }
 
 void FMacWindow::Enable( bool bEnable )
@@ -483,7 +483,7 @@ void FMacWindow::Enable( bool bEnable )
 	MainThreadCall(^{
 		SCOPED_AUTORELEASE_POOL;
 		[WindowHandle setIgnoresMouseEvents: !bEnable];
-	}, true);
+	}, UE4NilEventMode, true);
 }
 
 bool FMacWindow::IsPointInWindow( int32 X, int32 Y ) const
@@ -547,7 +547,7 @@ void FMacWindow::SetText(const TCHAR* const Text)
 			[NSApp changeWindowsItem: WindowHandle title: (NSString*)CFName filename: NO];
 		}
 		CFRelease( CFName );
-	}, true);
+	}, UE4NilEventMode, true);
 }
 
 bool FMacWindow::IsRegularWindow() const
@@ -576,7 +576,7 @@ void FMacWindow::OnDisplayReconfiguration(CGDirectDisplayID Display, CGDisplayCh
 			{
 				[WindowHandle setDisplayReconfiguring: false];
 			}
-		}, true);
+		});
 	}
 }
 
@@ -587,7 +587,7 @@ bool FMacWindow::OnIMKKeyDown(NSEvent* Event)
 		return MainThreadReturn(^{
 			FCocoaTextView* View = (FCocoaTextView*)[WindowHandle openGLView];
 			return View && [View imkKeyDown:Event];
-		});
+		}, UE4IMEEventMode);
 	}
 	else
 	{
diff --git a/Engine/Source/Runtime/Core/Public/Mac/CocoaThread.h b/Engine/Source/Runtime/Core/Public/Mac/CocoaThread.h
--- a/Engine/Source/Runtime/Core/Public/Mac/CocoaThread.h
+++ b/Engine/Source/Runtime/Core/Public/Mac/CocoaThread.h
@@ -5,6 +5,14 @@
 
 #ifdef __OBJC__
 
+/* Custom run-loop modes for UE4 that process only certain kinds of events to simulate Windows event ordering. */
+extern NSString* UE4NilEventMode; /* Process only mandatory events */
+extern NSString* UE4ShowEventMode; /* Process only show window events */
+extern NSString* UE4ResizeEventMode; /* Process only resize/move window events */
+extern NSString* UE4FullscreenEventMode; /* Process only fullscreen mode events */
+extern NSString* UE4CloseEventMode; /* Process only close window events */
+extern NSString* UE4IMEEventMode; /* Process only input method events */
+
 @interface NSThread (FCocoaThread)
 + (NSThread*) gameThread; // Returns the main game thread, or nil if has yet to be constructed.
 + (bool) isGameThread; // True if the current thread is the main game thread, else false.
@@ -15,6 +23,7 @@
 @interface NSRunLoop (FCocoaRunLoop)
 + (NSRunLoop*)gameRunLoop; // The run loop for the main game thread, or nil if the main game thread has not been started.
 - (void)performBlock:(dispatch_block_t)Block forMode:(NSString *)Mode wake:(bool)bWakeRunLoop; // Add Block to the run loop for processing in the given Mode, only waking the run loop if bWakeRunLoop is set otherwise processing occurs the next time the run-loop is processed.
+- (NSString*)intendedMode;
 @end
 
 @interface FCocoaGameThread : NSThread
@@ -24,31 +33,33 @@
 - (void)main; // Override that sets the variable backing +[NSRunLoop gameRunLoop], do not override in a subclass.
 @end
 
+NSString* InGameRunLoopMode(NSArray* AllowedModes);
+
 template<typename ReturnType>
-ReturnType PerformBlockOnRunLoop(ReturnType (^Block)(void), NSRunLoop* const RunLoop)
+ReturnType PerformBlockOnRunLoop(ReturnType (^Block)(void), NSRunLoop* const RunLoop, NSString* SendMode, NSString* WaitMode)
 {
 	__block ReturnType ReturnValue;
-	[[NSThread currentThread] performBlock:^{ ReturnValue = Block(); } onRunLoop:RunLoop forMode:NSDefaultRunLoopMode wait:true inMode:NSDefaultRunLoopMode];
+	[[NSThread currentThread] performBlock:^{ ReturnValue = Block(); } onRunLoop:RunLoop forMode:SendMode wait:true inMode:WaitMode];
 	return ReturnValue;
 }
 
-void PerformBlockOnRunLoop(dispatch_block_t Block, NSRunLoop* const RunLoop, bool const bWait);
+void PerformBlockOnRunLoop(dispatch_block_t Block, NSRunLoop* const RunLoop, NSString* SendMode, NSString* WaitMode, bool const bWait);
 
 template<typename ReturnType>
-ReturnType MainThreadReturn(ReturnType (^Block)(void))
+ReturnType MainThreadReturn(ReturnType (^Block)(void), NSString* WaitMode = NSDefaultRunLoopMode)
 {
-	return PerformBlockOnRunLoop(Block, [NSRunLoop mainRunLoop]);
+	return PerformBlockOnRunLoop(Block, [NSRunLoop mainRunLoop], NSDefaultRunLoopMode, WaitMode);
 }
 
-void MainThreadCall(dispatch_block_t Block, bool const bWait);
+void MainThreadCall(dispatch_block_t Block, NSString* WaitMode = NSDefaultRunLoopMode, bool const bWait = true);
 
 template<typename ReturnType>
-ReturnType GameThreadReturn(ReturnType (^Block)(void))
+ReturnType GameThreadReturn(ReturnType (^Block)(void), NSString* SendMode = NSDefaultRunLoopMode)
 {
-	return PerformBlockOnRunLoop(Block, [NSRunLoop gameRunLoop]);
+	return PerformBlockOnRunLoop(Block, [NSRunLoop gameRunLoop], SendMode, NSDefaultRunLoopMode);
 }
 
-void GameThreadCall(dispatch_block_t Block, bool const bWait);
+void GameThreadCall(dispatch_block_t Block, NSString* SendMode = NSDefaultRunLoopMode, bool const bWait = true);
 
 void RunGameThread(id Target, SEL Selector);
 
diff --git a/Engine/Source/Runtime/Core/Public/Mac/MacEvent.h b/Engine/Source/Runtime/Core/Public/Mac/MacEvent.h
--- a/Engine/Source/Runtime/Core/Public/Mac/MacEvent.h
+++ b/Engine/Source/Runtime/Core/Public/Mac/MacEvent.h
@@ -27,14 +27,14 @@ class FMacEvent
 	FMacEvent(NSNotification* const Notification, NSWindow* const Window);
 	
 	// Send the event to the game run loop for processing using the specified SendMethod.
-	static void SendToGameRunLoop(FMacEvent const* const Event, EMacEventSendMethod SendMethod);
+	static void SendToGameRunLoop(FMacEvent const* const Event, EMacEventSendMethod SendMethod, NSString* SendMode = NSDefaultRunLoopMode);
 	
 public:
 	// Send an NSEvent to the Game run loop as an FMacEvent using the specified SendMethod
-	static void SendToGameRunLoop(NSEvent* const Event, EMacEventSendMethod SendMethod);
+	static void SendToGameRunLoop(NSEvent* const Event, EMacEventSendMethod SendMethod, NSString* SendMode = NSDefaultRunLoopMode);
 	
 	// Send an NSNotification to the Game run loop as an FMacEvent using the specified SendMethod
-	static void SendToGameRunLoop(NSNotification* const Notification, NSWindow* const Window, EMacEventSendMethod SendMethod);
+	static void SendToGameRunLoop(NSNotification* const Notification, NSWindow* const Window, EMacEventSendMethod SendMethod, NSString* SendMode = NSDefaultRunLoopMode);
 	
 	// Destructor
 	~FMacEvent(void);
diff --git a/Engine/Source/Runtime/Launch/Private/Mac/LaunchMac.cpp b/Engine/Source/Runtime/Launch/Private/Mac/LaunchMac.cpp
--- a/Engine/Source/Runtime/Launch/Private/Mac/LaunchMac.cpp
+++ b/Engine/Source/Runtime/Launch/Private/Mac/LaunchMac.cpp
@@ -112,7 +112,7 @@ void EngineCrashHandler(const FGenericCrashContext & GenericContext)
 				GEngine->DeferredCommands.Add(TEXT("EXIT"));
 			}
 		}
-	}, false);
+	}, NSDefaultRunLoopMode, false);
 }
 
 - (IBAction)showAboutWindow:(id)Sender
@@ -123,7 +123,7 @@ void EngineCrashHandler(const FGenericCrashContext & GenericContext)
 		{
 			FModuleManager::GetModuleChecked<IMainFrameModule>(TEXT("MainFrame")).ShowAboutWindow();
 		}
-	}, false);
+	}, NSDefaultRunLoopMode, false);
 #else
 	[NSApp orderFrontStandardAboutPanel:Sender];
 #endif
@@ -137,7 +137,7 @@ void EngineCrashHandler(const FGenericCrashContext & GenericContext)
 		{
 			FModuleManager::LoadModuleChecked<ISettingsModule>("Settings").ShowViewer(FName("Editor"), FName("General"), FName("Appearance"));
 		}
-	}, false);
+	}, NSDefaultRunLoopMode, false);
 }
 #endif
 
diff --git a/Engine/Source/Runtime/Slate/Private/Framework/MultiBox/Mac/MacMenu.cpp b/Engine/Source/Runtime/Slate/Private/Framework/MultiBox/Mac/MacMenu.cpp
--- a/Engine/Source/Runtime/Slate/Private/Framework/MultiBox/Mac/MacMenu.cpp
+++ b/Engine/Source/Runtime/Slate/Private/Framework/MultiBox/Mac/MacMenu.cpp
@@ -186,7 +186,7 @@ void FSlateMacMenu::UpdateWithMultiBox(const TSharedRef< FMultiBox >& MultiBox)
 		{
 			[HelpMenu release];
 		}
-	}, true);
+	});
 }
 
 static const TSharedRef<SWidget> FindTextBlockWidget(TSharedRef<SWidget> Content)
@@ -281,7 +281,7 @@ void FSlateMacMenu::ExecuteMenuItemAction(const TSharedRef< const class FMenuEnt
 			(*MenuBlock)->GetDirectActions().Execute();
 		}
         delete MenuBlock;
-	}, false);
+	}, NSDefaultRunLoopMode, false);
 }
 
 CFStringRef FSlateMacMenu::GetMenuItemKeyEquivalent(const TSharedRef<const class FMenuEntryBlock>& Block, uint32* OutModifiers)
@@ -392,7 +392,7 @@ void FSlateMacMenu::UpdateMenu(FMacMenu* Menu)
 		{
 			[NSApp setWindowsMenu:Menu];
 		}
-	}, true);
+	});
 }
 
 bool FSlateMacMenu::IsMenuItemEnabled(const TSharedRef<const class FMenuEntryBlock>& Block)
