Fix missing Total memory detection for HTML5.   #HTML5.

diff --git a/Engine/Source/Runtime/Core/Private/HTML5/HTML5PlatformMemory.cpp b/Engine/Source/Runtime/Core/Private/HTML5/HTML5PlatformMemory.cpp
--- a/Engine/Source/Runtime/Core/Private/HTML5/HTML5PlatformMemory.cpp
+++ b/Engine/Source/Runtime/Core/Private/HTML5/HTML5PlatformMemory.cpp
@@ -26,7 +26,7 @@ const FPlatformMemoryConstants& FHTML5PlatformMemory::GetConstants()
 #if PLATFORM_HTML5_WIN32
 		uint64 GTotalMemoryAvailable = 512 * 1024 * 1024;
 #else
-		extern uint64 GTotalMemoryAvailable;
+		uint64 GTotalMemoryAvailable =  EM_ASM_INT_V({ return Module.TOTAL_MEMORY; }); 
 #endif
 
 		MemoryConstants.TotalPhysical = GTotalMemoryAvailable;
