Rearranging ModuleManager to make it possible to lift out hot-reload and UBT functionality.

diff --git a/Engine/Source/Editor/UnrealEd/Private/UnrealEdMisc.cpp b/Engine/Source/Editor/UnrealEd/Private/UnrealEdMisc.cpp
--- a/Engine/Source/Editor/UnrealEd/Private/UnrealEdMisc.cpp
+++ b/Engine/Source/Editor/UnrealEd/Private/UnrealEdMisc.cpp
@@ -401,7 +401,7 @@ void FUnrealEdMisc::InitEngineAnalytics()
 			{
 				TArray< FAnalyticsEventAttribute > ModuleAttributes;
 				ModuleAttributes.Add(FAnalyticsEventAttribute(FString("ModuleName"), Module.Name));
-				ModuleAttributes.Add(FAnalyticsEventAttribute(FString("CompilationMethod"), Module.CompilationMethod));
+				ModuleAttributes.Add(FAnalyticsEventAttribute(FString("CompilationMethod"), FModuleManager::Get().GetModuleCompileMethod(*Module.Name)));
 				EngineAnalytics.RecordEvent(FString("Editor.Usage.Modules"), ModuleAttributes);
 			}
 		}
diff --git a/Engine/Source/Runtime/Core/Private/Modules/ModuleManager.cpp b/Engine/Source/Runtime/Core/Private/Modules/ModuleManager.cpp
--- a/Engine/Source/Runtime/Core/Private/Modules/ModuleManager.cpp
+++ b/Engine/Source/Runtime/Core/Private/Modules/ModuleManager.cpp
@@ -553,7 +553,7 @@ void FModuleManager::UnloadModulesAtShutdown()
 		TSharedRef< FModuleInfo > ModuleInfo( ModuleIt.Value() );
 
 		// Take this opportunity to update and write out the compile data before the editor shuts down
-		UpdateModuleCompileData(ModuleIt.Key(), ModuleInfo);
+		UpdateModuleCompileData(ModuleIt.Key());
 
 		// Only if already loaded
 		if( ModuleInfo->Module.IsValid() )
@@ -751,23 +751,6 @@ bool FModuleManager::QueryModule( const FName InModuleName, FModuleStatus& OutMo
 			OutModuleStatus.bIsGameModule = ModuleInfo->Module->IsGameModule();
 		}
 
-		if (!ModuleInfo->CompileData.bIsValid)
-		{
-			UpdateModuleCompileData(InModuleName, ModuleInfo);
-		}
-
-		FString CompileMethodString = ModuleManagerDefs::CompileMethodUnknown;
-		if (ModuleInfo->CompileData.CompileMethod == EModuleCompileMethod::Runtime)
-		{
-			CompileMethodString = ModuleManagerDefs::CompileMethodRuntime;
-		}
-		else if (ModuleInfo->CompileData.CompileMethod == EModuleCompileMethod::External)
-		{
-			CompileMethodString = ModuleManagerDefs::CompileMethodExternal;
-		}
-
-		OutModuleStatus.CompilationMethod = CompileMethodString;
-
 		return true;
 	}
 
@@ -795,28 +778,29 @@ void FModuleManager::QueryModules( TArray< FModuleStatus >& OutModuleStatuses )
 			{
 				ModuleStatus.bIsGameModule = CurModule->Module->IsGameModule();
 			}
-
-			if (!CurModule->CompileData.bIsValid)
-			{
-				UpdateModuleCompileData(CurModuleFName, CurModule);
-			}
-
-			FString CompileMethodString = ModuleManagerDefs::CompileMethodUnknown;
-			if (CurModule->CompileData.CompileMethod == EModuleCompileMethod::Runtime)
-			{
-				CompileMethodString = ModuleManagerDefs::CompileMethodRuntime;
-			}
-			else if (CurModule->CompileData.CompileMethod == EModuleCompileMethod::External)
-			{
-				CompileMethodString = ModuleManagerDefs::CompileMethodExternal;
-			}
-
-			ModuleStatus.CompilationMethod = CompileMethodString;
 		}
 		OutModuleStatuses.Add( ModuleStatus );
 	}
 }
 
+FString FModuleManager::GetModuleCompileMethod(FName InModuleName)
+{
+	if (!ModuleCompileData.Contains(InModuleName))
+	{
+		UpdateModuleCompileData(InModuleName);
+	}
+
+	switch(ModuleCompileData.FindChecked(InModuleName).Get().CompileMethod)
+	{
+	case EModuleCompileMethod::External:
+		return ModuleManagerDefs::CompileMethodExternal;
+	case EModuleCompileMethod::Runtime:
+		return ModuleManagerDefs::CompileMethodRuntime;
+	default:
+		return ModuleManagerDefs::CompileMethodUnknown;
+	}
+}
+
 
 bool FModuleManager::IsSolutionFilePresent()
 {
@@ -984,90 +968,94 @@ void FModuleManager::OnModuleCompileSucceeded(FName ModuleName, const FString& N
 	// If the compile succeeded, update the module info entry with the new file name for this module
 	ModuleInfo->Filename = NewModuleFilename;
 
-	OnModuleCompileSucceeded(ModuleName, ModuleInfo);
-}
-
-void FModuleManager::OnModuleCompileSucceeded(FName ModuleName, TSharedRef<FModuleManager::FModuleInfo> ModuleInfo)
-{
 #if !IS_MONOLITHIC && WITH_EDITOR
 	// UpdateModuleCompileData() should have been run before compiling so the
 	// data in ModuleInfo should be correct for the pre-compile dll file.
-	if (ensure(ModuleInfo->CompileData.bIsValid))
-	{
-		FDateTime FileTimeStamp;
-		bool bGotFileTimeStamp = GetModuleFileTimeStamp(ModuleInfo, FileTimeStamp);
+	FModuleCompilationData& CompileData = ModuleCompileData.FindChecked(ModuleName).Get();
 
-		ModuleInfo->CompileData.bHasFileTimeStamp = bGotFileTimeStamp;
-		ModuleInfo->CompileData.FileTimeStamp = FileTimeStamp;
+	FDateTime FileTimeStamp;
+	bool bGotFileTimeStamp = GetModuleFileTimeStamp(ModuleName, FileTimeStamp);
 
-		if (ModuleInfo->CompileData.bHasFileTimeStamp)
-		{
-			ModuleInfo->CompileData.CompileMethod = EModuleCompileMethod::Runtime;
-		}
-		else
-		{
-			ModuleInfo->CompileData.CompileMethod = EModuleCompileMethod::Unknown;
-		}
-		WriteModuleCompilationInfoToConfig(ModuleName, ModuleInfo);
+	CompileData.bHasFileTimeStamp = bGotFileTimeStamp;
+	CompileData.FileTimeStamp = FileTimeStamp;
+
+	if (CompileData.bHasFileTimeStamp)
+	{
+		CompileData.CompileMethod = EModuleCompileMethod::Runtime;
 	}
+	else
+	{
+		CompileData.CompileMethod = EModuleCompileMethod::Unknown;
+	}
+	WriteModuleCompilationInfoToConfig(ModuleName, CompileData);
 #endif
 }
 
-void FModuleManager::UpdateModuleCompileData(FName ModuleFName)
+void FModuleManager::UpdateModuleCompileData(FName ModuleName)
 {
-	TSharedRef< FModuleInfo > ModuleInfo = Modules.FindChecked(ModuleFName);
-	UpdateModuleCompileData(ModuleFName, ModuleInfo);
-}
+	// Find or create a compile data object for this module
+	TSharedRef<FModuleCompilationData>* CompileDataPtr = ModuleCompileData.Find(ModuleName);
+	if(CompileDataPtr == nullptr)
+	{
+		CompileDataPtr = &ModuleCompileData.Add(ModuleName, TSharedRef<FModuleCompilationData>(new FModuleCompilationData()));
+	}
 
-void FModuleManager::UpdateModuleCompileData(FName ModuleName, TSharedRef<FModuleManager::FModuleInfo> ModuleInfo)
-{
 	// reset the compile data before updating it
-	ModuleInfo->CompileData.bHasFileTimeStamp = false;
-	ModuleInfo->CompileData.FileTimeStamp = FDateTime(0);
-	ModuleInfo->CompileData.CompileMethod = EModuleCompileMethod::Unknown;
-	ModuleInfo->CompileData.bIsValid = true;
+	FModuleCompilationData& CompileData = CompileDataPtr->Get();
+	CompileData.bHasFileTimeStamp = false;
+	CompileData.FileTimeStamp = FDateTime(0);
+	CompileData.CompileMethod = EModuleCompileMethod::Unknown;
 
 #if !IS_MONOLITHIC && WITH_EDITOR
-	ReadModuleCompilationInfoFromConfig(ModuleName, ModuleInfo);
+	ReadModuleCompilationInfoFromConfig(ModuleName, CompileData);
 
 	FDateTime FileTimeStamp;
-	bool bGotFileTimeStamp = GetModuleFileTimeStamp(ModuleInfo, FileTimeStamp);
+	bool bGotFileTimeStamp = GetModuleFileTimeStamp(ModuleName, FileTimeStamp);
 
 	if (!bGotFileTimeStamp)
 	{
 		// File missing? Reset the cached timestamp and method to defaults and save them.
-		ModuleInfo->CompileData.bHasFileTimeStamp = false;
-		ModuleInfo->CompileData.FileTimeStamp = FDateTime(0);
-		ModuleInfo->CompileData.CompileMethod = EModuleCompileMethod::Unknown;
-		WriteModuleCompilationInfoToConfig(ModuleName, ModuleInfo);
+		CompileData.bHasFileTimeStamp = false;
+		CompileData.FileTimeStamp = FDateTime(0);
+		CompileData.CompileMethod = EModuleCompileMethod::Unknown;
+		WriteModuleCompilationInfoToConfig(ModuleName, CompileData);
 	}
 	else
 	{
-		if (ModuleInfo->CompileData.bHasFileTimeStamp)
+		if (CompileData.bHasFileTimeStamp)
 		{
-			if (FileTimeStamp > ModuleInfo->CompileData.FileTimeStamp + ModuleManagerDefs::TimeStampEpsilon)
+			if (FileTimeStamp > CompileData.FileTimeStamp + ModuleManagerDefs::TimeStampEpsilon)
 			{
 				// The file is newer than the cached timestamp
 				// The file must have been compiled externally
-				ModuleInfo->CompileData.FileTimeStamp = FileTimeStamp;
-				ModuleInfo->CompileData.CompileMethod = EModuleCompileMethod::External;
-				WriteModuleCompilationInfoToConfig(ModuleName, ModuleInfo);
+				CompileData.FileTimeStamp = FileTimeStamp;
+				CompileData.CompileMethod = EModuleCompileMethod::External;
+				WriteModuleCompilationInfoToConfig(ModuleName, CompileData);
 			}
 		}
 		else
 		{
 			// The cached timestamp and method are default value so this file has no history yet
 			// We can only set its timestamp and save
-			ModuleInfo->CompileData.bHasFileTimeStamp = true;
-			ModuleInfo->CompileData.FileTimeStamp = FileTimeStamp;
-			WriteModuleCompilationInfoToConfig(ModuleName, ModuleInfo);
+			CompileData.bHasFileTimeStamp = true;
+			CompileData.FileTimeStamp = FileTimeStamp;
+			WriteModuleCompilationInfoToConfig(ModuleName, CompileData);
 		}
 	}
 #endif
 }
 
+FString FModuleManager::GetModuleFilename(FName ModuleName) const
+{
+	return Modules.FindChecked(ModuleName)->Filename;
+}
 
-void FModuleManager::ReadModuleCompilationInfoFromConfig(FName ModuleName, TSharedRef<FModuleManager::FModuleInfo> ModuleInfo)
+void FModuleManager::SetModuleFilename(FName ModuleName, const FString& Filename)
+{
+	Modules.FindChecked(ModuleName)->Filename = Filename;
+}
+
+void FModuleManager::ReadModuleCompilationInfoFromConfig(FName ModuleName, FModuleCompilationData& CompileData)
 {
 	FString DateTimeString;
 	if (GConfig->GetString(*ModuleManagerDefs::CompilationInfoConfigSection, *FString::Printf(TEXT("%s.TimeStamp"), *ModuleName.ToString()), DateTimeString, GEditorUserSettingsIni))
@@ -1075,19 +1063,19 @@ void FModuleManager::ReadModuleCompilationInfoFromConfig(FName ModuleName, TShar
 		FDateTime TimeStamp;
 		if (!DateTimeString.IsEmpty() && FDateTime::Parse(DateTimeString, TimeStamp))
 		{
-			ModuleInfo->CompileData.bHasFileTimeStamp = true;
-			ModuleInfo->CompileData.FileTimeStamp = TimeStamp;
+			CompileData.bHasFileTimeStamp = true;
+			CompileData.FileTimeStamp = TimeStamp;
 
 			FString CompileMethodString;
 			if (GConfig->GetString(*ModuleManagerDefs::CompilationInfoConfigSection, *FString::Printf(TEXT("%s.LastCompileMethod"), *ModuleName.ToString()), CompileMethodString, GEditorUserSettingsIni))
 			{
 				if (CompileMethodString.Equals(ModuleManagerDefs::CompileMethodRuntime, ESearchCase::IgnoreCase))
 				{
-					ModuleInfo->CompileData.CompileMethod = EModuleCompileMethod::Runtime;
+					CompileData.CompileMethod = EModuleCompileMethod::Runtime;
 				}
 				else if (CompileMethodString.Equals(ModuleManagerDefs::CompileMethodExternal, ESearchCase::IgnoreCase))
 				{
-					ModuleInfo->CompileData.CompileMethod = EModuleCompileMethod::External;
+					CompileData.CompileMethod = EModuleCompileMethod::External;
 				}
 			}
 		}
@@ -1095,38 +1083,36 @@ void FModuleManager::ReadModuleCompilationInfoFromConfig(FName ModuleName, TShar
 }
 
 
-void FModuleManager::WriteModuleCompilationInfoToConfig(FName ModuleName, TSharedRef<FModuleManager::FModuleInfo> ModuleInfo)
+void FModuleManager::WriteModuleCompilationInfoToConfig(FName ModuleName, const FModuleCompilationData& CompileData)
 {
-	if (ensure(ModuleInfo->CompileData.bIsValid))
+	FString DateTimeString;
+	if (CompileData.bHasFileTimeStamp)
 	{
-		FString DateTimeString;
-		if (ModuleInfo->CompileData.bHasFileTimeStamp)
-		{
-			DateTimeString = ModuleInfo->CompileData.FileTimeStamp.ToString();
-		}
+		DateTimeString = CompileData.FileTimeStamp.ToString();
+	}
 
-		GConfig->SetString(*ModuleManagerDefs::CompilationInfoConfigSection, *FString::Printf(TEXT("%s.TimeStamp"), *ModuleName.ToString()), *DateTimeString, GEditorUserSettingsIni);
+	GConfig->SetString(*ModuleManagerDefs::CompilationInfoConfigSection, *FString::Printf(TEXT("%s.TimeStamp"), *ModuleName.ToString()), *DateTimeString, GEditorUserSettingsIni);
 
-		FString CompileMethodString = ModuleManagerDefs::CompileMethodUnknown;
-		if (ModuleInfo->CompileData.CompileMethod == EModuleCompileMethod::Runtime)
-		{
-			CompileMethodString = ModuleManagerDefs::CompileMethodRuntime;
-		}
-		else if (ModuleInfo->CompileData.CompileMethod == EModuleCompileMethod::External)
-		{
-			CompileMethodString = ModuleManagerDefs::CompileMethodExternal;
-		}
-
-		GConfig->SetString(*ModuleManagerDefs::CompilationInfoConfigSection, *FString::Printf(TEXT("%s.LastCompileMethod"), *ModuleName.ToString()), *CompileMethodString, GEditorUserSettingsIni);
+	FString CompileMethodString = ModuleManagerDefs::CompileMethodUnknown;
+	if (CompileData.CompileMethod == EModuleCompileMethod::Runtime)
+	{
+		CompileMethodString = ModuleManagerDefs::CompileMethodRuntime;
 	}
+	else if (CompileData.CompileMethod == EModuleCompileMethod::External)
+	{
+		CompileMethodString = ModuleManagerDefs::CompileMethodExternal;
+	}
+
+	GConfig->SetString(*ModuleManagerDefs::CompilationInfoConfigSection, *FString::Printf(TEXT("%s.LastCompileMethod"), *ModuleName.ToString()), *CompileMethodString, GEditorUserSettingsIni);
 }
 
 
-bool FModuleManager::GetModuleFileTimeStamp(TSharedRef<const FModuleInfo> ModuleInfo, FDateTime& OutFileTimeStamp)
+bool FModuleManager::GetModuleFileTimeStamp(FName ModuleName, FDateTime& OutFileTimeStamp) const
 {
-	if (IFileManager::Get().FileSize(*ModuleInfo->Filename) > 0)
+	FString Filename = GetModuleFilename(ModuleName);
+	if (IFileManager::Get().FileSize(*Filename) > 0)
 	{
-		OutFileTimeStamp = FDateTime(IFileManager::Get().GetTimeStamp(*ModuleInfo->Filename));
+		OutFileTimeStamp = FDateTime(IFileManager::Get().GetTimeStamp(*Filename));
 		return true;
 	}
 	return false;
@@ -1381,13 +1367,9 @@ bool FModuleManager::StartCompilingModuleDLLs(const FString& GameName, const TAr
 		}
 		Ar.Logf( TEXT( "Recompiling %s..." ), *ModuleNames[ CurModuleIndex ].ModuleName );
 
+		// prepare the compile info in the FModuleInfo so that it can be compared after compiling
 		FName ModuleFName(*ModuleNames[ CurModuleIndex ].ModuleName);
-		if (ensure(Modules.Contains(ModuleFName)))
-		{
-			// prepare the compile info in the FModuleInfo so that it can be compared after compiling
-			TSharedRef< FModuleInfo > ModuleInfo = Modules.FindChecked(ModuleFName);
-			UpdateModuleCompileData(ModuleFName, ModuleInfo);
-		}
+		UpdateModuleCompileData(ModuleFName);
 	}
 
 	FString ExtraArg;
@@ -1574,17 +1556,8 @@ void FModuleManager::CheckForFinishedModuleDLLCompile( const bool bWaitForComple
 					// Were we asked to assign a new file name for this module?
 					if( !CurModule.NewModuleFilename.IsEmpty() )
 					{
-						// Find the module
-						const FName CurModuleName = FName( *CurModule.ModuleName );
-						if( ensure( Modules.Contains( CurModuleName ) ) )
-						{
-							TSharedRef< FModuleInfo > ModuleInfo = Modules.FindChecked( CurModuleName );
-
-							// If the compile succeeded, update the module info entry with the new file name for this module
-							ModuleInfo->Filename = CurModule.NewModuleFilename;
-
-							OnModuleCompileSucceeded(FName(*CurModule.ModuleName), ModuleInfo);
-						}
+						// If the compile succeeded, update the module info entry with the new file name for this module
+						OnModuleCompileSucceeded(FName(*CurModule.ModuleName), CurModule.NewModuleFilename);
 					}
 				}
 
diff --git a/Engine/Source/Runtime/Core/Public/Modules/ModuleManager.h b/Engine/Source/Runtime/Core/Public/Modules/ModuleManager.h
--- a/Engine/Source/Runtime/Core/Public/Modules/ModuleManager.h
+++ b/Engine/Source/Runtime/Core/Public/Modules/ModuleManager.h
@@ -130,9 +130,6 @@ struct FModuleStatus
 
 	/** Whether this module contains game play code. */
 	bool bIsGameModule;
-
-	/** The compilation method of this module. */
-	FString CompilationMethod;
 };
 
 /**
@@ -271,6 +268,14 @@ public:
 	 */
 	void QueryModules( TArray<FModuleStatus>& OutModuleStatuses );
 
+	/**
+	 * Queries the compilation method for a given module.
+	 *
+	 * @param InModuleName Module to query the name of
+	 * @return A string describing the method used to compile the module.
+	 */
+	FString GetModuleCompileMethod(FName InModuleName);
+
 	/**
 	 * Unloads a specific module
 	 *
@@ -522,6 +527,12 @@ public:
 	/** Called when the compile data for a module need to be update in memory and written to config */
 	void UpdateModuleCompileData(FName ModuleName);
 
+	/** Gets the filename for a module */
+	FString GetModuleFilename(FName ModuleName) const;
+
+	/** Sets the filename for a module. The module is not reloaded immediately, but the new name will be used for subsequent unload/load events. */
+	void SetModuleFilename(FName ModuleName, const FString& Filename);
+
 public:
 
 	/**
@@ -635,9 +646,6 @@ protected:
 	 */
 	struct FModuleCompilationData
 	{
-		/** A flag set when the data it updated - loaded modules don't update this info until they are compiled or just before they unload */
-		bool bIsValid;
-
 		/** Has a timestamp been set for the .dll file */
 		bool bHasFileTimeStamp;
 
@@ -648,8 +656,7 @@ protected:
 		EModuleCompileMethod CompileMethod;
 
 		FModuleCompilationData()
-			: bIsValid(false)
-			, bHasFileTimeStamp(false)
+			: bHasFileTimeStamp(false)
 			, CompileMethod(EModuleCompileMethod::Unknown)
 		{ }
 	};
@@ -679,9 +686,6 @@ protected:
 		/** Arbitrary number that encodes the load order of this module, so we can shut them down in reverse order. */
 		int32 LoadOrder;
 
-		/** Last know compilation data for this module - undefined if CompileData.bIsValid is false */
-		FModuleCompilationData CompileData;
-
 		/** static that tracks the current load number. Incremented whenever we add a new module*/
 		static int32 CurrentLoadOrder;
 
@@ -738,20 +742,14 @@ protected:
 	/** Compares file versions between the current executing engine version and the specified dll */
 	static bool CheckModuleCompatibility(const TCHAR *Filename);
 
-	/** Called during CheckForFinishedModuleDLLCompile() for each successfully recomplied module */
-	void OnModuleCompileSucceeded(FName ModuleName, TSharedRef<FModuleInfo> ModuleInfo);
-
-	/** Called when the compile data for a module need to be update in memory and written to config */
-	void UpdateModuleCompileData(FName ModuleName, TSharedRef<FModuleInfo> ModuleInfo);
-
 	/** Called when a new module is added to the manager to get the saved compile data from config */
-	static void ReadModuleCompilationInfoFromConfig(FName ModuleName, TSharedRef<FModuleInfo> ModuleInfo);
+	static void ReadModuleCompilationInfoFromConfig(FName ModuleName, FModuleCompilationData& CompileData);
 
 	/** Saves the module's compile data to config */
-	static void WriteModuleCompilationInfoToConfig(FName ModuleName, TSharedRef<FModuleInfo> ModuleInfo);
+	static void WriteModuleCompilationInfoToConfig(FName ModuleName, const FModuleCompilationData& CompileData);
 
 	/** Access the module's file and read the timestamp from the file system. Returns true if the timestamp was read successfully. */
-	static bool GetModuleFileTimeStamp(TSharedRef<const FModuleInfo> ModuleInfo, FDateTime& OutFileTimeStamp);
+	bool GetModuleFileTimeStamp(FName ModuleName, FDateTime& OutFileTimeStamp) const;
 
 	/** Finds modules matching a given name wildcard. */
 	void FindModulePaths(const TCHAR *NamePattern, TMap<FName, FString> &OutModulePaths) const;
@@ -784,6 +782,9 @@ private:
 	/** Multicast delegate called to process any new loaded objects. */
 	FSimpleMulticastDelegate ProcessLoadedObjectsCallback;
 
+	/** Last known compilation data for each module */
+	TMap<FName, TSharedRef<FModuleCompilationData>> ModuleCompileData;
+
 	/** When compiling a module using an external application, stores the handle to the process that is running */
 	FProcHandle ModuleCompileProcessHandle;
 
