fixed check() on engine quit becuase of cvar sink

diff --git a/Engine/Source/Runtime/Core/Public/HAL/IConsoleManager.h b/Engine/Source/Runtime/Core/Public/HAL/IConsoleManager.h
--- a/Engine/Source/Runtime/Core/Public/HAL/IConsoleManager.h
+++ b/Engine/Source/Runtime/Core/Public/HAL/IConsoleManager.h
@@ -525,7 +525,7 @@ public:
 	/** Destructor, removes the console variable sink **/
 	virtual ~FAutoConsoleVariableSink()
 	{
-		IConsoleManager::Get().UnregisterConsoleVariableSink(Command);
+//disabled for now, destruction order makes this not always working		IConsoleManager::Get().UnregisterConsoleVariableSink(Command);
 	}
 
 	const FConsoleCommandDelegate& Command;
