Fix overlap of src and dest in memcpy().

Happened with the following callstack:

FGenericPlatformMemory::Memswap(void*, void*, unsigned long) (GenericPlatformMemory.h:248)
void Swap<TSharedRef<FWorkspaceItem, (ESPMode::Type)0> >(TSharedRef<FWorkspaceItem, (ESPMode::Type)0>&, TSharedRef<FWorkspaceItem, (ESPMode::Type)0>&) (UnrealMemory.h:64)
void Exchange<TSharedRef<FWorkspaceItem, (ESPMode::Type)0> >(TSharedRef<FWorkspaceItem, (ESPMode::Type)0>&, TSharedRef<FWorkspaceItem, (ESPMode::Type)0>&) (UnrealTemplate.h:421)
void SortInternal<TSharedRef<FWorkspaceItem, (ESPMode::Type)0>, TDereferenceWrapper<TSharedRef<FWorkspaceItem, (ESPMode::Type)0>, FWorkspaceItem::FWorkspaceItemSort> >(TSharedRef<FWorkspaceItem, (ESPMode::Type)0>*, int, TDereferenceWrapper<TSharedRef<FWorkspaceItem, (ESPMode::Type)0>, FWorkspaceItem::FWorkspaceItemSort> const&) (Sorting.h:102)
void Sort<TSharedRef<FWorkspaceItem, (ESPMode::Type)0>, FWorkspaceItem::FWorkspaceItemSort>(TSharedRef<FWorkspaceItem, (ESPMode::Type)0>*, int, FWorkspaceItem::FWorkspaceItemSort const&) (Sorting.h:168)
void TArray<TSharedRef<FWorkspaceItem, (ESPMode::Type)0>, FDefaultAllocator>::Sort<FWorkspaceItem::FWorkspaceItemSort>(FWorkspaceItem::FWorkspaceItemSort const&) (Array.h:1565)
FWorkspaceItem::SortChildren() (WorkspaceItem.h:105)
FWorkspaceItem::AddItem(TSharedRef<FWorkspaceItem, (ESPMode::Type)0> const&) (WorkspaceItem.h:83)
FWorkspaceItem::AddGroup(FText const&, FSlateIcon const&, bool) (WorkspaceItem.h:44)
FWorkspaceMenuStructure::ResetLevelEditorCategory() (WorkspaceMenuStructureModule.cpp:74)
FWorkspaceMenuStructure::FWorkspaceMenuStructure() (WorkspaceMenuStructureModule.cpp:92)

diff --git a/Engine/Source/Runtime/Core/Private/GenericPlatform/GenericPlatformMemory.cpp b/Engine/Source/Runtime/Core/Private/GenericPlatform/GenericPlatformMemory.cpp
--- a/Engine/Source/Runtime/Core/Private/GenericPlatform/GenericPlatformMemory.cpp
+++ b/Engine/Source/Runtime/Core/Private/GenericPlatform/GenericPlatformMemory.cpp
@@ -147,10 +147,13 @@ void FGenericPlatformMemory::DumpPlatformAndAllocatorStats( class FOutputDevice&
 
 void FGenericPlatformMemory::Memswap( void* Ptr1, void* Ptr2, SIZE_T Size )
 {
-	void* Temp = FMemory_Alloca(Size);
-	FPlatformMemory::Memcpy( Temp, Ptr1, Size );
-	FPlatformMemory::Memcpy( Ptr1, Ptr2, Size );
-	FPlatformMemory::Memcpy( Ptr2, Temp, Size );
+	if (Ptr1 != Ptr2)
+	{
+		void* Temp = FMemory_Alloca(Size);
+		FPlatformMemory::Memcpy( Temp, Ptr1, Size );
+		FPlatformMemory::Memcpy( Ptr1, Ptr2, Size );
+		FPlatformMemory::Memcpy( Ptr2, Temp, Size );
+	}
 }
 
 FGenericPlatformMemory::FSharedMemoryRegion::FSharedMemoryRegion(const FString & InName, uint32 InAccessMode, void * InAddress, SIZE_T InSize)
