Fix TOptional so it calls the constructor and destructor of the optional type at the correct times.
Added Emplace() and Reset() to TOptional.
Added checkSlow() that you aren't calling GetValue() on a TOptional that's not set, and fix something that was.

diff --git a/Engine/Source/Runtime/Core/Public/Containers/ContainerAllocationPolicies.h b/Engine/Source/Runtime/Core/Public/Containers/ContainerAllocationPolicies.h
--- a/Engine/Source/Runtime/Core/Public/Containers/ContainerAllocationPolicies.h
+++ b/Engine/Source/Runtime/Core/Public/Containers/ContainerAllocationPolicies.h
@@ -59,10 +59,7 @@ struct FScriptContainerElement
  * It needs to use template specialization as the MS_ALIGN and GCC_ALIGN macros require literal parameters.
  */
 template<int32 Size,uint32 Alignment>
-class TAlignedBytes
-{
-	uint8 Data[-Size]; // this intentionally won't compile, we don't support the requested alignment
-};
+class TAlignedBytes; // this intentionally won't compile, we don't support the requested alignment
 
 /** Unaligned storage. */
 template<int32 Size>
diff --git a/Engine/Source/Runtime/Core/Public/Misc/Optional.h b/Engine/Source/Runtime/Core/Public/Misc/Optional.h
--- a/Engine/Source/Runtime/Core/Public/Misc/Optional.h
+++ b/Engine/Source/Runtime/Core/Public/Misc/Optional.h
@@ -11,29 +11,177 @@ struct TOptional
 {
 public:
 	/** Construct an OptionaType with a valid value. */
-	TOptional( OptionalType InValue )
-		: bIsSet( true )
-		, Value( InValue )
+	TOptional(const OptionalType& InValue)
 	{
+		new(&Value) OptionalType(InValue);
+		bIsSet = true;
+	}
+	TOptional(OptionalType&& InValue)
+	{
+		new(&Value) OptionalType(MoveTemp(InValue));
+		bIsSet = true;
 	}
 
-	/** Construct an OptionalType with no value; i.e. NULL */
+	/** Construct an OptionalType with no value; i.e. unset */
 	TOptional()
-		: bIsSet( false )
-		, Value()
+		: bIsSet(false)
+	{
+	}
+
+	~TOptional()
+	{
+		Reset();
+	}
+
+	/** Copy/Move construction */
+	TOptional(const TOptional& InValue)
+		: bIsSet(false)
+	{
+		if (InValue.bIsSet)
+		{
+			new(&Value) OptionalType((const OptionalType&)InValue.Value);
+			bIsSet = true;
+		}
+	}
+	TOptional(TOptional&& InValue)
+		: bIsSet(false)
+	{
+		if (InValue.bIsSet)
+		{
+			new(&Value) OptionalType((OptionalType&&)InValue.Value);
+			bIsSet = true;
+		}
+	}
+
+	TOptional& operator=(const TOptional& InValue)
+	{
+		Reset();
+		if (InValue.bIsSet)
+		{
+			new(&Value) OptionalType((const OptionalType&)InValue.Value);
+			bIsSet = true;
+		}
+		return *this;
+	}
+	TOptional& operator=(TOptional&& InValue)
+	{
+		Reset();
+		if (InValue.bIsSet)
+		{
+			new(&Value) OptionalType((OptionalType&&)InValue.Value);
+			bIsSet = true;
+		}
+		return *this;
+	}
+
+	TOptional& operator=(const OptionalType& InValue)
 	{
+		Reset();
+		new(&Value) OptionalType(InValue);
+		bIsSet = true;
+		return *this;
+	}
+	TOptional& operator=(OptionalType&& InValue)
+	{
+		Reset();
+		new(&Value) OptionalType(MoveTemp(InValue));
+		bIsSet = true;
+		return *this;
+	}
+
+	void Reset()
+	{
+		if (bIsSet)
+		{
+			bIsSet = false;
+
+			// We need a typedef here because VC won't compile the destructor call below if OptionalType itself has a member called OptionalType
+			typedef OptionalType OptionalDestructOptionalType;
+			((OptionalType&)Value).OptionalDestructOptionalType::~OptionalDestructOptionalType();
+		}
 	}
 
-	/** @return true when the value is meaningulf; false if calling GetValue() is undefined. */
+#if PLATFORM_COMPILER_HAS_VARIADIC_TEMPLATES
+	template <typename... ArgsType>
+	void Emplace(ArgsType&&... Args)
+	{
+		Reset();
+		new(&Value) OptionalType(Forward<ArgsType>(Args)...);
+		bIsSet = true;
+	}
+#else
+	void Emplace()
+	{
+		Reset();
+		new(&Value) OptionalType();
+		bIsSet = true;
+	}
+
+	template <typename Arg0Type>
+	void Emplace(Arg0Type&& Arg0)
+	{
+		Reset();
+		new(&Value) OptionalType(Forward<Arg0Type>(Arg0));
+		bIsSet = true;
+	}
+
+	template <typename Arg0Type, typename Arg1Type>
+	void Emplace(Arg0Type&& Arg0, Arg1Type&& Arg1)
+	{
+		Reset();
+		new(&Value) OptionalType(Forward<Arg0Type>(Arg0), Forward<Arg1Type>(Arg1));
+		bIsSet = true;
+	}
+
+	template <typename Arg0Type, typename Arg1Type, typename Arg2Type>
+	void Emplace(Arg0Type&& Arg0, Arg1Type&& Arg1, Arg2Type&& Arg2)
+	{
+		Reset();
+		new(&Value) OptionalType(Forward<Arg0Type>(Arg0), Forward<Arg1Type>(Arg1), Forward<Arg2Type>(Arg2));
+		bIsSet = true;
+	}
+
+	template <typename Arg0Type, typename Arg1Type, typename Arg2Type, typename Arg3Type>
+	void Emplace(Arg0Type&& Arg0, Arg1Type&& Arg1, Arg2Type&& Arg2, Arg3Type&& Arg3)
+	{
+		Reset();
+		new(&Value) OptionalType(Forward<Arg0Type>(Arg0), Forward<Arg1Type>(Arg1), Forward<Arg2Type>(Arg2), Forward<Arg3Type>(Arg3));
+		bIsSet = true;
+	}
+#endif
+
+	friend bool operator==(const TOptional& lhs, const TOptional& rhs)
+	{
+		if (lhs.bIsSet != rhs.bIsSet)
+		{
+			return false;
+		}
+		if (!lhs.bIsSet) // both unset
+		{
+			return true;
+		}
+		return (OptionalType&)lhs.Value == (OptionalType&)rhs.Value;
+	}
+	friend bool operator!=(const TOptional& lhs, const TOptional& rhs)
+	{
+		return !(lhs == rhs);
+	}
+
+	/** @return true when the value is meaningful; false if calling GetValue() is undefined. */
 	bool IsSet() const { return bIsSet; }
+	FORCEINLINE_EXPLICIT_OPERATOR_BOOL() const { return bIsSet; }
 
 	/** @return The optional value; undefined when IsSet() returns false. */
-	const OptionalType& GetValue() const { return Value; }
+	const OptionalType& GetValue() const { checkSlow(IsSet()); return (OptionalType&)Value; }
+	      OptionalType& GetValue()       { checkSlow(IsSet()); return (OptionalType&)Value; }
+
+	const OptionalType* operator->() const { return &GetValue(); }
+	      OptionalType* operator->()       { return &GetValue(); }
 
 	/** @return The optional value when set; DefaultValue otherwise. */
-	const OptionalType& Get( const OptionalType& DefaultValue ) const { return IsSet() ? Value : DefaultValue; }
+	const OptionalType& Get(const OptionalType& DefaultValue) const { return IsSet() ? (OptionalType&)Value : DefaultValue; }
 
 private:
 	bool bIsSet;
-	OptionalType Value;
+	TTypeCompatibleBytes<OptionalType> Value;
 };
\ No newline at end of file
diff --git a/Engine/Source/Runtime/UMG/Private/Components/SpinBox.cpp b/Engine/Source/Runtime/UMG/Private/Components/SpinBox.cpp
--- a/Engine/Source/Runtime/UMG/Private/Components/SpinBox.cpp
+++ b/Engine/Source/Runtime/UMG/Private/Components/SpinBox.cpp
@@ -17,10 +17,10 @@ USpinBox::USpinBox(const FPostConstructInitializeProperties& PCIP)
 	SSpinBox<float>::FArguments Defaults;
 
 	Value = Defaults._Value.Get();
-	MinValue = Defaults._MinValue.Get().GetValue();
-	MaxValue = Defaults._MaxValue.Get().GetValue();
-	MinSliderValue = Defaults._MinSliderValue.Get().GetValue();
-	MaxSliderValue = Defaults._MaxSliderValue.Get().GetValue();
+	MinValue = Defaults._MinValue.Get().Get(0.0f);
+	MaxValue = Defaults._MaxValue.Get().Get(0.0f);
+	MinSliderValue = Defaults._MinSliderValue.Get().Get(0.0f);
+	MaxSliderValue = Defaults._MaxSliderValue.Get().Get(0.0f);
 	Delta = Defaults._Delta.Get();
 	SliderExponent = Defaults._SliderExponent.Get();
 	MinDesiredWidth = Defaults._MinDesiredWidth.Get();
