Fix "error: 'TAlignedBytes' defined as a struct template here but previously declared as a class template" on clang

diff --git a/Engine/Source/Runtime/Core/Public/Containers/ContainerAllocationPolicies.h b/Engine/Source/Runtime/Core/Public/Containers/ContainerAllocationPolicies.h
--- a/Engine/Source/Runtime/Core/Public/Containers/ContainerAllocationPolicies.h
+++ b/Engine/Source/Runtime/Core/Public/Containers/ContainerAllocationPolicies.h
@@ -59,7 +59,7 @@ struct FScriptContainerElement
  * It needs to use template specialization as the MS_ALIGN and GCC_ALIGN macros require literal parameters.
  */
 template<int32 Size,uint32 Alignment>
-class TAlignedBytes; // this intentionally won't compile, we don't support the requested alignment
+struct TAlignedBytes; // this intentionally won't compile, we don't support the requested alignment
 
 /** Unaligned storage. */
 template<int32 Size>
@@ -103,7 +103,7 @@ IMPLEMENT_ALIGNED_STORAGE(2);
 
 /** An untyped array of data with compile-time alignment and size derived from another type. */
 template<typename ElementType>
-class TTypeCompatibleBytes :
+struct TTypeCompatibleBytes :
 	public TAlignedBytes<
 		sizeof(ElementType),
 		ALIGNOF(ElementType)
