Hopefully fix "error: cannot cast from lvalue of type 'TTypeCompatibleBytes' to rvalue reference type 'FIntPoint &&'; types are not compatible" on HTML5

diff --git a/Engine/Source/Runtime/Core/Public/Misc/Optional.h b/Engine/Source/Runtime/Core/Public/Misc/Optional.h
--- a/Engine/Source/Runtime/Core/Public/Misc/Optional.h
+++ b/Engine/Source/Runtime/Core/Public/Misc/Optional.h
@@ -48,7 +48,7 @@ public:
 	{
 		if (InValue.bIsSet)
 		{
-			new(&Value) OptionalType((OptionalType&&)InValue.Value);
+			new(&Value) OptionalType(MoveTemp((OptionalType&)InValue.Value));
 			bIsSet = true;
 		}
 	}
@@ -68,7 +68,7 @@ public:
 		Reset();
 		if (InValue.bIsSet)
 		{
-			new(&Value) OptionalType((OptionalType&&)InValue.Value);
+			new(&Value) OptionalType(MoveTemp((OptionalType&)InValue.Value));
 			bIsSet = true;
 		}
 		return *this;
