Use Apple's private CoreSymbolicate framework for symbolisation on OS X rather than relying on atos. This is faster, doesn't require Xcode & is also available on iOS, but can't be used to ship on the AppStore's.

diff --git a/Engine/Source/Editor/UnrealEd/Private/SourceCodeNavigation.cpp b/Engine/Source/Editor/UnrealEd/Private/SourceCodeNavigation.cpp
--- a/Engine/Source/Editor/UnrealEd/Private/SourceCodeNavigation.cpp
+++ b/Engine/Source/Editor/UnrealEd/Private/SourceCodeNavigation.cpp
@@ -17,6 +17,7 @@
 #include <mach-o/nlist.h>
 #include <mach-o/stab.h>
 #include <cxxabi.h>
+#include "ApplePlatformSymbolication.h"
 #endif
 
 DEFINE_LOG_CATEGORY(LogSelectionDetails);
@@ -508,29 +509,11 @@ void FSourceCodeNavigationImpl::NavigateToFunctionSource( const FString& Functio
 							
 							if(FunctionSymbolName == SymbolName)
 							{
-								uint64 Address = SymbolEntry.n_value;
-								// Mavericks requires additional params to atos to silence a deprecation warning. On Yosemite same params cause "unrecognized option" error
-								FString AtoSCommand = FMacPlatformMisc::IsRunningOnMavericks() ? FString::Printf(TEXT("-nowarning -arch x86_64 -d -o \"%s\" 0x%x"), *FullModulePath, Address) : FString::Printf(TEXT("-arch x86_64 -o \"%s\" 0x%x"), *FullModulePath, Address);
-								int32 ReturnCode = 0;
-								FString Results;
-								FPlatformProcess::ExecProcess( TEXT("/usr/bin/atos"), *AtoSCommand, &ReturnCode, &Results, NULL );
-								if(ReturnCode == 0)
+								FString ModuleName = FPaths::GetCleanFilename(FullModulePath);
+								FProgramCounterSymbolInfo Info;
+								if(FApplePlatformSymbolication::SymbolInfoForFunctionFromModule((char const*)(StringTable+SymbolEntry.n_un.n_strx), TCHAR_TO_UTF8(*ModuleName), Info))
 								{
-									int32 FirstIndex = -1;
-									int32 LastIndex = -1;
-									if(Results.FindChar(TCHAR('('), FirstIndex) && Results.FindLastChar(TCHAR('('), LastIndex) && FirstIndex != LastIndex)
-									{
-										int32 CloseIndex = -1;
-										int32 ColonIndex = -1;
-										if(Results.FindLastChar(TCHAR(':'), ColonIndex) && Results.FindLastChar(TCHAR(')'), CloseIndex))
-										{
-											int32 FileNamePos = LastIndex+1;
-											int32 FileNameLen = ColonIndex-FileNamePos;
-											FString FileName = Results.Mid(FileNamePos, FileNameLen);
-											FString LineNumber = Results.Mid(ColonIndex + 1, CloseIndex-(ColonIndex + 1));
-											SourceCodeAccessor.OpenFileAtLine( FileName, FCString::Atoi(*LineNumber), 0 );
-										}
-									}
+									SourceCodeAccessor.OpenFileAtLine( Info.Filename, Info.LineNumber, 0 );
 								}
 								break;
 							}
diff --git a/Engine/Source/Runtime/Core/Core.Build.cs b/Engine/Source/Runtime/Core/Core.Build.cs
--- a/Engine/Source/Runtime/Core/Core.Build.cs
+++ b/Engine/Source/Runtime/Core/Core.Build.cs
@@ -1,4 +1,4 @@
-﻿// Copyright 1998-2014 Epic Games, Inc. All Rights Reserved.
+// Copyright 1998-2014 Epic Games, Inc. All Rights Reserved.
 
 using UnrealBuildTool;
 using System;
@@ -86,6 +86,7 @@ public class Core : ModuleRules
 			{
 				PublicAdditionalLibraries.Add("/System/Library/PrivateFrameworks/MultitouchSupport.framework/Versions/Current/MultitouchSupport");
 			}
+			PublicAdditionalLibraries.Add("/System/Library/PrivateFrameworks/CoreSymbolication.framework/Versions/Current/CoreSymbolication");
 		}
 		else if (Target.Platform == UnrealTargetPlatform.IOS)
 		{
diff --git a/Engine/Source/Runtime/Core/Private/Apple/ApplePlatformStackWalk.cpp b/Engine/Source/Runtime/Core/Private/Apple/ApplePlatformStackWalk.cpp
--- a/Engine/Source/Runtime/Core/Private/Apple/ApplePlatformStackWalk.cpp
+++ b/Engine/Source/Runtime/Core/Private/Apple/ApplePlatformStackWalk.cpp
@@ -5,6 +5,7 @@
 =============================================================================*/
 
 #include "Core.h"
+#include "ApplePlatformSymbolication.h"
 #include <execinfo.h>
 #include <dlfcn.h>
 #include <mach-o/dyld.h>
@@ -118,147 +119,63 @@ bool FApplePlatformStackWalk::ProgramCounterToHumanReadableString( int32 Current
 
 void FApplePlatformStackWalk::ProgramCounterToSymbolInfo( uint64 ProgramCounter, FProgramCounterSymbolInfo&  SymbolInfo )
 {
-	Dl_info DylibInfo;
-	int32 Result = dladdr((const void*)ProgramCounter, &DylibInfo);
-	if (Result == 0)
-	{
-		return;
-	}
-
-	int32 Status = 0;
-	ANSICHAR* DemangledName = NULL;
-	
-	// Increased the size of the demangle destination to reduce the chances that abi::__cxa_demangle will allocate
-	// this causes the app to hang as malloc isn't signal handler safe. Ideally we wouldn't call this function in a handler.
-	size_t DemangledNameLen = 65536;
-	ANSICHAR DemangledNameBuffer[65536]= {0};
-	DemangledName = abi::__cxa_demangle(DylibInfo.dli_sname, DemangledNameBuffer, &DemangledNameLen, &Status);
-
-	if (DemangledName)
+	bool bOK = FApplePlatformSymbolication::SymbolInfoForAddress(ProgramCounter, SymbolInfo);
+	if(!bOK)
 	{
-		// C++ function
-		FCStringAnsi::Sprintf(SymbolInfo.FunctionName, "%s ", DemangledName);
-	}
-	else if (DylibInfo.dli_sname && strchr(DylibInfo.dli_sname, ']'))
-	{
-		// ObjC function
-		FCStringAnsi::Sprintf(SymbolInfo.FunctionName, "%s ", DylibInfo.dli_sname);
-	}
-	else if(DylibInfo.dli_sname)
-	{
-		// C function
-		FCStringAnsi::Sprintf(SymbolInfo.FunctionName, "%s() ", DylibInfo.dli_sname);
-	}
-	else
-	{
-		// Unknown!
-		FCStringAnsi::Sprintf(SymbolInfo.FunctionName, "[Unknown]() ");
-	}
-
-	bool OK = false;
-	
-#if PLATFORM_MAC
-	{
-		static char FileName[65536];
-		FMemory::MemSet(FileName, 0);
+		Dl_info DylibInfo;
+		int32 Result = dladdr((const void*)ProgramCounter, &DylibInfo);
+		if (Result == 0)
+		{
+			return;
+		}
+		
+		int32 Status = 0;
+		ANSICHAR* DemangledName = NULL;
+		
+		// Increased the size of the demangle destination to reduce the chances that abi::__cxa_demangle will allocate
+		// this causes the app to hang as malloc isn't signal handler safe. Ideally we wouldn't call this function in a handler.
+		size_t DemangledNameLen = 65536;
+		ANSICHAR DemangledNameBuffer[65536]= {0};
+		DemangledName = abi::__cxa_demangle(DylibInfo.dli_sname, DemangledNameBuffer, &DemangledNameLen, &Status);
 		
-		// Use fork() & execl() as they are async-signal safe, ExecProcess can fail in Cocoa
-		int32 ReturnCode = 0;
-		int FileDesc[2];
-		pipe(FileDesc);
-		pid_t ForkPID = fork();
-		if(ForkPID == 0)
+		if (DemangledName)
 		{
-			// Child
-			close(FileDesc[0]);
-			dup2(FileDesc[1], STDOUT_FILENO);
-			
-			static char FBase[65];
-			static char Address[65];
-			FCStringAnsi::Sprintf(FBase, "%p", (void*)DylibInfo.dli_fbase);
-			FCStringAnsi::Sprintf(Address, "%p", (void*)ProgramCounter);
-			
-			// Mavericks requires additional params to atos to silence a deprecation warning. On Yosemite same params cause "unrecognized option" error
-			if (FPlatformMisc::IsRunningOnMavericks())
-			{
-				execl("/usr/bin/atos", "-nowarning", "-arch", "x86_64", "-d", "-l", FBase, "-o", DylibInfo.dli_fname, Address, NULL);
-			}
-			else
-			{
-				execl("/usr/bin/atos", "-arch", "x86_64", "-l", FBase, "-o", DylibInfo.dli_fname, Address, NULL);
-			}
-			close(FileDesc[1]);
+			// C++ function
+			FCStringAnsi::Sprintf(SymbolInfo.FunctionName, "%s ", DemangledName);
 		}
-		else
+		else if (DylibInfo.dli_sname && strchr(DylibInfo.dli_sname, ']'))
 		{
-			// Parent
-			int StatLoc = 0;
-			pid_t WaitPID = waitpid(ForkPID, &StatLoc, 0);
-			int Bytes = 0;
-			close(FileDesc[1]);
-			int32 FirstIndex = -1;
-			int32 LastIndex = -1;
-			char* Buffer = FileName;
-			uint32 Len = ARRAY_COUNT(FileName) - 1;
-			while ((Bytes = read(FileDesc[0], Buffer, Len)) > 0)
-			{
-				Buffer += Bytes;
-				Len -= Bytes;
-			}
-			close(FileDesc[0]);
-			ReturnCode = WEXITSTATUS(StatLoc);
+			// ObjC function
+			FCStringAnsi::Sprintf(SymbolInfo.FunctionName, "%s ", DylibInfo.dli_sname);
 		}
-		
-		if(ReturnCode == 0)
+		else if(DylibInfo.dli_sname)
 		{
-			char* FirstBracket = FCStringAnsi::Strchr(FileName, '(');
-			char* SecondBracket = FCStringAnsi::Strrchr(FileName, '(');
-			if(FirstBracket && SecondBracket && FirstBracket != SecondBracket)
-			{
-				char* CloseBracket = SecondBracket;
-				CloseBracket = FCStringAnsi::Strrchr(CloseBracket, ')');
-				if(CloseBracket)
-				{
-					char* Colon = SecondBracket;
-					Colon = FCStringAnsi::Strchr(Colon, ':');
-					if(Colon)
-					{
-						*Colon = 0;
-					}
-					*SecondBracket = 0;
-					
-					FCStringAnsi::Strncpy(SymbolInfo.Filename, (SecondBracket+1), ARRAY_COUNT(SymbolInfo.Filename)-1);
-					if(Colon)
-					{
-						// I'm going to assume that this is safe too
-						SymbolInfo.LineNumber = FCStringAnsi::Atoi(Colon+1);
-					}
-					OK = true;
-				}
-			}
+			// C function
+			FCStringAnsi::Sprintf(SymbolInfo.FunctionName, "%s() ", DylibInfo.dli_sname);
 		}
-	}
-#endif
-	
-	if(!OK)
-	{
+		else
+		{
+			// Unknown!
+			FCStringAnsi::Sprintf(SymbolInfo.FunctionName, "[Unknown]() ");
+		}
+		
 		// No line number found.
 		FCStringAnsi::Strcat(SymbolInfo.Filename, "Unknown");
 		SymbolInfo.LineNumber = 0;
+		
+		// Write out Module information.
+		ANSICHAR* DylibPath = (ANSICHAR*)DylibInfo.dli_fname;
+		ANSICHAR* DylibName = FCStringAnsi::Strrchr(DylibPath, '/');
+		if(DylibName)
+		{
+			DylibName += 1;
+		}
+		else
+		{
+			DylibName = DylibPath;
+		}
+		FCStringAnsi::Strcpy(SymbolInfo.ModuleName, DylibName);
 	}
-
-	// Write out Module information.
-	ANSICHAR* DylibPath = (ANSICHAR*)DylibInfo.dli_fname;
-	ANSICHAR* DylibName = FCStringAnsi::Strrchr(DylibPath, '/');
-	if(DylibName)
-	{
-		DylibName += 1;
-	}
-	else
-	{
-		DylibName = DylibPath;
-	}
-	FCStringAnsi::Strcpy(SymbolInfo.ModuleName, DylibName);
 }
 
 int32 FApplePlatformStackWalk::GetProcessModuleCount()
diff --git a/Engine/Source/Runtime/Core/Private/Apple/ApplePlatformSymbolication.cpp b/Engine/Source/Runtime/Core/Private/Apple/ApplePlatformSymbolication.cpp
new file mode 100644
--- /dev/null
+++ b/Engine/Source/Runtime/Core/Private/Apple/ApplePlatformSymbolication.cpp
@@ -0,0 +1,123 @@
+// Copyright 1998-2014 Epic Games, Inc. All Rights Reserved.
+
+/*=============================================================================
+	ApplePlatformSymbolication.cpp: Apple platform implementation of symbolication
+=============================================================================*/
+
+#include "Core.h"
+
+#include "ApplePlatformSymbolication.h"
+#include <CoreFoundation/CoreFoundation.h>
+#include <mach/mach.h>
+
+#if PLATFORM_MAC
+extern "C"
+{
+	struct CSTypeRef
+	{
+		void* Buffer0;
+		void* Buffer1;
+	};
+	
+	typedef CSTypeRef CSSymbolicatorRef;
+	typedef CSTypeRef CSSourceInfoRef;
+	typedef CSTypeRef CSSymbolRef;
+	typedef CSTypeRef CSSymbolOwnerRef;
+	
+	struct CSRange
+	{
+		uint64 Location;
+		uint64 Length;
+	};
+	
+	#define kCSNow								0x80000000u
+	
+	Boolean CSEqual(CSTypeRef CS1, CSTypeRef CS2);
+	CFIndex CSGetRetainCount(CSTypeRef CS);
+	Boolean CSIsNull(CSTypeRef CS);
+	CSTypeRef CSRetain(CSTypeRef CS);
+	void CSRelease(CSTypeRef CS);
+	void CSShow(CSTypeRef CS);
+	
+	CSSymbolicatorRef CSSymbolicatorCreateWithTask(task_t Task);
+	
+	CSSourceInfoRef CSSymbolicatorGetSourceInfoWithAddressAtTime(CSSymbolicatorRef Symbolicator, vm_address_t Address, uint64_t Time);
+	CSSymbolRef CSSymbolicatorGetSymbolWithMangledNameFromSymbolOwnerWithNameAtTime(CSSymbolicatorRef Symbolicator, const char* Symbol, const char* Name, uint64_t Time);
+	
+	const char* CSSymbolGetName(CSSymbolRef Symbol);
+	CSRange CSSymbolGetRange(CSSymbolRef Symbol);
+	CSSymbolOwnerRef CSSourceInfoGetSymbolOwner(CSSourceInfoRef Info);
+	
+	const char* CSSymbolOwnerGetName(CSSymbolOwnerRef symbol);
+	
+	int CSSourceInfoGetLineNumber(CSSourceInfoRef Info);
+	const char* CSSourceInfoGetPath(CSSourceInfoRef Info);
+	CSRange CSSourceInfoGetRange(CSSourceInfoRef Info);
+	CSSymbolRef CSSourceInfoGetSymbol(CSSourceInfoRef Info);
+}
+#endif
+
+bool FApplePlatformSymbolication::SymbolInfoForAddress(uint64 ProgramCounter, FProgramCounterSymbolInfo& Info)
+{
+#if PLATFORM_MAC
+	bool bOK = false;
+	CSSymbolicatorRef Symbolicator = CSSymbolicatorCreateWithTask(mach_task_self());
+	if(!CSIsNull(Symbolicator))
+	{
+		CSSourceInfoRef Symbol = CSSymbolicatorGetSourceInfoWithAddressAtTime(Symbolicator, (vm_address_t)ProgramCounter, kCSNow);
+		
+		if(!CSIsNull(Symbol))
+		{
+			Info.LineNumber = CSSourceInfoGetLineNumber(Symbol);
+			FCStringAnsi::Sprintf(Info.Filename, CSSourceInfoGetPath(Symbol));
+			FCStringAnsi::Sprintf(Info.FunctionName, CSSymbolGetName(CSSourceInfoGetSymbol(Symbol)));
+			CSRange CodeRange = CSSourceInfoGetRange(Symbol);
+			Info.SymbolDisplacement = (ProgramCounter - CodeRange.Location);
+			
+			CSSymbolOwnerRef Owner = CSSourceInfoGetSymbolOwner(Symbol);
+			if(!CSIsNull(Owner))
+			{
+				ANSICHAR const* DylibName = CSSymbolOwnerGetName(Owner);
+				FCStringAnsi::Strcpy(Info.ModuleName, DylibName);
+				
+				bOK = Info.LineNumber != 0;
+				
+				CSRelease(Owner);
+			}
+			
+			CSRelease(Symbol);
+		}
+		
+		CSRelease(Symbolicator);
+	}
+	return bOK;
+#else
+	return false;
+#endif
+}
+
+bool FApplePlatformSymbolication::SymbolInfoForFunctionFromModule(ANSICHAR const* MangledName, ANSICHAR const* ModuleName, FProgramCounterSymbolInfo& Info)
+{
+#if PLATFORM_MAC
+	bool bOK = false;
+	CSSymbolicatorRef Symbolicator = CSSymbolicatorCreateWithTask(mach_task_self());
+	if(!CSIsNull(Symbolicator))
+	{
+		CSSymbolRef Symbol = CSSymbolicatorGetSymbolWithMangledNameFromSymbolOwnerWithNameAtTime(Symbolicator, MangledName, ModuleName, kCSNow);
+		
+		if(!CSIsNull(Symbol))
+		{
+			CSRange CodeRange = CSSymbolGetRange(Symbol);
+		
+			bOK = SymbolInfoForAddress(CodeRange.Location, Info);
+	
+			CSRelease(Symbol);
+		}
+		
+		CSRelease(Symbolicator);
+	}
+	return bOK;
+#else
+	return false;
+#endif
+}
diff --git a/Engine/Source/Runtime/Core/Public/Apple/ApplePlatformSymbolication.h b/Engine/Source/Runtime/Core/Public/Apple/ApplePlatformSymbolication.h
new file mode 100644
--- /dev/null
+++ b/Engine/Source/Runtime/Core/Public/Apple/ApplePlatformSymbolication.h
@@ -0,0 +1,18 @@
+// Copyright 1998-2014 Epic Games, Inc. All Rights Reserved.
+
+/*=============================================================================================
+	ApplePlatformSymbolication.h: Apple platform implementation of symbolication
+==============================================================================================*/
+
+#pragma once
+
+#include "ApplePlatformStackWalk.h"
+
+/**
+ * Apple platform implementation of symbolication
+ */
+struct CORE_API FApplePlatformSymbolication
+{
+	static bool SymbolInfoForAddress(uint64 ProgramCounter, FProgramCounterSymbolInfo& Info);
+	static bool SymbolInfoForFunctionFromModule(ANSICHAR const* MangledName, ANSICHAR const* ModuleName, FProgramCounterSymbolInfo& Info);
+};
