Allow compiling for Windows without named pipes support.

Fix for case when PLATFORM_SUPPORTS_NAMED_PIPES is defined to 0 (e.g. for debugging purposes).

diff --git a/Engine/Source/Runtime/Core/Private/Windows/WindowsPlatformNamedPipe.cpp b/Engine/Source/Runtime/Core/Private/Windows/WindowsPlatformNamedPipe.cpp
--- a/Engine/Source/Runtime/Core/Private/Windows/WindowsPlatformNamedPipe.cpp
+++ b/Engine/Source/Runtime/Core/Private/Windows/WindowsPlatformNamedPipe.cpp
@@ -16,6 +16,7 @@ static void VerifyWinResult(BOOL bResult, const TCHAR* InMessage)
 	verify(bResult != 0);
 }
 
+#if PLATFORM_SUPPORTS_NAMED_PIPES
 FWindowsPlatformNamedPipe::FWindowsPlatformNamedPipe() :
 	Pipe(NULL),
 	LastWaitingTime(0.0f),
@@ -523,3 +524,4 @@ bool FWindowsPlatformNamedPipe::HasFailed() const
 
 	return false;
 }
+#endif // PLATFORM_SUPPORTS_NAMED_PIPES
