Fix stat id redundant creation.

Change stat id creation to static declarations.

diff --git a/Engine/Source/Developer/FunctionalTesting/Private/FunctionalTestingManager.cpp b/Engine/Source/Developer/FunctionalTesting/Private/FunctionalTestingManager.cpp
--- a/Engine/Source/Developer/FunctionalTesting/Private/FunctionalTestingManager.cpp
+++ b/Engine/Source/Developer/FunctionalTesting/Private/FunctionalTestingManager.cpp
@@ -198,12 +198,13 @@ void UFunctionalTestingManager::OnWorldCleanedUp(UWorld* World, bool bSessionEnd
 void UFunctionalTestingManager::OnTestDone(AFunctionalTest* FTest)
 {
 	// add a delay
+	DECLARE_CYCLE_STAT(TEXT("FSimpleDelegateGraphTask.Requesting to build next tile if necessary"),
+		STAT_FSimpleDelegateGraphTask_RequestingToBuildNextTileIfNecessary,
+		STATGROUP_TaskGraphTasks);
+
 	FSimpleDelegateGraphTask::CreateAndDispatchWhenReady(
-		FSimpleDelegateGraphTask::FDelegate::CreateUObject(this, &UFunctionalTestingManager::NotifyTestDone, FTest)
-		, TEXT("Requesting to build next tile if necessary")
-		, NULL
-		, ENamedThreads::GameThread
-		);
+		FSimpleDelegateGraphTask::FDelegate::CreateUObject(this, &UFunctionalTestingManager::NotifyTestDone, FTest),
+		GET_STATID(STAT_FSimpleDelegateGraphTask_RequestingToBuildNextTileIfNecessary), NULL, ENamedThreads::GameThread);
 }
 
 void UFunctionalTestingManager::NotifyTestDone(AFunctionalTest* FTest)
diff --git a/Engine/Source/Developer/LogVisualizer/Private/SLogBar.cpp b/Engine/Source/Developer/LogVisualizer/Private/SLogBar.cpp
--- a/Engine/Source/Developer/LogVisualizer/Private/SLogBar.cpp
+++ b/Engine/Source/Developer/LogVisualizer/Private/SLogBar.cpp
@@ -420,13 +420,15 @@ void SLogBar::OnCurrentTimeChanged(float NewTime)
 
 	if (bUpdateSelection)
 	{
+		DECLARE_CYCLE_STAT(TEXT("FSimpleDelegateGraphTask.Updating Log Entry shown in LogVisualizer"),
+			STAT_FSimpleDelegateGraphTask_UpdatingLogEntryShownInLogVisualizer,
+			STATGROUP_TaskGraphTasks);
+
 		// find entry closest to desired time
 		FSimpleDelegateGraphTask::CreateAndDispatchWhenReady(
-			FSimpleDelegateGraphTask::FDelegate::CreateSP(this, &SLogBar::SelectEntryAtIndex, GetEntryIndexAtTime(NewTime))
-			, TEXT("Updating Log Entry shown in LogVisualizer")
-			, NULL
-			, ENamedThreads::GameThread
-			);
+			FSimpleDelegateGraphTask::FDelegate::CreateSP(this, &SLogBar::SelectEntryAtIndex, GetEntryIndexAtTime(NewTime)),
+			GET_STATID(STAT_FSimpleDelegateGraphTask_UpdatingLogEntryShownInLogVisualizer), NULL, ENamedThreads::GameThread
+		);
 	}
 }
 
diff --git a/Engine/Source/Developer/Profiler/Private/ProfilerSession.cpp b/Engine/Source/Developer/Profiler/Private/ProfilerSession.cpp
--- a/Engine/Source/Developer/Profiler/Private/ProfilerSession.cpp
+++ b/Engine/Source/Developer/Profiler/Private/ProfilerSession.cpp
@@ -207,21 +207,32 @@ void FProfilerSession::UpdateAggregatedEventGraphData( const uint32 FrameIndex )
 	{
 		FGraphEventArray EventGraphCombineTasks;
 
-		new (EventGraphCombineTasks) FGraphEventRef(FSimpleDelegateGraphTask::CreateAndDispatchWhenReady
+		DECLARE_CYCLE_STAT(TEXT("FSimpleDelegateGraphTask.EventGraphData.CombineAndFindMax"),
+			STAT_FSimpleDelegateGraphTask_EventGraphData_CombineAndFindMax,
+			STATGROUP_TaskGraphTasks);
+		DECLARE_CYCLE_STAT(TEXT("FSimpleDelegateGraphTask.EventGraphData.EventGraphCombineAndAdd"),
+			STAT_FSimpleDelegateGraphTask_EventGraphData_EventGraphCombineAndAdd,
+			STATGROUP_TaskGraphTasks);
+
+		new (EventGraphCombineTasks)FGraphEventRef(FSimpleDelegateGraphTask::CreateAndDispatchWhenReady
 		(
 			FSimpleDelegateGraphTask::FDelegate::CreateRaw( this, &FProfilerSession::EventGraphCombineAndMax, EventGraphDataCurrent, NumFrames ), 
-			TEXT("EventGraphData.CombineAndFindMax"), nullptr
+			GET_STATID(STAT_FSimpleDelegateGraphTask_EventGraphData_CombineAndFindMax), nullptr
 		));
 
 		new (EventGraphCombineTasks) FGraphEventRef(FSimpleDelegateGraphTask::CreateAndDispatchWhenReady
 		(
 			FSimpleDelegateGraphTask::FDelegate::CreateRaw( this, &FProfilerSession::EventGraphCombineAndAdd, EventGraphDataCurrent, NumFrames ), 
-			TEXT("EventGraphData.EventGraphCombineAndAdd"), nullptr
+			GET_STATID(STAT_FSimpleDelegateGraphTask_EventGraphData_EventGraphCombineAndAdd), nullptr
 		));
 
+		DECLARE_CYCLE_STAT(TEXT("FNullGraphTask.EventGraphData.CombineJoinAndContinue"),
+			STAT_FNullGraphTask_EventGraphData_CombineJoinAndContinue,
+			STATGROUP_TaskGraphTasks);
+
 		// JoinThreads
 		CompletionSync = TGraphTask<FNullGraphTask>::CreateTask( &EventGraphCombineTasks, ENamedThreads::GameThread )
-			.ConstructAndDispatchWhenReady( TEXT("EventGraphData.CombineJoinAndContinue"), ENamedThreads::AnyThread );
+			.ConstructAndDispatchWhenReady(GET_STATID(STAT_FNullGraphTask_EventGraphData_CombineJoinAndContinue), ENamedThreads::AnyThread);
 	}
 	else
 	{
diff --git a/Engine/Source/Editor/UnrealEd/Private/SourceCodeNavigation.cpp b/Engine/Source/Editor/UnrealEd/Private/SourceCodeNavigation.cpp
--- a/Engine/Source/Editor/UnrealEd/Private/SourceCodeNavigation.cpp
+++ b/Engine/Source/Editor/UnrealEd/Private/SourceCodeNavigation.cpp
@@ -592,13 +592,21 @@ void FSourceCodeNavigation::NavigateToFunctionSourceAsync( const FString& Functi
 	}
 
 	// Kick off asynchronous task to load symbols
+	DECLARE_CYCLE_STAT(TEXT("FDelegateGraphTask.EditorSourceCodeNavigation"),
+		STAT_FDelegateGraphTask_EditorSourceCodeNavigation,
+		STATGROUP_TaskGraphTasks);
+
 	FGraphEventRef PreloadSymbolsAsyncResult(
-		FDelegateGraphTask::CreateAndDispatchWhenReady( FDelegateGraphTask::FDelegate::CreateStatic(&FLocal::PreloadSymbolsTaskWrapper ), TEXT( "EditorSourceCodeNavigation" ) ) );
+		FDelegateGraphTask::CreateAndDispatchWhenReady(FDelegateGraphTask::FDelegate::CreateStatic(&FLocal::PreloadSymbolsTaskWrapper), GET_STATID(STAT_FDelegateGraphTask_EditorSourceCodeNavigation)));
 
 	// add a dependent task to run on the main thread when symbols are loaded
 	FGraphEventRef UnusedAsyncResult(
-		FDelegateGraphTask::CreateAndDispatchWhenReady( FDelegateGraphTask::FDelegate::CreateStatic(
-				&FLocal::NavigateToFunctionSourceTaskWrapper, NavigateFunctionParams, CompileNotificationPtr ), TEXT( "EditorSourceCodeNavigation" ), PreloadSymbolsAsyncResult, ENamedThreads::GameThread, ENamedThreads::GameThread ));
+		FDelegateGraphTask::CreateAndDispatchWhenReady(
+			FDelegateGraphTask::FDelegate::CreateStatic(
+				&FLocal::NavigateToFunctionSourceTaskWrapper, NavigateFunctionParams, CompileNotificationPtr), GET_STATID(STAT_FDelegateGraphTask_EditorSourceCodeNavigation),
+				PreloadSymbolsAsyncResult, ENamedThreads::GameThread, ENamedThreads::GameThread
+			)
+		);
 }
 
 
diff --git a/Engine/Source/Runtime/AIModule/Private/Navigation/NavigationComponent.cpp b/Engine/Source/Runtime/AIModule/Private/Navigation/NavigationComponent.cpp
--- a/Engine/Source/Runtime/AIModule/Private/Navigation/NavigationComponent.cpp
+++ b/Engine/Source/Runtime/AIModule/Private/Navigation/NavigationComponent.cpp
@@ -112,12 +112,14 @@ void UNavigationComponent::OnPathEvent(FNavigationPath* InvalidatedPath, ENavPat
 
 				Path->EnableRecalculationOnInvalidation(true);
 
-				/*FSimpleDelegateGraphTask::CreateAndDispatchWhenReady(
-					FSimpleDelegateGraphTask::FDelegate::CreateUObject(this, &UNavigationComponent::DeferredRepathToGoal)
-					, TEXT("Deferred restore path")
-					, NULL
-					, ENamedThreads::GameThread
-					);*/
+				/*
+				const static TStatId StatId = TStatId::CreateAndRegisterFromName(TEXT("FSimpleDelegateGraphTask.Deferred restore path"));
+
+				FSimpleDelegateGraphTask::CreateAndDispatchWhenReady(
+					FSimpleDelegateGraphTask::FDelegate::CreateUObject(this, &UNavigationComponent::DeferredRepathToGoal),
+					StatId, NULL, ENamedThreads::GameThread
+				);
+				*/
 			}
 
 			// reset state
@@ -556,12 +558,14 @@ bool UNavigationComponent::FindPathToActor(const AActor* NewGoalActor, TSharedPt
 
 		if (bDoAsyncPathfinding == true)
 		{
+			DECLARE_CYCLE_STAT(TEXT("FSimpleDelegateGraphTask.Deferred resume path"),
+				STAT_FSimpleDelegateGraphTask_DeferredResumePath,
+				STATGROUP_TaskGraphTasks);
+
 			FSimpleDelegateGraphTask::CreateAndDispatchWhenReady(
-				FSimpleDelegateGraphTask::FDelegate::CreateUObject(this, &UNavigationComponent::DeferredResumePath)
-				, TEXT("Deferred resume path")
-				, NULL
-				, ENamedThreads::GameThread
-				);
+				FSimpleDelegateGraphTask::FDelegate::CreateUObject(this, &UNavigationComponent::DeferredResumePath),
+				GET_STATID(STAT_FSimpleDelegateGraphTask_DeferredResumePath), NULL, ENamedThreads::GameThread
+			);
 			bPathGenerationSucceeded = true;
 		}
 		else
@@ -652,12 +656,14 @@ bool UNavigationComponent::FindPathToLocation(const FVector& DestLocation, TShar
 
 		if (bDoAsyncPathfinding == true)
 		{
+			DECLARE_CYCLE_STAT(TEXT("FSimpleDelegateGraphTask.Deferred resume path"),
+				STAT_FSimpleDelegateGraphTask_DeferredResumePath,
+				STATGROUP_TaskGraphTasks);
+
 			FSimpleDelegateGraphTask::CreateAndDispatchWhenReady(
-				FSimpleDelegateGraphTask::FDelegate::CreateUObject(this, &UNavigationComponent::DeferredResumePath)
-				, TEXT("Deferred resume path")
-				, NULL
-				, ENamedThreads::GameThread
-				);
+				FSimpleDelegateGraphTask::FDelegate::CreateUObject(this, &UNavigationComponent::DeferredResumePath),
+				GET_STATID(STAT_FSimpleDelegateGraphTask_DeferredResumePath), NULL, ENamedThreads::GameThread
+			);
 			bPathGenerationSucceeded = true;
 		}
 		else
diff --git a/Engine/Source/Runtime/Core/Private/Async/TaskGraph.cpp b/Engine/Source/Runtime/Core/Private/Async/TaskGraph.cpp
--- a/Engine/Source/Runtime/Core/Private/Async/TaskGraph.cpp
+++ b/Engine/Source/Runtime/Core/Private/Async/TaskGraph.cpp
@@ -1139,7 +1139,11 @@ void FGraphEvent::DispatchSubsequents(TArray<FBaseGraphTask*>& NewTasks, ENamedT
 		FGraphEventArray TempEventsToWaitFor;
 		Exchange(EventsToWaitFor, TempEventsToWaitFor);
 		// create the Gather...this uses a special version of private CreateTask that "assumes" the subsequent list (which other threads might still be adding too).
-		TGraphTask<FNullGraphTask>::CreateTask( FGraphEventRef( this ), &TempEventsToWaitFor, CurrentThreadIfKnown ).ConstructAndDispatchWhenReady( TEXT( "DontCompleteUntil" ) );
+		DECLARE_CYCLE_STAT(TEXT("FNullGraphTask.DontCompleteUntil"),
+			STAT_FNullGraphTask_DontCompleteUntil,
+			STATGROUP_TaskGraphTasks);
+
+		TGraphTask<FNullGraphTask>::CreateTask(FGraphEventRef(this), &TempEventsToWaitFor, CurrentThreadIfKnown).ConstructAndDispatchWhenReady(GET_STATID(STAT_FNullGraphTask_DontCompleteUntil));
 		return;
 	}
 
diff --git a/Engine/Source/Runtime/Core/Private/Stats/Stats2.cpp b/Engine/Source/Runtime/Core/Private/Stats/Stats2.cpp
--- a/Engine/Source/Runtime/Core/Private/Stats/Stats2.cpp
+++ b/Engine/Source/Runtime/Core/Private/Stats/Stats2.cpp
@@ -531,9 +531,6 @@ IStatGroupEnableManager& IStatGroupEnableManager::Get()
 	return *Singleton;
 }
 
-
-
-
 uint32 FThreadStats::TlsSlot = 0;
 FThreadSafeCounter FThreadStats::MasterEnableCounter;
 FThreadSafeCounter FThreadStats::MasterEnableUpdateNumber;
@@ -979,7 +976,12 @@ void FThreadStats::WaitForStats()
 				FTaskGraphInterface::Get().WaitUntilTaskCompletes(LastFramesEvents[(CurrentEventIndex + MAX_STAT_LAG - 1) % MAX_STAT_LAG], ENamedThreads::GameThread_Local);
 			}
 		}
-		LastFramesEvents[(CurrentEventIndex + MAX_STAT_LAG - 1) % MAX_STAT_LAG] = TGraphTask<FNullGraphTask>::CreateTask(NULL, ENamedThreads::GameThread).ConstructAndDispatchWhenReady(TEXT("StatWaitFence"), FPlatformProcess::SupportsMultithreading() ? ENamedThreads::StatsThread : ENamedThreads::GameThread);
+
+		DECLARE_CYCLE_STAT(TEXT("FNullGraphTask.StatWaitFence"),
+			STAT_FNullGraphTask_StatWaitFence,
+			STATGROUP_TaskGraphTasks);
+
+		LastFramesEvents[(CurrentEventIndex + MAX_STAT_LAG - 1) % MAX_STAT_LAG] = TGraphTask<FNullGraphTask>::CreateTask(NULL, ENamedThreads::GameThread).ConstructAndDispatchWhenReady(GET_STATID(STAT_FNullGraphTask_StatWaitFence), FPlatformProcess::SupportsMultithreading() ? ENamedThreads::StatsThread : ENamedThreads::GameThread);
 		CurrentEventIndex++;
 	}
 }
diff --git a/Engine/Source/Runtime/Core/Private/Stats/StatsCommand.cpp b/Engine/Source/Runtime/Core/Private/Stats/StatsCommand.cpp
--- a/Engine/Source/Runtime/Core/Private/Stats/StatsCommand.cpp
+++ b/Engine/Source/Runtime/Core/Private/Stats/StatsCommand.cpp
@@ -596,11 +596,14 @@ struct FHUDGroupManager
 			StatsMasterEnableSubtract();
 			bEnabled = false;
 
+			DECLARE_CYCLE_STAT(TEXT("FSimpleDelegateGraphTask.StatsToGame"),
+				STAT_FSimpleDelegateGraphTask_StatsToGame,
+				STATGROUP_TaskGraphTasks);
+
 			FSimpleDelegateGraphTask::CreateAndDispatchWhenReady
 			(
-				FSimpleDelegateGraphTask::FDelegate::CreateRaw(&FHUDGroupGameThreadRenderer::Get(), &FHUDGroupGameThreadRenderer::NewData, (FGameThreadHudData*)nullptr)
-				, TEXT("StatsToGame")
-				, nullptr, ENamedThreads::GameThread
+				FSimpleDelegateGraphTask::FDelegate::CreateRaw(&FHUDGroupGameThreadRenderer::Get(), &FHUDGroupGameThreadRenderer::NewData, (FGameThreadHudData*)nullptr),
+				GET_STATID(STAT_FSimpleDelegateGraphTask_StatsToGame), nullptr, ENamedThreads::GameThread
 			);
 		}
 	}
@@ -856,11 +859,14 @@ struct FHUDGroupManager
 				}
 			}
 
+			DECLARE_CYCLE_STAT(TEXT("FSimpleDelegateGraphTask.StatsHierToGame"),
+				STAT_FSimpleDelegateGraphTask_StatsHierToGame,
+				STATGROUP_TaskGraphTasks);
+
 			FSimpleDelegateGraphTask::CreateAndDispatchWhenReady
 			(
-				FSimpleDelegateGraphTask::FDelegate::CreateRaw(&FHUDGroupGameThreadRenderer::Get(), &FHUDGroupGameThreadRenderer::NewData, ToGame)
-				, TEXT("StatsHierToGame")
-				, nullptr, ENamedThreads::GameThread
+				FSimpleDelegateGraphTask::FDelegate::CreateRaw(&FHUDGroupGameThreadRenderer::Get(), &FHUDGroupGameThreadRenderer::NewData, ToGame),
+				GET_STATID(STAT_FSimpleDelegateGraphTask_StatsHierToGame), nullptr, ENamedThreads::GameThread
 			);
 		}	
 	}
@@ -1365,12 +1371,14 @@ bool DirectStatsCommand(const TCHAR* Cmd, bool bBlockForCompletion /*= false*/,
 			FHUDGroupGameThreadRenderer::Get();
 			FStatGroupGameThreadNotifier::Get();
 #endif
+			DECLARE_CYCLE_STAT(TEXT("FSimpleDelegateGraphTask.StatCmd"),
+				STAT_FSimpleDelegateGraphTask_StatCmd,
+				STATGROUP_TaskGraphTasks);
+
 			FGraphEventRef CompleteHandle = FSimpleDelegateGraphTask::CreateAndDispatchWhenReady(
-				FSimpleDelegateGraphTask::FDelegate::CreateStatic(&StatCmd, FString(Cmd) + AddArgs)
-				, TEXT("StatCmd")
-				, NULL
-				, ThreadType
-				);
+				FSimpleDelegateGraphTask::FDelegate::CreateStatic(&StatCmd, FString(Cmd) + AddArgs),
+				GET_STATID(STAT_FSimpleDelegateGraphTask_StatCmd), NULL, ThreadType
+			);
 			if (bBlockForCompletion)
 			{
 				FTaskGraphInterface::Get().WaitUntilTaskCompletes(CompleteHandle);
@@ -1396,12 +1404,15 @@ static void GetPermanentStats_StatsThread(TArray<FStatMessage>* OutStats)
 
 void GetPermanentStats(TArray<FStatMessage>& OutStats)
 {
+	DECLARE_CYCLE_STAT(TEXT("FSimpleDelegateGraphTask.GetPermanentStatsString_StatsThread"),
+		STAT_FSimpleDelegateGraphTask_GetPermanentStatsString_StatsThread,
+		STATGROUP_TaskGraphTasks);
+
 	FGraphEventRef CompleteHandle = FSimpleDelegateGraphTask::CreateAndDispatchWhenReady(
-		FSimpleDelegateGraphTask::FDelegate::CreateStatic(&GetPermanentStats_StatsThread, &OutStats)
-		, TEXT("GetPermanentStatsString_StatsThread")
-		, NULL
-		, FPlatformProcess::SupportsMultithreading() ? ENamedThreads::StatsThread : ENamedThreads::GameThread
-		);
+		FSimpleDelegateGraphTask::FDelegate::CreateStatic(&GetPermanentStats_StatsThread, &OutStats),
+		GET_STATID(STAT_FSimpleDelegateGraphTask_GetPermanentStatsString_StatsThread), NULL,
+		FPlatformProcess::SupportsMultithreading() ? ENamedThreads::StatsThread : ENamedThreads::GameThread
+	);
 	FTaskGraphInterface::Get().WaitUntilTaskCompletes(CompleteHandle);
 }
 
diff --git a/Engine/Source/Runtime/Core/Private/Stats/StatsData.cpp b/Engine/Source/Runtime/Core/Private/Stats/StatsData.cpp
--- a/Engine/Source/Runtime/Core/Private/Stats/StatsData.cpp
+++ b/Engine/Source/Runtime/Core/Private/Stats/StatsData.cpp
@@ -1283,11 +1283,14 @@ void FStatsThreadState::FindOrAddMetaData(FStatMessage const& Item)
 		}
 
 		// Add the info to the task graph so we can inform the game thread
+		DECLARE_CYCLE_STAT(TEXT("FSimpleDelegateGraphTask.StatsGroupToGame"),
+			STAT_FSimpleDelegateGraphTask_StatsGroupToGame,
+			STATGROUP_TaskGraphTasks);
+
 		FSimpleDelegateGraphTask::CreateAndDispatchWhenReady
 		(
-			FSimpleDelegateGraphTask::FDelegate::CreateRaw(&FStatGroupGameThreadNotifier::Get(), &FStatGroupGameThreadNotifier::NewData, Item.NameAndInfo)
-			, TEXT("StatsGroupToGame")
-			, nullptr, ENamedThreads::GameThread
+			FSimpleDelegateGraphTask::FDelegate::CreateRaw(&FStatGroupGameThreadNotifier::Get(), &FStatGroupGameThreadNotifier::NewData, Item.NameAndInfo),
+			GET_STATID(STAT_FSimpleDelegateGraphTask_StatsGroupToGame), nullptr, ENamedThreads::GameThread
 		);
 	}
 	else
diff --git a/Engine/Source/Runtime/Core/Private/UObject/UnrealNames.cpp b/Engine/Source/Runtime/Core/Private/UObject/UnrealNames.cpp
--- a/Engine/Source/Runtime/Core/Private/UObject/UnrealNames.cpp
+++ b/Engine/Source/Runtime/Core/Private/UObject/UnrealNames.cpp
@@ -1127,12 +1127,14 @@ static class FFNameExec: private FSelfRegisteringExec
 
 				for (int32 ThreadIndex = 0; ThreadIndex < FTest::NUM_TASKS; ThreadIndex++)
 				{
+					DECLARE_CYCLE_STAT(TEXT("FSimpleDelegateGraphTask.FName Test"),
+						STAT_FSimpleDelegateGraphTask_FName_Test,
+						STATGROUP_TaskGraphTasks);
+
 					new (Handles) FGraphEventRef(FSimpleDelegateGraphTask::CreateAndDispatchWhenReady(
 						FSimpleDelegateGraphTask::FDelegate::CreateRaw(&Test, &FTest::Thread),
-						TEXT("FName Test"),
-						NULL,
-						ENamedThreads::AnyThread
-						));
+						GET_STATID(STAT_FSimpleDelegateGraphTask_FName_Test), NULL,
+						ENamedThreads::AnyThread));
 				}
 				FTaskGraphInterface::Get().WaitUntilTasksComplete(Handles, ENamedThreads::GameThread);
 				check(Test.TestCounter.GetValue() == FTest::NUM_TESTS * FTest::NUM_TASKS);
diff --git a/Engine/Source/Runtime/Core/Public/Async/TaskGraphInterfaces.h b/Engine/Source/Runtime/Core/Public/Async/TaskGraphInterfaces.h
--- a/Engine/Source/Runtime/Core/Public/Async/TaskGraphInterfaces.h
+++ b/Engine/Source/Runtime/Core/Public/Async/TaskGraphInterfaces.h
@@ -832,38 +832,21 @@ private:
 };
 
 /**
- * Class that generalizes functionality of creating, storing and exposing stat id for custom name.
- * Should only be used as a base of a task graph class with custom stat name.
+ * Class that generalizes functionality of storing and exposing custom stat id.
+ * Should only be used as a base of a task graph class with custom stat id.
  */
-class FCustomNameGraphTaskBase
+class FCustomStatIDGraphTaskBase
 {
 public:
 	/** 
-	 * Creates the custom stat id.
+	 * Constructor.
 	 *
-	 * @param TaskClassName The name of generic class name for stats e.g. FNullTaskGraph
-	 * @param InTaskName The name of this task.
+	 * @param StatId The stat id for this task.
 	 */
-	FCustomNameGraphTaskBase(const TCHAR* TaskClassName, const TCHAR* TaskName)
+	FCustomStatIDGraphTaskBase(const TStatId& StatId)
 	{
 #if STATS
-		// Commented out for now as its generating huge overhead.
-		/*
-		const FString LongName = FString(TaskClassName) + TEXT(".") + FString(TaskName);
-		const FName StatName = FName( *LongName );
-
-		FStartupMessages::Get().AddMetadata(StatName, *LongName,
-			STAT_GROUP_TO_FStatGroup(STATGROUP_TaskGraphTasks)::GetGroupName(),
-			STAT_GROUP_TO_FStatGroup(STATGROUP_TaskGraphTasks)::GetGroupCategory(),
-			STAT_GROUP_TO_FStatGroup(STATGROUP_TaskGraphTasks)::GetDescription(),
-			true, EStatDataType::ST_int64, true);
-
-		StatID = IStatGroupEnableManager::Get().GetHighPerformanceEnableForStat(StatName,
-			STAT_GROUP_TO_FStatGroup(STATGROUP_TaskGraphTasks)::GetGroupName(),
-			STAT_GROUP_TO_FStatGroup(STATGROUP_TaskGraphTasks)::GetGroupCategory(),
-			STAT_GROUP_TO_FStatGroup(STATGROUP_TaskGraphTasks)::DefaultEnable,
-			true, EStatDataType::ST_int64, *LongName, true);
-		*/
+		StatID = StatId;
 #endif
 	}
 
@@ -875,31 +858,29 @@ public:
 	FORCEINLINE TStatId GetStatId() const
 	{
 #if STATS
-		// See the comment above.
-		return TStatId(); /* StatID; */
+		return StatID;
 #endif
 		return TStatId();
 	}
 
 private:
-	// See the comment above.
 	/** Stat id of this object. */
-	/* STAT(TStatId StatID;) */
+	STAT(TStatId StatID;)
 };
 
 /** 
  *	FNullGraphTask is a task that does nothing. It can be used to "gather" tasks into one prerequisite.
  **/
-class FNullGraphTask : public FCustomNameGraphTaskBase
+class FNullGraphTask : public FCustomStatIDGraphTaskBase
 {
 public:
 	/** 
 	 *	Constructor
-	 *	@param InTaskName; Name of the task (usually for debugging and tools)
+	 *	@param StatId The stat id for this task.
 	 *	@param InDesiredThread; Thread to run on, can be ENamedThreads::AnyThread
 	**/
-	FNullGraphTask(const TCHAR*	InTaskName, ENamedThreads::Type InDesiredThread = ENamedThreads::AnyThread)
-		: FCustomNameGraphTaskBase(TEXT("FNullGraphTask"), InTaskName)
+	FNullGraphTask(const TStatId& StatId, ENamedThreads::Type InDesiredThread = ENamedThreads::AnyThread)
+		: FCustomStatIDGraphTaskBase(StatId)
 		, DesiredThread(InDesiredThread)
 	{
 	}
@@ -966,7 +947,7 @@ private:
 };
 
 /** Task class for simple delegate based tasks. This is less efficient than a custom task, doesn't provide the task arguments, doesn't allow specification of the current thread, etc. **/
-class FSimpleDelegateGraphTask : public FCustomNameGraphTaskBase
+class FSimpleDelegateGraphTask : public FCustomStatIDGraphTaskBase
 {
 public:
 	DECLARE_DELEGATE(FDelegate);
@@ -989,11 +970,11 @@ public:
 	/**
 	  * Task constructor
 	  * @param InTaskDelegate - delegate to execute when the prerequisites are complete
-	  * @param InTaskName - Name of task for debugging or analysis tools, must be static memory valid forever
+	  *	@param StatId The stat id for this task.
 	  * @param InDesiredThread - Thread to run on
 	**/
-	FSimpleDelegateGraphTask(const FDelegate& InTaskDeletegate, const TCHAR* InTaskName, ENamedThreads::Type InDesiredThread)
-		: FCustomNameGraphTaskBase(TEXT("FSimpleDelegateGraphTask"), InTaskName)
+	FSimpleDelegateGraphTask(const FDelegate& InTaskDeletegate, const TStatId& StatId, ENamedThreads::Type InDesiredThread)
+		: FCustomStatIDGraphTaskBase(StatId)
 		, TaskDelegate(InTaskDeletegate)
 		, DesiredThread(InDesiredThread)
 	{
@@ -1002,34 +983,34 @@ public:
 	/**
 	  * Create a task and dispatch it when the prerequisites are complete
 	  * @param InTaskDelegate - delegate to execute when the prerequisites are complete
-	  * @param InTaskName - Name of task for debugging or analysis tools, must be static memory valid forever
+	  * @param InStatId - StatId of task for debugging or analysis tools
 	  * @param InPrerequisites - Handles for prerequisites for this task, can be NULL if there are no prerequisites
 	  * @param InDesiredThread - Thread to run on
 	  * @return completion handle for the new task 
 	**/
-	static FGraphEventRef CreateAndDispatchWhenReady(const FDelegate& InTaskDeletegate, const TCHAR* InTaskName, const FGraphEventArray* InPrerequisites = NULL, ENamedThreads::Type InDesiredThread = ENamedThreads::AnyThread)
+	static FGraphEventRef CreateAndDispatchWhenReady(const FDelegate& InTaskDeletegate, const TStatId& InStatId, const FGraphEventArray* InPrerequisites = NULL, ENamedThreads::Type InDesiredThread = ENamedThreads::AnyThread)
 	{
-		return TGraphTask<FSimpleDelegateGraphTask>::CreateTask(InPrerequisites).ConstructAndDispatchWhenReady<const FDelegate&>(InTaskDeletegate, InTaskName, InDesiredThread);
+		return TGraphTask<FSimpleDelegateGraphTask>::CreateTask(InPrerequisites).ConstructAndDispatchWhenReady<const FDelegate&>(InTaskDeletegate, InStatId, InDesiredThread);
 	}
 	/**
 	  * Create a task and dispatch it when the prerequisites are complete
 	  * @param InTaskDelegate - delegate to execute when the prerequisites are complete
-	  * @param InTaskName - Name of task for debugging or analysis tools, must be static memory valid forever
+	  * @param InStatId - StatId of task for debugging or analysis tools
 	  * @param InPrerequisite - Handle for a single prerequisite for this task
 	  * @param InDesiredThread - Thread to run on
 	  * @return completion handle for the new task 
 	**/
-	static FGraphEventRef CreateAndDispatchWhenReady(const FDelegate& InTaskDeletegate, const TCHAR* InTaskName, const FGraphEventRef& InPrerequisite, ENamedThreads::Type InDesiredThread = ENamedThreads::AnyThread)
+	static FGraphEventRef CreateAndDispatchWhenReady(const FDelegate& InTaskDeletegate, const TStatId& InStatId, const FGraphEventRef& InPrerequisite, ENamedThreads::Type InDesiredThread = ENamedThreads::AnyThread)
 	{
 		FGraphEventArray Prerequisites;
 		check(InPrerequisite.GetReference());
 		Prerequisites.Add(InPrerequisite);
-		return CreateAndDispatchWhenReady(InTaskDeletegate, InTaskName, &Prerequisites, InDesiredThread);
+		return CreateAndDispatchWhenReady(InTaskDeletegate, InStatId, &Prerequisites, InDesiredThread);
 	}
 };
 
 /** Task class for more full featured delegate based tasks. Still less efficient than a custom task, but provides all of the args **/
-class FDelegateGraphTask : public FCustomNameGraphTaskBase
+class FDelegateGraphTask : public FCustomStatIDGraphTaskBase
 {
 public:
 	DECLARE_DELEGATE_TwoParams( FDelegate,ENamedThreads::Type, const  FGraphEventRef& );
@@ -1054,11 +1035,11 @@ public:
 	/**
 	  * Task constructor
 	  * @param InTaskDelegate - delegate to execute when the prerequisites are complete
-	  * @param InTaskName - Name of task for debugging or analysis tools, must be static memory valid forever
+	  *	@param InStatId - The stat id for this task.
 	  * @param InDesiredThread - Thread to run on
 	**/
-	FDelegateGraphTask(const FDelegate& InTaskDeletegate, const TCHAR* InTaskName, ENamedThreads::Type InDesiredThread)
-		: FCustomNameGraphTaskBase(TEXT("FDelegateGraphTask"), InTaskName)
+	FDelegateGraphTask(const FDelegate& InTaskDeletegate, const TStatId& InStatId, ENamedThreads::Type InDesiredThread)
+		: FCustomStatIDGraphTaskBase(InStatId)
 		, TaskDelegate(InTaskDeletegate)
 		, DesiredThread(InDesiredThread)
 	{
@@ -1067,31 +1048,31 @@ public:
 	/**
 	  * Create a task and dispatch it when the prerequisites are complete
 	  * @param InTaskDelegate - delegate to execute when the prerequisites are complete
-	  * @param InTaskName - Name of task for debugging or analysis tools, must be static memory valid forever
+	  *	@param InStatId - The stat id for this task.
 	  * @param InPrerequisites - Handles for prerequisites for this task, can be NULL if there are no prerequisites
 	  * @param InCurrentThreadIfKnown - This thread, if known
 	  * @param InDesiredThread - Thread to run on
 	  * @return completion handle for the new task 
 	**/
-	static FGraphEventRef CreateAndDispatchWhenReady(const FDelegate& InTaskDeletegate, const TCHAR* InTaskName, const FGraphEventArray* InPrerequisites = NULL, ENamedThreads::Type InCurrentThreadIfKnown = ENamedThreads::AnyThread, ENamedThreads::Type InDesiredThread = ENamedThreads::AnyThread)
+	static FGraphEventRef CreateAndDispatchWhenReady(const FDelegate& InTaskDeletegate, const TStatId& InStatId, const FGraphEventArray* InPrerequisites = NULL, ENamedThreads::Type InCurrentThreadIfKnown = ENamedThreads::AnyThread, ENamedThreads::Type InDesiredThread = ENamedThreads::AnyThread)
 	{
-		return TGraphTask<FDelegateGraphTask>::CreateTask(InPrerequisites, InCurrentThreadIfKnown).ConstructAndDispatchWhenReady<const FDelegate&>(InTaskDeletegate, InTaskName, InDesiredThread);
+		return TGraphTask<FDelegateGraphTask>::CreateTask(InPrerequisites, InCurrentThreadIfKnown).ConstructAndDispatchWhenReady<const FDelegate&>(InTaskDeletegate, InStatId, InDesiredThread);
 	}
 	/**
 	  * Create a task and dispatch it when the prerequisites are complete
 	  * @param InTaskDelegate - delegate to execute when the prerequisites are complete
-	  * @param InTaskName - Name of task for debugging or analysis tools, must be static memory valid forever
+	  *	@param InStatId - The stat id for this task.
 	  * @param InPrerequisite - Handle for a single prerequisite for this task
 	  * @param InCurrentThreadIfKnown - This thread, if known
 	  * @param InDesiredThread - Thread to run on
 	  * @return completion handle for the new task 
 	**/
-	static FGraphEventRef CreateAndDispatchWhenReady(const FDelegate& InTaskDeletegate, const TCHAR* InTaskName, const FGraphEventRef& InPrerequisite, ENamedThreads::Type InCurrentThreadIfKnown = ENamedThreads::AnyThread, ENamedThreads::Type InDesiredThread = ENamedThreads::AnyThread)
+	static FGraphEventRef CreateAndDispatchWhenReady(const FDelegate& InTaskDeletegate, const TStatId& InStatId, const FGraphEventRef& InPrerequisite, ENamedThreads::Type InCurrentThreadIfKnown = ENamedThreads::AnyThread, ENamedThreads::Type InDesiredThread = ENamedThreads::AnyThread)
 	{
 		FGraphEventArray Prerequisites;
 		check(InPrerequisite.GetReference());
 		Prerequisites.Add(InPrerequisite);
-		return CreateAndDispatchWhenReady(InTaskDeletegate, InTaskName, &Prerequisites, InCurrentThreadIfKnown, InDesiredThread);
+		return CreateAndDispatchWhenReady(InTaskDeletegate, InStatId, &Prerequisites, InCurrentThreadIfKnown, InDesiredThread);
 	}
 };
 
@@ -1150,7 +1131,14 @@ public:
 				Pending[Index]->Release(); // remove the ref count we added when we added it to the lock free list
 			}
 			// start a new task that won't complete until all of these tasks have executed, plus any tasks that they create when they run
-			CompleteHandle = FDelegateGraphTask::CreateAndDispatchWhenReady(FDelegateGraphTask::FDelegate::CreateRaw(this, &FCompletionList::ChainWaitForPrerequisites), TEXT("WaitOnCompletionList"), &PendingHandles);
+			DECLARE_CYCLE_STAT(TEXT("FDelegateGraphTask.WaitOnCompletionList"),
+				STAT_FDelegateGraphTask_WaitOnCompletionList,
+				STATGROUP_TaskGraphTasks);
+
+			CompleteHandle = FDelegateGraphTask::CreateAndDispatchWhenReady(
+				FDelegateGraphTask::FDelegate::CreateRaw(this, &FCompletionList::ChainWaitForPrerequisites),
+				GET_STATID(STAT_FDelegateGraphTask_WaitOnCompletionList), &PendingHandles
+			);
 		}
 		return CompleteHandle;
 	}
diff --git a/Engine/Source/Runtime/Engine/Classes/AI/Navigation/NavigationData.h b/Engine/Source/Runtime/Engine/Classes/AI/Navigation/NavigationData.h
--- a/Engine/Source/Runtime/Engine/Classes/AI/Navigation/NavigationData.h
+++ b/Engine/Source/Runtime/Engine/Classes/AI/Navigation/NavigationData.h
@@ -456,12 +456,14 @@ public:
 		SharedPath->SetNavigationDataUsed(this);
 		SharedPath->SetQuerier(Querier);
 
+		DECLARE_CYCLE_STAT(TEXT("FSimpleDelegateGraphTask.Adding a path to ActivePaths"),
+			STAT_FSimpleDelegateGraphTask_AddingPathToActivePaths,
+			STATGROUP_TaskGraphTasks);
+
 		FSimpleDelegateGraphTask::CreateAndDispatchWhenReady(
-			FSimpleDelegateGraphTask::FDelegate::CreateUObject(this, &ANavigationData::RegisterActivePath, SharedPath)
-			, TEXT("Adding a path to ActivePaths")
-			, NULL
-			, ENamedThreads::GameThread
-			);
+			FSimpleDelegateGraphTask::FDelegate::CreateUObject(this, &ANavigationData::RegisterActivePath, SharedPath),
+			GET_STATID(STAT_FSimpleDelegateGraphTask_AddingPathToActivePaths), NULL, ENamedThreads::GameThread
+		);
 
 		return SharedPath;
 	}
diff --git a/Engine/Source/Runtime/Engine/Private/AI/Navigation/NavigationSystem.cpp b/Engine/Source/Runtime/Engine/Private/AI/Navigation/NavigationSystem.cpp
--- a/Engine/Source/Runtime/Engine/Private/AI/Navigation/NavigationSystem.cpp
+++ b/Engine/Source/Runtime/Engine/Private/AI/Navigation/NavigationSystem.cpp
@@ -706,10 +706,13 @@ void UNavigationSystem::AbortAsyncFindPathRequest(uint32 AsynPathQueryID)
 
 void UNavigationSystem::TriggerAsyncQueries(TArray<FAsyncPathFindingQuery>& PathFindingQueries)
 {
+	DECLARE_CYCLE_STAT(TEXT("FSimpleDelegateGraphTask.NavigationSystem batched async queries"),
+		STAT_FSimpleDelegateGraphTask_NavigationSystemBatchedAsyncQueries,
+		STATGROUP_TaskGraphTasks);
+
 	FSimpleDelegateGraphTask::CreateAndDispatchWhenReady(
-		FSimpleDelegateGraphTask::FDelegate::CreateUObject(this, &UNavigationSystem::PerformAsyncQueries, PathFindingQueries)
-		, TEXT("NavigationSystem batched async queries")
-		);
+		FSimpleDelegateGraphTask::FDelegate::CreateUObject(this, &UNavigationSystem::PerformAsyncQueries, PathFindingQueries),
+		GET_STATID(STAT_FSimpleDelegateGraphTask_NavigationSystemBatchedAsyncQueries));
 }
 
 static void AsyncQueryDone(FAsyncPathFindingQuery Query)
@@ -754,12 +757,13 @@ void UNavigationSystem::PerformAsyncQueries(TArray<FAsyncPathFindingQuery> PathF
 
 		// @todo make it return more informative results (bResult == false)
 		// trigger calling delegate on main thread - otherwise it may depend too much on stuff being thread safe
+		DECLARE_CYCLE_STAT(TEXT("FSimpleDelegateGraphTask.Async nav query finished"),
+			STAT_FSimpleDelegateGraphTask_AsyncNavQueryFinished,
+			STATGROUP_TaskGraphTasks);
+
 		FSimpleDelegateGraphTask::CreateAndDispatchWhenReady(
-			FSimpleDelegateGraphTask::FDelegate::CreateStatic(AsyncQueryDone, *Query)
-			, TEXT("Async nav query finished")
-			, NULL
-			, ENamedThreads::GameThread
-			);
+			FSimpleDelegateGraphTask::FDelegate::CreateStatic(AsyncQueryDone, *Query),
+			GET_STATID(STAT_FSimpleDelegateGraphTask_AsyncNavQueryFinished), NULL, ENamedThreads::GameThread);
 	}
 }
 
@@ -1276,12 +1280,13 @@ void UNavigationSystem::RequestRegistration(ANavigationData* NavData, bool bTrig
 		if (bTriggerRegistrationProcessing == true)
 		{
 			// trigger registration candidates processing
+			DECLARE_CYCLE_STAT(TEXT("FSimpleDelegateGraphTask.Process registration candidates"),
+				STAT_FSimpleDelegateGraphTask_ProcessRegistrationCandidates,
+				STATGROUP_TaskGraphTasks);
+
 			FSimpleDelegateGraphTask::CreateAndDispatchWhenReady(
-				FSimpleDelegateGraphTask::FDelegate::CreateUObject(this, &UNavigationSystem::ProcessRegistrationCandidates)
-				, TEXT("Process registration candidates")
-				, NULL
-				, ENamedThreads::GameThread
-				);
+				FSimpleDelegateGraphTask::FDelegate::CreateUObject(this, &UNavigationSystem::ProcessRegistrationCandidates),
+				GET_STATID(STAT_FSimpleDelegateGraphTask_ProcessRegistrationCandidates), NULL, ENamedThreads::GameThread);
 		}
 	}
 	else
diff --git a/Engine/Source/Runtime/Engine/Private/AI/Navigation/RecastNavMesh.cpp b/Engine/Source/Runtime/Engine/Private/AI/Navigation/RecastNavMesh.cpp
--- a/Engine/Source/Runtime/Engine/Private/AI/Navigation/RecastNavMesh.cpp
+++ b/Engine/Source/Runtime/Engine/Private/AI/Navigation/RecastNavMesh.cpp
@@ -1409,12 +1409,13 @@ void ARecastNavMesh::GetDebugTileBounds(FBox& OuterBox, int32& NumTilesX, int32&
 void ARecastNavMesh::RequestDrawingUpdate()
 {
 #if !UE_BUILD_SHIPPING
+	DECLARE_CYCLE_STAT(TEXT("FSimpleDelegateGraphTask.Requesting navmesh redraw"),
+		STAT_FSimpleDelegateGraphTask_RequestingNavmeshRedraw,
+		STATGROUP_TaskGraphTasks);
+
 	FSimpleDelegateGraphTask::CreateAndDispatchWhenReady(
-		FSimpleDelegateGraphTask::FDelegate::CreateUObject(this, &ARecastNavMesh::UpdateDrawing)
-		, TEXT("Requesting navmesh redraw")
-		, NULL
-		, ENamedThreads::GameThread
-		);
+		FSimpleDelegateGraphTask::FDelegate::CreateUObject(this, &ARecastNavMesh::UpdateDrawing),
+		GET_STATID(STAT_FSimpleDelegateGraphTask_RequestingNavmeshRedraw), NULL, ENamedThreads::GameThread);
 #endif // !UE_BUILD_SHIPPING
 }
 
diff --git a/Engine/Source/Runtime/Engine/Private/AI/Navigation/RecastNavMeshGenerator.cpp b/Engine/Source/Runtime/Engine/Private/AI/Navigation/RecastNavMeshGenerator.cpp
--- a/Engine/Source/Runtime/Engine/Private/AI/Navigation/RecastNavMeshGenerator.cpp
+++ b/Engine/Source/Runtime/Engine/Private/AI/Navigation/RecastNavMeshGenerator.cpp
@@ -142,12 +142,13 @@ public:
 				SCOPE_SECONDS_COUNTER(ThisTime);
 				RecastNavMeshGenerator->GenerateTile(TileId, Version);	
 			
+				DECLARE_CYCLE_STAT(TEXT("FSimpleDelegateGraphTask.Requesting to build next tile if necessary"),
+					STAT_FSimpleDelegateGraphTask_RequestingToBuildNextTileIfNecessary,
+					STATGROUP_TaskGraphTasks);
+
 				FSimpleDelegateGraphTask::CreateAndDispatchWhenReady(
-					FSimpleDelegateGraphTask::FDelegate::CreateThreadSafeSP(RecastNavMeshGenerator, &FRecastNavMeshGenerator::UpdateTileGenerationWorkers, TileId)
-					, TEXT("Requesting to build next tile if necessary")
-					, NULL
-					, ENamedThreads::GameThread
-					);
+					FSimpleDelegateGraphTask::FDelegate::CreateThreadSafeSP(RecastNavMeshGenerator, &FRecastNavMeshGenerator::UpdateTileGenerationWorkers, TileId),
+					GET_STATID(STAT_FSimpleDelegateGraphTask_RequestingToBuildNextTileIfNecessary), NULL, ENamedThreads::GameThread);
 			}
 			INC_FLOAT_STAT_BY(STAT_Navigation_CumulativeBuildTime,(float)ThisTime*1000);
 		}
@@ -3520,11 +3521,14 @@ void FRecastNavMeshGenerator::RequestGeneration()
 	}
 	else
 	{
+		DECLARE_CYCLE_STAT(TEXT("FSimpleDelegateGraphTask.Requesting navmesh rebuild from scratch from RequestGeneration"),
+			STAT_FSimpleDelegateGraphTask_RequestingNavmeshRebuildFromScratchFromRequestGeneration,
+			STATGROUP_TaskGraphTasks);
+
 		FSimpleDelegateGraphTask::CreateAndDispatchWhenReady(
-			FSimpleDelegateGraphTask::FDelegate::CreateThreadSafeSP(this, &FRecastNavMeshGenerator::RequestGeneration)
-			, TEXT("Requesting navmesh rebuild from scratch from RequestGeneration")
-			, NULL
-			, ENamedThreads::GameThread
+			FSimpleDelegateGraphTask::FDelegate::CreateThreadSafeSP(this, &FRecastNavMeshGenerator::RequestGeneration),
+			GET_STATID(STAT_FSimpleDelegateGraphTask_RequestingNavmeshRebuildFromScratchFromRequestGeneration),
+			NULL, ENamedThreads::GameThread
 		);
 	}
 }
@@ -3538,10 +3542,14 @@ void FRecastNavMeshGenerator::RequestDirtyTilesRebuild()
 
 	bRebuildDirtyTilesRequested = true;
 
+	DECLARE_CYCLE_STAT(TEXT("FSimpleDelegateGraphTask.NavMesh Tiles Regeneration"),
+		STAT_FSimpleDelegateGraphTask_NavMeshTilesRegeneration,
+		STATGROUP_TaskGraphTasks);
+
 		// kick off a task to rebuild influenced tiles			
 	FSimpleDelegateGraphTask::CreateAndDispatchWhenReady(
-		FSimpleDelegateGraphTask::FDelegate::CreateThreadSafeSP(this, &FRecastNavMeshGenerator::RegenerateDirtyTiles)
-		, TEXT("NavMesh Tiles Regeneration"), NULL, ENamedThreads::GameThread
+		FSimpleDelegateGraphTask::FDelegate::CreateThreadSafeSP(this, &FRecastNavMeshGenerator::RegenerateDirtyTiles),
+		GET_STATID(STAT_FSimpleDelegateGraphTask_NavMeshTilesRegeneration), NULL, ENamedThreads::GameThread
 	);
 }
 
@@ -3762,11 +3770,14 @@ void FRecastNavMeshGenerator::UpdateBuilding()
 		}
 
 		// need to send it to main thread - uses some gamethread-only iterators
+		DECLARE_CYCLE_STAT(TEXT("FSimpleDelegateGraphTask.Requesting navmesh regen from UpdateBuilding"),
+			STAT_FSimpleDelegateGraphTask_RequestingNavmeshRegenFromUpdateBuilding,
+			STATGROUP_TaskGraphTasks);
+
 		FSimpleDelegateGraphTask::CreateAndDispatchWhenReady(
-			FSimpleDelegateGraphTask::FDelegate::CreateThreadSafeSP(this, &FRecastNavMeshGenerator::TriggerGeneration)
-			, TEXT("Requesting navmesh regen from UpdateBuilding")
-			, NULL
-			, ENamedThreads::GameThread
+			FSimpleDelegateGraphTask::FDelegate::CreateThreadSafeSP(this, &FRecastNavMeshGenerator::TriggerGeneration),
+			GET_STATID(STAT_FSimpleDelegateGraphTask_RequestingNavmeshRegenFromUpdateBuilding), NULL,
+			ENamedThreads::GameThread
 		);
 	}
 	else if(!UnrealNavBounds.IsValid)
@@ -3916,12 +3927,14 @@ bool FRecastNavMeshGenerator::AddTile(FRecastTileGenerator* TileGenerator, AReca
 		if (bOperationSuccessful)
 		{
 			// send off to game thread to not use critical section for modifying AsyncGenerationResultContainer
+			DECLARE_CYCLE_STAT(TEXT("FSimpleDelegateGraphTask.Storing async nav generetion result"),
+				STAT_FSimpleDelegateGraphTask_StoringAsyncNavGeneretionResult,
+				STATGROUP_TaskGraphTasks);
+
 			FSimpleDelegateGraphTask::CreateAndDispatchWhenReady(
-				FSimpleDelegateGraphTask::FDelegate::CreateThreadSafeSP(this, &FRecastNavMeshGenerator::StoreAsyncResults, AsyncResults)
-				, TEXT("Storing async nav generetion result")
-				, NULL
-				, ENamedThreads::GameThread
-				);
+				FSimpleDelegateGraphTask::FDelegate::CreateThreadSafeSP(this, &FRecastNavMeshGenerator::StoreAsyncResults, AsyncResults),
+				GET_STATID(STAT_FSimpleDelegateGraphTask_StoringAsyncNavGeneretionResult), NULL, ENamedThreads::GameThread
+			);
 		}
 	}
 	INC_FLOAT_STAT_BY(STAT_Navigation_CumulativeBuildTime,(float)ThisTime*1000);
@@ -4308,9 +4321,15 @@ void FRecastNavMeshGenerator::UpdateTileGenerationWorkers(int32 TileId)
 	// prepare next batch of dirty generators in next tick
 	if (GeneratorsQueue.Num() == 0 && DirtyGenerators.Num() > 0)
 	{
+		DECLARE_CYCLE_STAT(TEXT("FSimpleDelegateGraphTask.Prepare navmesh tile generators"),
+			STAT_FSimpleDelegateGraphTask_PrepareNavmeshTileGenerators,
+			STATGROUP_TaskGraphTasks);
+
 		FSimpleDelegateGraphTask::CreateAndDispatchWhenReady(
 			FSimpleDelegateGraphTask::FDelegate::CreateRaw(this, &FRecastNavMeshGenerator::StartDirtyGenerators), 
-			TEXT("Prepare navmesh tile generators"), NULL, ENamedThreads::GameThread);
+			GET_STATID(STAT_FSimpleDelegateGraphTask_PrepareNavmeshTileGenerators), NULL,
+			ENamedThreads::GameThread
+		);
 	}
 }
 
@@ -4382,12 +4401,13 @@ bool FRecastNavMeshGenerator::TransferGeneratedData()
 		DestNavMesh->GetRecastNavMeshImpl()->SetRecastMesh(DetourMesh, /*bOwnData=*/false);
 
 		// this should be done synchronously
+		DECLARE_CYCLE_STAT(TEXT("FSimpleDelegateGraphTask.Requesting navmesh redraw"),
+			STAT_FSimpleDelegateGraphTask_RequestingNavmeshRedraw,
+			STATGROUP_TaskGraphTasks);
+
 		FSimpleDelegateGraphTask::CreateAndDispatchWhenReady(
-			FSimpleDelegateGraphTask::FDelegate::CreateUObject(DestNavMesh.Get(), &ARecastNavMesh::UpdateNavMeshDrawing)
-			, TEXT("Requesting navmesh redraw")
-			, NULL
-			, ENamedThreads::GameThread
-			);
+			FSimpleDelegateGraphTask::FDelegate::CreateUObject(DestNavMesh.Get(), &ARecastNavMesh::UpdateNavMeshDrawing),
+			GET_STATID(STAT_FSimpleDelegateGraphTask_RequestingNavmeshRedraw), NULL, ENamedThreads::GameThread);
 
 		return true;
 	}
diff --git a/Engine/Source/Runtime/Engine/Private/ActorComponent.cpp b/Engine/Source/Runtime/Engine/Private/ActorComponent.cpp
--- a/Engine/Source/Runtime/Engine/Private/ActorComponent.cpp
+++ b/Engine/Source/Runtime/Engine/Private/ActorComponent.cpp
@@ -499,12 +499,14 @@ void UActorComponent::SetComponentTickEnabledAsync(bool bEnabled)
 {
 	if (!IsTemplate() && PrimaryComponentTick.bCanEverTick)
 	{
+		DECLARE_CYCLE_STAT(TEXT("FSimpleDelegateGraphTask.SetComponentTickEnabledAsync"),
+			STAT_FSimpleDelegateGraphTask_SetComponentTickEnabledAsync,
+			STATGROUP_TaskGraphTasks);
+
 		FSimpleDelegateGraphTask::CreateAndDispatchWhenReady(
-			FSimpleDelegateGraphTask::FDelegate::CreateUObject(this, &UActorComponent::SetComponentTickEnabled, bEnabled)
-			, TEXT("SetComponentTickEnabledAsync")
-			, NULL
-			, ENamedThreads::GameThread
-			);
+			FSimpleDelegateGraphTask::FDelegate::CreateUObject(this, &UActorComponent::SetComponentTickEnabled, bEnabled),
+			GET_STATID(STAT_FSimpleDelegateGraphTask_SetComponentTickEnabledAsync), NULL, ENamedThreads::GameThread
+		);
 	}
 }
 
diff --git a/Engine/Source/Runtime/Engine/Private/Materials/Material.cpp b/Engine/Source/Runtime/Engine/Private/Materials/Material.cpp
--- a/Engine/Source/Runtime/Engine/Private/Materials/Material.cpp
+++ b/Engine/Source/Runtime/Engine/Private/Materials/Material.cpp
@@ -857,12 +857,15 @@ bool UMaterial::CheckMaterialUsage_Concurrent(EMaterialUsage Usage, const bool b
 
 			FScopedEvent Event;
 			FCallSMU CallSMU(const_cast<UMaterial*>(this), Usage, bSkipPrim, bUsageSetSuccessfully, Event);
+
+			DECLARE_CYCLE_STAT(TEXT("FSimpleDelegateGraphTask.CheckMaterialUsage"),
+				STAT_FSimpleDelegateGraphTask_CheckMaterialUsage,
+				STATGROUP_TaskGraphTasks);
+
 			FSimpleDelegateGraphTask::CreateAndDispatchWhenReady(
-				FSimpleDelegateGraphTask::FDelegate::CreateRaw(&CallSMU, &FCallSMU::Task)
-				, TEXT("CheckMaterialUsage")
-				, NULL
-				, ENamedThreads::GameThread_Local
-				);
+				FSimpleDelegateGraphTask::FDelegate::CreateRaw(&CallSMU, &FCallSMU::Task),
+				GET_STATID(STAT_FSimpleDelegateGraphTask_CheckMaterialUsage), NULL, ENamedThreads::GameThread_Local
+			);
 		}
 	}
 	return bUsageSetSuccessfully;
diff --git a/Engine/Source/Runtime/Engine/Private/Materials/MaterialInstance.cpp b/Engine/Source/Runtime/Engine/Private/Materials/MaterialInstance.cpp
--- a/Engine/Source/Runtime/Engine/Private/Materials/MaterialInstance.cpp
+++ b/Engine/Source/Runtime/Engine/Private/Materials/MaterialInstance.cpp
@@ -897,12 +897,15 @@ bool UMaterialInstance::CheckMaterialUsage_Concurrent(const EMaterialUsage Usage
 
 				FScopedEvent Event;
 				FCallSMU CallSMU(const_cast<UMaterialInstance*>(this), Usage, bSkipPrim, bUsageSetSuccessfully, Event);
+
+				DECLARE_CYCLE_STAT(TEXT("FSimpleDelegateGraphTask.CheckMaterialUsage"),
+					STAT_FSimpleDelegateGraphTask_CheckMaterialUsage,
+					STATGROUP_TaskGraphTasks);
+
 				FSimpleDelegateGraphTask::CreateAndDispatchWhenReady(
-					FSimpleDelegateGraphTask::FDelegate::CreateRaw(&CallSMU, &FCallSMU::Task)
-					, TEXT("CheckMaterialUsage")
-					, NULL
-					, ENamedThreads::GameThread_Local
-					);
+					FSimpleDelegateGraphTask::FDelegate::CreateRaw(&CallSMU, &FCallSMU::Task),
+					GET_STATID(STAT_FSimpleDelegateGraphTask_CheckMaterialUsage), NULL, ENamedThreads::GameThread_Local
+				);
 			}
 		}
 		return bUsageSetSuccessfully;
diff --git a/Engine/Source/Runtime/Engine/Private/NavMeshRenderingComponent.cpp b/Engine/Source/Runtime/Engine/Private/NavMeshRenderingComponent.cpp
--- a/Engine/Source/Runtime/Engine/Private/NavMeshRenderingComponent.cpp
+++ b/Engine/Source/Runtime/Engine/Private/NavMeshRenderingComponent.cpp
@@ -474,12 +474,14 @@ void UNavMeshRenderingComponent::GatherData(struct FNavMeshSceneProxyData* Curre
 			CurrentData->MeshBuilders.Add(DebugMeshData);
 
 			// if we got here it means there's something being built, or is fresh so we better redraw in some time
+			DECLARE_CYCLE_STAT(TEXT("FSimpleDelegateGraphTask.Requesting navmesh redraw"),
+				STAT_FSimpleDelegateGraphTask_RequestingNavmeshRedraw,
+				STATGROUP_TaskGraphTasks);
+
 			FSimpleDelegateGraphTask::CreateAndDispatchWhenReady(
-				FSimpleDelegateGraphTask::FDelegate::CreateUObject(NavMesh, &ARecastNavMesh::MarkComponentsRenderStateDirty)
-				, TEXT("Requesting navmesh redraw")
-				, NULL
-				, ENamedThreads::GameThread
-				);
+				FSimpleDelegateGraphTask::FDelegate::CreateUObject(NavMesh, &ARecastNavMesh::MarkComponentsRenderStateDirty),
+				GET_STATID(STAT_FSimpleDelegateGraphTask_RequestingNavmeshRedraw), NULL, ENamedThreads::GameThread
+			);
 		}
 
 		if (NavMesh->bDrawClusters)
diff --git a/Engine/Source/Runtime/Engine/Private/Particles/ParticleComponents.cpp b/Engine/Source/Runtime/Engine/Private/Particles/ParticleComponents.cpp
--- a/Engine/Source/Runtime/Engine/Private/Particles/ParticleComponents.cpp
+++ b/Engine/Source/Runtime/Engine/Private/Particles/ParticleComponents.cpp
@@ -3889,18 +3889,24 @@ void UParticleSystemComponent::TickComponent(float DeltaTime, enum ELevelTick Ti
 	else
 	{
 		// set up async task and the game thread task to finalize the results.
+		DECLARE_CYCLE_STAT(TEXT("FSimpleDelegateGraphTask.AsyncParticleTick"),
+			STAT_FSimpleDelegateGraphTask_AsyncParticleTick,
+			STATGROUP_TaskGraphTasks);
+
 		AsyncWork = FSimpleDelegateGraphTask::CreateAndDispatchWhenReady(
-			FSimpleDelegateGraphTask::FDelegate::CreateUObject(this, &UParticleSystemComponent::ComputeTickComponent_Concurrent)
-			, TEXT("AsyncParticleTick")
-			, NULL
-			, ENamedThreads::AnyThread
-			);
+			FSimpleDelegateGraphTask::FDelegate::CreateUObject(this, &UParticleSystemComponent::ComputeTickComponent_Concurrent),
+			GET_STATID(STAT_FSimpleDelegateGraphTask_AsyncParticleTick), NULL, ENamedThreads::AnyThread
+		);
+
+		DECLARE_CYCLE_STAT(TEXT("FSimpleDelegateGraphTask.FinalizeParticleTick"),
+			STAT_FSimpleDelegateGraphTask_FinalizeParticleTick,
+			STATGROUP_TaskGraphTasks);
+
 		FGraphEventRef Finalize = FSimpleDelegateGraphTask::CreateAndDispatchWhenReady(
-			FSimpleDelegateGraphTask::FDelegate::CreateUObject(this, &UParticleSystemComponent::FinalizeTickComponent)
-			, TEXT("FinalizeParticleTick")
-			, AsyncWork
-			, ENamedThreads::GameThread
-			);
+			FSimpleDelegateGraphTask::FDelegate::CreateUObject(this, &UParticleSystemComponent::FinalizeTickComponent),
+			GET_STATID(STAT_FSimpleDelegateGraphTask_FinalizeParticleTick), AsyncWork, ENamedThreads::GameThread
+		);
+
 		ThisTickFunction->GetCompletionHandle()->DontCompleteUntil(Finalize);
 	}
 }
diff --git a/Engine/Source/Runtime/Engine/Private/PhysicsEngine/PhysLevel.cpp b/Engine/Source/Runtime/Engine/Private/PhysicsEngine/PhysLevel.cpp
--- a/Engine/Source/Runtime/Engine/Private/PhysicsEngine/PhysLevel.cpp
+++ b/Engine/Source/Runtime/Engine/Private/PhysicsEngine/PhysLevel.cpp
@@ -218,7 +218,16 @@ void FEndPhysicsTickFunction::ExecuteTick(float DeltaTime, enum ELevelTick TickT
 	if (PhysicsComplete.GetReference() && !PhysicsComplete->IsComplete())
 	{
 		// don't release the next tick group until the physics has completed and we have run FinishPhysicsSim
-		MyCompletionGraphEvent->DontCompleteUntil(FSimpleDelegateGraphTask::CreateAndDispatchWhenReady(FSimpleDelegateGraphTask::FDelegate::CreateUObject(Target, &UWorld::FinishPhysicsSim), TEXT("FinishPhysicsSim"), PhysicsComplete, ENamedThreads::GameThread));
+		DECLARE_CYCLE_STAT(TEXT("FSimpleDelegateGraphTask.FinishPhysicsSim"),
+			STAT_FSimpleDelegateGraphTask_FinishPhysicsSim,
+			STATGROUP_TaskGraphTasks);
+
+		MyCompletionGraphEvent->DontCompleteUntil(
+			FSimpleDelegateGraphTask::CreateAndDispatchWhenReady(
+				FSimpleDelegateGraphTask::FDelegate::CreateUObject(Target, &UWorld::FinishPhysicsSim),
+				GET_STATID(STAT_FSimpleDelegateGraphTask_FinishPhysicsSim), PhysicsComplete, ENamedThreads::GameThread
+			)
+		);
 	}
 	else
 	{
diff --git a/Engine/Source/Runtime/Engine/Private/PhysicsEngine/PhysScene.cpp b/Engine/Source/Runtime/Engine/Private/PhysicsEngine/PhysScene.cpp
--- a/Engine/Source/Runtime/Engine/Private/PhysicsEngine/PhysScene.cpp
+++ b/Engine/Source/Runtime/Engine/Private/PhysicsEngine/PhysScene.cpp
@@ -350,7 +350,15 @@ bool FPhysScene::SubstepSimulation(uint32 SceneType, FGraphEventRef &InOutComple
 		//we have valid scene and subtime so enqueue task
 		PhysXCompletionTask* Task = new PhysXCompletionTask(InOutCompletionEvent, PScene->getTaskManager());
 		ENamedThreads::Type NamedThread = PhysSingleThreadedMode() ? ENamedThreads::GameThread : ENamedThreads::AnyThread;
-		FSimpleDelegateGraphTask::CreateAndDispatchWhenReady(FSimpleDelegateGraphTask::FDelegate::CreateRaw(PhysSubSteppers[SceneType], &FPhysSubstepTask::StepSimulation, Task), TEXT("SubstepSimulationImp"), NULL, NamedThread);
+
+		DECLARE_CYCLE_STAT(TEXT("FSimpleDelegateGraphTask.SubstepSimulationImp"),
+			STAT_FSimpleDelegateGraphTask_SubstepSimulationImp,
+			STATGROUP_TaskGraphTasks);
+
+		FSimpleDelegateGraphTask::CreateAndDispatchWhenReady(
+			FSimpleDelegateGraphTask::FDelegate::CreateRaw(PhysSubSteppers[SceneType], &FPhysSubstepTask::StepSimulation, Task),
+			GET_STATID(STAT_FSimpleDelegateGraphTask_SubstepSimulationImp), NULL, NamedThread
+		);
 		return true;
 	}
 #endif
@@ -842,7 +850,18 @@ void FPhysScene::StartFrame()
 		if (PhysicsSubsceneCompletion[PST_Sync].GetReference())
 		{
 			MainScenePrerequisites.Add(PhysicsSubsceneCompletion[PST_Sync]);
-			new (FinishPrerequisites)FGraphEventRef(FDelegateGraphTask::CreateAndDispatchWhenReady(FDelegateGraphTask::FDelegate::CreateRaw(this, &FPhysScene::SceneCompletionTask, PST_Sync), TEXT("ProcessPhysScene_Sync"), &MainScenePrerequisites, ENamedThreads::GameThread, ENamedThreads::GameThread));
+
+			DECLARE_CYCLE_STAT(TEXT("FDelegateGraphTask.ProcessPhysScene_Sync"),
+				STAT_FDelegateGraphTask_ProcessPhysScene_Sync,
+				STATGROUP_TaskGraphTasks);
+
+			new (FinishPrerequisites)FGraphEventRef(
+				FDelegateGraphTask::CreateAndDispatchWhenReady(
+					FDelegateGraphTask::FDelegate::CreateRaw(this, &FPhysScene::SceneCompletionTask, PST_Sync),
+					GET_STATID(STAT_FDelegateGraphTask_ProcessPhysScene_Sync), &MainScenePrerequisites,
+					ENamedThreads::GameThread, ENamedThreads::GameThread
+				)
+			);
 		}
 	}
 
@@ -851,7 +870,17 @@ void FPhysScene::StartFrame()
 		TickPhysScene(PST_Async, PhysicsSubsceneCompletion[PST_Async]);
 		if (PhysicsSubsceneCompletion[PST_Async].GetReference())
 		{
-			new (FinishPrerequisites)FGraphEventRef(FDelegateGraphTask::CreateAndDispatchWhenReady(FDelegateGraphTask::FDelegate::CreateRaw(this, &FPhysScene::SceneCompletionTask, PST_Async), TEXT("ProcessPhysScene_Async"), PhysicsSubsceneCompletion[PST_Async], ENamedThreads::GameThread, ENamedThreads::GameThread));
+			DECLARE_CYCLE_STAT(TEXT("FDelegateGraphTask.ProcessPhysScene_Async"),
+				STAT_FDelegateGraphTask_ProcessPhysScene_Async,
+				STATGROUP_TaskGraphTasks);
+
+			new (FinishPrerequisites)FGraphEventRef(
+				FDelegateGraphTask::CreateAndDispatchWhenReady(
+					FDelegateGraphTask::FDelegate::CreateRaw(this, &FPhysScene::SceneCompletionTask, PST_Async),
+					GET_STATID(STAT_FDelegateGraphTask_ProcessPhysScene_Async), PhysicsSubsceneCompletion[PST_Async],
+					ENamedThreads::GameThread, ENamedThreads::GameThread
+				)
+			);
 		}
 	}
 
@@ -860,8 +889,12 @@ void FPhysScene::StartFrame()
 	{
 		if (FinishPrerequisites.Num() > 1)  // we don't need to create a new task if we only have one prerequisite
 		{
-			PhysicsSceneCompletion = TGraphTask<FNullGraphTask>::CreateTask(&FinishPrerequisites, ENamedThreads::GameThread).ConstructAndDispatchWhenReady(TEXT("ProcessPhysScene_Join"),
-				PhysSingleThreadedMode() ? ENamedThreads::GameThread : ENamedThreads::AnyThread);
+			DECLARE_CYCLE_STAT(TEXT("FNullGraphTask.ProcessPhysScene_Join"),
+				STAT_FNullGraphTask_ProcessPhysScene_Join,
+				STATGROUP_TaskGraphTasks);
+
+			PhysicsSceneCompletion = TGraphTask<FNullGraphTask>::CreateTask(&FinishPrerequisites, ENamedThreads::GameThread).ConstructAndDispatchWhenReady(
+				GET_STATID(STAT_FNullGraphTask_ProcessPhysScene_Join), PhysSingleThreadedMode() ? ENamedThreads::GameThread : ENamedThreads::AnyThread);
 		}
 		else
 		{
@@ -880,7 +913,17 @@ void FPhysScene::StartCloth()
 	{
 		if (PhysicsSubsceneCompletion[PST_Cloth].GetReference())
 		{
-			new (FinishPrerequisites)FGraphEventRef(FDelegateGraphTask::CreateAndDispatchWhenReady(FDelegateGraphTask::FDelegate::CreateRaw(this, &FPhysScene::SceneCompletionTask, PST_Cloth), TEXT("ProcessPhysScene_Cloth"), PhysicsSubsceneCompletion[PST_Cloth], ENamedThreads::GameThread, ENamedThreads::GameThread));
+			DECLARE_CYCLE_STAT(TEXT("FDelegateGraphTask.ProcessPhysScene_Cloth"),
+				STAT_FDelegateGraphTask_ProcessPhysScene_Cloth,
+				STATGROUP_TaskGraphTasks);
+
+			new (FinishPrerequisites)FGraphEventRef(
+				FDelegateGraphTask::CreateAndDispatchWhenReady(
+					FDelegateGraphTask::FDelegate::CreateRaw(this, &FPhysScene::SceneCompletionTask, PST_Cloth),
+					GET_STATID(STAT_FDelegateGraphTask_ProcessPhysScene_Cloth), PhysicsSubsceneCompletion[PST_Cloth],
+					ENamedThreads::GameThread, ENamedThreads::GameThread
+				)
+			);
 		}
 	}
 
@@ -890,7 +933,15 @@ void FPhysScene::StartCloth()
 		TickPhysScene(PST_Async, PhysicsSubsceneCompletion[PST_Async]);
 		if (PhysicsSubsceneCompletion[PST_Async].GetReference())
 		{
-			FrameLaggedPhysicsSubsceneCompletion[PST_Async] = FDelegateGraphTask::CreateAndDispatchWhenReady(FDelegateGraphTask::FDelegate::CreateRaw(this, &FPhysScene::SceneCompletionTask, PST_Async), TEXT("ProcessPhysScene_Async"), PhysicsSubsceneCompletion[PST_Async], ENamedThreads::GameThread, ENamedThreads::GameThread);
+			DECLARE_CYCLE_STAT(TEXT("FDelegateGraphTask.ProcessPhysScene_Async"),
+				STAT_FDelegateGraphTask_ProcessPhysScene_Async,
+				STATGROUP_TaskGraphTasks);
+
+			FrameLaggedPhysicsSubsceneCompletion[PST_Async] = FDelegateGraphTask::CreateAndDispatchWhenReady(
+				FDelegateGraphTask::FDelegate::CreateRaw(this, &FPhysScene::SceneCompletionTask, PST_Async),
+				GET_STATID(STAT_FDelegateGraphTask_ProcessPhysScene_Async), PhysicsSubsceneCompletion[PST_Async],
+				ENamedThreads::GameThread, ENamedThreads::GameThread
+			);
 		}
 	}
 }
diff --git a/Engine/Source/Runtime/Engine/Private/PhysicsEngine/PhysSubstepTasks.cpp b/Engine/Source/Runtime/Engine/Private/PhysicsEngine/PhysSubstepTasks.cpp
--- a/Engine/Source/Runtime/Engine/Private/PhysicsEngine/PhysSubstepTasks.cpp
+++ b/Engine/Source/Runtime/Engine/Private/PhysicsEngine/PhysSubstepTasks.cpp
@@ -299,7 +299,14 @@ void FPhysSubstepTask::SubstepSimulationStart()
 	CompletionEvent = FGraphEvent::CreateGraphEvent();
 	PhysXCompletionTask* SubstepTask = new PhysXCompletionTask(CompletionEvent, PAScene->getTaskManager());
 	ENamedThreads::Type NamedThread = PhysSingleThreadedMode() ? ENamedThreads::GameThread : ENamedThreads::AnyThread;
-	FDelegateGraphTask::CreateAndDispatchWhenReady(FDelegateGraphTask::FDelegate::CreateRaw(this, &FPhysSubstepTask::SubstepSimulationEnd), TEXT("ProcessPhysSubstepSimulation"), CompletionEvent, NamedThread, NamedThread);
+
+	DECLARE_CYCLE_STAT(TEXT("FDelegateGraphTask.ProcessPhysSubstepSimulation"),
+		STAT_FDelegateGraphTask_ProcessPhysSubstepSimulation,
+		STATGROUP_TaskGraphTasks);
+
+	FDelegateGraphTask::CreateAndDispatchWhenReady(
+		FDelegateGraphTask::FDelegate::CreateRaw(this, &FPhysSubstepTask::SubstepSimulationEnd),
+		GET_STATID(STAT_FDelegateGraphTask_ProcessPhysSubstepSimulation), CompletionEvent, NamedThread, NamedThread);
 
 	++CurrentSubStep;	
 
diff --git a/Engine/Source/Runtime/Engine/Private/TickTaskManager.cpp b/Engine/Source/Runtime/Engine/Private/TickTaskManager.cpp
--- a/Engine/Source/Runtime/Engine/Private/TickTaskManager.cpp
+++ b/Engine/Source/Runtime/Engine/Private/TickTaskManager.cpp
@@ -175,9 +175,15 @@ public:
 		}
 		else
 		{
+			DECLARE_CYCLE_STAT(TEXT("FDelegateGraphTask.DispatchTickGroup"),
+				STAT_FDelegateGraphTask_DispatchTickGroup,
+				STATGROUP_TaskGraphTasks);
+
 			// dispatch the tick group on another thread, that way, the game thread can be processing ticks while ticks are being queued by another thread
 			FTaskGraphInterface::Get().WaitUntilTaskCompletes(
-				FDelegateGraphTask::CreateAndDispatchWhenReady(FDelegateGraphTask::FDelegate::CreateRaw(this, &FTickTaskSequencer::DispatchTickGroup, WorldTickGroup), TEXT("DispatchTickGroup")),	
+				FDelegateGraphTask::CreateAndDispatchWhenReady(
+					FDelegateGraphTask::FDelegate::CreateRaw(this, &FTickTaskSequencer::DispatchTickGroup, WorldTickGroup),
+					GET_STATID(STAT_FDelegateGraphTask_DispatchTickGroup)),
 				ENamedThreads::GameThread);
 		}
 
diff --git a/Engine/Source/Runtime/Engine/Private/UnrealEngine.cpp b/Engine/Source/Runtime/Engine/Private/UnrealEngine.cpp
--- a/Engine/Source/Runtime/Engine/Private/UnrealEngine.cpp
+++ b/Engine/Source/Runtime/Engine/Private/UnrealEngine.cpp
@@ -4012,7 +4012,18 @@ bool UEngine::HandleDebugCommand( const TCHAR* Cmd, FOutputDevice& Ar )
 				UE_LOG(LogEngine, Fatal, TEXT("Crashing the worker thread at your request") );
 			}
 		};
-		FTaskGraphInterface::Get().WaitUntilTaskCompletes(FDelegateGraphTask::CreateAndDispatchWhenReady(FDelegateGraphTask::FDelegate::CreateStatic(FThread::Crash), TEXT("FThread::Crash")), ENamedThreads::GameThread);
+
+		DECLARE_CYCLE_STAT(TEXT("FDelegateGraphTask.FThread::Crash"),
+			STAT_FDelegateGraphTask_FThread__Crash,
+			STATGROUP_TaskGraphTasks);
+
+		FTaskGraphInterface::Get().WaitUntilTaskCompletes(
+			FDelegateGraphTask::CreateAndDispatchWhenReady(
+				FDelegateGraphTask::FDelegate::CreateStatic(FThread::Crash),
+				GET_STATID(STAT_FDelegateGraphTask_FThread__Crash)
+			),
+			ENamedThreads::GameThread
+		);
 		return true;
 	}
 	if( FParse::Command(&Cmd,TEXT("THREADCHECK")) )
@@ -4025,7 +4036,18 @@ bool UEngine::HandleDebugCommand( const TCHAR* Cmd, FOutputDevice& Ar )
 				check(!"Crashing a worker thread via check(0) at your request");
 			}
 		};
-		FTaskGraphInterface::Get().WaitUntilTaskCompletes(FDelegateGraphTask::CreateAndDispatchWhenReady(FDelegateGraphTask::FDelegate::CreateStatic(FThread::Check), TEXT("FThread::Check")), ENamedThreads::GameThread);
+
+		DECLARE_CYCLE_STAT(TEXT("FDelegateGraphTask.FThread::Check"),
+			STAT_FDelegateGraphTask_FThread__Check,
+			STATGROUP_TaskGraphTasks);
+
+		FTaskGraphInterface::Get().WaitUntilTaskCompletes(
+			FDelegateGraphTask::CreateAndDispatchWhenReady(
+				FDelegateGraphTask::FDelegate::CreateStatic(FThread::Check),
+				GET_STATID(STAT_FDelegateGraphTask_FThread__Check)
+			),
+			ENamedThreads::GameThread
+		);
 		return true;
 	}
 	if( FParse::Command(&Cmd,TEXT("THREADGPF")) )
@@ -4038,7 +4060,18 @@ bool UEngine::HandleDebugCommand( const TCHAR* Cmd, FOutputDevice& Ar )
 				*(int32 *)3 = 123;
 			}
 		};
-		FTaskGraphInterface::Get().WaitUntilTaskCompletes(FDelegateGraphTask::CreateAndDispatchWhenReady(FDelegateGraphTask::FDelegate::CreateStatic(FThread::GPF), TEXT("FThread::GPF")), ENamedThreads::GameThread);
+
+		DECLARE_CYCLE_STAT(TEXT("FDelegateGraphTask.FThread::GPF"),
+			STAT_FDelegateGraphTask_FThread__GPF,
+			STATGROUP_TaskGraphTasks);
+
+		FTaskGraphInterface::Get().WaitUntilTaskCompletes(
+			FDelegateGraphTask::CreateAndDispatchWhenReady(
+				FDelegateGraphTask::FDelegate::CreateStatic(FThread::GPF),
+				GET_STATID(STAT_FDelegateGraphTask_FThread__GPF)
+			),
+			ENamedThreads::GameThread
+		);
 		return true;
 	}
 	else if( FParse::Command( &Cmd, TEXT( "THREADENSURE" ) ) )
@@ -4051,7 +4084,17 @@ bool UEngine::HandleDebugCommand( const TCHAR* Cmd, FOutputDevice& Ar )
 				ensure( 0 );
 			}
 		};
-		FTaskGraphInterface::Get().WaitUntilTaskCompletes( FDelegateGraphTask::CreateAndDispatchWhenReady( FDelegateGraphTask::FDelegate::CreateStatic( FThread::Ensure ), TEXT( "FThread::Ensure" ) ), ENamedThreads::GameThread );
+
+		DECLARE_CYCLE_STAT(TEXT("FThread::Ensure"),
+			STAT_FThread__Ensure,
+			STATGROUP_TaskGraphTasks);
+
+		FTaskGraphInterface::Get().WaitUntilTaskCompletes(
+			FDelegateGraphTask::CreateAndDispatchWhenReady(
+				FDelegateGraphTask::FDelegate::CreateStatic(FThread::Ensure),
+				GET_STATID(STAT_FThread__Ensure)),
+			ENamedThreads::GameThread
+		);
 		return true;
 	}
 	else if( FParse::Command( &Cmd, TEXT( "THREADFATAL" ) ) )
@@ -4064,7 +4107,17 @@ bool UEngine::HandleDebugCommand( const TCHAR* Cmd, FOutputDevice& Ar )
 				LowLevelFatalError( TEXT( "FError::LowLevelFatal test" ) );
 			}
 		};
-		FTaskGraphInterface::Get().WaitUntilTaskCompletes( FDelegateGraphTask::CreateAndDispatchWhenReady( FDelegateGraphTask::FDelegate::CreateStatic( FThread::Fatal ), TEXT( "FThread::Fatal" ) ), ENamedThreads::GameThread );
+
+		DECLARE_CYCLE_STAT(TEXT("FThread::Fatal"),
+			STAT_FThread__Fatal,
+			STATGROUP_TaskGraphTasks);
+
+		FTaskGraphInterface::Get().WaitUntilTaskCompletes(
+			FDelegateGraphTask::CreateAndDispatchWhenReady(
+				FDelegateGraphTask::FDelegate::CreateStatic(FThread::Fatal),
+				GET_STATID(STAT_FThread__Fatal)),
+			ENamedThreads::GameThread
+		);
 		return true;
 	}
 	else if( FParse::Command(&Cmd,TEXT("CRASH")) )
diff --git a/Engine/Source/Runtime/RHI/Private/GPUProfiler.cpp b/Engine/Source/Runtime/RHI/Private/GPUProfiler.cpp
--- a/Engine/Source/Runtime/RHI/Private/GPUProfiler.cpp
+++ b/Engine/Source/Runtime/RHI/Private/GPUProfiler.cpp
@@ -234,11 +234,13 @@ void FGPUProfilerEventNodeFrame::DumpEventTree()
 
 				TSharedPtr<FVisualizerEvent> VisualizerData = CreateVisualizerData( EventTree );
 
-				FSimpleDelegateGraphTask::CreateAndDispatchWhenReady
-				(
-					FSimpleDelegateGraphTask::FDelegate::CreateRaw(&DisplayProfilerVisualizer, &FDisplayProfilerVisualizer::Thread, VisualizerData )
-					, TEXT("DisplayProfilerVisualizer")
-					, nullptr, ENamedThreads::GameThread
+				DECLARE_CYCLE_STAT(TEXT("FSimpleDelegateGraphTask.DisplayProfilerVisualizer"),
+					STAT_FSimpleDelegateGraphTask_DisplayProfilerVisualizer,
+					STATGROUP_TaskGraphTasks);
+
+				FSimpleDelegateGraphTask::CreateAndDispatchWhenReady(
+					FSimpleDelegateGraphTask::FDelegate::CreateRaw(&DisplayProfilerVisualizer, &FDisplayProfilerVisualizer::Thread, VisualizerData),
+					GET_STATID(STAT_FSimpleDelegateGraphTask_DisplayProfilerVisualizer), nullptr, ENamedThreads::GameThread
 				);
 			}
 
diff --git a/Engine/Source/Runtime/RenderCore/Private/RenderingThread.cpp b/Engine/Source/Runtime/RenderCore/Private/RenderingThread.cpp
--- a/Engine/Source/Runtime/RenderCore/Private/RenderingThread.cpp
+++ b/Engine/Source/Runtime/RenderCore/Private/RenderingThread.cpp
@@ -88,11 +88,13 @@ FSuspendRenderingThread::FSuspendRenderingThread( bool bInRecreateThread )
 			
 			if (GIsThreadedRendering)
 			{
+				DECLARE_CYCLE_STAT(TEXT("FSimpleDelegateGraphTask.SuspendRendering"),
+					STAT_FSimpleDelegateGraphTask_SuspendRendering,
+					STATGROUP_TaskGraphTasks);
+
 				FGraphEventRef CompleteHandle = FSimpleDelegateGraphTask::CreateAndDispatchWhenReady(
 					FSimpleDelegateGraphTask::FDelegate::CreateStatic(&SuspendRendering),
-					TEXT("SuspendRendering"),
-					NULL,
-					ENamedThreads::RenderThread);
+					GET_STATID(STAT_FSimpleDelegateGraphTask_SuspendRendering), NULL, ENamedThreads::RenderThread);
 
 				// Busy wait while Kismet debugging, to avoid opportunistic execution of game thread tasks
 				// If the game thread is already executing tasks, then we have no choice but to spin
@@ -110,11 +112,13 @@ FSuspendRenderingThread::FSuspendRenderingThread( bool bInRecreateThread )
 				check(GIsRenderingThreadSuspended);
 			
 				// Now tell the render thread to busy wait until it's resumed
+				DECLARE_CYCLE_STAT(TEXT("FSimpleDelegateGraphTask.WaitAndResumeRendering"),
+					STAT_FSimpleDelegateGraphTask_WaitAndResumeRendering,
+					STATGROUP_TaskGraphTasks);
+
 				FSimpleDelegateGraphTask::CreateAndDispatchWhenReady(
 					FSimpleDelegateGraphTask::FDelegate::CreateStatic(&WaitAndResumeRendering),
-					TEXT("WaitAndResumeRendering"),
-					NULL,
-					ENamedThreads::RenderThread);
+					GET_STATID(STAT_FSimpleDelegateGraphTask_WaitAndResumeRendering), NULL, ENamedThreads::RenderThread);
 			}
 			else
 			{
@@ -141,11 +145,14 @@ FSuspendRenderingThread::~FSuspendRenderingThread()
 			StartRenderingThread();
             
             // Now tell the render thread to set it self to real time mode
+			DECLARE_CYCLE_STAT(TEXT("FSimpleDelegateGraphTask.SetRealTimeMode"),
+				STAT_FSimpleDelegateGraphTask_SetRealTimeMode,
+				STATGROUP_TaskGraphTasks);
+
             FSimpleDelegateGraphTask::CreateAndDispatchWhenReady(
                 FSimpleDelegateGraphTask::FDelegate::CreateStatic(&FPlatformProcess::SetRealTimeMode),
-                TEXT("SetRealTimeMode"),
-                NULL,
-                ENamedThreads::RenderThread);
+				GET_STATID(STAT_FSimpleDelegateGraphTask_SetRealTimeMode), NULL, ENamedThreads::RenderThread
+			);
         }
 	}
 	else
@@ -588,16 +595,22 @@ void FRenderCommandFence::BeginFence()
 	}
 	else
 	{
+		DECLARE_CYCLE_STAT(TEXT("FNullGraphTask.FenceRenderCommand"),
+			STAT_FNullGraphTask_FenceRenderCommand,
+			STATGROUP_TaskGraphTasks);
+
 		if (IsFenceComplete())
 		{
-			CompletionEvent = TGraphTask<FNullGraphTask>::CreateTask(NULL, ENamedThreads::GameThread).ConstructAndDispatchWhenReady(TEXT("FenceRenderCommand"), ENamedThreads::RenderThread);
+			CompletionEvent = TGraphTask<FNullGraphTask>::CreateTask(NULL, ENamedThreads::GameThread).ConstructAndDispatchWhenReady(
+				GET_STATID(STAT_FNullGraphTask_FenceRenderCommand), ENamedThreads::RenderThread);
 		}
 		else
 		{
 			// we already had a fence, so we will chain this one to the old one as a prerequisite
 			FGraphEventArray Prerequistes;
 			Prerequistes.Add(CompletionEvent);
-			CompletionEvent = TGraphTask<FNullGraphTask>::CreateTask(&Prerequistes, ENamedThreads::GameThread).ConstructAndDispatchWhenReady(TEXT("FenceRenderCommand"), ENamedThreads::RenderThread);
+			CompletionEvent = TGraphTask<FNullGraphTask>::CreateTask(&Prerequistes, ENamedThreads::GameThread).ConstructAndDispatchWhenReady(
+				GET_STATID(STAT_FNullGraphTask_FenceRenderCommand), ENamedThreads::RenderThread);
 		}
 	}
 }
diff --git a/Engine/Source/Runtime/Renderer/Private/Renderer.cpp b/Engine/Source/Runtime/Renderer/Private/Renderer.cpp
--- a/Engine/Source/Runtime/Renderer/Private/Renderer.cpp
+++ b/Engine/Source/Runtime/Renderer/Private/Renderer.cpp
@@ -113,12 +113,14 @@ void FRendererModule::DebugLogOnCrash()
 			}
 		} Test;
 
-		FSimpleDelegateGraphTask::CreateAndDispatchWhenReady
-			(
-			FSimpleDelegateGraphTask::FDelegate::CreateRaw(&Test, &FTest::Thread)
-			, TEXT("DumpDataAfterCrash")
-			, nullptr, ENamedThreads::GameThread
-			);
+		DECLARE_CYCLE_STAT(TEXT("FSimpleDelegateGraphTask.DumpDataAfterCrash"),
+			STAT_FSimpleDelegateGraphTask_DumpDataAfterCrash,
+			STATGROUP_TaskGraphTasks);
+
+		FSimpleDelegateGraphTask::CreateAndDispatchWhenReady(
+			FSimpleDelegateGraphTask::FDelegate::CreateRaw(&Test, &FTest::Thread),
+			GET_STATID(STAT_FSimpleDelegateGraphTask_DumpDataAfterCrash), nullptr, ENamedThreads::GameThread
+		);
 	}
 }
 
