Fixed some IME activation/deactivation issues

TTP# 343949 - CRASH: Adding Japanese comments to blueprints by pressing "c" key leads to crash

Newly created windows now have their HWND associated with the disabled IME contexts to prevent you being able to compose when there is no text area active. The only exception to this is when a new window opens while composing (eg, to show a search filter) - in this case we have to try and preserve the current IME context as the user is still typing.

To make sure that all of the windows which may be using an IME context are disabled when a text area loses focus, we now keep a list of known windows, and go through and re-associate each HWND with the disabled IME contexts when DeactivateContext is called.

There was also an issue where Slate widgets could be culled (and thus, unregister their context) without being told they had lost focus, which resulted in a crash as the active context had been unregistered. The Slate widgets in question now make sure they deactivate themselves before unregistering their context.

This also fixes an issue where an IMM based IME would incorrectly set the selection range when it was cancelled.

diff --git a/Engine/Source/Runtime/Core/Private/Mac/CocoaTextView.cpp b/Engine/Source/Runtime/Core/Private/Mac/CocoaTextView.cpp
--- a/Engine/Source/Runtime/Core/Private/Mac/CocoaTextView.cpp
+++ b/Engine/Source/Runtime/Core/Private/Mac/CocoaTextView.cpp
@@ -156,6 +156,11 @@
 	[[self inputContext] discardMarkedText];
 }
 
+- (bool)isActiveInputMethod:(const TSharedRef<ITextInputMethodContext>&)InContext
+{
+	return IMMContext == InContext;
+}
+
 //@protocol NSTextInputClient
 //@required
 /* The receiver inserts aString replacing the content specified by replacementRange. aString can be either an NSString or NSAttributedString instance.
diff --git a/Engine/Source/Runtime/Core/Private/Mac/MacTextInputMethodSystem.cpp b/Engine/Source/Runtime/Core/Private/Mac/MacTextInputMethodSystem.cpp
--- a/Engine/Source/Runtime/Core/Private/Mac/MacTextInputMethodSystem.cpp
+++ b/Engine/Source/Runtime/Core/Private/Mac/MacTextInputMethodSystem.cpp
@@ -119,6 +119,10 @@ void FMacTextInputMethodSystem::Terminate()
 }
 
 // ITextInputMethodSystem Interface Begin
+void FMacTextInputMethodSystem::ApplyDefaults(const TSharedRef<FGenericWindow>& InWindow)
+{
+}
+
 TSharedPtr<ITextInputMethodChangeNotifier> FMacTextInputMethodSystem::RegisterContext(const TSharedRef<ITextInputMethodContext>& Context)
 {
 	TSharedRef<ITextInputMethodChangeNotifier> Notifier = MakeShareable( new FTextInputMethodChangeNotifier(Context) );
@@ -221,4 +225,33 @@ void FMacTextInputMethodSystem::DeactivateContext(const TSharedRef<ITextInputMet
 		UE_LOG(LogMacTextInputMethodSystem, Error, TEXT("Deactivating a context failed when its window couldn't be found."));
 	}
 }
+
+bool FMacTextInputMethodSystem::IsActiveContext(const TSharedRef<ITextInputMethodContext>& Context) const
+{
+	const TWeakPtr<ITextInputMethodChangeNotifier>& NotifierRef = ContextMap[Context];
+	if(!NotifierRef.IsValid())
+	{
+		UE_LOG(LogMacTextInputMethodSystem, Error, TEXT("Checking for an active context failed when its registration couldn't be found."));
+		return false;
+	}
+	
+	TSharedPtr<ITextInputMethodChangeNotifier> Notifier(NotifierRef.Pin());
+	FTextInputMethodChangeNotifier* MacNotifier = (FTextInputMethodChangeNotifier*)Notifier.Get();
+	TSharedPtr<FGenericWindow> GenericWindow = MacNotifier->GetContextWindow();
+	if(GenericWindow.IsValid())
+	{
+		FCocoaWindow* CocoaWindow = (FCocoaWindow*)GenericWindow->GetOSWindowHandle();
+		if(CocoaWindow && [CocoaWindow openGLView])
+		{
+			NSView* GLView = [CocoaWindow openGLView];
+			if([GLView isKindOfClass:[FCocoaTextView class]])
+			{
+				FCocoaTextView* TextView = (FCocoaTextView*)GLView;
+				return [TextView isActiveInputMethod:Context];
+			}
+		}
+	}
+	UE_LOG(LogMacTextInputMethodSystem, Error, TEXT("Checking for an active context failed when its window couldn't be found."));
+	return false;
+}
 // ITextInputMethodSystem Interface End
diff --git a/Engine/Source/Runtime/Core/Private/Windows/WindowsTextInputMethodSystem.cpp b/Engine/Source/Runtime/Core/Private/Windows/WindowsTextInputMethodSystem.cpp
--- a/Engine/Source/Runtime/Core/Private/Windows/WindowsTextInputMethodSystem.cpp
+++ b/Engine/Source/Runtime/Core/Private/Windows/WindowsTextInputMethodSystem.cpp
@@ -151,7 +151,7 @@ bool FWindowsTextInputMethodSystem::Initialize()
 	{
 		const HKL KeyboardLayout = ::GetKeyboardLayout(0);
 
-		// We might already be using an IME if its set as the default language
+		// We might already be using an IME if it's set as the default language
 		// If so, work out what kind of IME it is
 		TF_INPUTPROCESSORPROFILE TSFProfile;
 		if(SUCCEEDED(TSFInputProcessorProfileManager->GetActiveProfile(GUID_TFCAT_TIP_KEYBOARD, &TSFProfile)) && TSFProfile.hkl && TSFProfile.dwProfileType == TF_PROFILETYPE_INPUTPROCESSOR)
@@ -175,6 +175,7 @@ bool FWindowsTextInputMethodSystem::InitializeIMM()
 {
 	UE_LOG(LogWindowsTextInputMethodSystem, Verbose, TEXT("Initializing IMM..."));
 
+	IMMContextId = ::ImmCreateContext();
 	UpdateIMMProperty(::GetKeyboardLayout(0));
 
 	UE_LOG(LogWindowsTextInputMethodSystem, Verbose, TEXT("Initialized IMM!"));
@@ -184,13 +185,13 @@ bool FWindowsTextInputMethodSystem::InitializeIMM()
 
 void FWindowsTextInputMethodSystem::UpdateIMMProperty(HKL KeyboardLayoutHandle)
 {
-	IMEProperties = ::ImmGetProperty(KeyboardLayoutHandle, IGP_PROPERTY);
+	IMMProperties = ::ImmGetProperty(KeyboardLayoutHandle, IGP_PROPERTY);
 }
 
 bool FWindowsTextInputMethodSystem::ShouldDrawIMMCompositionString() const
 {
 	// If the IME doesn't have any kind of special UI and it draws the composition window at the caret, we can draw it ourselves.
-	return !(IMEProperties & IME_PROP_SPECIAL_UI) && (IMEProperties & IME_PROP_AT_CARET);
+	return !(IMMProperties & IME_PROP_SPECIAL_UI) && (IMMProperties & IME_PROP_AT_CARET);
 }
 
 void FWindowsTextInputMethodSystem::UpdateIMMWindowPositions(HIMC IMMContext)
@@ -449,6 +450,55 @@ void FWindowsTextInputMethodSystem::Terminate()
 	}
 
 	TSFThreadManager.Reset();
+
+	::ImmDestroyContext(IMMContextId);
+}
+
+void FWindowsTextInputMethodSystem::ClearStaleWindowHandles()
+{
+	auto KnownWindowsTmp = KnownWindows;
+	for(const TWeakPtr<FGenericWindow>& Window : KnownWindowsTmp)
+	{
+		if(!Window.IsValid())
+		{
+			KnownWindows.Remove(Window);
+		}
+	}
+}
+
+void FWindowsTextInputMethodSystem::ApplyDefaults(const TSharedRef<FGenericWindow>& InWindow)
+{
+	ClearStaleWindowHandles();
+	KnownWindows.Add(InWindow);
+
+	const HWND Hwnd = reinterpret_cast<HWND>(InWindow->GetOSWindowHandle());
+
+	// Typically we disable the IME when a new window is created, however, we don't want to do that if we have an active context, 
+	// as it will immediately disable the IME (which may be undesirable for things like a search suggestions window appearing)
+	// In that case, we set the window to use the same IME context as is currently active. If the window actually takes focus
+	// away from the currently active IME context, then that will be taken care of in DeactivateContext, and all windows using the
+	// IME context will be disabled
+	ITfDocumentMgr* TSFDocumentManagerToSet = nullptr;
+	HIMC IMMContextToSet = nullptr;
+	if(ActiveContext.IsValid())
+	{
+		TSFThreadManager->GetFocus(&TSFDocumentManagerToSet);
+		IMMContextToSet = IMMContextId;
+	}
+
+	// TSF Implementation
+	if(TSFDocumentManagerToSet)
+	{
+		TSFThreadManager->SetFocus(TSFDocumentManagerToSet);
+	}
+	else
+	{
+		ITfDocumentMgr* Unused;
+		TSFThreadManager->AssociateFocus(Hwnd, TSFDisabledDocumentManager, &Unused);
+	}
+
+	// IMM Implementation
+	::ImmAssociateContext(Hwnd, IMMContextToSet);
 }
 
 TSharedPtr<ITextInputMethodChangeNotifier> FWindowsTextInputMethodSystem::RegisterContext(const TSharedRef<ITextInputMethodContext>& Context)
@@ -511,6 +561,8 @@ void FWindowsTextInputMethodSystem::UnregisterContext(const TSharedRef<ITextInpu
 
 	HRESULT Result;
 
+	check(ActiveContext != Context);
+
 	check(ContextToInternalContextMap.Contains(Context));
 	FInternalContext& InternalContext = ContextToInternalContextMap[Context];
 
@@ -539,16 +591,17 @@ void FWindowsTextInputMethodSystem::ActivateContext(const TSharedRef<ITextInputM
 	FInternalContext& InternalContext = ContextToInternalContextMap[Context];
 
 	const TSharedPtr<FGenericWindow> GenericWindow = Context->GetWindow();
-	const HWND Hwnd = GenericWindow.IsValid() ? reinterpret_cast<HWND>(GenericWindow->GetOSWindowHandle()) : nullptr;
+	InternalContext.WindowHandle = GenericWindow.IsValid() ? reinterpret_cast<HWND>(GenericWindow->GetOSWindowHandle()) : nullptr;
 
 	// IMM Implementation
 	InternalContext.IMMContext.IsComposing = false;
 	InternalContext.IMMContext.IsDeactivating = false;
+	::ImmAssociateContext(InternalContext.WindowHandle, IMMContextId);
 
 	// TSF Implementation
 	TComPtr<FTextStoreACP>& TextStore = InternalContext.TSFContext;
 	ITfDocumentMgr* Unused;
-	Result = TSFThreadManager->AssociateFocus(Hwnd, TextStore->TSFDocumentManager, &Unused);
+	Result = TSFThreadManager->AssociateFocus(InternalContext.WindowHandle, TextStore->TSFDocumentManager, &Unused);
 	if(FAILED(Result))
 	{
 		UE_LOG(LogWindowsTextInputMethodSystem, Error, TEXT("Activating a context failed while setting focus on a TSF document manager."));
@@ -560,42 +613,38 @@ void FWindowsTextInputMethodSystem::ActivateContext(const TSharedRef<ITextInputM
 void FWindowsTextInputMethodSystem::DeactivateContext(const TSharedRef<ITextInputMethodContext>& Context)
 {
 	UE_LOG(LogWindowsTextInputMethodSystem, Verbose, TEXT("Deactivating context %p..."), &(Context.Get()));
-	HRESULT Result;
-
+	
 	FInternalContext& InternalContext = ContextToInternalContextMap[Context];
 
-	const TSharedPtr<FGenericWindow> GenericWindow = Context->GetWindow();
-	const HWND Hwnd = GenericWindow.IsValid() ? reinterpret_cast<HWND>(GenericWindow->GetOSWindowHandle()) : nullptr;
-
-	// TSF Implementation
-	TComPtr<FTextStoreACP> TextStore = InternalContext.TSFContext;
+	// IMM Implementation
+	HIMC IMMContext = ::ImmGetContext(InternalContext.WindowHandle);
+	InternalContext.IMMContext.IsDeactivating = true;
+	// Request the composition is completed to ensure that the composition input UI is closed, and that a WM_IME_ENDCOMPOSITION message is sent
+	::ImmNotifyIME(IMMContext, NI_COMPOSITIONSTR, CPS_COMPLETE, 0);
+	::ImmReleaseContext(InternalContext.WindowHandle, IMMContext);
 
-	ITfDocumentMgr* FocusedDocumentManger = nullptr;
-	Result = TSFThreadManager->GetFocus(&(FocusedDocumentManger));
-	if(FAILED(Result))
-	{
-		UE_LOG(LogWindowsTextInputMethodSystem, Error, TEXT("Deactivating a context failed while getting the currently focused document manager."));
-	}
-	else if(FocusedDocumentManger == TextStore->TSFDocumentManager)
+	// Ensure all known windows are associated with a disabled IME
+	ClearStaleWindowHandles();
+	for(const TWeakPtr<FGenericWindow>& Window : KnownWindows)
 	{
-		ITfDocumentMgr* Unused;
-		Result = TSFThreadManager->AssociateFocus(Hwnd, TSFDisabledDocumentManager, &Unused);
-		if(FAILED(Result))
+		TSharedPtr<FGenericWindow> WindowPtr = Window.Pin();
+		if(WindowPtr.IsValid())
 		{
-			UE_LOG(LogWindowsTextInputMethodSystem, Error, TEXT("Deactivating a context failed while setting focus to the disabled TSF document manager."));
+			const HWND Hwnd = reinterpret_cast<HWND>(WindowPtr->GetOSWindowHandle());
+			if(Hwnd)
+			{
+				//  TSF Implementation
+				ITfDocumentMgr* Unused;
+				if(FAILED(TSFThreadManager->AssociateFocus(Hwnd, TSFDisabledDocumentManager, &Unused)))
+				{
+					UE_LOG(LogWindowsTextInputMethodSystem, Error, TEXT("Deactivating a context failed while setting focus to the disabled TSF document manager."));
+				}
+
+				// IMM Implementation
+				::ImmAssociateContext(InternalContext.WindowHandle, nullptr);
+			}
 		}
 	}
-	else
-	{
-		UE_LOG(LogWindowsTextInputMethodSystem, Verbose, TEXT("Deactivating a context ignored because it is not the focused TSF document manager."));
-	}
-
-	// IMM Implementation
-	HIMC IMMContext = ::ImmGetContext(Hwnd);
-	InternalContext.IMMContext.IsDeactivating = true;
-	// Request the composition is completed to ensure that the composition input UI is closed, and that a WM_IME_ENDCOMPOSITION message is sent
-	::ImmNotifyIME(IMMContext, NI_COMPOSITIONSTR, CPS_COMPLETE, 0);
-	::ImmReleaseContext(Hwnd, IMMContext);
 
 	// General Implementation
 	ActiveContext = nullptr;
@@ -603,6 +652,11 @@ void FWindowsTextInputMethodSystem::DeactivateContext(const TSharedRef<ITextInpu
 	UE_LOG(LogWindowsTextInputMethodSystem, Verbose, TEXT("Deactivated context %p!"), &(Context.Get()));
 }
 
+bool FWindowsTextInputMethodSystem::IsActiveContext(const TSharedRef<ITextInputMethodContext>& Context) const
+{
+	return ActiveContext == Context;
+}
+
 void FWindowsTextInputMethodSystem::OnIMEActivationStateChanged(const bool bIsEnabled)
 {
 	if(bIsEnabled)
@@ -627,7 +681,8 @@ void FWindowsTextInputMethodSystem::OnIMEActivationStateChanged(const bool bIsEn
 	}
 	else
 	{
-		check(CurrentAPI != EAPI::Unknown);
+		// It seems that switching away from an IMM based IME doesn't generate a deactivation notification
+		//check(CurrentAPI != EAPI::Unknown);
 
 		CurrentAPI = EAPI::Unknown;
 		UE_LOG(LogWindowsTextInputMethodSystem, Verbose, TEXT("IME system now deactivated."));
@@ -687,13 +742,6 @@ int32 FWindowsTextInputMethodSystem::ProcessMessage(HWND hwnd, uint32 msg, WPARA
 				
 				UE_LOG(LogWindowsTextInputMethodSystem, Verbose, TEXT("Beginning IMM composition."));
 			}
-			else
-			{
-				// No valid context, so don't allow this composition to start
-				HIMC IMMContext = ::ImmGetContext(hwnd);
-				::ImmNotifyIME(IMMContext, NI_COMPOSITIONSTR, CPS_CANCEL, 0);
-				::ImmReleaseContext(hwnd, IMMContext);
-			}
 
 			return DefWindowProc(hwnd, msg, wParam, lParam);
 		}
@@ -720,20 +768,11 @@ int32 FWindowsTextInputMethodSystem::ProcessMessage(HWND hwnd, uint32 msg, WPARA
 				{
 					UE_LOG(LogWindowsTextInputMethodSystem, Verbose, TEXT("WM_IME_COMPOSITION Composition Canceled"));
 
-					// We might have been canceled from a mouse move, so we need to take the current selection so we can restore it properly
-					// otherwise calling SetTextInRange can cause the cursor/selection to jump around
-					uint32 SelectionBeginIndex = 0;
-					uint32 SelectionLength = 0;
-					ITextInputMethodContext::ECaretPosition SelectionCaretPosition = ITextInputMethodContext::ECaretPosition::Ending;
-					ActiveContext->GetSelectionRange(SelectionBeginIndex, SelectionLength, SelectionCaretPosition);
-
 					// Clear Composition
 					ActiveContext->SetTextInRange(InternalContext.IMMContext.CompositionBeginIndex, InternalContext.IMMContext.CompositionLength, TEXT(""));
 					InternalContext.IMMContext.CompositionLength = 0;
 					ActiveContext->UpdateCompositionRange(InternalContext.IMMContext.CompositionBeginIndex, 0);
-
-					// Restore any previous selection
-					ActiveContext->SetSelectionRange(SelectionBeginIndex, SelectionLength, SelectionCaretPosition);
+					ActiveContext->SetSelectionRange(InternalContext.IMMContext.CompositionBeginIndex, 0, ITextInputMethodContext::ECaretPosition::Beginning);
 
 					EndIMMComposition();
 				}
diff --git a/Engine/Source/Runtime/Core/Public/GenericPlatform/ITextInputMethodSystem.h b/Engine/Source/Runtime/Core/Public/GenericPlatform/ITextInputMethodSystem.h
--- a/Engine/Source/Runtime/Core/Public/GenericPlatform/ITextInputMethodSystem.h
+++ b/Engine/Source/Runtime/Core/Public/GenericPlatform/ITextInputMethodSystem.h
@@ -153,6 +153,11 @@ public:
 class ITextInputMethodSystem
 {
 public:
+	/**
+	 * Called when a window is first created to allow the text input system to apply any default settings
+	 */
+	virtual void ApplyDefaults(const TSharedRef<FGenericWindow>& InWindow) = 0;
+
 	/**
 	 * Registers an implemented context interface object with the system to receive callbacks and provides an implemented
 	 * notifier interface object. Editable text should call this method after constructing an implemented context 
@@ -185,4 +190,13 @@ public:
 	 * @param	Context	The context to be activated.
 	 */
 	virtual void DeactivateContext(const TSharedRef<ITextInputMethodContext>& Context) = 0;
-};
\ No newline at end of file
+
+	/**
+	 * Test to see whether the provided context is the currently active context
+	 *
+	 * @param	Context	The context to be checked.
+	 *
+	 * @return	True if the context is active, false otherwise
+	 */
+	virtual bool IsActiveContext(const TSharedRef<ITextInputMethodContext>& Context) const = 0;
+};
diff --git a/Engine/Source/Runtime/Core/Public/Mac/CocoaTextView.h b/Engine/Source/Runtime/Core/Public/Mac/CocoaTextView.h
--- a/Engine/Source/Runtime/Core/Public/Mac/CocoaTextView.h
+++ b/Engine/Source/Runtime/Core/Public/Mac/CocoaTextView.h
@@ -15,6 +15,7 @@
 
 - (void)activateInputMethod:(const TSharedRef<ITextInputMethodContext>&)InContext;
 - (void)deactivateInputMethod;
+- (bool)isActiveInputMethod:(const TSharedRef<ITextInputMethodContext>&)InContext;
 
 - (bool)imkKeyDown:(NSEvent*)Event;
 
diff --git a/Engine/Source/Runtime/Core/Public/Mac/MacTextInputMethodSystem.h b/Engine/Source/Runtime/Core/Public/Mac/MacTextInputMethodSystem.h
--- a/Engine/Source/Runtime/Core/Public/Mac/MacTextInputMethodSystem.h
+++ b/Engine/Source/Runtime/Core/Public/Mac/MacTextInputMethodSystem.h
@@ -12,10 +12,12 @@ public:
 	void Terminate();
 	
 	// ITextInputMethodSystem Interface Begin
+	virtual void ApplyDefaults(const TSharedRef<FGenericWindow>& InWindow) override;
 	virtual TSharedPtr<ITextInputMethodChangeNotifier> RegisterContext(const TSharedRef<ITextInputMethodContext>& Context) override;
 	virtual void UnregisterContext(const TSharedRef<ITextInputMethodContext>& Context) override;
 	virtual void ActivateContext(const TSharedRef<ITextInputMethodContext>& Context) override;
 	virtual void DeactivateContext(const TSharedRef<ITextInputMethodContext>& Context) override;
+	virtual bool IsActiveContext(const TSharedRef<ITextInputMethodContext>& Context) const override;
 	// ITextInputMethodSystem Interface End
 	
 private:
diff --git a/Engine/Source/Runtime/Core/Public/Windows/WindowsTextInputMethodSystem.h b/Engine/Source/Runtime/Core/Public/Windows/WindowsTextInputMethodSystem.h
--- a/Engine/Source/Runtime/Core/Public/Windows/WindowsTextInputMethodSystem.h
+++ b/Engine/Source/Runtime/Core/Public/Windows/WindowsTextInputMethodSystem.h
@@ -63,10 +63,12 @@ public:
 	void Terminate();
 
 	// ITextInputMethodSystem Interface Begin
+	virtual void ApplyDefaults(const TSharedRef<FGenericWindow>& InWindow) override;
 	virtual TSharedPtr<ITextInputMethodChangeNotifier> RegisterContext(const TSharedRef<ITextInputMethodContext>& Context) override;
 	virtual void UnregisterContext(const TSharedRef<ITextInputMethodContext>& Context) override;
 	virtual void ActivateContext(const TSharedRef<ITextInputMethodContext>& Context) override;
 	virtual void DeactivateContext(const TSharedRef<ITextInputMethodContext>& Context) override;
+	virtual bool IsActiveContext(const TSharedRef<ITextInputMethodContext>& Context) const override;
 	// ITextInputMethodSystem Interface End
 
 	int32 ProcessMessage(HWND hwnd, uint32 msg, WPARAM wParam, LPARAM lParam);
@@ -85,6 +87,8 @@ private:
 
 	void OnIMEActivationStateChanged(const bool bIsEnabled);
 
+	void ClearStaleWindowHandles();
+
 private:
 	TSharedPtr<ITextInputMethodContext> ActiveContext;
 	enum class EAPI
@@ -104,7 +108,19 @@ private:
 
 	struct FInternalContext
 	{
+		FInternalContext()
+			: WindowHandle(nullptr)
+		{
+			IMMContext.IsComposing = false;
+			IMMContext.IsDeactivating = false;
+			IMMContext.CompositionBeginIndex = 0;
+			IMMContext.CompositionLength = 0;
+		}
+
+		HWND WindowHandle;
+
 		TComPtr<FTextStoreACP> TSFContext;
+
 		struct
 		{
 			bool IsComposing;
@@ -116,7 +132,10 @@ private:
 	TMap< TWeakPtr<ITextInputMethodContext>, FInternalContext > ContextToInternalContextMap;
 
 	// IMM Implementation
-	DWORD IMEProperties;
+	HIMC IMMContextId;
+	DWORD IMMProperties;
+
+	TSet<TWeakPtr<FGenericWindow>> KnownWindows;
 };
 
 #include "HideWindowsPlatformTypes.h"
\ No newline at end of file
diff --git a/Engine/Source/Runtime/Slate/Private/Framework/Application/SlateApplication.cpp b/Engine/Source/Runtime/Slate/Private/Framework/Application/SlateApplication.cpp
--- a/Engine/Source/Runtime/Slate/Private/Framework/Application/SlateApplication.cpp
+++ b/Engine/Source/Runtime/Slate/Private/Framework/Application/SlateApplication.cpp
@@ -1066,6 +1066,12 @@ TSharedRef< FGenericWindow > FSlateApplication::MakeWindow( TSharedRef<SWindow>
 
 	PlatformApplication->InitializeWindow( NewWindow, Definition, NativeParent, bShowImmediately );
 
+	ITextInputMethodSystem* const TextInputMethodSystem = PlatformApplication->GetTextInputMethodSystem();
+	if ( TextInputMethodSystem )
+	{
+		TextInputMethodSystem->ApplyDefaults( NewWindow );
+	}
+
 	return NewWindow;
 }
 
diff --git a/Engine/Source/Runtime/Slate/Private/Widgets/Input/SEditableText.cpp b/Engine/Source/Runtime/Slate/Private/Widgets/Input/SEditableText.cpp
--- a/Engine/Source/Runtime/Slate/Private/Widgets/Input/SEditableText.cpp
+++ b/Engine/Source/Runtime/Slate/Private/Widgets/Input/SEditableText.cpp
@@ -38,7 +38,14 @@ SEditableText::~SEditableText()
 	ITextInputMethodSystem* const TextInputMethodSystem = FSlateApplication::IsInitialized() ? FSlateApplication::Get().GetTextInputMethodSystem() : nullptr;
 	if(TextInputMethodSystem)
 	{
-		TextInputMethodSystem->UnregisterContext(TextInputMethodContext.ToSharedRef());
+		TSharedRef<FTextInputMethodContext> TextInputMethodContextRef = TextInputMethodContext.ToSharedRef();
+		if(TextInputMethodSystem->IsActiveContext(TextInputMethodContextRef))
+		{
+			// This can happen if an entire tree of widgets is culled, as Slate isn't notified of the focus loss, the widget is just deleted
+			TextInputMethodSystem->DeactivateContext(TextInputMethodContextRef);
+		}
+
+		TextInputMethodSystem->UnregisterContext(TextInputMethodContextRef);
 	}
 }
 
diff --git a/Engine/Source/Runtime/Slate/Private/Widgets/Text/SMultiLineEditableText.cpp b/Engine/Source/Runtime/Slate/Private/Widgets/Text/SMultiLineEditableText.cpp
--- a/Engine/Source/Runtime/Slate/Private/Widgets/Text/SMultiLineEditableText.cpp
+++ b/Engine/Source/Runtime/Slate/Private/Widgets/Text/SMultiLineEditableText.cpp
@@ -200,7 +200,14 @@ SMultiLineEditableText::~SMultiLineEditableText()
 	ITextInputMethodSystem* const TextInputMethodSystem = FSlateApplication::IsInitialized() ? FSlateApplication::Get().GetTextInputMethodSystem() : nullptr;
 	if(TextInputMethodSystem)
 	{
-		TextInputMethodSystem->UnregisterContext(TextInputMethodContext.ToSharedRef());
+		TSharedRef<FTextInputMethodContext> TextInputMethodContextRef = TextInputMethodContext.ToSharedRef();
+		if(TextInputMethodSystem->IsActiveContext(TextInputMethodContextRef))
+		{
+			// This can happen if an entire tree of widgets is culled, as Slate isn't notified of the focus loss, the widget is just deleted
+			TextInputMethodSystem->DeactivateContext(TextInputMethodContextRef);
+		}
+
+		TextInputMethodSystem->UnregisterContext(TextInputMethodContextRef);
 	}
 }
 
