Improvements to Mac crash reporting, including more of the relevant crash data, erasing meaningless stack entries & fixing some bugs.

diff --git a/Engine/Source/Runtime/Core/Private/Apple/ApplePlatformCrashContext.cpp b/Engine/Source/Runtime/Core/Private/Apple/ApplePlatformCrashContext.cpp
--- a/Engine/Source/Runtime/Core/Private/Apple/ApplePlatformCrashContext.cpp
+++ b/Engine/Source/Runtime/Core/Private/Apple/ApplePlatformCrashContext.cpp
@@ -72,7 +72,7 @@ int32 FApplePlatformCrashContext::ReportCrash() const
 		*StackBuffer = 0;
 		
 		// Walk the stack and dump it to the allocated memory (ignore first 2 callstack lines as those are in stack walking code)
-		FPlatformStackWalk::StackWalkAndDump( StackBuffer, ARRAY_COUNT(MinidumpCallstackInfo) - 1, 2, Context );
+		FPlatformStackWalk::StackWalkAndDump( StackBuffer, ARRAY_COUNT(MinidumpCallstackInfo) - 1, 6, Context );
 		
 #if WITH_EDITORONLY_DATA
 		FCString::Strncat( GErrorHist, ANSI_TO_TCHAR(MinidumpCallstackInfo), ARRAY_COUNT(GErrorHist) - 1 );
@@ -167,9 +167,16 @@ ANSICHAR* FApplePlatformCrashContext::ItoANSI(uint64 Val, uint64 Base)
 	
 	uint64 i = 62;
 	
-	for(; Val && i ; --i, Val /= Base)
+	if(Val)
 	{
-		InternalBuffer[i] = "0123456789abcdef"[Val % Base];
+		for(; Val && i ; --i, Val /= Base)
+		{
+			InternalBuffer[i] = "0123456789abcdef"[Val % Base];
+		}
+	}
+	else
+	{
+		InternalBuffer[i--] = '0';
 	}
 	
 	return &InternalBuffer[i+1];
@@ -181,9 +188,16 @@ TCHAR* FApplePlatformCrashContext::ItoTCHAR(uint64 Val, uint64 Base)
 	
 	uint64 i = 62;
 	
-	for(; Val && i ; --i, Val /= Base)
+	if(Val)
+	{
+		for(; Val && i ; --i, Val /= Base)
+		{
+			InternalBuffer[i] = TEXT("0123456789abcdef")[Val % Base];
+		}
+	}
+	else
 	{
-		InternalBuffer[i] = TEXT("0123456789abcdef")[Val % Base];
+		InternalBuffer[i--] = TEXT('0');
 	}
 	
 	return &InternalBuffer[i+1];
diff --git a/Engine/Source/Runtime/Core/Private/Apple/ApplePlatformStackWalk.cpp b/Engine/Source/Runtime/Core/Private/Apple/ApplePlatformStackWalk.cpp
--- a/Engine/Source/Runtime/Core/Private/Apple/ApplePlatformStackWalk.cpp
+++ b/Engine/Source/Runtime/Core/Private/Apple/ApplePlatformStackWalk.cpp
@@ -260,22 +260,24 @@ void NewReportEnsure( const TCHAR* ErrorMessage )
 int32 ReportCrash(ucontext_t *Context, int32 Signal, struct __siginfo* Info)
 {
 	static bool GAlreadyCreatedMinidump = false;
+	
 	// Only create a minidump the first time this function is called.
 	// (Can be called the first time from the RenderThread, then a second time from the MainThread.)
 	if ( GAlreadyCreatedMinidump == false )
 	{
 		GAlreadyCreatedMinidump = true;
 
+		// No malloc in a signal handler as it is unsafe & will deadlock the application!
 		const SIZE_T StackTraceSize = 65535;
-		ANSICHAR* StackTrace = (ANSICHAR*) FMemory::Malloc( StackTraceSize );
+		static ANSICHAR StackTrace[ StackTraceSize ];
 		StackTrace[0] = 0;
+		
 		// Walk the stack and dump it to the allocated memory.
 		FPlatformStackWalk::StackWalkAndDump( StackTrace, StackTraceSize, 0, Context );
 #if WITH_EDITORONLY_DATA
         FCString::Strncat( GErrorHist, ANSI_TO_TCHAR(StackTrace), ARRAY_COUNT(GErrorHist) - 1 );
 		CreateExceptionInfoString(Signal, Info);
 #endif
-		FMemory::Free( StackTrace );
 	}
 
 	return 0;
diff --git a/Engine/Source/Runtime/Core/Private/Mac/MacPlatformMisc.cpp b/Engine/Source/Runtime/Core/Private/Mac/MacPlatformMisc.cpp
--- a/Engine/Source/Runtime/Core/Private/Mac/MacPlatformMisc.cpp
+++ b/Engine/Source/Runtime/Core/Private/Mac/MacPlatformMisc.cpp
@@ -14,6 +14,7 @@
 #include "Runtime/Launch/Resources/Version.h"
 #include "EngineVersion.h"
 
+#include <dlfcn.h>
 #include <IOKit/IOKitLib.h>
 #include <IOKit/network/IOEthernetInterface.h>
 #include <IOKit/network/IONetworkInterface.h>
@@ -156,6 +157,8 @@ struct MacApplicationInfo
 		
 		LCID = FString::Printf(TEXT("%d"), FInternationalization::Get().GetCurrentCulture()->GetLCID());
 		
+		PrimaryGPU = FPlatformMisc::GetPrimaryGPUBrand();
+		
 		// Notification handler to check we are running from a battery - this only applies to MacBook's.
 		notify_handler_t PowerSourceNotifyHandler = ^(int32 Token){
 			RunningOnBattery = false;
@@ -209,6 +212,7 @@ struct MacApplicationInfo
 	FString LCID;
 	FString CommandLine;
 	FString BranchBaseDir;
+	FString PrimaryGPU;
 };
 static MacApplicationInfo GMacAppInfo;
 
@@ -1223,7 +1227,31 @@ void FMacCrashContext::GenerateReport(char const* DiagnosticsPath) const
 		for(uint32 Index = 0; Index < ModuleCount; Index++)
 		{
 			char const* ModulePath = _dyld_get_image_name(Index);
-			WriteLine(ReportFile, ModulePath);
+			uint32 Version = 0;
+			struct mach_header_64* Header = (struct mach_header_64*)_dyld_get_image_header(Index);
+			struct load_command *CurrentCommand = (struct load_command *)( (char *)Header + sizeof(struct mach_header_64) );
+			if( Header->magic == MH_MAGIC_64 )
+			{
+				for( int32 i = 0; i < Header->ncmds; i++ )
+				{
+					if( CurrentCommand->cmd == LC_LOAD_DYLIB )
+					{
+						struct dylib_command *DylibCommand = (struct dylib_command *) CurrentCommand;
+						Version = DylibCommand->dylib.current_version;
+						Version = ((Version & 0xff) + ((Version >> 8) & 0xff) * 100 + ((Version >> 16) & 0xffff) * 10000);
+						break;
+					}
+					
+					CurrentCommand = (struct load_command *)( (char *)CurrentCommand + CurrentCommand->cmdsize );
+				}
+			}
+			
+			// Add version information (very useful!)
+			FCStringAnsi::Strncpy(Line, ModulePath, PATH_MAX);
+			FCStringAnsi::Strcat(Line, PATH_MAX, " (");
+			FCStringAnsi::Strcat(Line, PATH_MAX, ItoANSI(Version, 10));
+			FCStringAnsi::Strcat(Line, PATH_MAX, ")");
+			WriteLine(ReportFile, Line);
 		}
 		WriteLine(ReportFile, "<MODULES END>");
 		WriteLine(ReportFile);
@@ -1239,6 +1267,8 @@ void FMacCrashContext::GenerateWindowsErrorReport(char const* WERPath) const
 	int ReportFile = open(WERPath, O_CREAT|O_WRONLY, 0766);
 	if (ReportFile != -1)
 	{
+		TCHAR Line[PATH_MAX] = {};
+		
 		// write BOM
 		static uint16 ByteOrderMarker = 0xFEFF;
 		write(ReportFile, &ByteOrderMarker, sizeof(ByteOrderMarker));
@@ -1306,13 +1336,75 @@ void FMacCrashContext::GenerateWindowsErrorReport(char const* WERPath) const
 		WriteUTF16String(ReportFile, ItoTCHAR(ENGINE_VERSION_LOWORD, 10));
 		WriteLine(ReportFile, TEXT("</Parameter1>"));
 
+		// App time stamp
 		WriteLine(ReportFile, TEXT("\t\t<Parameter2>528f2d37</Parameter2>"));													// FIXME: supply valid?
-		WriteLine(ReportFile, TEXT("\t\t<Parameter3>KERNELBASE.dll</Parameter3>"));												// FIXME: supply valid?
-		WriteLine(ReportFile, TEXT("\t\t<Parameter4>6.1.7601.18015</Parameter4>"));												// FIXME: supply valid?
-		WriteLine(ReportFile, TEXT("\t\t<Parameter5>50b8479b</Parameter5>"));													// FIXME: supply valid?
-		WriteLine(ReportFile, TEXT("\t\t<Parameter6>00000001</Parameter6>"));													// FIXME: supply valid?
-		WriteLine(ReportFile, TEXT("\t\t<Parameter7>0000000000009E5D</Parameter7>"));											// FIXME: supply valid?
-		WriteLine(ReportFile, TEXT("\t\t<Parameter8>!!</Parameter8>"));															// FIXME: supply valid?
+		
+		uint64 BackTrace[7] = {0, 0, 0, 0, 0, 0, 0};
+		FPlatformStackWalk::CaptureStackBackTrace( BackTrace, 7, Context );
+		
+		if(BackTrace[6] != 0)
+		{
+			Dl_info Info;
+			dladdr((void*)BackTrace[6], &Info);
+			
+			// Crash Module name
+			WriteUTF16String(ReportFile, TEXT("\t\t<Parameter3>"));
+			for(uint32 i = 0; i < PATH_MAX && (i < FCStringAnsi::Strlen(Info.dli_fname) + 1); i++)
+			{
+				Line[i] = Info.dli_fname[i];
+			}
+			WriteLine(ReportFile, Line);
+			WriteLine(ReportFile, TEXT("</Parameter3>"));
+			
+			// Check header
+			uint32 Version = 0;
+			uint32 TimeStamp = 0;
+			struct mach_header_64* Header = (struct mach_header_64*)Info.dli_fbase;
+			struct load_command *CurrentCommand = (struct load_command *)( (char *)Header + sizeof(struct mach_header_64) );
+			if( Header->magic == MH_MAGIC_64 )
+			{
+				for( int32 i = 0; i < Header->ncmds; i++ )
+				{
+					if( CurrentCommand->cmd == LC_LOAD_DYLIB )
+					{
+						struct dylib_command *DylibCommand = (struct dylib_command *) CurrentCommand;
+						Version = DylibCommand->dylib.current_version;
+						TimeStamp = DylibCommand->dylib.timestamp;
+						Version = ((Version & 0xff) + ((Version >> 8) & 0xff) * 100 + ((Version >> 16) & 0xffff) * 10000);
+						break;
+					}
+					
+					CurrentCommand = (struct load_command *)( (char *)CurrentCommand + CurrentCommand->cmdsize );
+				}
+			}
+			
+			// Module version
+			WriteUTF16String(ReportFile, TEXT("\t\t<Parameter4>"));
+			WriteUTF16String(ReportFile, ItoTCHAR(Version, 10));
+			WriteLine(ReportFile, TEXT("</Parameter4>"));
+			
+			// Module time stamp
+			WriteUTF16String(ReportFile, TEXT("\t\t<Parameter5>"));
+			WriteUTF16String(ReportFile, ItoTCHAR(TimeStamp, 16));
+			WriteLine(ReportFile, TEXT("</Parameter5>"));
+			
+			// MethodDef token -> no equivalent
+			WriteLine(ReportFile, TEXT("\t\t<Parameter6>00000001</Parameter6>"));
+			
+			// IL Offset -> Function pointer
+			WriteUTF16String(ReportFile, TEXT("\t\t<Parameter7>"));
+			WriteUTF16String(ReportFile, ItoTCHAR(BackTrace[6], 16));
+			WriteLine(ReportFile, TEXT("</Parameter7>"));
+		}
+		
+		// Fault type -> Signal
+		WriteUTF16String(ReportFile, TEXT("\t\t<Parameter8>"));
+		for(uint32 i = 0; i < PATH_MAX && (i < FCStringAnsi::Strlen(SignalDescription) + 1); i++)
+		{
+			Line[i] = SignalDescription[i];
+		}
+		WriteLine(ReportFile, Line);
+		WriteLine(ReportFile, TEXT("</Parameter8>"));
 		
 		WriteUTF16String(ReportFile, TEXT("\t\t<Parameter9>"));
 		WriteUTF16String(ReportFile, *GMacAppInfo.BranchBaseDir);
@@ -1337,7 +1429,7 @@ void FMacCrashContext::GenerateWindowsErrorReport(char const* WERPath) const
 		WriteUTF16String(ReportFile, *GMacAppInfo.MachineUUID);
 		WriteLine(ReportFile, TEXT("</MID>"));
 		
-		WriteLine(ReportFile, TEXT("\t\t<SystemManufacturer>Apple Inc.</SystemManufacturer>"));						// FIXME: supply valid?
+		WriteLine(ReportFile, TEXT("\t\t<SystemManufacturer>Apple Inc.</SystemManufacturer>"));
 		
 		WriteUTF16String(ReportFile, TEXT("\t\t<SystemProductName>"));
 		WriteUTF16String(ReportFile, *GMacAppInfo.MachineModel);
@@ -1347,7 +1439,12 @@ void FMacCrashContext::GenerateWindowsErrorReport(char const* WERPath) const
 		WriteUTF16String(ReportFile, *GMacAppInfo.BiosRelease);
 		WriteUTF16String(ReportFile, TEXT("-"));
 		WriteUTF16String(ReportFile, *GMacAppInfo.BiosRevision);
-		WriteLine(ReportFile, TEXT("</BIOSVersion>"));											// FIXME: supply valid?
+		WriteLine(ReportFile, TEXT("</BIOSVersion>"));
+		
+		WriteUTF16String(ReportFile, TEXT("\t\t<GraphicsCard>"));
+		WriteUTF16String(ReportFile, *GMacAppInfo.PrimaryGPU);
+		WriteLine(ReportFile, TEXT("</GraphicsCard>"));
+		
 		WriteLine(ReportFile, TEXT("\t</SystemInformation>"));
 		
 		WriteLine(ReportFile, TEXT("</WERReportMetadata>"));
