More correct setting of window flags (PR #404).

- Contributed by [REDACTED_PERSON].
- Does not change the behavior though.

diff --git a/Engine/Source/Runtime/Core/Private/Linux/LinuxWindow.cpp b/Engine/Source/Runtime/Core/Private/Linux/LinuxWindow.cpp
--- a/Engine/Source/Runtime/Core/Private/Linux/LinuxWindow.cpp
+++ b/Engine/Source/Runtime/Core/Private/Linux/LinuxWindow.cpp
@@ -76,13 +76,19 @@ void FLinuxWindow::Initialize( FLinuxApplication* const Application, const TShar
 	int32 WindowWidth = ClientWidth;
 	int32 WindowHeight = ClientHeight;
 
-	WindowStyle |= SDL_WINDOW_OPENGL | SDL_WINDOW_SHOWN | SDL_WINDOW_INPUT_FOCUS | SDL_WINDOW_MOUSE_FOCUS;
+	WindowStyle |= SDL_WINDOW_OPENGL | SDL_WINDOW_SHOWN;
 
-	if	( !Definition->HasOSWindowBorder )
+	if ( !Definition->HasOSWindowBorder )
 	{
 		WindowStyle |= SDL_WINDOW_BORDERLESS;
 	}
 
+	if ( Definition->AcceptsInput )
+	{
+		// @TODO: check if we need these flags at all
+		WindowStyle |= SDL_WINDOW_INPUT_FOCUS | SDL_WINDOW_MOUSE_FOCUS;
+	}
+
 	if ( Definition->HasSizingFrame )
 	{
 		WindowStyle |= SDL_WINDOW_RESIZABLE;
