Assertion to catch FPlatformMemory::Memswap() misuse.
- Suggested during code review of the previous change (CL 2279960 ).

diff --git a/Engine/Source/Runtime/Core/Private/GenericPlatform/GenericPlatformMemory.cpp b/Engine/Source/Runtime/Core/Private/GenericPlatform/GenericPlatformMemory.cpp
--- a/Engine/Source/Runtime/Core/Private/GenericPlatform/GenericPlatformMemory.cpp
+++ b/Engine/Source/Runtime/Core/Private/GenericPlatform/GenericPlatformMemory.cpp
@@ -149,6 +149,10 @@ void FGenericPlatformMemory::Memswap( void* Ptr1, void* Ptr2, SIZE_T Size )
 {
 	if (Ptr1 != Ptr2)
 	{
+		// check that Ptr1 and Ptr2 do not overlap in undefined ways
+		checkf(reinterpret_cast<uint8 *>(Ptr1)+Size <= reinterpret_cast<uint8 *>(Ptr2) || reinterpret_cast<uint8 *>(Ptr2)+Size <= reinterpret_cast<uint8 *>(Ptr1),
+			TEXT("Pointers given to FPlatformMemory::Memswap() point to overlapping memory areas, results are undefined."));
+
 		void* Temp = FMemory_Alloca(Size);
 		FPlatformMemory::Memcpy( Temp, Ptr1, Size );
 		FPlatformMemory::Memcpy( Ptr1, Ptr2, Size );
