Moved CSV parser from a developer module, to the Engine/Core/Serialization/Csv.

This functionality is required inside the engine in more uses than we originally thought, so has been moved accordingly.

diff --git a/Engine/Source/Developer/CsvParser/CsvParser.Build.cs b/Engine/Source/Developer/CsvParser/CsvParser.Build.cs
deleted file mode 100644
diff --git a/Engine/Source/Developer/CsvParser/Private/CsvParser.cpp b/Engine/Source/Developer/CsvParser/Private/CsvParser.cpp
deleted file mode 100644
diff --git a/Engine/Source/Developer/CsvParser/Private/CsvParserPrivatePCH.h b/Engine/Source/Developer/CsvParser/Private/CsvParserPrivatePCH.h
deleted file mode 100644
diff --git a/Engine/Source/Developer/CsvParser/Private/CsvParserTests.cpp b/Engine/Source/Developer/CsvParser/Private/CsvParserTests.cpp
deleted file mode 100644
diff --git a/Engine/Source/Developer/CsvParser/Public/CsvParser.h b/Engine/Source/Developer/CsvParser/Public/CsvParser.h
deleted file mode 100644
diff --git a/Engine/Source/Runtime/Core/Private/Serialization/Csv/CsvParser.cpp b/Engine/Source/Runtime/Core/Private/Serialization/Csv/CsvParser.cpp
new file mode 100644
--- /dev/null
+++ b/Engine/Source/Runtime/Core/Private/Serialization/Csv/CsvParser.cpp
@@ -0,0 +1,141 @@
+// Copyright 1998-2014 Epic Games, Inc. All Rights Reserved.
+
+#include "Core.h"
+#include "Serialization/Csv/CsvParser.h"
+
+FCsvParser::FCsvParser(FString InSourceString)
+	: SourceString(MoveTemp(InSourceString))		
+{
+	if (!SourceString.IsEmpty())
+	{
+		ParseRows();
+	}
+}
+		
+void FCsvParser::ParseRows()
+{
+	BufferStart = &SourceString[0];
+	ReadAt = BufferStart;
+
+	EParseResult Result;
+	do
+	{
+		Result = ParseRow();
+	} while(Result != EParseResult::EndOfString);
+}
+
+FCsvParser::EParseResult FCsvParser::ParseRow()
+{
+	// Check for an empty line
+	const int8 NewLineSize = MeasureNewLine(ReadAt);
+	if (NewLineSize)
+	{
+		ReadAt += NewLineSize;
+		return *ReadAt ? EParseResult::EndOfCell : EParseResult::EndOfString;
+	}
+
+	EParseResult Result;
+
+	Rows.Emplace();
+	do
+	{
+		Result = ParseCell();
+	}
+	while(Result == EParseResult::EndOfCell);
+
+	return Result;
+}
+
+FCsvParser::EParseResult FCsvParser::ParseCell()
+{
+	TCHAR* WriteAt = BufferStart + (ReadAt - BufferStart);
+
+	// Check if this cell is quoted. Whitespace between cell opening and quote is invalid.
+	bool bQuoted = *ReadAt == '"';
+
+	if (bQuoted)
+	{
+		// Skip over the first quote
+		ReadAt = ++WriteAt;
+	}
+	
+	Rows.Last().Add(ReadAt);
+
+	while( *ReadAt )
+	{
+		// If the cell value is quoted, we swallow everything until we find a closing quote
+		if (bQuoted)
+		{
+			if (*ReadAt == '"')
+			{
+				// RFC 4180 specifies that double quotes are escaped as ""
+
+				int32 NumQuotes = 0;
+				while (*(ReadAt + NumQuotes) == '"') ++NumQuotes;
+
+				if (NumQuotes % 2 != 0)
+				{
+					// No longer quoted if there are an even number of quotes before this one
+					bQuoted = false;
+				}
+
+				ReadAt += NumQuotes;
+
+				// Unescape the double quotes
+				// We leave the write pos pointing at the trailing closing quote if present so
+				// it gets overwritten by any subsequent text in the cell
+				NumQuotes /= 2;
+				while(NumQuotes-- > 0) *(WriteAt++) = '"';
+
+				continue;
+			}
+		}
+		else
+		{
+			// Check for the end of a row (new line)
+			const int8 NewLineSize = MeasureNewLine(ReadAt);
+			if (NewLineSize != 0)
+			{
+				// Null terminate the cell
+				*WriteAt = '\0';
+				ReadAt += NewLineSize;
+
+				return *ReadAt ? EParseResult::EndOfRow : EParseResult::EndOfString;
+			}
+			else if (*ReadAt == ',')
+			{
+				*WriteAt = '\0';
+				++ReadAt;
+
+				return *ReadAt ? EParseResult::EndOfCell : EParseResult::EndOfString;
+			}
+		}
+		if (WriteAt != ReadAt)
+		{
+			(*WriteAt++) = (*ReadAt++);
+		}
+		else
+		{
+			ReadAt = ++WriteAt;
+		}
+	}
+
+	return EParseResult::EndOfString;
+}
+
+int8 FCsvParser::MeasureNewLine(const TCHAR* At)
+{
+	switch(*At)
+	{
+		case '\r':
+			if (*(At+1) == '\n')	// could be null
+			{
+				return 2;
+			}
+			return 1;
+		case '\n':
+			return 1;
+		default:
+			return 0;
+	}
+}
diff --git a/Engine/Source/Runtime/Core/Private/Serialization/Csv/CsvParserTests.cpp b/Engine/Source/Runtime/Core/Private/Serialization/Csv/CsvParserTests.cpp
new file mode 100644
--- /dev/null
+++ b/Engine/Source/Runtime/Core/Private/Serialization/Csv/CsvParserTests.cpp
@@ -0,0 +1,129 @@
+// Copyright 1998-2014 Epic Games, Inc. All Rights Reserved.
+
+#include "Core.h"
+#include "Serialization/Csv/CsvParser.h"
+
+namespace CsvParser_Tests
+{
+
+	/** Check that the specified row of cells matches what we expect */
+	template <uint32 N>
+	bool CheckRow(const TArray<const TCHAR*>& Row, const FString (&Expected)[N], int32 RowIndex, FAutomationTestBase* Test)
+	{
+		if (Row.Num() != N)
+		{
+			Test->AddError(FString::Printf(TEXT("Row %d has the wrong number of cells. Expected %d, but parsed result has %d."), RowIndex, N, Row.Num()));
+			return false;
+		}
+
+		for (int32 Cell = 0; Cell < Row.Num(); ++Cell)
+		{
+			if (FCString::Strcmp(Row[Cell], *Expected[Cell]) != 0)
+			{
+				Test->AddError(FString::Printf(TEXT("Cell %d on row %d did not parse correctly. Expected:\n%s\nBut parsed result was:\n%s"), Cell, RowIndex, *Expected[Cell], Row[Cell]));
+				return false;
+			}
+		}
+
+		return true;
+	}
+	/** Check that the specified rows all match the specified results */
+	template <uint32 RowCount, uint32 CellCount>
+	bool CheckRows(const FCsvParser::FRows& Rows, const FString (&Expected)[RowCount][CellCount], FAutomationTestBase* Test)
+	{
+		if (Rows.Num() != RowCount)
+		{
+			Test->AddError(FString::Printf(TEXT("Parser generated wrong number of rows. Expected %d, but parsed result has %d."), RowCount, Rows.Num()));
+			return false;
+		}
+
+		for (int32 Row = 0; Row < RowCount; ++Row)
+		{
+			if (!CheckRow(Rows[Row], Expected[Row], Row, Test))
+			{
+				return false;
+			}
+		}
+
+		return true;
+	}
+}
+
+IMPLEMENT_SIMPLE_AUTOMATION_TEST(FBasicTest, "Core.CSV Parser.Basic", 0)
+
+bool FBasicTest::RunTest(const FString& Parameters)
+{
+	FString Str =	"1,2,3,4\n5,6,7,8";
+	const FCsvParser Parser(MoveTemp(Str));
+
+	const FString Expected[][4] = {
+		{ "1", "2", "3", "4" },
+		{ "5", "6", "7", "8" }
+	};
+
+	return CsvParser_Tests::CheckRows(Parser.GetRows(), Expected, this);
+}
+
+IMPLEMENT_SIMPLE_AUTOMATION_TEST(FQuotedStringTest, "Core.CSV Parser.Quoted String", 0)
+
+bool FQuotedStringTest::RunTest(const FString& Parameters)
+{
+	FString Str =	"\"Quoted with nested \"\"quotes\"\", and \"\"commas\"\"\" \"unquoted due to \"whitespace\" , second \" unquoted\n"
+					"\"Quoted\nString,With\nNewlines\",\"\"\n"
+					"1,\"2\"\n";
+	const FCsvParser Parser(MoveTemp(Str));
+
+	const FString Expected[][2] = {
+		{ "Quoted with nested \"quotes\", and \"commas\" \"unquoted due to \"whitespace\" ", " second \" unquoted" },
+		{ "Quoted\nString,With\nNewlines", "" },
+		{ "1", "2" },
+	};
+
+	return CsvParser_Tests::CheckRows(Parser.GetRows(), Expected, this);
+}
+
+IMPLEMENT_SIMPLE_AUTOMATION_TEST(FEmptyCellsTest, "Core.CSV Parser.Empty Cells", 0)
+
+bool FEmptyCellsTest::RunTest(const FString& Parameters)
+{
+	FString Str =	",,,\n"
+					",,,\n"
+					" , , , \n";
+	const FCsvParser Parser(MoveTemp(Str));
+
+	const FString Expected[][4] = {
+		{ "", "", "", "" },
+		{ "", "", "", "" },
+		{ " ", " ", " ", " " },	// Whitespace is significant in csv
+	};
+
+	return CsvParser_Tests::CheckRows(Parser.GetRows(), Expected, this);
+}
+
+IMPLEMENT_SIMPLE_AUTOMATION_TEST(FMixedNewLinesTest, "Core.CSV Parser.Mixed Newlines", 0)
+
+bool FMixedNewLinesTest::RunTest(const FString& Parameters)
+{
+	const FCsvParser Parser(FString("1\r2\n3\r\n4\n\n5\r\r6\r\n\r\n7\n\r8"));
+
+	const FString Expected[][1] = {
+		{ "1" }, { "2" }, { "3" }, { "4" }, { "5" }, { "6" }, { "7" }, { "8" },
+	};
+
+	return CsvParser_Tests::CheckRows(Parser.GetRows(), Expected, this);
+}
+
+IMPLEMENT_SIMPLE_AUTOMATION_TEST(FEmptyStringTest, "Core.CSV Parser.Empty String", 0)
+
+bool FEmptyStringTest::RunTest(const FString& Parameters)
+{
+	return	FCsvParser(FString()).GetRows().Num() == 0 &&
+			FCsvParser(FString("")).GetRows().Num() == 0;
+}
+
+IMPLEMENT_SIMPLE_AUTOMATION_TEST(FOnlyNewLinesTest, "Core.CSV Parser.Only Newlines", 0)
+
+bool FOnlyNewLinesTest::RunTest(const FString& Parameters)
+{
+	return	FCsvParser(FString("\r\n\r\r\n\n\r\n\n\r")).GetRows().Num() == 0;
+}
\ No newline at end of file
diff --git a/Engine/Source/Runtime/Core/Public/Serialization/Csv/CsvParser.h b/Engine/Source/Runtime/Core/Public/Serialization/Csv/CsvParser.h
new file mode 100644
--- /dev/null
+++ b/Engine/Source/Runtime/Core/Public/Serialization/Csv/CsvParser.h
@@ -0,0 +1,63 @@
+// Copyright 1998-2014 Epic Games, Inc. All Rights Reserved.
+
+#pragma once
+
+
+/** A simple, efficient csv parser */
+struct FCsvParser
+{
+	typedef TArray<TArray<const TCHAR*>> FRows;
+
+	/**
+	 *	Construct with a string.
+	 *	Takes a copy because the parser tramples over the file to create null-terminated cell strings
+	 *	Avoid a copy by passing overship of the source string with MoveTemp()
+	 */
+	CORE_API FCsvParser(FString InSourceString);
+
+	/**
+	 *	Const access to the parsed rows
+	 *	Rows are stored as arrays of parsed, null-terminated cell-strings.
+	 *	Object (and contained strings) is valid for the lifetime of the parser.
+	 */
+	const FRows& GetRows() const
+	{
+		return Rows;
+	}
+
+private:
+	/** Non copyable - has pointers to itself */
+	FCsvParser(const FCsvParser&);
+	FCsvParser& operator=(const FCsvParser&);
+
+	/** Internal enum to control parsing progress */
+	enum struct EParseResult { EndOfCell, EndOfRow, EndOfString };
+
+	/** Parse all the rows in the source data */
+	void ParseRows();
+
+	/** Parse a single row at the current read position */
+	EParseResult ParseRow();
+
+	/** Parse a cell at the current read position */
+	EParseResult ParseCell();
+
+private:
+
+	/** Raw csv buffer, sliced and diced by the parser */
+	FString SourceString;
+
+	/** Fixed pointer to the start of the buffer */
+	TCHAR* BufferStart;
+
+	/** Current parser read position */
+	const TCHAR* ReadAt;
+
+	/** Array of cell string arrays */
+	FRows Rows;
+
+private:
+
+	/** Helper to measure the size of a new line at the current read position */
+	int8 MeasureNewLine(const TCHAR* At);
+};
diff --git a/Engine/Source/Runtime/Engine/Classes/Curves/CurveBase.h b/Engine/Source/Runtime/Engine/Classes/Curves/CurveBase.h
--- a/Engine/Source/Runtime/Engine/Classes/Curves/CurveBase.h
+++ b/Engine/Source/Runtime/Engine/Classes/Curves/CurveBase.h
@@ -476,13 +476,11 @@ public:
 
 	// Begin UCurveBase interface
 
-#if WITH_EDITOR
 	/** 
 	 *	Create curve from CSV style comma-separated string. 
 	 *	@return	Set of problems encountered while processing input
 	 */
 	TArray<FString> CreateCurveFromCSVString(const FString& InString);
-#endif
 
 	/** Reset all curve data */
 	void ResetCurve();
diff --git a/Engine/Source/Runtime/Engine/Classes/Engine/CurveTable.h b/Engine/Source/Runtime/Engine/Classes/Engine/CurveTable.h
--- a/Engine/Source/Runtime/Engine/Classes/Engine/CurveTable.h
+++ b/Engine/Source/Runtime/Engine/Classes/Engine/CurveTable.h
@@ -61,7 +61,6 @@ class UCurveTable : public UObject
 	/** Output entire contents of table as JSON */
 	ENGINE_API bool WriteTableAsJSON(const TSharedRef< TJsonWriter<TCHAR, TPrettyJsonPrintPolicy<TCHAR> > >& JsonWriter) const;
 
-#if WITH_EDITOR
 	/** 
 	 *	Create table from CSV style comma-separated string. 
 	 *	RowCurve must be defined before calling this function. 
@@ -70,7 +69,6 @@ class UCurveTable : public UObject
 	 *	@return	Set of problems encountered while processing input
 	 */
 	ENGINE_API TArray<FString> CreateTableFromCSVString(const FString& InString, ERichCurveInterpMode InterpMode = RCIM_Linear);
-#endif
 
 	/** Empty the table info (will not clear RowCurve) */
 	ENGINE_API void EmptyTable();
diff --git a/Engine/Source/Runtime/Engine/Engine.Build.cs b/Engine/Source/Runtime/Engine/Engine.Build.cs
--- a/Engine/Source/Runtime/Engine/Engine.Build.cs
+++ b/Engine/Source/Runtime/Engine/Engine.Build.cs
@@ -67,11 +67,6 @@ public class Engine : ModuleRules
 			}
         );
 
-        if (UEBuildConfiguration.bBuildEditor)
-		{
-            PrivateDependencyModuleNames.AddRange(new string[] { "CsvParser" });
-		}
-
         CircularlyReferencedDependentModules.Add("AIModule");
 
 		DynamicallyLoadedModuleNames.AddRange(
diff --git a/Engine/Source/Runtime/Engine/Private/CurveBase.cpp b/Engine/Source/Runtime/Engine/Private/CurveBase.cpp
--- a/Engine/Source/Runtime/Engine/Private/CurveBase.cpp
+++ b/Engine/Source/Runtime/Engine/Private/CurveBase.cpp
@@ -1,6 +1,7 @@
 // Copyright 1998-2014 Epic Games, Inc. All Rights Reserved.
 
 #include "EnginePrivate.h"
+#include "Serialization/Csv/CsvParser.h"
 
 void FKeyHandleMap::Add( const FKeyHandle& InHandle, int32 InIndex )
 {
@@ -868,9 +869,6 @@ void UCurveBase::ResetCurve()
 	}
 }
 
-#if WITH_EDITOR
-#include "CsvParser.h"
-
 TArray<FString> UCurveBase::CreateCurveFromCSVString(const FString& InString)
 {
 	// Array used to store problems about curve import
@@ -931,8 +929,6 @@ TArray<FString> UCurveBase::CreateCurveFromCSVString(const FString& InString)
 	return OutProblems;
 }
 
-#endif	// WITH_EDITOR
-
 //////////////////////////////////////////////////////////////////////////
 
 int32 FIntegralCurve::GetNumKeys() const
diff --git a/Engine/Source/Runtime/Engine/Private/CurveTable.cpp b/Engine/Source/Runtime/Engine/Private/CurveTable.cpp
--- a/Engine/Source/Runtime/Engine/Private/CurveTable.cpp
+++ b/Engine/Source/Runtime/Engine/Private/CurveTable.cpp
@@ -2,6 +2,7 @@
 
 #include "EnginePrivate.h"
 #include "Json.h"
+#include "Serialization/Csv/CsvParser.h"
 
 DEFINE_LOG_CATEGORY(LogCurveTable);
 
@@ -229,9 +230,6 @@ void GetCurveValues(const TArray<const TCHAR*>& Cells, TArray<float>* Values)
 	}
 }
 
-#if WITH_EDITOR
-#include "CsvParser.h"
-
 TArray<FString> UCurveTable::CreateTableFromCSVString(const FString& InString, ERichCurveInterpMode InterpMode)
 {
 	// Array used to store problems about table creation
@@ -306,8 +304,6 @@ TArray<FString> UCurveTable::CreateTableFromCSVString(const FString& InString, E
 	return OutProblems;
 }
 
-#endif	// WITH_EDITOR
-
 //////////////////////////////////////////////////////////////////////////
 
 
diff --git a/Engine/Source/Runtime/Engine/Private/DataTable.cpp b/Engine/Source/Runtime/Engine/Private/DataTable.cpp
--- a/Engine/Source/Runtime/Engine/Private/DataTable.cpp
+++ b/Engine/Source/Runtime/Engine/Private/DataTable.cpp
@@ -2,6 +2,7 @@
 
 #include "EnginePrivate.h"
 #include "Json.h"
+#include "Serialization/Csv/CsvParser.h"
 
 DEFINE_LOG_CATEGORY(LogDataTable);
 
@@ -235,8 +236,6 @@ UProperty* UDataTable::FindTableProperty(const FName& PropertyName) const
 
 #if WITH_EDITOR || HACK_HEADER_GENERATOR
 
-#include "CsvParser.h"
-
 struct FPropertyDisplayNameHelper
 {
 	static inline FString Get(const UProperty* Prop, const FString& DefaultName)
