Remove hardcoded font path for windows.

diff --git a/Engine/Source/Runtime/Core/Private/Windows/WindowsPlatformMisc.cpp b/Engine/Source/Runtime/Core/Private/Windows/WindowsPlatformMisc.cpp
--- a/Engine/Source/Runtime/Core/Private/Windows/WindowsPlatformMisc.cpp
+++ b/Engine/Source/Runtime/Core/Private/Windows/WindowsPlatformMisc.cpp
@@ -2276,13 +2276,6 @@ FText FWindowsPlatformMisc::GetFileManagerName()
 	return NSLOCTEXT("WindowsPlatform", "FileManagerName", "Explorer");
 }
 
-TArray<uint8> FWindowsPlatformMisc::GetSystemFontBytes()
-{
-	TArray<uint8> FontBytes;
-	FFileHelper::LoadFileToArray( FontBytes, TEXT("C:\\Windows\\Fonts\\comicbd.ttf") );
-	return FontBytes;
-}
-
 bool FWindowsPlatformMisc::IsRunningOnBattery()
 {
 	SYSTEM_POWER_STATUS status;
diff --git a/Engine/Source/Runtime/Core/Public/Windows/WindowsPlatformMisc.h b/Engine/Source/Runtime/Core/Public/Windows/WindowsPlatformMisc.h
--- a/Engine/Source/Runtime/Core/Public/Windows/WindowsPlatformMisc.h
+++ b/Engine/Source/Runtime/Core/Public/Windows/WindowsPlatformMisc.h
@@ -207,8 +207,6 @@ struct CORE_API FWindowsPlatformMisc : public FGenericPlatformMisc
 	/** @return Get the name of the platform specific file manager (Explorer) */
 	static FText GetFileManagerName();
 
-	static TArray<uint8> GetSystemFontBytes();
-
 	/**
 	 * Returns whether the platform is running on battery power or not.
 	 */
