Core: temp hack to make FTimespan BP math functions work

diff --git a/Engine/Source/Runtime/Core/Public/Misc/Timespan.h b/Engine/Source/Runtime/Core/Public/Misc/Timespan.h
--- a/Engine/Source/Runtime/Core/Public/Misc/Timespan.h
+++ b/Engine/Source/Runtime/Core/Public/Misc/Timespan.h
@@ -559,7 +559,7 @@ protected:
 	 */
 	void CORE_API Assign( int32 Days, int32 Hours, int32 Minutes, int32 Seconds, int32 Milliseconds );
 
-private:
+public://private:
 
 	/** Holds the time span in 100 nanoseconds resolution. */
 	int64 Ticks;
diff --git a/Engine/Source/Runtime/CoreUObject/Classes/Object.h b/Engine/Source/Runtime/CoreUObject/Classes/Object.h
--- a/Engine/Source/Runtime/CoreUObject/Classes/Object.h
+++ b/Engine/Source/Runtime/CoreUObject/Classes/Object.h
@@ -686,6 +686,7 @@ struct FDateTime
 USTRUCT(immutable, noexport, BlueprintType)
 struct FTimespan
 {
+	UPROPERTY()
 	int64 Ticks;
 };
 
