Core: temp fix for destructable mesh deserialization crash

diff --git a/Engine/Source/Runtime/Core/Public/Misc/DateTime.h b/Engine/Source/Runtime/Core/Public/Misc/DateTime.h
--- a/Engine/Source/Runtime/Core/Public/Misc/DateTime.h
+++ b/Engine/Source/Runtime/Core/Public/Misc/DateTime.h
@@ -673,7 +673,7 @@ protected:
 	 */
 	static const int32 DaysToMonth[];
 
-private:
+public://private:
 
 	/** Holds the ticks in 100 nanoseconds resolution since January 1, 0001 A.D. */
 	int64 Ticks;
diff --git a/Engine/Source/Runtime/CoreUObject/Classes/Object.h b/Engine/Source/Runtime/CoreUObject/Classes/Object.h
--- a/Engine/Source/Runtime/CoreUObject/Classes/Object.h
+++ b/Engine/Source/Runtime/CoreUObject/Classes/Object.h
@@ -677,6 +677,7 @@ public:
 USTRUCT(immutable, noexport, BlueprintType)
 struct FDateTime
 {
+	UPROPERTY()
 	int64 Ticks;
 };
 
