Changed FMacPlatformMisc::MessageBoxExt() to use NSAlert on the main thread

diff --git a/Engine/Source/Runtime/Core/Private/Mac/MacPlatformMisc.cpp b/Engine/Source/Runtime/Core/Private/Mac/MacPlatformMisc.cpp
--- a/Engine/Source/Runtime/Core/Private/Mac/MacPlatformMisc.cpp
+++ b/Engine/Source/Runtime/Core/Private/Mac/MacPlatformMisc.cpp
@@ -639,169 +639,173 @@ EAppReturnType::Type FMacPlatformMisc::MessageBoxExt(EAppMsgType::Type MsgType,
 {
 	SCOPED_AUTORELEASE_POOL;
 
-	EAppReturnType::Type RetValue = EAppReturnType::Cancel;
-	NSInteger Result;
-
 	if (MacApplication)
 	{
 		MacApplication->UseMouseCaptureWindow(false);
 	}
 
-	NSAlert* AlertPanel = [NSAlert new];
-	[AlertPanel setInformativeText:FString(Text).GetNSString()];
-	[AlertPanel setMessageText:FString(Caption).GetNSString()];
+	EAppReturnType::Type RetValue = MainThreadReturn(^{
+		EAppReturnType::Type RetValue = EAppReturnType::Cancel;
+		NSInteger Result;
 
-	switch (MsgType)
-	{
-		case EAppMsgType::Ok:
-			[AlertPanel addButtonWithTitle:@"OK"];
-			[AlertPanel runModal];
-			RetValue = EAppReturnType::Ok;
-			break;
+		NSAlert* AlertPanel = [NSAlert new];
+		[AlertPanel setInformativeText:FString(Text).GetNSString()];
+		[AlertPanel setMessageText:FString(Caption).GetNSString()];
 
-		case EAppMsgType::YesNo:
-			[AlertPanel addButtonWithTitle:@"Yes"];
-			[AlertPanel addButtonWithTitle:@"No"];
-			Result = [AlertPanel runModal];
-			if (Result == NSAlertFirstButtonReturn)
-			{
-				RetValue = EAppReturnType::Yes;
-			}
-			else if (Result == NSAlertSecondButtonReturn)
-			{
-				RetValue = EAppReturnType::No;
-			}
-			break;
-
-		case EAppMsgType::OkCancel:
-			[AlertPanel addButtonWithTitle:@"OK"];
-			[AlertPanel addButtonWithTitle:@"Cancel"];
-			Result = [AlertPanel runModal];
-			if (Result == NSAlertFirstButtonReturn)
-			{
+		switch (MsgType)
+		{
+			case EAppMsgType::Ok:
+				[AlertPanel addButtonWithTitle:@"OK"];
+				[AlertPanel runModal];
 				RetValue = EAppReturnType::Ok;
-			}
-			else if (Result == NSAlertSecondButtonReturn)
-			{
-				RetValue = EAppReturnType::Cancel;
-			}
-			break;
+				break;
 
-		case EAppMsgType::YesNoCancel:
-			[AlertPanel addButtonWithTitle:@"Yes"];
-			[AlertPanel addButtonWithTitle:@"No"];
-			[AlertPanel addButtonWithTitle:@"Cancel"];
-			Result = [AlertPanel runModal];
-			if (Result == NSAlertFirstButtonReturn)
-			{
-				RetValue = EAppReturnType::Yes;
-			}
-			else if (Result == NSAlertSecondButtonReturn)
-			{
-				RetValue = EAppReturnType::No;
-			}
-			else
-			{
-				RetValue = EAppReturnType::Cancel;
-			}
-			break;
+			case EAppMsgType::YesNo:
+				[AlertPanel addButtonWithTitle:@"Yes"];
+				[AlertPanel addButtonWithTitle:@"No"];
+				Result = [AlertPanel runModal];
+				if (Result == NSAlertFirstButtonReturn)
+				{
+					RetValue = EAppReturnType::Yes;
+				}
+				else if (Result == NSAlertSecondButtonReturn)
+				{
+					RetValue = EAppReturnType::No;
+				}
+				break;
 
-		case EAppMsgType::CancelRetryContinue:
-			[AlertPanel addButtonWithTitle:@"Continue"];
-			[AlertPanel addButtonWithTitle:@"Retry"];
-			[AlertPanel addButtonWithTitle:@"Cancel"];
-			Result = [AlertPanel runModal];
-			if (Result == NSAlertFirstButtonReturn)
-			{
-				RetValue = EAppReturnType::Continue;
-			}
-			else if (Result == NSAlertSecondButtonReturn)
-			{
-				RetValue = EAppReturnType::Retry;
-			}
-			else
-			{
-				RetValue = EAppReturnType::Cancel;
-			}
-			break;
+			case EAppMsgType::OkCancel:
+				[AlertPanel addButtonWithTitle:@"OK"];
+				[AlertPanel addButtonWithTitle:@"Cancel"];
+				Result = [AlertPanel runModal];
+				if (Result == NSAlertFirstButtonReturn)
+				{
+					RetValue = EAppReturnType::Ok;
+				}
+				else if (Result == NSAlertSecondButtonReturn)
+				{
+					RetValue = EAppReturnType::Cancel;
+				}
+				break;
 
-		case EAppMsgType::YesNoYesAllNoAll:
-			[AlertPanel addButtonWithTitle:@"Yes"];
-			[AlertPanel addButtonWithTitle:@"No"];
-			[AlertPanel addButtonWithTitle:@"Yes to all"];
-			[AlertPanel addButtonWithTitle:@"No to all"];
-			Result = [AlertPanel runModal];
-			if (Result == NSAlertFirstButtonReturn)
-			{
-				RetValue = EAppReturnType::Yes;
-			}
-			else if (Result == NSAlertSecondButtonReturn)
-			{
-				RetValue = EAppReturnType::No;
-			}
-			else if (Result == NSAlertThirdButtonReturn)
-			{
-				RetValue = EAppReturnType::YesAll;
-			}
-			else
-			{
-				RetValue = EAppReturnType::NoAll;
-			}
-			break;
+			case EAppMsgType::YesNoCancel:
+				[AlertPanel addButtonWithTitle:@"Yes"];
+				[AlertPanel addButtonWithTitle:@"No"];
+				[AlertPanel addButtonWithTitle:@"Cancel"];
+				Result = [AlertPanel runModal];
+				if (Result == NSAlertFirstButtonReturn)
+				{
+					RetValue = EAppReturnType::Yes;
+				}
+				else if (Result == NSAlertSecondButtonReturn)
+				{
+					RetValue = EAppReturnType::No;
+				}
+				else
+				{
+					RetValue = EAppReturnType::Cancel;
+				}
+				break;
 
-		case EAppMsgType::YesNoYesAllNoAllCancel:
-			[AlertPanel addButtonWithTitle:@"Yes"];
-			[AlertPanel addButtonWithTitle:@"No"];
-			[AlertPanel addButtonWithTitle:@"Yes to all"];
-			[AlertPanel addButtonWithTitle:@"No to all"];
-			[AlertPanel addButtonWithTitle:@"Cancel"];
-			Result = [AlertPanel runModal];
-			if (Result == NSAlertFirstButtonReturn)
-			{
-				RetValue = EAppReturnType::Yes;
-			}
-			else if (Result == NSAlertSecondButtonReturn)
-			{
-				RetValue = EAppReturnType::No;
-			}
-			else if (Result == NSAlertThirdButtonReturn)
-			{
-				RetValue = EAppReturnType::YesAll;
-			}
-			else if (Result == NSAlertThirdButtonReturn + 1)
-			{
-				RetValue = EAppReturnType::NoAll;
-			}
-			else
-			{
-				RetValue = EAppReturnType::Cancel;
-			}
-			break;
+			case EAppMsgType::CancelRetryContinue:
+				[AlertPanel addButtonWithTitle:@"Continue"];
+				[AlertPanel addButtonWithTitle:@"Retry"];
+				[AlertPanel addButtonWithTitle:@"Cancel"];
+				Result = [AlertPanel runModal];
+				if (Result == NSAlertFirstButtonReturn)
+				{
+					RetValue = EAppReturnType::Continue;
+				}
+				else if (Result == NSAlertSecondButtonReturn)
+				{
+					RetValue = EAppReturnType::Retry;
+				}
+				else
+				{
+					RetValue = EAppReturnType::Cancel;
+				}
+				break;
 
-		case EAppMsgType::YesNoYesAll:
-			[AlertPanel addButtonWithTitle:@"Yes"];
-			[AlertPanel addButtonWithTitle:@"No"];
-			[AlertPanel addButtonWithTitle:@"Yes to all"];
-			Result = [AlertPanel runModal];
-			if (Result == NSAlertFirstButtonReturn)
-			{
-				RetValue = EAppReturnType::Yes;
-			}
-			else if (Result == NSAlertSecondButtonReturn)
-			{
-				RetValue = EAppReturnType::No;
-			}
-			else
-			{
-				RetValue = EAppReturnType::YesAll;
-			}
-			break;
+			case EAppMsgType::YesNoYesAllNoAll:
+				[AlertPanel addButtonWithTitle:@"Yes"];
+				[AlertPanel addButtonWithTitle:@"No"];
+				[AlertPanel addButtonWithTitle:@"Yes to all"];
+				[AlertPanel addButtonWithTitle:@"No to all"];
+				Result = [AlertPanel runModal];
+				if (Result == NSAlertFirstButtonReturn)
+				{
+					RetValue = EAppReturnType::Yes;
+				}
+				else if (Result == NSAlertSecondButtonReturn)
+				{
+					RetValue = EAppReturnType::No;
+				}
+				else if (Result == NSAlertThirdButtonReturn)
+				{
+					RetValue = EAppReturnType::YesAll;
+				}
+				else
+				{
+					RetValue = EAppReturnType::NoAll;
+				}
+				break;
 
-		default:
-			break;
-	}
+			case EAppMsgType::YesNoYesAllNoAllCancel:
+				[AlertPanel addButtonWithTitle:@"Yes"];
+				[AlertPanel addButtonWithTitle:@"No"];
+				[AlertPanel addButtonWithTitle:@"Yes to all"];
+				[AlertPanel addButtonWithTitle:@"No to all"];
+				[AlertPanel addButtonWithTitle:@"Cancel"];
+				Result = [AlertPanel runModal];
+				if (Result == NSAlertFirstButtonReturn)
+				{
+					RetValue = EAppReturnType::Yes;
+				}
+				else if (Result == NSAlertSecondButtonReturn)
+				{
+					RetValue = EAppReturnType::No;
+				}
+				else if (Result == NSAlertThirdButtonReturn)
+				{
+					RetValue = EAppReturnType::YesAll;
+				}
+				else if (Result == NSAlertThirdButtonReturn + 1)
+				{
+					RetValue = EAppReturnType::NoAll;
+				}
+				else
+				{
+					RetValue = EAppReturnType::Cancel;
+				}
+				break;
+
+			case EAppMsgType::YesNoYesAll:
+				[AlertPanel addButtonWithTitle:@"Yes"];
+				[AlertPanel addButtonWithTitle:@"No"];
+				[AlertPanel addButtonWithTitle:@"Yes to all"];
+				Result = [AlertPanel runModal];
+				if (Result == NSAlertFirstButtonReturn)
+				{
+					RetValue = EAppReturnType::Yes;
+				}
+				else if (Result == NSAlertSecondButtonReturn)
+				{
+					RetValue = EAppReturnType::No;
+				}
+				else
+				{
+					RetValue = EAppReturnType::YesAll;
+				}
+				break;
+
+			default:
+				break;
+		}
+
+		[AlertPanel release];
 
-	[AlertPanel release];
+		return RetValue;
+	});
 
 	if (MacApplication)
 	{
