Core: better Blueprint and serialization support for FDateTime and FTimespan

diff --git a/Engine/Source/Runtime/Core/Private/Misc/DateTime.cpp b/Engine/Source/Runtime/Core/Private/Misc/DateTime.cpp
--- a/Engine/Source/Runtime/Core/Private/Misc/DateTime.cpp
+++ b/Engine/Source/Runtime/Core/Private/Misc/DateTime.cpp
@@ -1,9 +1,5 @@
 ﻿// Copyright 1998-2014 Epic Games, Inc. All Rights Reserved.
 
-/*=============================================================================
-	DateTime.cpp: Implements the FDateTime class.
-=============================================================================*/
-
 #include "Core.h"
 
 
@@ -52,15 +48,6 @@ FDateTime::FDateTime( int32 Year, int32 Month, int32 Day, int32 Hour, int32 Minu
 }
 
 
-/* FDateTime friend operators
- *****************************************************************************/
-
-FArchive& operator<<( FArchive& Ar, FDateTime& DateTime )
-{
-	return Ar << DateTime.Ticks;
-}
-
-
 /* FDateTime interface
  *****************************************************************************/
 
@@ -181,6 +168,14 @@ bool FDateTime::ImportTextItem( const TCHAR*& Buffer, int32 PortFlags, UObject*
 }
 
 
+bool FDateTime::Serialize( FArchive& Ar )
+{
+	Ar << *this;
+
+	return true;
+}
+
+
 FString FDateTime::ToIso8601( ) const
 {
 	return ToString(TEXT("%Y-%m-%dT%H:%M:%S.%sZ"));
@@ -453,3 +448,12 @@ FDateTime FDateTime::UtcNow( )
 
 	return FDateTime(Year, Month, Day, Hour, Minute, Second, Millisecond);
 }
+
+
+/* FDateTime friend functions
+ *****************************************************************************/
+
+FArchive& operator<<( FArchive& Ar, FDateTime& DateTime )
+{
+	return Ar << DateTime.Ticks;
+}
diff --git a/Engine/Source/Runtime/Core/Private/Misc/Timespan.cpp b/Engine/Source/Runtime/Core/Private/Misc/Timespan.cpp
--- a/Engine/Source/Runtime/Core/Private/Misc/Timespan.cpp
+++ b/Engine/Source/Runtime/Core/Private/Misc/Timespan.cpp
@@ -1,15 +1,34 @@
 // Copyright 1998-2014 Epic Games, Inc. All Rights Reserved.
 
-/*=============================================================================
-	Timespan.cpp: Implements the FTimespan class.
-=============================================================================*/
-
 #include "Core.h"
 
 
 /* FTimespan interface
  *****************************************************************************/
 
+bool FTimespan::ExportTextItem( FString& ValueStr, FTimespan const& DefaultValue, UObject* Parent, int32 PortFlags, UObject* ExportRootScope ) const
+{
+	ValueStr += ToString(TEXT("%N%d.%h:%m:%s.%f"));
+
+	return true;
+}
+
+
+bool FTimespan::ImportTextItem( const TCHAR*& Buffer, int32 PortFlags, UObject* Parent, FOutputDevice* ErrorText )
+{
+	// @todo gmp: implement FTimespan::ImportTextItem
+	return false;
+}
+
+
+bool FTimespan::Serialize( FArchive& Ar )
+{
+	Ar << *this;
+
+	return true;
+}
+
+
 FString FTimespan::ToString( ) const
 {
 	if (GetDays() == 0)
@@ -110,7 +129,6 @@ bool FTimespan::Parse( const FString& TimespanString, FTimespan& OutTimespan )
 	TokenString.ReplaceInline(TEXT("-"), TEXT(":"));
 
 	TArray<FString> Tokens;
-
 	TokenString.ParseIntoArray(&Tokens, TEXT(":"), true);
 
 	if (Tokens.Num() == 4)
@@ -149,8 +167,7 @@ FArchive& operator<<( FArchive& Ar, FTimespan& Timespan )
 void FTimespan::Assign( int32 Days, int32 Hours, int32 Minutes, int32 Seconds, int32 Milliseconds )
 {
 	int64 totalms = 1000 * (60 * 60 * 24 * (int64)Days + 60 * 60 * (int64)Hours + 60 * (int64)Minutes + (int64)Seconds) + (int64)Milliseconds;
-
 	check((totalms >= MinValue().GetTotalMilliseconds()) && (totalms <= MaxValue().GetTotalMilliseconds()));
 
 	Ticks = totalms * ETimespan::TicksPerMillisecond;
-}
\ No newline at end of file
+}
diff --git a/Engine/Source/Runtime/Core/Public/Misc/DateTime.h b/Engine/Source/Runtime/Core/Public/Misc/DateTime.h
--- a/Engine/Source/Runtime/Core/Public/Misc/DateTime.h
+++ b/Engine/Source/Runtime/Core/Public/Misc/DateTime.h
@@ -1,14 +1,11 @@
 // Copyright 1998-2014 Epic Games, Inc. All Rights Reserved.
 
-/*=============================================================================
-	DateTime.h: Declares the FDateTime class.
-=============================================================================*/
-
 #pragma once
 
 
-class UObject;
+class FArchive;
 class FOutputDevice;
+class UObject;
 
 
 namespace EDayOfWeek
@@ -236,7 +233,7 @@ public:
 public:
 
 	/**
-	 * Exports the DateTime value to a string.
+	 * Exports the date and time value to a string.
 	 *
 	 * @param ValueStr Will hold the string value.
 	 * @param DefaultValue The default value.
@@ -423,7 +420,7 @@ public:
 	CORE_API int32 GetYear( ) const;
 
 	/**
-	 * Imports the DateTime value from a text buffer.
+	 * Imports a date and time value from a text buffer.
 	 *
 	 * @param Buffer The text buffer to import from.
 	 * @param PortFlags Not used.
@@ -456,6 +453,14 @@ public:
 		return (GetHour() < 12);
 	}
 
+	/**
+	 * Serializes this date and time from or into the specified archive.
+	 *
+	 * @param Ar The archive to serialize from or into.
+	 * @return true on success, false otherwise.
+	 */
+	CORE_API bool Serialize( FArchive& Ar );
+
 	/**
 	 * Returns the ISO-8601 string representation of the FDateTime.
 	 *
@@ -637,8 +642,6 @@ public:
 
 public:
 
-	friend class UObject;
-
 	/**
 	 * Serializes the given date and time from or into the specified archive.
 	 *
@@ -673,7 +676,7 @@ protected:
 	 */
 	static const int32 DaysToMonth[];
 
-public://private:
+private:
 
 	/** Holds the ticks in 100 nanoseconds resolution since January 1, 0001 A.D. */
 	int64 Ticks;
diff --git a/Engine/Source/Runtime/Core/Public/Misc/Timespan.h b/Engine/Source/Runtime/Core/Public/Misc/Timespan.h
--- a/Engine/Source/Runtime/Core/Public/Misc/Timespan.h
+++ b/Engine/Source/Runtime/Core/Public/Misc/Timespan.h
@@ -1,14 +1,11 @@
 // Copyright 1998-2014 Epic Games, Inc. All Rights Reserved.
 
-/*=============================================================================
-	Timespan.h: Declares the FTimespan class.
-=============================================================================*/
-
 #pragma once
 
 
-class UObject;
+class FArchive;
 class FOutputDevice;
+class UObject;
 
 
 namespace ETimespan
@@ -255,6 +252,19 @@ public:
 
 public:
 
+	/**
+	 * Exports this time span value to a string.
+	 *
+	 * @param ValueStr Will hold the string value.
+	 * @param DefaultValue The default value.
+	 * @param Parent Not used.
+	 * @param PortFlags Not used.
+	 * @param ExportRootScope Not used.
+	 * @return true on success, false otherwise.
+	 * @see ImportTextItem
+	 */
+	CORE_API bool ExportTextItem( FString& ValueStr, FTimespan const& DefaultValue, UObject* Parent, int32 PortFlags, UObject* ExportRootScope ) const;
+
 	/**
 	 * Gets the days component of this time span.
 	 *
@@ -387,6 +397,27 @@ public:
 		return ((double)Ticks / ETimespan::TicksPerSecond);
 	}
 
+	/**
+	 * Imports a time span value from a text buffer.
+	 *
+	 * @param Buffer The text buffer to import from.
+	 * @param PortFlags Not used.
+	 * @param Parent Not used.
+	 * @param ErrorText The output device for error logging.
+	 * @return true on success, false otherwise.
+	 * @see ExportTextItem
+	 */
+	CORE_API bool ImportTextItem( const TCHAR*& Buffer, int32 PortFlags, UObject* Parent, FOutputDevice* ErrorText );
+
+
+	/**
+	 * Serializes this time span from or into the specified archive.
+	 *
+	 * @param Ar The archive to serialize from or into.
+	 * @return true on success, false otherwise.
+	 */
+	CORE_API bool Serialize( FArchive& Ar );
+
 	/**
 	 * Returns the string representation of this time span using a default format.
 	 *
@@ -533,7 +564,7 @@ public:
 	 *
 	 * @todo gmp: Figure out better include order in Core.h so this can be inlined.
 	 */
-	friend CORE_API class FArchive& operator<<( class FArchive& Ar, FTimespan& Timespan );
+	friend CORE_API FArchive& operator<<( FArchive& Ar, FTimespan& Timespan );
 
 	/**
 	 * Gets the hash for the specified time span.
@@ -559,7 +590,7 @@ protected:
 	 */
 	void CORE_API Assign( int32 Days, int32 Hours, int32 Minutes, int32 Seconds, int32 Milliseconds );
 
-public://private:
+private:
 
 	/** Holds the time span in 100 nanoseconds resolution. */
 	int64 Ticks;
diff --git a/Engine/Source/Runtime/CoreUObject/Classes/Object.h b/Engine/Source/Runtime/CoreUObject/Classes/Object.h
--- a/Engine/Source/Runtime/CoreUObject/Classes/Object.h
+++ b/Engine/Source/Runtime/CoreUObject/Classes/Object.h
@@ -677,7 +677,6 @@ public:
 USTRUCT(immutable, noexport, BlueprintType)
 struct FDateTime
 {
-	UPROPERTY()
 	int64 Ticks;
 };
 
@@ -687,7 +686,6 @@ struct FDateTime
 USTRUCT(immutable, noexport, BlueprintType)
 struct FTimespan
 {
-	UPROPERTY()
 	int64 Ticks;
 };
 
diff --git a/Engine/Source/Runtime/CoreUObject/Private/UObject/Property.cpp b/Engine/Source/Runtime/CoreUObject/Private/UObject/Property.cpp
--- a/Engine/Source/Runtime/CoreUObject/Private/UObject/Property.cpp
+++ b/Engine/Source/Runtime/CoreUObject/Private/UObject/Property.cpp
@@ -234,13 +234,29 @@ struct TStructOpsTypeTraits<FDateTime> : public TStructOpsTypeTraitsBase
 {
 	enum 
 	{
+		WithCopy = true,
 		WithExportTextItem = true,
 		WithImportTextItem = true,
+		WithSerializer = true,
 		WithZeroConstructor = true,
 	};
 };
 IMPLEMENT_STRUCT(DateTime);
 
+template<>
+struct TStructOpsTypeTraits<FTimespan> : public TStructOpsTypeTraitsBase
+{
+	enum 
+	{
+		WithCopy = true,
+		WithExportTextItem = true,
+		WithImportTextItem = false, // @todo gmp: implement FTimespan::ImportTextItem
+		WithSerializer = true,
+		WithZeroConstructor = true,
+	};
+};
+IMPLEMENT_STRUCT(Timespan);
+
 IMPLEMENT_STRUCT(StringAssetReference);
 
 IMPLEMENT_STRUCT(StringClassReference);
