Core: Moved FGuid import/export functions into .cpp

diff --git a/Engine/Source/Runtime/Core/Private/Misc/Guid.cpp b/Engine/Source/Runtime/Core/Private/Misc/Guid.cpp
--- a/Engine/Source/Runtime/Core/Private/Misc/Guid.cpp
+++ b/Engine/Source/Runtime/Core/Private/Misc/Guid.cpp
@@ -1,15 +1,37 @@
 // Copyright 1998-2014 Epic Games, Inc. All Rights Reserved.
 
-/*=============================================================================
-	Guid.cpp: Implements the FGuid class.
-=============================================================================*/
-
 #include "Core.h"
 
 
 /* FGuid interface
  *****************************************************************************/
 
+bool FGuid::ExportTextItem( FString& ValueStr, FGuid const& DefaultValue, UObject* Parent, int32 PortFlags, class UObject* ExportRootScope ) const
+{
+	ValueStr += ToString();
+
+	return true;
+}
+
+
+bool FGuid::ImportTextItem( const TCHAR*& Buffer, int32 PortFlags, class UObject* Parent, FOutputDevice* ErrorText )
+{
+	if (FPlatformString::Strlen(Buffer) < 32)
+	{
+		return false;
+	}
+
+	if (!ParseExact(FString(Buffer).Left(32), EGuidFormats::Digits, *this))
+	{
+		return false;
+	}
+
+	Buffer += 32;
+
+	return true;
+}
+
+
 FString FGuid::ToString( EGuidFormats::Type Format ) const
 {
 	switch (Format)
diff --git a/Engine/Source/Runtime/Core/Public/Misc/Guid.h b/Engine/Source/Runtime/Core/Public/Misc/Guid.h
--- a/Engine/Source/Runtime/Core/Public/Misc/Guid.h
+++ b/Engine/Source/Runtime/Core/Public/Misc/Guid.h
@@ -1,9 +1,5 @@
 // Copyright 1998-2014 Epic Games, Inc. All Rights Reserved.
 
-/*=============================================================================
-	Guid.h: Declares the FGuid class.
-=============================================================================*/
-
 #pragma once
 
 
@@ -201,12 +197,7 @@ public:
 	 * @return true on success, false otherwise.
 	 * @see ImportTextItem
 	 */
-	bool ExportTextItem( FString& ValueStr, FGuid const& DefaultValue, UObject* Parent, int32 PortFlags, class UObject* ExportRootScope ) const
-	{
-		ValueStr += ToString();
-
-		return true;
-	}
+	CORE_API bool ExportTextItem( FString& ValueStr, FGuid const& DefaultValue, UObject* Parent, int32 PortFlags, class UObject* ExportRootScope ) const;
 
 	/**
 	 * Imports the GUIDs value from a text buffer.
@@ -218,22 +209,7 @@ public:
 	 * @return true on success, false otherwise.
 	 * @see ExportTextItem
 	 */
-	bool ImportTextItem( const TCHAR*& Buffer, int32 PortFlags, class UObject* Parent, FOutputDevice* ErrorText )
-	{
-		if (FPlatformString::Strlen(Buffer) < 32)
-		{
-			return false;
-		}
-
-		if (!ParseExact(FString(Buffer).Left(32), EGuidFormats::Digits, *this))
-		{
-			return false;
-		}
-
-		Buffer += 32;
-
-		return true;
-	}
+	CORE_API bool ImportTextItem( const TCHAR*& Buffer, int32 PortFlags, class UObject* Parent, FOutputDevice* ErrorText );
 
 	/**
 	 * Checks whether this GUID is valid or not.
