Core: documentation cleanup pass; NULL to nullptr

diff --git a/Engine/Source/Runtime/Core/Public/HAL/ThreadingBase.h b/Engine/Source/Runtime/Core/Public/HAL/ThreadingBase.h
--- a/Engine/Source/Runtime/Core/Public/HAL/ThreadingBase.h
+++ b/Engine/Source/Runtime/Core/Public/HAL/ThreadingBase.h
@@ -1,18 +1,20 @@
 // Copyright 1998-2014 Epic Games, Inc. All Rights Reserved.
 
-/*=============================================================================
-	ThreadingBase.h: Declares base level interfaces for threading support.
-=============================================================================*/
-
 #pragma once
 
+
 /**
  * The list of enumerated thread priorities we support
  */
 enum EThreadPriority
 {
+	/** Normal priority. */
 	TPri_Normal,
+
+	/** Above normal priority. */
 	TPri_AboveNormal,
+
+	/** Below normal priority. */
 	TPri_BelowNormal
 };
 
@@ -34,11 +36,10 @@ public:
 	 * Manually reset events stay triggered until reset.
 	 * Named events share the same underlying event.
 	 *
-	 * @param bIsManualReset - Whether the event requires manual reseting or not.
-	 *
+	 * @param bIsManualReset Whether the event requires manual reseting or not.
 	 * @return true if the event was created, false otherwise.
 	 */
-	virtual bool Create (bool bIsManualReset = false) = 0;
+	virtual bool Create( bool bIsManualReset = false ) = 0;
 
 	/**
 	 * Triggers the event so any waiting threads are activated
@@ -55,19 +56,18 @@ public:
 	 *
 	 * A wait time of MAX_uint32 is treated as infinite wait.
 	 *
-	 * @param WaitTime					- The time to wait (in milliseconds).
-	 * @param bIgnoreThreadIdleStats	- If true, ignores ThreadIdleStats
-	 *
+	 * @param WaitTime The time to wait (in milliseconds).
+	 * @param bIgnoreThreadIdleStats If true, ignores ThreadIdleStats
 	 * @return true if the event was triggered, false if the wait timed out.
 	 */
-	virtual bool Wait(uint32 WaitTime, const bool bIgnoreThreadIdleStats = false) = 0;
+	virtual bool Wait( uint32 WaitTime, const bool bIgnoreThreadIdleStats = false ) = 0;
 
 	/**
 	 * Waits an infinite amount of time for the event to be triggered.
 	 *
 	 * @return true if the event was triggered.
 	 */
-	bool Wait ()
+	bool Wait()
 	{
 		return Wait(MAX_uint32);
 	}
@@ -75,23 +75,21 @@ public:
 	/**
 	 * Waits the specified amount of time for the event to be triggered.
 	 *
-	 * @param WaitTime					- The time to wait.
-	 * @param bIgnoreThreadIdleStats	- If true, ignores ThreadIdleStats
-	 *
+	 * @param WaitTime The time to wait.
+	 * @param bIgnoreThreadIdleStats If true, ignores ThreadIdleStats
 	 * @return true if the event was triggered, false if the wait timed out.
 	 */
-	bool Wait (const FTimespan& WaitTime, const bool bIgnoreThreadIdleStats = false)
+	bool Wait( const FTimespan& WaitTime, const bool bIgnoreThreadIdleStats = false )
 	{
 		return Wait(WaitTime.GetTotalMilliseconds(), bIgnoreThreadIdleStats);
 	}
 
-
 public:
 
 	/**
 	 * Virtual destructor.
 	 */
-	virtual ~FEvent () { }
+	virtual ~FEvent() { }
 };
 
 
@@ -124,7 +122,7 @@ public:
 	{
 		Event->Wait();
 		ReturnToPool(Event);
-		Event = NULL;
+		Event = nullptr;
 	}
 
 	/**
@@ -155,10 +153,9 @@ protected:
 	/**
 	 * Returns an event object to the pool.
 	 *
-	 * @param Event - The event object to return.
+	 * @param Event The event object to return.
 	 */
-	static void ReturnToPool(FEvent* Event);
-
+	static void ReturnToPool( FEvent* Event );
 
 private:
 
@@ -182,6 +179,7 @@ private:
 class CORE_API FRunnable
 {
 public:
+
 	/**
 	 * Initializes the runnable object.
 	 *
@@ -190,7 +188,7 @@ public:
 	 *
 	 * @return True if initialization was successful, false otherwise
 	 */
-	virtual bool Init ()
+	virtual bool Init()
 	{
 		return true;
 	}
@@ -202,29 +200,32 @@ public:
 	 *
 	 * @return The exit code of the runnable object
 	 */
-	virtual uint32 Run () = 0;
+	virtual uint32 Run() = 0;
 
 	/**
 	 * Stops the runnable object.
 	 *
 	 * This is called if a thread is requested to terminate early.
 	 */
-	virtual void Stop () { }
+	virtual void Stop() { }
 
 	/**
 	 * Exits the runnable object.
 	 *
 	 * Called in the context of the aggregating thread to perform any cleanup.
 	 */
-	virtual void Exit () { }
+	virtual void Exit() { }
 
 	/**
 	 * Gets single thread interface pointer used for ticking this runnable when multithreading is disabled.
 	 * If the interface is not implemented, this runnable will not be ticked when FPlatformProcess::SupportsMultithreading() is false.
 	 *
-	 * @return Pointer to the single thread interface or NULL if not implemented.
+	 * @return Pointer to the single thread interface or nullptr if not implemented.
 	 */
-	virtual class FSingleThreadRunnable* GetSingleThreadInterface() { return NULL; }
+	virtual class FSingleThreadRunnable* GetSingleThreadInterface( )
+	{
+		return nullptr;
+	}
 
 public:
 
@@ -234,6 +235,7 @@ public:
 	virtual ~FRunnable() { }
 };
 
+
 /**
  * Interface for runnable threads.
  *
@@ -250,12 +252,9 @@ public:
 	* @param ThreadName Name of the thread
 	* @param bAutoDeleteSelf Whether to delete this object on exit
 	* @param bAutoDeleteRunnable Whether to delete the runnable object on exit
-	* @param InStackSize The size of the stack to create. 0 means use the
-	* current thread's stack size
-	* @param InThreadPri Tells the thread whether it needs to adjust its
-	* priority or not. Defaults to normal priority
-	*
-	* @return The newly created thread or NULL if it failed
+	* @param InStackSize The size of the stack to create. 0 means use the current thread's stack size
+	* @param InThreadPri Tells the thread whether it needs to adjust its priority or not. Defaults to normal priority
+	* @return The newly created thread or nullptr if it failed
 	*/
 	DEPRECATED(4.3, "Function deprecated. Use FRunnableThread::Create without bAutoDeleteSelf and bAutoDeleteRunnable params and delete thread and runnable manually.")
 	CORE_API static FRunnableThread* Create(
@@ -272,26 +271,23 @@ public:
 	 *
 	 * @param InRunnable The runnable object to execute
 	 * @param ThreadName Name of the thread
-	 * @param InStackSize The size of the stack to create. 0 means use the
-	 * current thread's stack size
-	 * @param InThreadPri Tells the thread whether it needs to adjust its
-	 * priority or not. Defaults to normal priority
-	 *
-	 * @return The newly created thread or NULL if it failed
+	 * @param InStackSize The size of the stack to create. 0 means use the current thread's stack size
+	 * @param InThreadPri Tells the thread whether it needs to adjust its priority or not. Defaults to normal priority
+	 * @return The newly created thread or nullptr if it failed
 	 */
 	CORE_API static FRunnableThread* Create(
 		class FRunnable* InRunnable,
 		const TCHAR* ThreadName,
 		uint32 InStackSize = 0,
 		EThreadPriority InThreadPri = TPri_Normal,
-		uint64 InThreadAffinityMask = FPlatformAffinity::GetNoAffinityMask());
+		uint64 InThreadAffinityMask = FPlatformAffinity::GetNoAffinityMask() );
 
 	/**
 	 * Changes the thread priority of the currently running thread
 	 *
 	 * @param NewPriority The thread priority to change to
 	 */
-	virtual void SetThreadPriority (EThreadPriority NewPriority) = 0;
+	virtual void SetThreadPriority( EThreadPriority NewPriority ) = 0;
 
 	/**
 	 * Tells the thread to either pause execution or resume depending on the
@@ -299,7 +295,7 @@ public:
 	 *
 	 * @param bShouldPause Whether to pause the thread (true) or resume (false)
 	 */
-	virtual void Suspend (bool bShouldPause = true) = 0;
+	virtual void Suspend( bool bShouldPause = true ) = 0;
 
 	/**
 	 * Tells the thread to exit. If the caller needs to know when the thread
@@ -309,33 +305,32 @@ public:
 	 * NOTE: having a thread forcibly destroyed can cause leaks in TLS, etc.
 	 *
 	 * @param bShouldWait If true, the call will wait for the thread to exit
-	 *
 	 * @return True if the thread exited graceful, false otherwise
 	 */
-	virtual bool Kill (bool bShouldWait = false) = 0;
+	virtual bool Kill( bool bShouldWait = false ) = 0;
 
 	/**
 	 * Halts the caller until this thread is has completed its work.
 	 */
-	virtual void WaitForCompletion () = 0;
+	virtual void WaitForCompletion() = 0;
 
 	/**
 	 * Thread ID for this thread 
 	 *
 	 * @return ID that was set by CreateThread
 	 */
-	virtual uint32 GetThreadID () = 0;
+	virtual uint32 GetThreadID() = 0;
 
 	/**
 	 * Retrieves the given name of the thread
 	 *
 	 * @return Name that was set by CreateThread
 	 */
-	virtual FString GetThreadName () = 0;
+	virtual FString GetThreadName() = 0;
 
 	struct FThreadRegistry
 	{
-		void Add(uint32 ID, FRunnableThread* Thread)
+		void Add( uint32 ID, FRunnableThread* Thread )
 		{
 			Lock();
 			Registry.FindOrAdd( ID ) = Thread;
@@ -343,7 +338,7 @@ public:
 			Unlock();
 		}
 
-		void Remove(uint32 ID)
+		void Remove( uint32 ID )
 		{
 			Lock();
 			Registry.Remove(ID);
@@ -384,12 +379,13 @@ public:
 			return Registry.CreateConstIterator();
 		}
 
-		FRunnableThread* GetThread(uint32 ID)
+		FRunnableThread* GetThread( uint32 ID )
 		{
 			return Registry.FindRef(ID);
 		}
 
 	private:
+
 		TMap<uint32, FRunnableThread*> Registry;
 		bool Updated;
 		FCriticalSection CriticalSection;
@@ -410,9 +406,7 @@ public:
 	/**
 	 * Virtual destructor
 	 */
-	virtual ~FRunnableThread ();
-
-
+	virtual ~FRunnableThread();
 protected:
 
 	/**
@@ -420,16 +414,13 @@ protected:
 	 *
 	 * @param InRunnable The runnable object to execute
 	 * @param ThreadName Name of the thread
-	 * @param InStackSize The size of the stack to create. 0 means use the
-	 * current thread's stack size
-	 * @param InThreadPri Tells the thread whether it needs to adjust its
-	 * priority or not. Defaults to normal priority
-	 *
+	 * @param InStackSize The size of the stack to create. 0 means use the current thread's stack size
+	 * @param InThreadPri Tells the thread whether it needs to adjust its priority or not. Defaults to normal priority
 	 * @return True if the thread and all of its initialization was successful, false otherwise
 	 */
-	virtual bool CreateInternal (FRunnable* InRunnable, const TCHAR* ThreadName,
+	virtual bool CreateInternal( FRunnable* InRunnable, const TCHAR* ThreadName,
 		uint32 InStackSize = 0,
-		EThreadPriority InThreadPri = TPri_Normal, uint64 InThreadAffinityMask = 0) = 0;
+		EThreadPriority InThreadPri = TPri_Normal, uint64 InThreadAffinityMask = 0 ) = 0;
 
 	/** Called when the this runnable has been destroyed, so we should clean-up memory allocated by misc classes. */
 	FSimpleMulticastDelegate ThreadDestroyedDelegate;
@@ -439,28 +430,46 @@ protected:
 /**
  * Fake event object used when running with only one thread.
  */
-class FSingleThreadEvent : public FEvent
+class FSingleThreadEvent
+	: public FEvent
 {
 	/** Flag to know whether this event has been triggered. */
 	bool bTriggered;
+
 	/** Should this event reset automatically or not. */
 	bool bManualReset;
 
 public:
+
+	/**
+	 * Default constructor.
+	 */
 	FSingleThreadEvent()
 		: bTriggered(false)
 		, bManualReset(false)
-	{}
+	{ }
+
+public:
 
 	// FEvent Interface
-	virtual bool Create(bool bIsManualReset = false) 
+
+	virtual bool Create( bool bIsManualReset = false ) 
 	{ 
 		bManualReset = bIsManualReset;
 		return true; 
 	}
-	virtual void Trigger() { bTriggered = true; }
-	virtual void Reset() { bTriggered = false; }
-	virtual bool Wait(uint32 WaitTime, const bool bIgnoreThreadIdleStats = false) 
+
+	virtual void Trigger()
+	{
+		bTriggered = true;
+	}
+
+	virtual void Reset()
+	{
+		bTriggered = false;
+	}
+
+	virtual bool Wait( uint32 WaitTime, const bool bIgnoreThreadIdleStats = false ) 
 	{ 
 		// With only one thread it's assumed the event has been triggered
 		// before Wait is called, otherwise it would end up waiting forever or always fail.
@@ -478,7 +487,8 @@ public:
 class CORE_API FSingleThreadRunnable
 {
 public:
-	virtual ~FSingleThreadRunnable() {}
+
+	virtual ~FSingleThreadRunnable() { }
 
 	/* Tick function. */
 	virtual void Tick() = 0;
@@ -500,14 +510,14 @@ public:
 	 *
 	 * @param Thread Fake thread object.
 	 */
-	void AddThread(class FFakeThread* Thread);
+	void AddThread( class FFakeThread* Thread );
 
 	/**
 	 * Used internally to remove fake thread object.
 	 *
 	 * @param Thread Fake thread object to be removed.
 	 */
-	void RemoveThread(class FFakeThread* Thread);
+	void RemoveThread( class FFakeThread* Thread );
 
 	/**
 	 * Ticks all fake threads and their runnable objects.
@@ -538,7 +548,7 @@ public:
 	 * This is where the real thread work is done. All work that is done for
 	 * this queued object should be done from within the call to this function.
 	 */
-	virtual void DoThreadedWork () = 0;
+	virtual void DoThreadedWork() = 0;
 
 	/**
 	 * Tells the queued work that it is being abandoned so that it can do
@@ -546,8 +556,7 @@ public:
 	 * abandoned before completion. NOTE: This requires the object to delete
 	 * itself using whatever heap it was allocated in.
 	 */
-	virtual void Abandon () = 0;
-
+	virtual void Abandon() = 0;
 
 public:
 
@@ -555,7 +564,7 @@ public:
 	 * Virtual destructor so that child implementations are guaranteed a chance
 	 * to clean up any resources they allocated.
 	 */
-	virtual ~FQueuedWork () {}
+	virtual ~FQueuedWork() { }
 };
 
 
@@ -575,10 +584,9 @@ public:
 	 * @param InNumQueuedThreads Specifies the number of threads to use in the pool
 	 * @param StackSize The size of stack the threads in the pool need (32K default)
 	 * @param ThreadPriority priority of new pool thread
-	 *
 	 * @return Whether the pool creation was successful or not
 	 */
-	virtual bool Create(uint32 InNumQueuedThreads,uint32 StackSize = (32 * 1024),EThreadPriority ThreadPriority=TPri_Normal) = 0;
+	virtual bool Create( uint32 InNumQueuedThreads, uint32 StackSize = (32 * 1024), EThreadPriority ThreadPriority=TPri_Normal ) = 0;
 
 	/**
 	 * Tells the pool to clean up all background threads
@@ -591,7 +599,7 @@ public:
 	 *
 	 * @param InQueuedWork The work that needs to be done asynchronously
 	 */
-	virtual void AddQueuedWork(FQueuedWork* InQueuedWork) = 0;
+	virtual void AddQueuedWork( FQueuedWork* InQueuedWork ) = 0;
 
 	/**
 	 * Attempts to retract a previously queued task.
@@ -599,7 +607,7 @@ public:
 	 * @param InQueuedWork The work to try to retract
 	 * @return true if the work was retracted
 	 */
-	virtual bool RetractQueuedWork(FQueuedWork* InQueuedWork) = 0;
+	virtual bool RetractQueuedWork( FQueuedWork* InQueuedWork ) = 0;
 
 	/**
 	 * Places a thread back into the available pool
@@ -607,25 +615,21 @@ public:
 	 * @param InQueuedThread The thread that is ready to be pooled
 	 * @return next job or null if there is no job available now
 	 */
-	virtual FQueuedWork* ReturnToPoolOrGetNextJob(class FQueuedThread* InQueuedThread) = 0;
-
+	virtual FQueuedWork* ReturnToPoolOrGetNextJob( class FQueuedThread* InQueuedThread ) = 0;
 
 public:
 
 	/**
 	 * Virtual destructor.
 	 */
-	virtual ~FQueuedThreadPool () 
-	{
-	}
-
+	virtual ~FQueuedThreadPool() { }
 
 public:
 
 	/**
 	 * Allocates a thread pool
 	 */
-	static FQueuedThreadPool* Allocate ();
+	static FQueuedThreadPool* Allocate();
 };
 
 
@@ -657,9 +661,9 @@ public:
 	 * guarantees as to which values you will get up to the caller to not care, synchronize
 	 * or other way to make those guarantees.
 	 *
-	 * @param Other - The other thread safe counter to copy
+	 * @param Other The other thread safe counter to copy
 	 */
-	FThreadSafeCounter (const FThreadSafeCounter& Other)
+	FThreadSafeCounter( const FThreadSafeCounter& Other )
 	{
 		Counter = Other.GetValue();
 	}
@@ -669,7 +673,7 @@ public:
 	 *
 	 * @param Value	Value to initialize counter to
 	 */
-	FThreadSafeCounter (int32 Value)
+	FThreadSafeCounter( int32 Value )
 	{
 		Counter = Value;
 	}
@@ -679,7 +683,7 @@ public:
 	 *
 	 * @return the new, incremented value
 	 */
-	int32 Increment ()
+	int32 Increment()
 	{
 		return FPlatformAtomics::InterlockedIncrement(&Counter);
 	}
@@ -690,7 +694,7 @@ public:
 	 * @param Amount Amount to increase the counter by
 	 * @return the old value
 	 */
-	int32 Add (int32 Amount)
+	int32 Add( int32 Amount )
 	{
 		return FPlatformAtomics::InterlockedAdd(&Counter, Amount);
 	}
@@ -700,7 +704,7 @@ public:
 	 *
 	 * @return the new, decremented value
 	 */
-	int32 Decrement ()
+	int32 Decrement()
 	{
 		return FPlatformAtomics::InterlockedDecrement(&Counter);
 	}
@@ -711,7 +715,7 @@ public:
 	 * @param Amount Amount to decrease the counter by
 	 * @return the old value
 	 */
-	int32 Subtract (int32 Amount)
+	int32 Subtract( int32 Amount )
 	{
 		return FPlatformAtomics::InterlockedAdd(&Counter, -Amount);
 	}
@@ -722,7 +726,7 @@ public:
 	 * @param Value	Value to set the counter to
 	 * @return The old value
 	 */
-	int32 Set (int32 Value)
+	int32 Set( int32 Value )
 	{
 		return FPlatformAtomics::InterlockedExchange(&Counter, Value);
 	}
@@ -732,7 +736,7 @@ public:
 	 *
 	 * @return the old value.
 	 */
-	int32 Reset ()
+	int32 Reset()
 	{
 		return FPlatformAtomics::InterlockedExchange(&Counter, 0);
 	}
@@ -742,63 +746,67 @@ public:
 	 *
 	 * @return the current value
 	 */
-	int32 GetValue () const
+	int32 GetValue() const
 	{
 		return Counter;
 	}
 
-
 private:
-	// Hidden on purpose as usage wouldn't be thread safe.
-	void operator=(const FThreadSafeCounter& Other){}
+
+	/** Hidden on purpose as usage wouldn't be thread safe. */
+	void operator=( const FThreadSafeCounter& Other ){}
 
 	/** Thread-safe counter */
 	volatile int32 Counter;
 };
 
+
 /** Fake Thread safe counter, used to avoid cluttering code with ifdefs when counters are only used for debugging. */
 class FNoopCounter
 {
 public:
-	FNoopCounter()
-	{
-	}
-	FNoopCounter( const FNoopCounter& Other )
-	{
-	}
-	FNoopCounter( int32 Value )
-	{
-	}
+
+	FNoopCounter() { }
+	FNoopCounter( const FNoopCounter& Other ) { }
+	FNoopCounter( int32 Value ) { }
+
 	int32 Increment()
 	{
 		return 0;
 	}
+
 	int32 Add( int32 Amount )
 	{
 		return 0;
 	}
+
 	int32 Decrement()
 	{
 		return 0;
 	}
+
 	int32 Subtract( int32 Amount )
 	{
 		return 0;
 	}
+
 	int32 Set( int32 Value )
 	{
 		return 0;
 	}
+
 	int32 Reset()
 	{
 		return 0;
 	}
+
 	int32 GetValue() const
 	{
 		return 0;
 	}
 };
 
+
 /**
  * Implements a scope lock.
  *
@@ -825,40 +833,35 @@ public:
 	 *
 	 * @param InSynchObject The synchronization object to manage
 	 */
-	FScopeLock (FCriticalSection* InSynchObject)
+	FScopeLock( FCriticalSection* InSynchObject )
 		: SynchObject(InSynchObject)
 	{
 		check(SynchObject);
-
 		SynchObject->Lock();
 	}
 
 	/**
 	 * Destructor that performs a release on the synchronization object
 	 */
-	~FScopeLock ()
+	~FScopeLock()
 	{
 		check(SynchObject);
-
 		SynchObject->Unlock();
 	}
-
-
 private:
 
 	// Default constructor (hidden on purpose).
-	FScopeLock ();
+	FScopeLock();
 
 	// Copy constructor( hidden on purpose).
-	FScopeLock (FScopeLock* InScopeLock);
+	FScopeLock( FScopeLock* InScopeLock);
 
 	// Assignment operator (hidden on purpose).
-	FScopeLock& operator= (FScopeLock& InScopeLock)
+	FScopeLock& operator=( FScopeLock& InScopeLock )
 	{
 		return *this;
 	}
 
-
 private:
 
 	// Holds the synchronization object to aggregate and scope manage.
@@ -900,22 +903,24 @@ struct FThreadSingleton
 	typedef FThreadSingleton* (*TCreateSingletonFuncPtr)();
 	typedef void (FThreadSingleton::*TDestroySingletonFuncPtr)();
 
-	CORE_API virtual ~FThreadSingleton()
-	{}
+	CORE_API virtual ~FThreadSingleton() { }
 
 	virtual void DeleteThis() = 0;
 };
 
+
 /** Thread singleton initializer. */
 class FThreadSingletonInitializer
 {
 public:
+
 	/**
 	* @return an instance of a singleton for the current thread.
 	*/
 	static CORE_API FThreadSingleton* Get( const FThreadSingleton::TCreateSingletonFuncPtr CreateFunc, const FThreadSingleton::TDestroySingletonFuncPtr DestroyFunc, uint32& TlsSlot );
 };
 
+
 /**
  * This a special version of singleton. It means that there is created only one instance for each thread.
  * Calling Get() method is thread-safe, but first call should be done on the game thread.
@@ -934,10 +939,11 @@ class TThreadSingleton : public FThreadSingleton
 	}
 
 protected:
+
 	/** Default constructor. */
 	TThreadSingleton()
-		: ThreadId( FPlatformTLS::GetCurrentThreadId() )
-	{}
+		: ThreadId(FPlatformTLS::GetCurrentThreadId())
+	{ }
 
 	/**
 	 * @return a new instance of the thread singleton.
@@ -960,14 +966,16 @@ protected:
 	const uint32 ThreadId;
 
 public:
+
 	/**
 	 *	@return an instance of a singleton for the current thread.
 	 */
 	FORCEINLINE static T& Get()
 	{
-		return *(T*)FThreadSingletonInitializer::Get( &T::Create, &FThreadSingleton::DeleteThis, T::GetTlsSlot() );
+		return *(T*)FThreadSingletonInitializer::Get(&T::Create, &FThreadSingleton::DeleteThis, T::GetTlsSlot());
 	}
 };
 
+
 #define DECLARE_THREAD_SINGLETON(ClassType) \
 	static auto& GForceInitAtBoot_ThreadSingletonInitializer_##ClassType = ClassType::Get();
