Adding Serialization/Csv to Core public include paths

diff --git a/Engine/Source/Runtime/Core/Core.Build.cs b/Engine/Source/Runtime/Core/Core.Build.cs
--- a/Engine/Source/Runtime/Core/Core.Build.cs
+++ b/Engine/Source/Runtime/Core/Core.Build.cs
@@ -27,6 +27,7 @@ public class Core : ModuleRules
 				"Runtime/Core/Public/ProfilingDebugging",
 				"Runtime/Core/Public/Serialization",
 				"Runtime/Core/Public/Serialization/Json",
+				"Runtime/Core/Public/Serialization/Csv",
 				"Runtime/Core/Public/Stats",
 				"Runtime/Core/Public/Templates",
 				"Runtime/Core/Public/UObject",
diff --git a/Engine/Source/Runtime/Core/Private/Serialization/Csv/CsvParser.cpp b/Engine/Source/Runtime/Core/Private/Serialization/Csv/CsvParser.cpp
--- a/Engine/Source/Runtime/Core/Private/Serialization/Csv/CsvParser.cpp
+++ b/Engine/Source/Runtime/Core/Private/Serialization/Csv/CsvParser.cpp
@@ -1,7 +1,7 @@
 // Copyright 1998-2014 Epic Games, Inc. All Rights Reserved.
 
 #include "Core.h"
-#include "Serialization/Csv/CsvParser.h"
+#include "CsvParser.h"
 
 FCsvParser::FCsvParser(FString InSourceString)
 	: SourceString(MoveTemp(InSourceString))		
diff --git a/Engine/Source/Runtime/Core/Private/Serialization/Csv/CsvParserTests.cpp b/Engine/Source/Runtime/Core/Private/Serialization/Csv/CsvParserTests.cpp
--- a/Engine/Source/Runtime/Core/Private/Serialization/Csv/CsvParserTests.cpp
+++ b/Engine/Source/Runtime/Core/Private/Serialization/Csv/CsvParserTests.cpp
@@ -1,7 +1,7 @@
 // Copyright 1998-2014 Epic Games, Inc. All Rights Reserved.
 
 #include "Core.h"
-#include "Serialization/Csv/CsvParser.h"
+#include "CsvParser.h"
 
 namespace CsvParser_Tests
 {
diff --git a/Engine/Source/Runtime/Engine/Private/CurveBase.cpp b/Engine/Source/Runtime/Engine/Private/CurveBase.cpp
--- a/Engine/Source/Runtime/Engine/Private/CurveBase.cpp
+++ b/Engine/Source/Runtime/Engine/Private/CurveBase.cpp
@@ -1,7 +1,7 @@
 // Copyright 1998-2014 Epic Games, Inc. All Rights Reserved.
 
 #include "EnginePrivate.h"
-#include "Serialization/Csv/CsvParser.h"
+#include "CsvParser.h"
 
 void FKeyHandleMap::Add( const FKeyHandle& InHandle, int32 InIndex )
 {
diff --git a/Engine/Source/Runtime/Engine/Private/CurveTable.cpp b/Engine/Source/Runtime/Engine/Private/CurveTable.cpp
--- a/Engine/Source/Runtime/Engine/Private/CurveTable.cpp
+++ b/Engine/Source/Runtime/Engine/Private/CurveTable.cpp
@@ -2,7 +2,7 @@
 
 #include "EnginePrivate.h"
 #include "Json.h"
-#include "Serialization/Csv/CsvParser.h"
+#include "CsvParser.h"
 
 DEFINE_LOG_CATEGORY(LogCurveTable);
 
diff --git a/Engine/Source/Runtime/Engine/Private/DataTable.cpp b/Engine/Source/Runtime/Engine/Private/DataTable.cpp
--- a/Engine/Source/Runtime/Engine/Private/DataTable.cpp
+++ b/Engine/Source/Runtime/Engine/Private/DataTable.cpp
@@ -2,7 +2,7 @@
 
 #include "EnginePrivate.h"
 #include "Json.h"
-#include "Serialization/Csv/CsvParser.h"
+#include "CsvParser.h"
 
 DEFINE_LOG_CATEGORY(LogDataTable);
 
