Changed checkSlow() in TOptional to a checkf() as requested.
There is existing code which was (incorrectly) calling GetValue() on an unset TOptional without checking IsSet(), which "worked" before (returned a default-constructed object by coincidence) but would now return uninitialized memory. The checkSlow (which only fired in debug) catches new code well (programmers tend to run debug builds), but not old code (which gets run by non-coders in non-debug builds), and the crash reports were rather bizarre.

diff --git a/Engine/Source/Runtime/Core/Public/Misc/Optional.h b/Engine/Source/Runtime/Core/Public/Misc/Optional.h
--- a/Engine/Source/Runtime/Core/Public/Misc/Optional.h
+++ b/Engine/Source/Runtime/Core/Public/Misc/Optional.h
@@ -172,8 +172,8 @@ public:
 	FORCEINLINE_EXPLICIT_OPERATOR_BOOL() const { return bIsSet; }
 
 	/** @return The optional value; undefined when IsSet() returns false. */
-	const OptionalType& GetValue() const { checkSlow(IsSet()); return (OptionalType&)Value; }
-	      OptionalType& GetValue()       { checkSlow(IsSet()); return (OptionalType&)Value; }
+	const OptionalType& GetValue() const { checkf(IsSet(), TEXT("It is an error to call GetValue() on an unset TOptional. Please either check IsSet() or use Get(DefaultValue) instead.")); return (OptionalType&)Value; }
+	      OptionalType& GetValue()       { checkf(IsSet(), TEXT("It is an error to call GetValue() on an unset TOptional. Please either check IsSet() or use Get(DefaultValue) instead.")); return (OptionalType&)Value; }
 
 	const OptionalType* operator->() const { return &GetValue(); }
 	      OptionalType* operator->()       { return &GetValue(); }
