Can't call the automatic string conversion functions - they allocate memory which in a signal handler isn't a wise thing to do.

diff --git a/Engine/Source/Runtime/Core/Private/Apple/ApplePlatformCrashContext.cpp b/Engine/Source/Runtime/Core/Private/Apple/ApplePlatformCrashContext.cpp
--- a/Engine/Source/Runtime/Core/Private/Apple/ApplePlatformCrashContext.cpp
+++ b/Engine/Source/Runtime/Core/Private/Apple/ApplePlatformCrashContext.cpp
@@ -74,10 +74,8 @@ int32 FApplePlatformCrashContext::ReportCrash() const
 		// Walk the stack and dump it to the allocated memory (ignore first 2 callstack lines as those are in stack walking code)
 		FPlatformStackWalk::StackWalkAndDump( StackBuffer, ARRAY_COUNT(MinidumpCallstackInfo) - 1, 6, Context );
 		
-#if WITH_EDITORONLY_DATA
-		FCString::Strncat( GErrorHist, ANSI_TO_TCHAR(MinidumpCallstackInfo), ARRAY_COUNT(GErrorHist) - 1 );
+		FUTF8ToTCHAR_Convert::Convert(GErrorHist, ARRAY_COUNT(GErrorHist) - 1, MinidumpCallstackInfo, FCStringAnsi::Strlen(MinidumpCallstackInfo));
 		CreateExceptionInfoString(Signal, Info);
-#endif
 	}
 	
 	return 0;
