CrashReport - GetEngineMode moved to FPlatformMisc, removed duplicates

diff --git a/Engine/Source/Runtime/Core/Private/GenericPlatform/GenericPlatformMisc.cpp b/Engine/Source/Runtime/Core/Private/GenericPlatform/GenericPlatformMisc.cpp
--- a/Engine/Source/Runtime/Core/Private/GenericPlatform/GenericPlatformMisc.cpp
+++ b/Engine/Source/Runtime/Core/Private/GenericPlatform/GenericPlatformMisc.cpp
@@ -815,4 +815,13 @@ FString FGenericPlatformMisc::GetEpicAccountId()
 void FGenericPlatformMisc::SetEpicAccountId( const FString& AccountId )
 {
 	SetStoredValue( TEXT( "Epic Games" ), TEXT( "Unreal Engine/Identifiers" ), TEXT( "AccountId" ), AccountId );
-}
\ No newline at end of file
+}
+
+const TCHAR* FGenericPlatformMisc::GetEngineMode()
+{
+	return	
+		IsRunningCommandlet() ? TEXT( "Commandlet" ) :
+		GIsEditor ? TEXT( "Editor" ) :
+		IsRunningDedicatedServer() ? TEXT( "Server" ) :
+		TEXT( "Game" );
+}
diff --git a/Engine/Source/Runtime/Core/Private/Linux/LinuxPlatformMisc.cpp b/Engine/Source/Runtime/Core/Private/Linux/LinuxPlatformMisc.cpp
--- a/Engine/Source/Runtime/Core/Private/Linux/LinuxPlatformMisc.cpp
+++ b/Engine/Source/Runtime/Core/Private/Linux/LinuxPlatformMisc.cpp
@@ -1105,17 +1105,6 @@ void FLinuxCrashContext::GenerateReport(const FString & DiagnosticsPath) const
  */
 void GenerateWindowsErrorReport(const FString & WERPath)
 {
-	struct FLocalHelpers
-	{
-		static const TCHAR* GetEngineMode()
-		{
-			return	IsRunningCommandlet()?	 	TEXT("Commandlet") :
-					GIsEditor?				 	TEXT("Editor") :
-					IsRunningDedicatedServer()?	TEXT("Server") :
-												TEXT("Game");
-		}
-	};
-
 	FArchive* ReportFile = IFileManager::Get().CreateFileWriter(*WERPath);
 	if (ReportFile != NULL)
 	{
@@ -1155,7 +1144,7 @@ void GenerateWindowsErrorReport(const FString & WERPath)
 		WriteLine(ReportFile, TEXT("\t\t<Parameter6>00000000</Parameter6>"));													// FIXME: supply valid?
 		WriteLine(ReportFile, TEXT("\t\t<Parameter7>0000000000000000</Parameter7>"));											// FIXME: supply valid?
 		WriteLine(ReportFile, *FString::Printf(TEXT("\t\t<Parameter8>!%s!</Parameter8>"), FCommandLine::Get()));				// FIXME: supply valid? Only partially valid
-		WriteLine(ReportFile, *FString::Printf(TEXT("\t\t<Parameter9>%s!%s!%s!%d</Parameter9>"), TEXT( BRANCH_NAME ), FPlatformProcess::BaseDir(), FLocalHelpers::GetEngineMode(), BUILT_FROM_CHANGELIST));
+		WriteLine(ReportFile, *FString::Printf(TEXT("\t\t<Parameter9>%s!%s!%s!%d</Parameter9>"), TEXT( BRANCH_NAME ), FPlatformProcess::BaseDir(), FPlatformMisc::GetEngineMode(), BUILT_FROM_CHANGELIST));
 		WriteLine(ReportFile, TEXT("\t</ProblemSignatures>"));
 
 		WriteLine(ReportFile, TEXT("\t<DynamicSignatures>"));
diff --git a/Engine/Source/Runtime/Core/Private/Mac/MacPlatformMisc.cpp b/Engine/Source/Runtime/Core/Private/Mac/MacPlatformMisc.cpp
--- a/Engine/Source/Runtime/Core/Private/Mac/MacPlatformMisc.cpp
+++ b/Engine/Source/Runtime/Core/Private/Mac/MacPlatformMisc.cpp
@@ -34,27 +34,6 @@
  */
 struct MacApplicationInfo
 {
-	/**
-	 * Get a string description of the mode the engine was running in when it crashed
-	 */
-	static const TCHAR* GetEngineMode()
-	{
-		if( IsRunningCommandlet() )
-		{
-			return TEXT( "Commandlet" );
-		}
-		else if( GIsEditor )
-		{
-			return TEXT( "Editor" );
-		}
-		else if( IsRunningDedicatedServer() )
-		{
-			return TEXT( "Server" );
-		}
-		
-		return TEXT( "Game" );
-	}
-	
 	void Init()
 	{
 		SCOPED_AUTORELEASE_POOL;
@@ -118,7 +97,7 @@ struct MacApplicationInfo
 		
 		FString CrashVideoPath = FPaths::GameLogDir() + TEXT("CrashVideo.avi");
 		
-		BranchBaseDir = FString::Printf(TEXT("%s!%s!%s!%d"), TEXT( BRANCH_NAME ), FPlatformProcess::BaseDir(), GetEngineMode(), BUILT_FROM_CHANGELIST);
+		BranchBaseDir = FString::Printf( TEXT( "%s!%s!%s!%d" ), TEXT( BRANCH_NAME ), FPlatformProcess::BaseDir(), FPlatformMisc::GetEngineMode(), BUILT_FROM_CHANGELIST );
 		
 		// Get the paths that the files will actually have been saved to
 		FString LogDirectory = FPaths::GameLogDir();
diff --git a/Engine/Source/Runtime/Core/Private/Windows/WindowsPlatformExceptionHandling.cpp b/Engine/Source/Runtime/Core/Private/Windows/WindowsPlatformExceptionHandling.cpp
--- a/Engine/Source/Runtime/Core/Private/Windows/WindowsPlatformExceptionHandling.cpp
+++ b/Engine/Source/Runtime/Core/Private/Windows/WindowsPlatformExceptionHandling.cpp
@@ -50,17 +50,6 @@ bool WriteMinidump(const TCHAR* Path, LPEXCEPTION_POINTERS ExceptionInfo)
 	return result == TRUE;
 }
 
-/** 
- * Get a string description of the mode the engine was running in when it crashed
- */
-const TCHAR* GetEngineMode()
-{
-	return	IsRunningCommandlet()?	 	TEXT("Commandlet") :
-			GIsEditor?				 	TEXT("Editor") :
-			IsRunningDedicatedServer()?	TEXT("Server") :
-										TEXT("Game");
-}
-
 #if WINVER > 0x502	// Windows Error Reporting is not supported on Windows XP
 
 /** 
@@ -182,7 +171,7 @@ void SetReportParameters( HREPORT ReportHandle, EXCEPTION_POINTERS* ExceptionInf
 	StringCchPrintf( StringBuffer, MAX_SPRINTF, TEXT( "!%s!AssertLog=\"%s\"" ), FCommandLine::Get(), LocalBuffer );
 	Result = WerReportSetParameter( ReportHandle, WER_P8, TEXT( "Commandline" ), StringBuffer );
 
-	StringCchPrintf( StringBuffer, MAX_SPRINTF, TEXT( "%s!%s!%s!%d" ), TEXT( BRANCH_NAME ), FPlatformProcess::BaseDir(), GetEngineMode(), BUILT_FROM_CHANGELIST );
+	StringCchPrintf( StringBuffer, MAX_SPRINTF, TEXT( "%s!%s!%s!%d" ), TEXT( BRANCH_NAME ), FPlatformProcess::BaseDir(), FPlatformMisc::GetEngineMode, BUILT_FROM_CHANGELIST );
 	Result = WerReportSetParameter( ReportHandle, WER_P9, TEXT( "BranchBaseDir" ), StringBuffer );
 }
 
@@ -355,7 +344,7 @@ int32 ReportCrashUsingCrashReportClient(EXCEPTION_POINTERS* ExceptionInfo, const
 			FText MessageTitle( FText::Format(
 				NSLOCTEXT( "MessageDialog", "AppHasCrashed", "The {0} {1} has crashed and will close" ),
 				FText::FromString( ReportInformation.wzApplicationName ),
-				FText::FromString( GetEngineMode() )
+				FText::FromString( FPlatformMisc::GetEngineMode() )
 				) );
 			FMessageDialog::Open( EAppMsgType::Ok, FText::FromString( GErrorHist ), &MessageTitle );
 		}
diff --git a/Engine/Source/Runtime/Core/Public/GenericPlatform/GenericPlatformMisc.h b/Engine/Source/Runtime/Core/Public/GenericPlatform/GenericPlatformMisc.h
--- a/Engine/Source/Runtime/Core/Public/GenericPlatform/GenericPlatformMisc.h
+++ b/Engine/Source/Runtime/Core/Public/GenericPlatform/GenericPlatformMisc.h
@@ -182,6 +182,117 @@ namespace EAppReturnType
 	};
 }
 
+/**
+ *	Contains crash's description that is common for all platform and can be obtained in run-time.
+ *	Since we can't used dynamic memory allocation everything is defined as simple as possible.
+ *	This may change in the future.
+ */
+struct CORE_API FGenericCrashDescription
+{
+	/**
+	 * The name of the game that crashed. (AppID)
+	 */
+	TCHAR GameName[64];
+
+	/**
+	 * The mode the game was in e.g. editor.
+	 */
+	TCHAR EngineMode[64];
+
+	/**
+	 * The platform that crashed e.g. Win64.
+	 * Main platform name [detailed information like OS version]
+	 */
+	TCHAR Platform[64];
+
+	/**
+	* Encoded engine version. (AppVersion)
+	* E.g. 4.3.0.0-2215663+UE4-Releases+4.3
+	* BuildVersion-BuiltFromCL-BranchName
+	* @EngineVersion	varchar(64)
+	*/
+	//FEngineVersion EngineVersion;
+
+	//Architecture;
+
+	//NumberOfCores;
+
+	//CrashedModuleName;
+
+	//LoadedModules ? ;
+
+	//static const FPlatformMemoryConstants& GetConstants();
+	//static FPlatformMemoryStats GetStats();
+
+	/**
+	 * The three component version of the app e.g. 4.4.1
+	 */
+	TCHAR BuildVersion[32];
+
+	/**
+	 * Built from changelist.
+	 */
+	uint32 BuiltFromCL;
+
+	/**
+	 * The name of the branch this game was built out of.
+	 */
+	TCHAR BranchName[32];
+
+	/**
+	 * The command line of the application that crashed.
+	 */
+	TCHAR CommandLine[512];
+
+	/**
+	 * The base directory where the app was running.
+	 */
+	TCHAR BaseDir[512];
+
+	/**
+	 * The language ID the application that crashed.
+	 */
+	TCHAR LanguageLCID[64];
+
+	/**
+	 * The name of the user that caused this crash.
+	 */
+	TCHAR UserName[64];
+
+	/**
+	 * The unique ID used to identify the machine the crash occurred on.
+	 */
+	TCHAR MachineId[64];
+
+	/**
+	 * The Epic account ID for the user who last used the Launcher.
+	 */
+	TCHAR EpicAccountId[64];
+
+	/**
+	 * A string representing the callstack of the crash.
+	 */
+	TCHAR CallStack[16384];
+
+	/**
+	 * A string representing the user description of the crash.
+	 */
+	TCHAR UserDescription[512];
+
+	/**
+	 * The error message, can be assertion message, ensure message or message from the fatal error.
+	 */
+	TCHAR ErrorMessage[512];
+
+	/**
+	 * The UTC time the crash occurred.
+	 */
+	//	FDateTime TimeOfCrash;
+
+	/** Lenght of the crash's description, in bytes. */
+	int32 Lenght;
+};
+
 struct CORE_API FGenericCrashContext
 {
 };
@@ -826,6 +937,11 @@ public:
 	 */
 	static void SetEpicAccountId( const FString& AccountId );
 
+	/** 
+	 * Get a string description of the mode the engine was running in.
+	 */
+	static const TCHAR* GetEngineMode();
+
 #if !UE_BUILD_SHIPPING
 protected:
 	/** Whether the user should be prompted to allow for a remote debugger to be attached */
