fix for holding splash screen with landscape on iPhone 5/5c/5s
also refactored the code to utilize string append so we will always use the default for the splash screen in the future if something is not found properly
#ue4
#ios

diff --git a/Engine/Source/Runtime/Core/Private/IOS/IOSAppDelegate.cpp b/Engine/Source/Runtime/Core/Private/IOS/IOSAppDelegate.cpp
--- a/Engine/Source/Runtime/Core/Private/IOS/IOSAppDelegate.cpp
+++ b/Engine/Source/Runtime/Core/Private/IOS/IOSAppDelegate.cpp
@@ -391,6 +391,9 @@ void InstallSignalHandlers()
 	self.Window = [[UIWindow alloc] initWithFrame:MainFrame];
 	self.Window.screen = [UIScreen mainScreen];
     
+    // get the native scale
+    const float NativeScale = [[UIScreen mainScreen] scale];
+    
 	//Make this the primary window, and show it.
 	[self.Window makeKeyAndVisible];
 
@@ -400,27 +403,41 @@ void InstallSignalHandlers()
     NSMutableString* path = [[NSMutableString alloc]init];
     [path setString: [[NSBundle mainBundle] resourcePath]];
     UIImageOrientation orient = UIImageOrientationUp;
-    if (MainFrame.size.height == 480)
+    NSMutableString* ImageString = [[NSMutableString alloc]init];
+    [ImageString appendString:@"Default"];
+    if (MainFrame.size.height == 320 && !self.bDeviceInPortraitMode)
     {
-        [path setString: [path stringByAppendingPathComponent:@"Default.png"]];
+        [ImageString appendString:@"-568h"];
+        orient = UIImageOrientationLeft;
     }
     else if (MainFrame.size.height == 568)
     {
-        [path setString: [path stringByAppendingPathComponent:@"Default-568h.png"]];
+        [ImageString appendString:@"-568h"];
     }
     else if (MainFrame.size.height == 1024 && !self.bDeviceInPortraitMode)
     {
-        [path setString: [path stringByAppendingPathComponent:@"Default-Landscape.png"]];
+        [ImageString appendString:@"-Landscape"];
         orient = UIImageOrientationRight;
     }
     else if (MainFrame.size.height == 1024)
     {
-        [path setString: [path stringByAppendingPathComponent:@"Default-Portrait.png"]];
+        [ImageString appendString:@"-Portrait"];
     }
     else if (MainFrame.size.height == 768 && !self.bDeviceInPortraitMode)
     {
-        [path setString: [path stringByAppendingPathComponent:@"Default-Landscape.png"]];
+        [ImageString appendString:@"-Landscape"];
+    }
+
+    if (NativeScale > 1.0f)
+    {
+        [ImageString appendString:@"@2x.png"];
     }
+    else
+    {
+        [ImageString appendString:@".png"];
+    }
+
+    [path setString: [path stringByAppendingPathComponent:ImageString]];
     UIImage* image = [[UIImage alloc] initWithContentsOfFile: path];
     [path release];
     UIImage* imageToDisplay = [UIImage imageWithCGImage: [image CGImage] scale: 1.0 orientation: orient];
