Better reporting of broken stack entries in Mac OS X crash reports including hardening the crash debug helper.

diff --git a/Engine/Source/Developer/CrashDebugHelper/Private/Mac/CrashDebugHelperMac.cpp b/Engine/Source/Developer/CrashDebugHelper/Private/Mac/CrashDebugHelperMac.cpp
--- a/Engine/Source/Developer/CrashDebugHelper/Private/Mac/CrashDebugHelperMac.cpp
+++ b/Engine/Source/Developer/CrashDebugHelper/Private/Mac/CrashDebugHelperMac.cpp
@@ -220,11 +220,8 @@ static int32 SymboliseStackInfo(TArray<FCrashModuleInfo> const& ModuleInfo, FStr
 			break;
 		}
 	}
-	check(Module.Name.Len() > 0);
-	
-	bool bSuccess = FApplePlatformSymbolication::SymbolInfoForStrippedSymbol((ProgramCounter - Module.BaseOfImage), TCHAR_TO_UTF8(*Module.Name), TCHAR_TO_UTF8(*Module.Report), Info);
-	
-	if(bSuccess)
+
+	if((Module.Name.Len() > 0) && FApplePlatformSymbolication::SymbolInfoForStrippedSymbol((ProgramCounter - Module.BaseOfImage), TCHAR_TO_UTF8(*Module.Name), TCHAR_TO_UTF8(*Module.Report), Info))
 	{
 		if(FCStringAnsi::Strlen(Info.FunctionName) > 0)
 		{
diff --git a/Engine/Source/Runtime/Core/Private/Mac/MacPlatformMisc.cpp b/Engine/Source/Runtime/Core/Private/Mac/MacPlatformMisc.cpp
--- a/Engine/Source/Runtime/Core/Private/Mac/MacPlatformMisc.cpp
+++ b/Engine/Source/Runtime/Core/Private/Mac/MacPlatformMisc.cpp
@@ -1458,7 +1458,13 @@ void FMacCrashContext::GenerateMinidump(char const* Path) const
 			WriteUTF16String(ReportFile, ItoTCHAR(CurrentDepth - IgnoreCount, 10));
 			WriteUTF16String(ReportFile, TEXT("\t"));
 			
-			dladdr((void*)StackTrace[CurrentDepth], &Info);
+			if(dladdr((void*)StackTrace[CurrentDepth], &Info) == 0)
+			{
+				Info.dli_fbase = nullptr;
+				Info.dli_fname = nullptr;
+				Info.dli_saddr = nullptr;
+				Info.dli_sname = nullptr;
+			}
 			if(Info.dli_fname && FCStringAnsi::Strrchr(Info.dli_fname, '/'))
 			{
 				char const* Name = FCStringAnsi::Strrchr(Info.dli_fname, '/');
@@ -1470,7 +1476,7 @@ void FMacCrashContext::GenerateMinidump(char const* Path) const
 			}
 			else
 			{
-				WriteUTF16String(ReportFile, TEXT("Unknown Module"));
+				WriteUTF16String(ReportFile, TEXT("[Unknown]"));
 			}
 			WriteUTF16String(ReportFile, TEXT("\t0x"));
 			WriteUTF16String(ReportFile, ItoTCHAR(StackTrace[CurrentDepth], 16));
@@ -1485,7 +1491,7 @@ void FMacCrashContext::GenerateMinidump(char const* Path) const
 			}
 			else
 			{
-				WriteUTF16String(ReportFile, TEXT("Unknown Function"));
+				WriteUTF16String(ReportFile, TEXT("[Unknown]"));
 			}
 			WriteUTF16String(ReportFile, TEXT(" + "));
 			WriteLine(ReportFile, ItoTCHAR(StackTrace[CurrentDepth] - (uint64)Info.dli_saddr, 10));
