Add self-assignment checks to TOptional

diff --git a/Engine/Source/Runtime/Core/Public/Misc/Optional.h b/Engine/Source/Runtime/Core/Public/Misc/Optional.h
--- a/Engine/Source/Runtime/Core/Public/Misc/Optional.h
+++ b/Engine/Source/Runtime/Core/Public/Misc/Optional.h
@@ -55,37 +55,49 @@ public:
 
 	TOptional& operator=(const TOptional& InValue)
 	{
-		Reset();
-		if (InValue.bIsSet)
+		if (&InValue != this)
 		{
-			new(&Value) OptionalType((const OptionalType&)InValue.Value);
-			bIsSet = true;
+			Reset();
+			if (InValue.bIsSet)
+			{
+				new(&Value) OptionalType((const OptionalType&)InValue.Value);
+				bIsSet = true;
+			}
 		}
 		return *this;
 	}
 	TOptional& operator=(TOptional&& InValue)
 	{
-		Reset();
-		if (InValue.bIsSet)
+		if (&InValue != this)
 		{
-			new(&Value) OptionalType(MoveTemp((OptionalType&)InValue.Value));
-			bIsSet = true;
+			Reset();
+			if (InValue.bIsSet)
+			{
+				new(&Value) OptionalType(MoveTemp((OptionalType&)InValue.Value));
+				bIsSet = true;
+			}
 		}
 		return *this;
 	}
 
 	TOptional& operator=(const OptionalType& InValue)
 	{
-		Reset();
-		new(&Value) OptionalType(InValue);
-		bIsSet = true;
+		if (&InValue != &(OptionalType&)Value)
+		{
+			Reset();
+			new(&Value) OptionalType(InValue);
+			bIsSet = true;
+		}
 		return *this;
 	}
 	TOptional& operator=(OptionalType&& InValue)
 	{
-		Reset();
-		new(&Value) OptionalType(MoveTemp(InValue));
-		bIsSet = true;
+		if (&InValue != &(OptionalType&)Value)
+		{
+			Reset();
+			new(&Value) OptionalType(MoveTemp(InValue));
+			bIsSet = true;
+		}
 		return *this;
 	}
 
