Falling back wide character version of ToUpper/ToLower/IsAlpha etc. to native C++ functions.

diff --git a/Engine/Source/Runtime/Core/Public/Misc/Char.h b/Engine/Source/Runtime/Core/Public/Misc/Char.h
--- a/Engine/Source/Runtime/Core/Public/Misc/Char.h
+++ b/Engine/Source/Runtime/Core/Public/Misc/Char.h
@@ -92,22 +92,22 @@ struct TChar : public TCharBase<T, sizeof(T)>
 {
 	typedef T CharType;
 public:
-	static inline CharType ToUpper( CharType c );
-	static inline CharType ToLower( CharType c );
-	static inline bool IsUpper( CharType c );
-	static inline bool IsLower( CharType c );
-	static inline bool IsAlpha( CharType c );
-	static inline bool IsPunct( CharType c );
+	static inline CharType ToUpper(CharType Char);
+	static inline CharType ToLower(CharType Char);
+	static inline bool IsUpper(CharType Char);
+	static inline bool IsLower(CharType Char);
+	static inline bool IsAlpha(CharType Char);
+	static inline bool IsPunct(CharType Char);
 
-	static inline bool IsAlnum( CharType c )			{ return IsAlpha(c) || IsDigit(c); }
-	static inline bool IsDigit( CharType c )			{ return c>=LITERAL(CharType, '0') && c<=LITERAL(CharType, '9'); }
-	static inline bool IsHexDigit( CharType c )			{ return IsDigit(c) || (c>=LITERAL(CharType, 'a') && c<=LITERAL(CharType, 'f')) || (c>=LITERAL(CharType, 'A') && c<=LITERAL(CharType, 'F')); }
-	static inline bool IsWhitespace( CharType c )		{ return c == LITERAL(CharType, ' ') || c == LITERAL(CharType, '\t'); }
+	static inline bool IsAlnum(CharType Char);
+	static inline bool IsDigit(CharType Char);
+	static inline bool IsHexDigit(CharType Char);
+	static inline bool IsWhitespace(CharType Char);
 
-	static inline bool IsUnderscore( CharType c )		{ return c == LITERAL(CharType, '_'); }
+	static inline bool IsUnderscore(CharType Char)		{ return Char == LITERAL(CharType, '_'); }
 
 public:
-	static inline bool IsLinebreak( CharType c )		{ return LineBreakImplementation<CharType, sizeof(CharType)>::IsLinebreak(c); }
+	static inline bool IsLinebreak(CharType Char)	{ return LineBreakImplementation<CharType, sizeof(CharType)>::IsLinebreak(Char); }
 };
 
 typedef TChar<TCHAR>    FChar;
@@ -118,105 +118,27 @@ typedef TChar<ANSICHAR> FCharAnsi;
 	TCHAR specialized functions
 -----------------------------------------------------------------------------*/
 
-#define UPPER_LOWER_DIFF	32
-
-template <> inline 
-TChar<WIDECHAR>::CharType TChar<WIDECHAR>::ToUpper(CharType c)
-{
-	// compiler generates incorrect code if we try to use TEXT('char') instead of the numeric values directly
-	//@hack - ideally, this would be data driven or use some sort of lookup table
-	// some special cases
-	switch (WIDECHAR(c))
-	{
-		// these special chars are not 32 apart
-		case 255: return 159; // diaeresis y
-		case 156: return 140; // digraph ae
-
-		// characters within the 192 - 255 range which have no uppercase/lowercase equivalents
-		case 240:
-		case 208:
-		case 223:
-		case 247:
-			return c;
-	}
-
-	if ( (c >= LITERAL(CharType, 'a') && c <= LITERAL(CharType, 'z')) || (c > 223 && c < 255) )
-	{
-		return c - UPPER_LOWER_DIFF;
-	}
-
-	// no uppercase equivalent
-	return c;
-}
-
-template <> inline 
-TChar<WIDECHAR>::CharType TChar<WIDECHAR>::ToLower( CharType c )
-{
-	// compiler generates incorrect code if we try to use TEXT('char') instead of the numeric values directly
-	// some special cases
-	switch (WIDECHAR(c))
-	{
-		// these are not 32 apart
-		case 159: return 255; // diaeresis y
-		case 140: return 156; // digraph ae
-
-		// characters within the 192 - 255 range which have no uppercase/lowercase equivalents
-		case 240:
-		case 208:
-		case 223:
-		case 247: 
-			return c;
-	}
-
-	if ( (c >= 192 && c < 223) || (c >= LITERAL(CharType, 'A') && c <= LITERAL(CharType, 'Z')) )
-	{
-		return c + UPPER_LOWER_DIFF;
-	}
-
-	// no lowercase equivalent
-	return c;
-}
-
-template <> inline 
-bool TChar<WIDECHAR>::IsUpper( CharType cc )
-{
-	WIDECHAR c(cc);
-	// compiler generates incorrect code if we try to use TEXT('char') instead of the numeric values directly
-	return (c == 159) || (c == 140)	// these are outside the standard range
-		|| (c == 240) || (c == 247)	// these have no lowercase equivalents
-		|| (c >= LITERAL(CharType, 'A') && c <= LITERAL(CharType, 'Z')) || (c >= 192 && c <= 223);
-}
-	
-template <> inline 
-bool TChar<WIDECHAR>::IsLower( CharType cc )
-{
-	WIDECHAR c(cc);
-	// compiler generates incorrect code if we try to use TEXT('char') instead of the numeric values directly
-	return (c == 156) 								// outside the standard range
-		|| (c == 215) || (c == 208) || (c== 223)	// these have no lower-case equivalents
-		|| (c >= LITERAL(CharType, 'a') && c <= LITERAL(CharType, 'z')) || (c >= 224 && c <= 255);
-}
-
-template <> inline 
-bool TChar<WIDECHAR>::IsAlpha( CharType cc )
-{
-	WIDECHAR c(cc);
-	// compiler generates incorrect code if we try to use TEXT('char') instead of the numeric values directly
-	return (c >= LITERAL(CharType, 'A') && c <= LITERAL(CharType, 'Z')) 
-		|| (c >= 192 && c <= 255)
-		|| (c >= LITERAL(CharType, 'a') && c <= LITERAL(CharType, 'z')) 
-		|| (c == 159) || (c == 140) || (c == 156);	// these are outside the standard range
-}
-
-template <> inline bool TChar<WIDECHAR>::IsPunct( CharType c )				{ return ::iswpunct( c ) != 0; }
-
+template <> inline TChar<WIDECHAR>::CharType TChar<WIDECHAR>::ToUpper(CharType Char)	{ return ::towupper(Char); }
+template <> inline TChar<WIDECHAR>::CharType TChar<WIDECHAR>::ToLower(CharType Char)	{ return ::towlower(Char); }
+template <> inline bool TChar<WIDECHAR>::IsUpper(CharType Char)							{ return ::iswupper(Char) != 0; }
+template <> inline bool TChar<WIDECHAR>::IsLower(CharType Char)							{ return ::iswlower(Char) != 0; }
+template <> inline bool TChar<WIDECHAR>::IsAlpha(CharType Char)							{ return ::iswalpha(Char) != 0; }
+template <> inline bool TChar<WIDECHAR>::IsPunct(CharType Char)							{ return ::iswpunct(Char) != 0; }
+template <> inline bool TChar<WIDECHAR>::IsAlnum(CharType Char)							{ return ::iswalnum(Char) != 0; }
+template <> inline bool TChar<WIDECHAR>::IsDigit(CharType Char)							{ return ::iswdigit(Char) != 0; }
+template <> inline bool TChar<WIDECHAR>::IsHexDigit(CharType Char)						{ return ::iswxdigit(Char) != 0; }
+template <> inline bool TChar<WIDECHAR>::IsWhitespace(CharType Char)					{ return ::iswspace(Char) != 0; }
 
 /*-----------------------------------------------------------------------------
 	ANSICHAR specialized functions
 -----------------------------------------------------------------------------*/
-template <> inline TChar<ANSICHAR>::CharType TChar<ANSICHAR>::ToUpper( CharType c )	{ return (ANSICHAR)::toupper(c); }
-template <> inline TChar<ANSICHAR>::CharType TChar<ANSICHAR>::ToLower( CharType c )	{ return (ANSICHAR)::tolower(c); }
-template <> inline bool TChar<ANSICHAR>::IsUpper( CharType c )							{ return ::isupper(static_cast<unsigned char>(c)) != 0; }
-template <> inline bool TChar<ANSICHAR>::IsLower(CharType c)							{ return ::islower(static_cast<unsigned char>(c)) != 0; }
-template <> inline bool TChar<ANSICHAR>::IsAlpha(CharType c)							{ return ::isalpha(static_cast<unsigned char>(c)) != 0; }
-template <> inline bool TChar<ANSICHAR>::IsPunct(CharType c)							{ return ::ispunct(static_cast<unsigned char>(c)) != 0; }
+template <> inline TChar<ANSICHAR>::CharType TChar<ANSICHAR>::ToUpper(CharType Char)	{ return ::toupper(Char); }
+template <> inline TChar<ANSICHAR>::CharType TChar<ANSICHAR>::ToLower(CharType Char)	{ return ::tolower(Char); }
+template <> inline bool TChar<ANSICHAR>::IsUpper(CharType Char)							{ return ::isupper(Char) != 0; }
+template <> inline bool TChar<ANSICHAR>::IsLower(CharType Char)							{ return ::islower(Char) != 0; }
+template <> inline bool TChar<ANSICHAR>::IsAlpha(CharType Char)							{ return ::isalpha(Char) != 0; }
+template <> inline bool TChar<ANSICHAR>::IsPunct(CharType Char)							{ return ::ispunct(Char) != 0; }
+template <> inline bool TChar<ANSICHAR>::IsAlnum(CharType Char)							{ return ::isalnum(Char) != 0; }
+template <> inline bool TChar<ANSICHAR>::IsDigit(CharType Char)							{ return ::isdigit(Char) != 0; }
+template <> inline bool TChar<ANSICHAR>::IsHexDigit(CharType Char)						{ return ::isxdigit(Char) != 0; }
+template <> inline bool TChar<ANSICHAR>::IsWhitespace(CharType Char)					{ return ::isspace(Char) != 0; }
