Use a pre-configured malloc zone allocator instead of the current allocator once execution enters the async. signal handler on Mac OS X. This prevents a problem with Mac OS X's malloc not being async. handler safe that causes the application to become unkillable if it crashes inside a malloc routine. Where the crash malloc can't satisfy the memory request during shutdown we'll just give up & launch the crash reporter & exit.

diff --git a/Engine/Source/Programs/CrashReporter/CrashReportClient/Private/Mac/CrashReportClientMainMac.cpp b/Engine/Source/Programs/CrashReporter/CrashReportClient/Private/Mac/CrashReportClientMainMac.cpp
--- a/Engine/Source/Programs/CrashReporter/CrashReportClient/Private/Mac/CrashReportClientMainMac.cpp
+++ b/Engine/Source/Programs/CrashReporter/CrashReportClient/Private/Mac/CrashReportClientMainMac.cpp
@@ -21,12 +21,10 @@ void CrashReporterCrashHandler(const FGenericCrashContext & GenericContext)
 	{
 		GWarn->Flush();
 	}
-	// @todo Disable calls to functions which eventually try and allocated memory, resulting in the application hanging before the crash reporter can launch
-	// We need to setup the memory allocator system to handle the case where we can't call the system malloc, free etc. to re-enable this.
 	if (GError)
 	{
 		GError->Flush();
-		//GError->HandleError();
+		GError->HandleError();
 	}
 	FPlatformMisc::RequestExit(true);
 }
diff --git a/Engine/Source/Runtime/Core/Private/Mac/MacMallocZone.cpp b/Engine/Source/Runtime/Core/Private/Mac/MacMallocZone.cpp
new file mode 100644
--- /dev/null
+++ b/Engine/Source/Runtime/Core/Private/Mac/MacMallocZone.cpp
@@ -0,0 +1,184 @@
+// Copyright 1998-2014 Epic Games, Inc. All Rights Reserved.
+
+#include "Core.h"
+#include "MacMallocZone.h"
+#include "MacPlatformMisc.h"
+#include <malloc/malloc.h>
+
+FMacMallocZone::FMacMallocZone( uint64 const InitialSize )
+: MemoryZone( malloc_create_zone( InitialSize, 0 ) )
+{
+}
+
+FMacMallocZone::~FMacMallocZone()
+{
+	if ( MemoryZone )
+	{
+		malloc_destroy_zone( MemoryZone );
+	}
+}
+
+void* FMacMallocZone::Malloc( SIZE_T Size, uint32 Alignment )
+{
+	check( MemoryZone );
+	void* Result = malloc_zone_malloc( MemoryZone, Size );
+	return Result;
+}
+
+void* FMacMallocZone::Realloc( void* Ptr, SIZE_T NewSize, uint32 Alignment )
+{
+	check( MemoryZone );
+	void* Result = malloc_zone_realloc( MemoryZone, Ptr, NewSize );
+	return Result;
+}
+
+void FMacMallocZone::Free( void* Ptr )
+{
+	check( MemoryZone );
+	malloc_zone_free( MemoryZone, Ptr );
+}
+
+bool FMacMallocZone::GetAllocationSize(void *Original, SIZE_T &SizeOut)
+{
+	SizeOut = malloc_size( Original );
+	return SizeOut > 0;
+}
+
+bool FMacMallocZone::IsInternallyThreadSafe() const
+{
+	return true;
+}
+
+bool FMacMallocZone::ValidateHeap()
+{
+	check( MemoryZone );
+	return malloc_zone_check( MemoryZone ) ;
+}
+
+const TCHAR * FMacMallocZone::GetDescriptiveName()
+{
+	return TEXT("MallocZone");
+}
+
+FMacMallocCrashHandler::FMacMallocCrashHandler( uint64 const InitialSize )
+: FMacMallocZone( InitialSize )
+, OriginalHeap( GMalloc )
+, CrashContext( nullptr )
+, ThreadId( -1 )
+{
+	check( GMalloc );
+}
+
+FMacMallocCrashHandler::~FMacMallocCrashHandler()
+{
+	// We crashed, so don't try and tidy the malloc zone
+	if ( ThreadId != -1 )
+	{
+		MemoryZone = nullptr;
+	}
+}
+	
+void FMacMallocCrashHandler::Enable( FMacCrashContext* Context, uint32 CrashedThreadId )
+{
+	check(Context);
+	CrashContext = Context;
+	ThreadId = CrashedThreadId;
+	OriginalHeap = GMalloc;
+	GMalloc = this;
+}
+	
+void* FMacMallocCrashHandler::Malloc( SIZE_T Size, uint32 Alignment )
+{
+	void* Result = nullptr;
+	if ( IsOnCrashedThread() )
+	{
+		Result = FMacMallocZone::Malloc( Size, Alignment );
+		if( !Result )
+		{
+			check(CrashContext);
+			CrashContext->GenerateCrashInfoAndLaunchReporter();
+		}
+	}
+	return Result;
+}
+
+void* FMacMallocCrashHandler::Realloc( void* Ptr, SIZE_T NewSize, uint32 Alignment )
+{
+	void* Result = nullptr;
+	if ( IsOnCrashedThread() )
+	{
+		if ( Ptr == nullptr || MemoryZone->size( MemoryZone, Ptr ) > 0 )
+		{
+			Result = FMacMallocZone::Realloc( Ptr, NewSize, Alignment );
+			if( NewSize && !Result )
+			{
+				check(CrashContext);
+				CrashContext->GenerateCrashInfoAndLaunchReporter();
+			}
+		}
+		else if ( NewSize )
+		{
+			if ( FCStringWide::Strcmp( OriginalHeap->GetDescriptiveName(), TEXT("ANSI") ) != 0 )
+			{
+				SIZE_T OldSize = 0;
+				if ( OriginalHeap->GetAllocationSize( Ptr, OldSize ) )
+				{
+					Result = Malloc( NewSize, Alignment );
+					if ( Result )
+					{
+						FMemory::Memcpy( Result, Ptr, FMath::Min( NewSize, OldSize ) );
+					}
+					else
+					{
+						check(CrashContext);
+						CrashContext->GenerateCrashInfoAndLaunchReporter();
+					}
+				}
+			}
+			else // Can't safely handle this, so just report & exit
+			{
+				check(CrashContext);
+				CrashContext->GenerateCrashInfoAndLaunchReporter();
+			}
+		}
+	}
+	return Result;
+}
+
+void FMacMallocCrashHandler::Free( void* Ptr )
+{
+	if ( IsOnCrashedThread() && Ptr && MemoryZone->size( MemoryZone, Ptr ) > 0 )
+	{
+		FMacMallocZone::Free( Ptr );
+	}
+}
+
+bool FMacMallocCrashHandler::GetAllocationSize( void *Original, SIZE_T &SizeOut )
+{
+	SizeOut = 0;
+	if ( IsOnCrashedThread() && Original )
+	{
+		SizeOut = MemoryZone->size( MemoryZone, Original );
+	}
+	return SizeOut > 0;
+}
+
+const TCHAR * FMacMallocCrashHandler::GetDescriptiveName()
+{
+	return TEXT("MallocCrashHandler");
+}
+
+bool FMacMallocCrashHandler::IsOnCrashedThread( void )
+{
+	// Suspend threads other than the crashed one to prevent serious memory errors.
+	// Only the crashed thread can do anything meaningful from here anyway.
+	if ( ThreadId == FPlatformTLS::GetCurrentThreadId() )
+	{
+		return true;
+	}
+	else
+	{
+		FPlatformProcess::SleepInfinite();
+		return false;
+	}
+}
diff --git a/Engine/Source/Runtime/Core/Private/Mac/MacPlatformMisc.cpp b/Engine/Source/Runtime/Core/Private/Mac/MacPlatformMisc.cpp
--- a/Engine/Source/Runtime/Core/Private/Mac/MacPlatformMisc.cpp
+++ b/Engine/Source/Runtime/Core/Private/Mac/MacPlatformMisc.cpp
@@ -13,6 +13,7 @@
 #include "CocoaThread.h"
 #include "Runtime/Launch/Resources/Version.h"
 #include "EngineVersion.h"
+#include "MacMallocZone.h"
 
 #include <dlfcn.h>
 #include <IOKit/IOKitLib.h>
@@ -1050,6 +1051,7 @@ bool FMacPlatformMisc::IsRunningOnMavericks()
 
 /** Global pointer to crash handler */
 void (* GCrashHandlerPointer)(const FGenericCrashContext & Context) = NULL;
+FMacMallocCrashHandler* GCrashMalloc = nullptr;
 
 /**
  * Good enough default crash reporter.
@@ -1069,9 +1071,7 @@ static void DefaultCrashHandler(FMacCrashContext const& Context)
 	if (GError)
 	{
 		GError->Flush();
-		// @todo Disable calls to functions which eventually try and allocated memory, resulting in the application hanging before the crash reporter can launch
-		// We need to setup the memory allocator system to handle the case where we can't call the system malloc, free etc. to re-enable this.
-		//GError->HandleError();
+		GError->HandleError();
 	}
 	
 	return Context.GenerateCrashInfoAndLaunchReporter();
@@ -1083,6 +1083,10 @@ static void PlatformCrashHandler(int32 Signal, siginfo_t* Info, void* Context)
 	FMacCrashContext CrashContext;
 	CrashContext.InitFromSignal(Signal, Info, Context);
 	
+	// Switch to crash handler malloc to avoid malloc reentrancy
+	check(GCrashMalloc);
+	GCrashMalloc->Enable(&CrashContext, FPlatformTLS::GetCurrentThreadId());
+	
 	if (GCrashHandlerPointer)
 	{
 		GCrashHandlerPointer(CrashContext);
@@ -1139,6 +1143,9 @@ void FMacPlatformMisc::SetCrashHandler(void (* CrashHandler)(const FGenericCrash
 {
 	GCrashHandlerPointer = CrashHandler;
 	
+	// Configure the crash handler malloc zone to reserve some VM space for itself
+	GCrashMalloc = new FMacMallocCrashHandler( 128 * 1024 * 1024 );
+	
 	struct sigaction Action;
 	FMemory::Memzero(&Action, sizeof(struct sigaction));
 	Action.sa_sigaction = PlatformCrashHandler;
@@ -1433,10 +1440,7 @@ void FMacCrashContext::GenerateMinidump(char const* Path) const
 		WriteLine(ReportFile, ItoTCHAR(Info->si_band, 16));
 		
 		WriteLine(ReportFile, TEXT("Application Specific Information:"));
-		for(uint32 i = 0; i < PATH_MAX && (i < FCStringAnsi::Strlen(SignalDescription) + 1); i++)
-		{
-			Line[i] = SignalDescription[i];
-		}
+		FUTF8ToTCHAR_Convert::Convert(Line, PATH_MAX, SignalDescription, FCStringAnsi::Strlen(SignalDescription));
 		WriteLine(ReportFile, Line);
 		
 		WriteLine(ReportFile, TEXT(""));
@@ -1468,10 +1472,7 @@ void FMacCrashContext::GenerateMinidump(char const* Path) const
 			if(Info.dli_fname && FCStringAnsi::Strrchr(Info.dli_fname, '/'))
 			{
 				char const* Name = FCStringAnsi::Strrchr(Info.dli_fname, '/');
-				for(uint32 i = 1; i < PATH_MAX && (i < FCStringAnsi::Strlen(Name) + 1); i++)
-				{
-					Line[i-1] = Name[i];
-				}
+				FUTF8ToTCHAR_Convert::Convert(Line, PATH_MAX, Name, FCStringAnsi::Strlen(Name));
 				WriteUTF16String(ReportFile, Line);
 			}
 			else
@@ -1483,10 +1484,7 @@ void FMacCrashContext::GenerateMinidump(char const* Path) const
 			WriteUTF16String(ReportFile, TEXT(" "));
 			if(Info.dli_sname && FCStringAnsi::Strlen(Info.dli_sname))
 			{
-				for(uint32 i = 0; i < PATH_MAX && (i < FCStringAnsi::Strlen(Info.dli_sname) + 1); i++)
-				{
-					Line[i] = Info.dli_sname[i];
-				}
+				FUTF8ToTCHAR_Convert::Convert(Line, PATH_MAX, Info.dli_sname, FCStringAnsi::Strlen(Info.dli_sname));
 				WriteUTF16String(ReportFile, Line);
 			}
 			else
@@ -1541,10 +1539,7 @@ void FMacCrashContext::GenerateMinidump(char const* Path) const
 			WriteUTF16String(ReportFile, TEXT(" "));
 			
 			char const* Name = FCStringAnsi::Strrchr(ModulePath, '/');
-			for(uint32 i = 1; i < PATH_MAX && (i < FCStringAnsi::Strlen(Name) + 1); i++)
-			{
-				Line[i-1] = Name[i];
-			}
+			FUTF8ToTCHAR_Convert::Convert(Line, PATH_MAX, Name, FCStringAnsi::Strlen(Name));
 			WriteUTF16String(ReportFile, Line);
 			WriteUTF16String(ReportFile, TEXT(" ("));
 			
@@ -1582,10 +1577,7 @@ void FMacCrashContext::GenerateMinidump(char const* Path) const
 			WriteUTF16String(ReportFile, ItoTCHAR(UUID[15], 16, 2));
 			WriteUTF16String(ReportFile, TEXT("> "));
 			
-			for(uint32 i = 0; i < PATH_MAX && (i < FCStringAnsi::Strlen(ModulePath) + 1); i++)
-			{
-				Line[i] = ModulePath[i];
-			}
+			FUTF8ToTCHAR_Convert::Convert(Line, PATH_MAX, ModulePath, FCStringAnsi::Strlen(ModulePath));
 			WriteLine(ReportFile, Line);
 		}
 		
@@ -1595,10 +1587,7 @@ void FMacCrashContext::GenerateMinidump(char const* Path) const
 		WriteUTF16String(ReportFile, TEXT("Model: "));
 		WriteUTF16String(ReportFile, *GMacAppInfo.MachineModel);
 		WriteUTF16String(ReportFile, TEXT(", "));
-		for(uint32 i = 0; i < PATH_MAX && (i < FCStringAnsi::Strlen(GMacAppInfo.MachineCPUString) + 1); i++)
-		{
-			Line[i] = GMacAppInfo.MachineCPUString[i];
-		}
+		FUTF8ToTCHAR_Convert::Convert(Line, PATH_MAX, GMacAppInfo.MachineCPUString, FCStringAnsi::Strlen(GMacAppInfo.MachineCPUString));
 		WriteLine(ReportFile, Line);
 		
 		WriteUTF16String(ReportFile, TEXT("Graphics: "));
diff --git a/Engine/Source/Runtime/Core/Public/Mac/MacMallocZone.h b/Engine/Source/Runtime/Core/Public/Mac/MacMallocZone.h
new file mode 100644
--- /dev/null
+++ b/Engine/Source/Runtime/Core/Public/Mac/MacMallocZone.h
@@ -0,0 +1,68 @@
+// Copyright 1998-2014 Epic Games, Inc. All Rights Reserved.
+
+#pragma once
+
+#include "MemoryBase.h"
+
+typedef struct _malloc_zone_t malloc_zone_t;
+struct FMacCrashContext;
+
+/*
+ * A malloc interface using a unique malloc zone
+ */
+class FMacMallocZone : public FMalloc
+{
+public:
+	FMacMallocZone( uint64 const InitialSize );
+	virtual ~FMacMallocZone();
+	
+	// FMalloc interface.
+	virtual void* Malloc( SIZE_T Size, uint32 Alignment ) override;
+	
+	virtual void* Realloc( void* Ptr, SIZE_T NewSize, uint32 Alignment ) override;
+	
+	virtual void Free( void* Ptr ) override;
+	
+	virtual bool GetAllocationSize(void *Original, SIZE_T &SizeOut) override;
+	
+	virtual bool IsInternallyThreadSafe() const override;
+	
+	virtual bool ValidateHeap() override;
+	
+	virtual const TCHAR * GetDescriptiveName() override;
+	
+protected:
+	malloc_zone_t* MemoryZone;
+};
+
+/*
+ * Specific FMacMallocZone for the crash handler malloc override so that we avoid malloc reentrancy problems
+ */
+class FMacMallocCrashHandler : public FMacMallocZone
+{
+public:
+	FMacMallocCrashHandler( uint64 const InitialSize );
+	
+	virtual ~FMacMallocCrashHandler();
+	
+	void Enable( FMacCrashContext* Context, uint32 CrashedThreadId );
+	
+	// FMalloc interface.
+	virtual void* Malloc( SIZE_T Size, uint32 Alignment ) override;
+	
+	virtual void* Realloc( void* Ptr, SIZE_T NewSize, uint32 Alignment ) override;
+	
+	virtual void Free( void* Ptr ) override;
+	
+	virtual bool GetAllocationSize( void *Original, SIZE_T &SizeOut ) override;
+	
+	virtual const TCHAR * GetDescriptiveName() override;
+
+private:	
+	bool IsOnCrashedThread(void);
+	
+private:
+	FMalloc* OriginalHeap;
+	FMacCrashContext* CrashContext;
+	int32 ThreadId;
+};
diff --git a/Engine/Source/Runtime/Launch/Private/Mac/LaunchMac.cpp b/Engine/Source/Runtime/Launch/Private/Mac/LaunchMac.cpp
--- a/Engine/Source/Runtime/Launch/Private/Mac/LaunchMac.cpp
+++ b/Engine/Source/Runtime/Launch/Private/Mac/LaunchMac.cpp
@@ -38,9 +38,9 @@ void EngineCrashHandler(const FGenericCrashContext & GenericContext)
 	if (GError)
 	{
 		GError->Flush();
-		//GError->HandleError();
+		GError->HandleError();
 	}
-	//LaunchStaticShutdownAfterError();
+	LaunchStaticShutdownAfterError();
 	return Context.GenerateCrashInfoAndLaunchReporter();
 }
 
