Remove log spam when checking if a Mac IME context is active.

diff --git a/Engine/Source/Runtime/Core/Private/Mac/MacTextInputMethodSystem.cpp b/Engine/Source/Runtime/Core/Private/Mac/MacTextInputMethodSystem.cpp
--- a/Engine/Source/Runtime/Core/Private/Mac/MacTextInputMethodSystem.cpp
+++ b/Engine/Source/Runtime/Core/Private/Mac/MacTextInputMethodSystem.cpp
@@ -251,7 +251,6 @@ bool FMacTextInputMethodSystem::IsActiveContext(const TSharedRef<ITextInputMetho
 			}
 		}
 	}
-	UE_LOG(LogMacTextInputMethodSystem, Error, TEXT("Checking for an active context failed when its window couldn't be found."));
 	return false;
 }
 // ITextInputMethodSystem Interface End
