#ttp 344974 WEX: Support Android System font

Added code to load the default system font on Android.
Changed file pathing so that it only rebases certain paths.

diff --git a/Engine/Source/Runtime/Core/Private/Android/AndroidFile.cpp b/Engine/Source/Runtime/Core/Private/Android/AndroidFile.cpp
--- a/Engine/Source/Runtime/Core/Private/Android/AndroidFile.cpp
+++ b/Engine/Source/Runtime/Core/Private/Android/AndroidFile.cpp
@@ -666,12 +666,15 @@ public:
 		}
 
 		FString Result = Filename;
-		Result.ReplaceInline(TEXT("../"), TEXT(""));
-		Result.ReplaceInline(TEXT(".."), TEXT(""));
-		Result.ReplaceInline(FPlatformProcess::BaseDir(), TEXT(""));
+		if (!Result.StartsWith(TEXT("/system/fonts")))
+		{
+			Result.ReplaceInline(TEXT("../"), TEXT(""));
+			Result.ReplaceInline(TEXT(".."), TEXT(""));
+			Result.ReplaceInline(FPlatformProcess::BaseDir(), TEXT(""));
 
-		static FString BasePath = GFilePathBase + FString("/") + GGameName + FString("/");
-		Result =  BasePath + Result;
+			static FString BasePath = GFilePathBase + FString("/") + GGameName + FString("/");
+			Result = BasePath + Result;
+		}
 		return Result;
 	}
 
diff --git a/Engine/Source/Runtime/Core/Private/Android/AndroidMisc.cpp b/Engine/Source/Runtime/Core/Private/Android/AndroidMisc.cpp
--- a/Engine/Source/Runtime/Core/Private/Android/AndroidMisc.cpp
+++ b/Engine/Source/Runtime/Core/Private/Android/AndroidMisc.cpp
@@ -268,7 +268,12 @@ bool FAndroidMisc::GetUseVirtualJoysticks()
 	return !FParse::Param(FCommandLine::Get(),TEXT("joystick"));
 }
 
-
+TArray<uint8> FAndroidMisc::GetSystemFontBytes()
+{
+	TArray<uint8> FontBytes;
+	FFileHelper::LoadFileToArray(FontBytes, TEXT("/system/fonts/DroidSans.ttf"));
+	return FontBytes;
+}
 
 void FAndroidMisc::SetVersionInfo( FString InAndroidVersion, FString InDeviceMake, FString InDeviceModel, FString InOSLanguage )
 {
diff --git a/Engine/Source/Runtime/Core/Public/Android/AndroidMisc.h b/Engine/Source/Runtime/Core/Public/Android/AndroidMisc.h
--- a/Engine/Source/Runtime/Core/Public/Android/AndroidMisc.h
+++ b/Engine/Source/Runtime/Core/Public/Android/AndroidMisc.h
@@ -78,6 +78,9 @@ struct CORE_API FAndroidMisc : public FGenericPlatformMisc
 	static uint32 GetKeyMap( uint16* KeyCodes, FString* KeyNames, uint32 MaxMappings );
 	static const TCHAR* GetDefaultDeviceProfileName() { return TEXT("Android"); }
 
+	/** @return Memory representing a true type or open type font provided by the platform as a default font for unreal to consume; empty array if the default font failed to load. */
+	static TArray<uint8> GetSystemFontBytes();
+
 	// ANDROID ONLY:
 	static void SetVersionInfo( FString AndroidVersion, FString DeviceMake, FString DeviceModel, FString OSLanguage );
 	static const FString GetAndroidVersion();
