Simplify the creation and destruction of CocoaWindow's on Mac OS X to fix some event ordering problems. The old retain/release code was more complicated than it needed to be - we just need to detach the window delegate at the right point. This should resolve a few crashes that the thread separation exposed.

diff --git a/Engine/Source/Runtime/Core/Private/Mac/CocoaWindow.cpp b/Engine/Source/Runtime/Core/Private/Mac/CocoaWindow.cpp
--- a/Engine/Source/Runtime/Core/Private/Mac/CocoaWindow.cpp
+++ b/Engine/Source/Runtime/Core/Private/Mac/CocoaWindow.cpp
@@ -484,6 +484,8 @@ NSString* NSPerformDragOperation = @"NSPerformDragOperation";
 	{
 		FMacEvent::SendToGameRunLoop(notification, self, EMacEventSendMethod::Async, InGameRunLoopMode(@[ UE4CloseEventMode ]));
 	}
+	self.bForwardEvents = false;
+	[self setDelegate:nil];
 }
 
 - (void)mouseDown:(NSEvent*)Event
diff --git a/Engine/Source/Runtime/Core/Private/Mac/MacWindow.cpp b/Engine/Source/Runtime/Core/Private/Mac/MacWindow.cpp
--- a/Engine/Source/Runtime/Core/Private/Mac/MacWindow.cpp
+++ b/Engine/Source/Runtime/Core/Private/Mac/MacWindow.cpp
@@ -85,89 +85,97 @@ void FMacWindow::Initialize( FMacApplication* const Application, const TSharedRe
 		WindowStyle &= ~(NSTexturedBackgroundWindowMask);
 	}
 
-	WindowHandle = MainThreadReturn(^{ return [[FCocoaWindow alloc] initWithContentRect: ViewRect styleMask: WindowStyle backing: NSBackingStoreBuffered defer: NO]; }, UE4NilEventMode);
-
-	if( WindowHandle == NULL )
-	{
-		// @todo Error message should be localized!
-		NSRunInformationalAlertPanel( @"Error", @"Window creation failed!", @"Yes", NULL, NULL );
-		check(0);
-		return;
-	}
-
 	MainThreadCall(^{
-		// Disable automatic release on close - explicit retain/release makes the Destroy() logic much easier to follow
-		[WindowHandle setReleasedWhenClosed: NO];
-		
-		[WindowHandle setWindowMode: WindowMode];
-		[WindowHandle setAcceptsInput: Definition->AcceptsInput];
-		[WindowHandle setDisplayReconfiguring: false];
-		[WindowHandle setAcceptsMouseMovedEvents: YES];
-		[WindowHandle setDelegate: WindowHandle];
-
-		// @todo: We really need a window type in tab manager to know what window level to use and whether or not a window should hide on deactivate
-		if(Definition->IsRegularWindow && (!InParent.IsValid() || Definition->IsModalWindow || (!Definition->SupportsMaximize && !Definition->SupportsMinimize)))
-		{
-			[WindowHandle setLevel: NSNormalWindowLevel];
-		}
-		else
-		{
-			[WindowHandle setLevel: NSFloatingWindowLevel];
-			[WindowHandle setHidesOnDeactivate: YES];
-		}
+		WindowHandle = [[FCocoaWindow alloc] initWithContentRect: ViewRect styleMask: WindowStyle backing: NSBackingStoreBuffered defer: NO];
 		
-		// Use of rounded corners will always render with system values for rounding
-		[WindowHandle setRoundedCorners: (Definition->CornerRadius != 0)];
-		
-		if( !Definition->HasOSWindowBorder )
+		if( WindowHandle != nullptr )
 		{
-			[WindowHandle setBackgroundColor: [NSColor darkGrayColor]];
-			[WindowHandle setHasShadow: YES];
-		}
-
-		[WindowHandle setOpaque: NO];
-
-		ReshapeWindow( X, Y, SizeX, SizeY );
-
-		if( Definition->IsRegularWindow )
-		{
-			CFStringRef CFName = FPlatformString::TCHARToCFString( *Definition->Title );
-			[WindowHandle setTitle: ( NSString *)CFName];
-			CFRelease( CFName );
+			[WindowHandle setWindowMode: WindowMode];
+			[WindowHandle setAcceptsInput: Definition->AcceptsInput];
+			[WindowHandle setDisplayReconfiguring: false];
+			[WindowHandle setAcceptsMouseMovedEvents: YES];
+			[WindowHandle setDelegate: WindowHandle];
+
+			// @todo: We really need a window type in tab manager to know what window level to use and whether or not a window should hide on deactivate
+			if(Definition->IsRegularWindow && (!InParent.IsValid() || Definition->IsModalWindow || (!Definition->SupportsMaximize && !Definition->SupportsMinimize)))
+			{
+				[WindowHandle setLevel: NSNormalWindowLevel];
+			}
+			else
+			{
+				[WindowHandle setLevel: NSFloatingWindowLevel];
+				[WindowHandle setHidesOnDeactivate: YES];
+			}
+			
+			// Use of rounded corners will always render with system values for rounding
+			[WindowHandle setRoundedCorners: (Definition->CornerRadius != 0)];
+			
+			if( !Definition->HasOSWindowBorder )
+			{
+				[WindowHandle setBackgroundColor: [NSColor darkGrayColor]];
+				[WindowHandle setHasShadow: YES];
+			}
 
-			[NSApp addWindowsItem: WindowHandle title: [WindowHandle title] filename: NO];
+			[WindowHandle setOpaque: NO];
 
-			// Tell Cocoa that we are opting into drag and drop.
-			// Only makes sense for regular windows (windows that last a while.)
-			[WindowHandle registerForDraggedTypes: [NSArray arrayWithObject: NSFilenamesPboardType]];
+			ReshapeWindow( X, Y, SizeX, SizeY );
 
-			if( Definition->HasOSWindowBorder )
+			if( Definition->IsRegularWindow )
+			{
+				CFStringRef CFName = FPlatformString::TCHARToCFString( *Definition->Title );
+				[WindowHandle setTitle: ( NSString *)CFName];
+				CFRelease( CFName );
+
+				[NSApp addWindowsItem: WindowHandle title: [WindowHandle title] filename: NO];
+
+				// Tell Cocoa that we are opting into drag and drop.
+				// Only makes sense for regular windows (windows that last a while.)
+				[WindowHandle registerForDraggedTypes: [NSArray arrayWithObject: NSFilenamesPboardType]];
+
+				if( Definition->HasOSWindowBorder )
+				{
+					[WindowHandle setCollectionBehavior: NSWindowCollectionBehaviorFullScreenPrimary|NSWindowCollectionBehaviorDefault|NSWindowCollectionBehaviorManaged|NSWindowCollectionBehaviorParticipatesInCycle];
+				}
+				else
+				{
+					[WindowHandle setCollectionBehavior: NSWindowCollectionBehaviorFullScreenAuxiliary|NSWindowCollectionBehaviorDefault|NSWindowCollectionBehaviorManaged|NSWindowCollectionBehaviorParticipatesInCycle];
+				}
+			}
+			else if(Definition->AppearsInTaskbar)
 			{
-				[WindowHandle setCollectionBehavior: NSWindowCollectionBehaviorFullScreenPrimary|NSWindowCollectionBehaviorDefault|NSWindowCollectionBehaviorManaged|NSWindowCollectionBehaviorParticipatesInCycle];
+				[WindowHandle setCollectionBehavior:NSWindowCollectionBehaviorFullScreenAuxiliary|NSWindowCollectionBehaviorDefault|NSWindowCollectionBehaviorManaged|NSWindowCollectionBehaviorParticipatesInCycle];
 			}
 			else
 			{
-				[WindowHandle setCollectionBehavior: NSWindowCollectionBehaviorFullScreenAuxiliary|NSWindowCollectionBehaviorDefault|NSWindowCollectionBehaviorManaged|NSWindowCollectionBehaviorParticipatesInCycle];
+				[WindowHandle setCollectionBehavior:NSWindowCollectionBehaviorCanJoinAllSpaces|NSWindowCollectionBehaviorTransient|NSWindowCollectionBehaviorIgnoresCycle];
 			}
-		}
-		else if(Definition->AppearsInTaskbar)
-		{
-			[WindowHandle setCollectionBehavior:NSWindowCollectionBehaviorFullScreenAuxiliary|NSWindowCollectionBehaviorDefault|NSWindowCollectionBehaviorManaged|NSWindowCollectionBehaviorParticipatesInCycle];
-		}
-		else
-		{
-			[WindowHandle setCollectionBehavior:NSWindowCollectionBehaviorCanJoinAllSpaces|NSWindowCollectionBehaviorTransient|NSWindowCollectionBehaviorIgnoresCycle];
-		}
 
-		if( Definition->SupportsTransparency )
-		{
-			SetOpacity( Definition->Opacity );
+			if( Definition->SupportsTransparency )
+			{
+				SetOpacity( Definition->Opacity );
+			}
+			else
+			{
+				SetOpacity( 1.0f );
+			}
 		}
 		else
 		{
-			SetOpacity( 1.0f );
+			// @todo Error message should be localized!
+			NSRunInformationalAlertPanel( @"Error", @"Window creation failed!", @"Yes", NULL, NULL );
+			check(0);
 		}
-	}, UE4ResizeEventMode, true);
+	}, UE4ShowEventMode, true);
+	
+	if( WindowHandle == nullptr )
+	{
+		return;
+	}
+	
+	if ( bShowImmediately )
+	{
+		Show();
+	}
 }
 
 FMacWindow::FMacWindow()
@@ -256,16 +264,12 @@ void FMacWindow::Destroy()
 {
 	if( WindowHandle )
 	{
-		// Retain the window as this FMacWindow may be destructed within OnWindowDestroyed.
-		// The destructor will release the window once, so we need to still have a reference.
-		// This makes it possible to correctly handle any NSEvent's sent to the window during destruction
-		// as the WindowHandle is still a valid NSWindow. Unlike HWNDs accessing a destructed NSWindow* is fatal.
-		FCocoaWindow* Window = [WindowHandle retain];
+		FCocoaWindow* Window = WindowHandle;
 		if(Definition->IsModalWindow)
 		{
-			RemoveModalWindow(WindowHandle);
+			RemoveModalWindow(Window);
 		}
-		Window.bForwardEvents = false;
+		
 		if( MacApplication->OnWindowDestroyed( Window ) )
 		{
 			// This FMacWindow may have been destructed by now & so the WindowHandle will probably be invalid memory.
@@ -277,18 +281,11 @@ void FMacWindow::Destroy()
 				// Activate specified previous window if still present, provided it isn't minimized
 				KeyWindow->SetWindowFocus();
 			}
-
-			// Close the window, but don't destruct it...
+			
 			MainThreadCall(^{
 				[Window performClose:nil];
+				WindowHandle = nil;
 			}, UE4CloseEventMode, true);
-			
-			// Since event handling of the change in focus may try to communicate with this window,
-			// dispatch a call to release it which should be handled afterward.
-			// This will ensure proper event handling & destruction order.
-			dispatch_async(dispatch_get_main_queue(), ^{
-				[Window release];
-			});
 		}
 	}
 }
