Added pragma to treat deprecation attribute as warning, not error, on Mac

diff --git a/Engine/Source/Runtime/Core/Public/Mac/MacPlatformCompilerPreSetup.h b/Engine/Source/Runtime/Core/Public/Mac/MacPlatformCompilerPreSetup.h
--- a/Engine/Source/Runtime/Core/Public/Mac/MacPlatformCompilerPreSetup.h
+++ b/Engine/Source/Runtime/Core/Public/Mac/MacPlatformCompilerPreSetup.h
@@ -7,5 +7,6 @@
 #pragma once
 
 #ifndef DISABLE_DEPRECATION
+	#pragma clang diagnostic warning "-Wdeprecated-declarations"
 	#define DEPRECATED(VERSION, MESSAGE) __attribute__((deprecated(MESSAGE " Please update your code to the new API before upgrading to the next release, otherwise your project will no longer compile.")))
 #endif // DISABLE_DEPRECATION
