On OS X don't query the window's type in a mouse moved event, just try and find the associated FMacWindow. Only valid windows will then be forwarded on to Slate.

diff --git a/Engine/Source/Runtime/Core/Private/Mac/MacEvent.cpp b/Engine/Source/Runtime/Core/Private/Mac/MacEvent.cpp
--- a/Engine/Source/Runtime/Core/Private/Mac/MacEvent.cpp
+++ b/Engine/Source/Runtime/Core/Private/Mac/MacEvent.cpp
@@ -13,7 +13,7 @@ FMacEvent::FMacEvent(NSEvent* const Event)
 	[Event retain];
 	
 	NSWindow* Window = FMacApplication::FindMacEventWindow(Event);
-	if(Window && [Window isKindOfClass:[FCocoaWindow class]])
+	if(Window)
 	{
 		MacWindow = FMacApplication::FindMacWindowByNSWindow((FCocoaWindow*)Window);
 	}
@@ -30,7 +30,7 @@ FMacEvent::FMacEvent(NSNotification* const Notification, NSWindow* const Window)
 	check(Notification);
 	[Notification retain];
 	
-	if(Window && [Window isKindOfClass:[FCocoaWindow class]])
+	if(Window)
 	{
 		MacWindow = FMacApplication::FindMacWindowByNSWindow((FCocoaWindow*)Window);
 	}
