Change the GenericPlatformMisc::SubmitErrorReport stub to match the behaviour required & expected. New platforms use the CrashReportClient, not the old auto-reporter which is only used on Windows XP. Remove the Mac implementation - it isn't needed.

diff --git a/Engine/Source/Runtime/Core/Private/GenericPlatform/GenericPlatformMisc.cpp b/Engine/Source/Runtime/Core/Private/GenericPlatform/GenericPlatformMisc.cpp
--- a/Engine/Source/Runtime/Core/Private/GenericPlatform/GenericPlatformMisc.cpp
+++ b/Engine/Source/Runtime/Core/Private/GenericPlatform/GenericPlatformMisc.cpp
@@ -8,6 +8,7 @@
 #include "ModuleManager.h"
 #include "VarargsHelper.h"
 #include "SecureHash.h"
+#include "ExceptionHandling.h"
 #include "Containers/Map.h"
 #include "../../Launch/Resources/Version.h"
 
@@ -199,7 +200,22 @@ FString FGenericPlatformMisc::GetUniqueDeviceId()
 
 void FGenericPlatformMisc::SubmitErrorReport( const TCHAR* InErrorHist, EErrorReportMode::Type InMode )
 {
-	UE_LOG(LogGenericPlatformMisc, Error, TEXT("This platform cannot submit a crash report. Report was:\n%s"), InErrorHist);
+	if ((!FPlatformMisc::IsDebuggerPresent() || GAlwaysReportCrash) && !FParse::Param(FCommandLine::Get(), TEXT("CrashForUAT")))
+	{
+		if ( GUseCrashReportClient )
+		{
+			int32 FromCommandLine = 0;
+			FParse::Value( FCommandLine::Get(), TEXT("AutomatedPerfTesting="), FromCommandLine );
+			if (FApp::IsUnattended() && FromCommandLine != 0 && FParse::Param(FCommandLine::Get(), TEXT("KillAllPopUpBlockingWindows")))
+			{
+				UE_LOG(LogGenericPlatformMisc, Error, TEXT("This platform does not implement KillAllPopUpBlockingWindows"));
+			}
+		}
+		else
+		{
+			UE_LOG(LogGenericPlatformMisc, Error, TEXT("This platform cannot submit a crash report. Report was:\n%s"), InErrorHist);
+		}
+	}
 }
 
 
diff --git a/Engine/Source/Runtime/Core/Private/Mac/MacPlatformMisc.cpp b/Engine/Source/Runtime/Core/Private/Mac/MacPlatformMisc.cpp
--- a/Engine/Source/Runtime/Core/Private/Mac/MacPlatformMisc.cpp
+++ b/Engine/Source/Runtime/Core/Private/Mac/MacPlatformMisc.cpp
@@ -452,21 +452,6 @@ TArray<uint8> FMacPlatformMisc::GetMacAddress()
 	return Result;
 }
 
-void FMacPlatformMisc::SubmitErrorReport(const TCHAR* InErrorHist, EErrorReportMode::Type InMode)
-{
-	if (GUseCrashReportClient && (!FPlatformMisc::IsDebuggerPresent() || GAlwaysReportCrash))
-	{
-		int32 FromCommandLine = 0;
-		FParse::Value( FCommandLine::Get(), TEXT("AutomatedPerfTesting="), FromCommandLine );
-		if (FApp::IsUnattended() && FromCommandLine != 0 && FParse::Param(FCommandLine::Get(), TEXT("KillAllPopUpBlockingWindows")))
-		{
-			abort();
-		}
-	}
-
-	// @todo Mac
-}
-
 void FMacPlatformMisc::PumpMessages( bool bFromMainLoop )
 {
 	if( bFromMainLoop )
diff --git a/Engine/Source/Runtime/Core/Public/Mac/MacPlatformMisc.h b/Engine/Source/Runtime/Core/Public/Mac/MacPlatformMisc.h
--- a/Engine/Source/Runtime/Core/Public/Mac/MacPlatformMisc.h
+++ b/Engine/Source/Runtime/Core/Public/Mac/MacPlatformMisc.h
@@ -23,7 +23,6 @@ struct CORE_API FMacPlatformMisc : public FGenericPlatformMisc
 	}
 
 	static TArray<uint8> GetMacAddress();
-	static void SubmitErrorReport( const TCHAR* InErrorHist, EErrorReportMode::Type InMode );
 #if !UE_BUILD_SHIPPING
 	FORCEINLINE static bool IsDebuggerPresent()
 	{
