Fix to comply with CL 2276830 "Consistent localization-related commandlet path-handling with regards to relative paths. From now on, please specify paths relative to to the project's directory or the Engine directory."

Checking to see if path is relative to Engine or Game, but the Game case is not going to work for now for other reasons (We don't load game *.locres intentionally, so that it falls back to source text.)

diff --git a/Engine/Source/Runtime/Core/Private/Internationalization/TextLocalizationManager.cpp b/Engine/Source/Runtime/Core/Private/Internationalization/TextLocalizationManager.cpp
--- a/Engine/Source/Runtime/Core/Private/Internationalization/TextLocalizationManager.cpp
+++ b/Engine/Source/Runtime/Core/Private/Internationalization/TextLocalizationManager.cpp
@@ -679,6 +679,25 @@ void FTextLocalizationManager::RegenerateResources(const FString& ConfigFilePath
 		}
 	}
 
+	// Source path needs to be relative to Engine or Game directory
+	FString ConfigFullPath = FPaths::ConvertRelativePathToFull(ConfigFilePath);
+	FString EngineFullPath = FPaths::ConvertRelativePathToFull(FPaths::EngineConfigDir());
+	bool IsEngineManifest = false;
+	
+	if (ConfigFullPath.StartsWith(EngineFullPath))
+	{
+		IsEngineManifest = true;
+	}
+
+	if (IsEngineManifest)
+	{
+		SourcePath = FPaths::Combine(*(FPaths::EngineDir()), *SourcePath);
+	}
+	else
+	{
+		SourcePath = FPaths::Combine(*(FPaths::GameDir()), *SourcePath);
+	}
+
 	TArray<TArray<uint8>> BackingBuffers;
 	BackingBuffers.SetNum(LocaleNames.Num());
 
