FPlatformMisc::MessageBoxExt implimentation for HTML5

diff --git a/Engine/Source/Runtime/Core/Core.Build.cs b/Engine/Source/Runtime/Core/Core.Build.cs
--- a/Engine/Source/Runtime/Core/Core.Build.cs
+++ b/Engine/Source/Runtime/Core/Core.Build.cs
@@ -141,6 +141,10 @@ public class Core : ModuleRules
 			AddThirdPartyPrivateStaticDependencies(Target, "SDL");
 			AddThirdPartyPrivateStaticDependencies(Target, "OpenAL");
 		}
+        else if (Target.Platform == UnrealTargetPlatform.HTML5 && Target.Architecture != "-win32")
+        {
+            PrivateDependencyModuleNames.Add("HTML5JS");
+        }
 
 
 		if ((UEBuildConfiguration.bIncludeADO == true) || (UEBuildConfiguration.bCompileAgainstEngine == true))
diff --git a/Engine/Source/Runtime/Core/Private/HTML5/HTML5PlatformMisc.cpp b/Engine/Source/Runtime/Core/Private/HTML5/HTML5PlatformMisc.cpp
--- a/Engine/Source/Runtime/Core/Private/HTML5/HTML5PlatformMisc.cpp
+++ b/Engine/Source/Runtime/Core/Private/HTML5/HTML5PlatformMisc.cpp
@@ -6,6 +6,7 @@
 
 #include "Core.h"
 #include "HTML5Application.h"
+#include "HTML5JavaScriptFx.h"
 #include <ctime>
 
 void FHTML5Misc::PlatformInit()
@@ -129,3 +130,22 @@ uint32 FHTML5Misc::GetKeyMap(uint16* KeyCodes, FString* KeyNames, uint32 MaxMapp
 
 	return NumMappings;
 }
+
+EAppReturnType::Type FHTML5Misc::MessageBoxExt(EAppMsgType::Type MsgType, const TCHAR* Text, const TCHAR* Caption)
+{
+
+#if PLATFORM_HTML5_BROWSER 
+
+	ANSICHAR* AText = TCHAR_TO_ANSI(Text);
+	ANSICHAR* ACaption = TCHAR_TO_ANSI(Caption);
+	return static_cast<EAppReturnType::Type>(UE_MessageBox(MsgType,AText,ACaption));
+
+#endif 
+
+#if PLATFORM_HTML5_WIN32
+
+	return FGenericPlatformMisc::MessageBoxExt(MsgType, Text, Caption); 
+
+#endif 
+
+}
diff --git a/Engine/Source/Runtime/Core/Public/HTML5/HTML5PlatformMisc.h b/Engine/Source/Runtime/Core/Public/HTML5/HTML5PlatformMisc.h
--- a/Engine/Source/Runtime/Core/Public/HTML5/HTML5PlatformMisc.h
+++ b/Engine/Source/Runtime/Core/Public/HTML5/HTML5PlatformMisc.h
@@ -18,7 +18,7 @@ struct CORE_API FHTML5Misc : public FGenericPlatformMisc
 	static class GenericApplication* CreateApplication();
 	static uint32 GetKeyMap(uint16* KeyCodes, FString* KeyNames, uint32 MaxMappings);
 	static uint32 GetCharKeyMap(uint16* KeyCodes, FString* KeyNames, uint32 MaxMappings);
-
+	static EAppReturnType::Type MessageBoxExt( EAppMsgType::Type MsgType, const TCHAR* Text, const TCHAR* Caption );
 	FORCEINLINE static int32 NumberOfCores()
 	{
 		return 1;
diff --git a/Engine/Source/Runtime/HTML5/HTML5JS/Private/HTML5JavaScriptFx.js b/Engine/Source/Runtime/HTML5/HTML5JS/Private/HTML5JavaScriptFx.js
--- a/Engine/Source/Runtime/HTML5/HTML5JS/Private/HTML5JavaScriptFx.js
+++ b/Engine/Source/Runtime/HTML5/HTML5JS/Private/HTML5JavaScriptFx.js
@@ -88,8 +88,20 @@ var UE_JavaScriptLibary =
 		if (b64encoded == null)
 			return false;
 		return true;
-	}
+	},
 
+    UE_MessageBox: function (type, message, caption ) {
+        // type maps to EAppMsgType::Type
+        if ( type == 0 ){
+            var text = Pointer_stringify(message);
+            if (!confirm(text)) 
+                return 0; 
+        } else {
+            var text = Pointer_stringify(message);
+            alert(text);
+        }
+        return 1; 
+    }
 };
 
 mergeInto(LibraryManager.library, UE_JavaScriptLibary);
diff --git a/Engine/Source/Runtime/HTML5/HTML5JS/Public/HTML5JavaScriptFx.h b/Engine/Source/Runtime/HTML5/HTML5JS/Public/HTML5JavaScriptFx.h
--- a/Engine/Source/Runtime/HTML5/HTML5JS/Public/HTML5JavaScriptFx.h
+++ b/Engine/Source/Runtime/HTML5/HTML5JS/Public/HTML5JavaScriptFx.h
@@ -6,5 +6,6 @@
 		 bool UE_DoesSaveGameExist(const char* Name, const int UserIndex);
 		 bool UE_SaveGame(const char* Name, const int UserIndex, const char* indata, int insize);
 		 bool UE_LoadGame(const char* Name, const int UserIndex, char **outdata, int* outsize);
+		 int  UE_MessageBox( int MsgType, const char* Text, const char* Caption);
 	}
 
