Translation Editor: Preview in Editor button should work for most strings now

diff --git a/Engine/Source/Runtime/Core/Private/Internationalization/TextLocalizationManager.cpp b/Engine/Source/Runtime/Core/Private/Internationalization/TextLocalizationManager.cpp
--- a/Engine/Source/Runtime/Core/Private/Internationalization/TextLocalizationManager.cpp
+++ b/Engine/Source/Runtime/Core/Private/Internationalization/TextLocalizationManager.cpp
@@ -632,6 +632,9 @@ namespace
 
 void FTextLocalizationManager::RegenerateResources(const FString& ConfigFilePath)
 {
+	// Add one to the revision index, so all FText's refresh.
+	++HeadCultureRevisionIndex;
+	
 	FInternationalization& I18N = FInternationalization::Get();
 
 	FString SectionName = TEXT("RegenerateResources");
