CrashReport - Removed verbosity flags, we most of the time display all data, removed duplicated FProgramCounterSymbolInfo, partially switched to unified ProgramCounterToHumanReadableString

diff --git a/Engine/Source/Runtime/Core/Private/Android/AndroidPlatformStackWalk.cpp b/Engine/Source/Runtime/Core/Private/Android/AndroidPlatformStackWalk.cpp
--- a/Engine/Source/Runtime/Core/Private/Android/AndroidPlatformStackWalk.cpp
+++ b/Engine/Source/Runtime/Core/Private/Android/AndroidPlatformStackWalk.cpp
@@ -10,7 +10,7 @@
 #include <cxxabi.h>
 #include <stdio.h>
 
-bool FAndroidPlatformStackWalk::ProgramCounterToHumanReadableString(int32 CurrentCallDepth, uint64 ProgramCounter, ANSICHAR* HumanReadableString, SIZE_T HumanReadableStringSize, EVerbosityFlags VerbosityFlags, FGenericCrashContext* Context)
+bool FAndroidPlatformStackWalk::ProgramCounterToHumanReadableString(int32 CurrentCallDepth, uint64 ProgramCounter, ANSICHAR* HumanReadableString, SIZE_T HumanReadableStringSize, FGenericCrashContext* Context)
 {
 	Dl_info DylibInfo;
 	int32 Result = dladdr((const void*)ProgramCounter, &DylibInfo);
@@ -25,10 +25,9 @@ bool FAndroidPlatformStackWalk::ProgramCounterToHumanReadableString(int32 Curren
 	// Write out function name.
 	FCStringAnsi::Strcat(HumanReadableString, HumanReadableStringSize, SymbolInfo.FunctionName);
 
-	if (VerbosityFlags & VF_DISPLAY_FILENAME)
+	// Get filename.
 	{
 		ANSICHAR FileNameLine[MAX_SPRINTF];
-
 		if (SymbolInfo.LineNumber == 0)
 		{
 			// No line number. Print out the module offset address instead.
@@ -43,7 +42,7 @@ bool FAndroidPlatformStackWalk::ProgramCounterToHumanReadableString(int32 Curren
 		FCStringAnsi::Strcat(HumanReadableString, HumanReadableStringSize, FileNameLine);
 	}
 
-	if (VerbosityFlags & VF_DISPLAY_MODULE)
+	// Get module name.
 	{
 		ANSICHAR ModuleName[MAX_SPRINTF];
 		// Write out Module information if there is sufficient space.
@@ -57,7 +56,7 @@ bool FAndroidPlatformStackWalk::ProgramCounterToHumanReadableString(int32 Curren
 	return true;
 }
 
-void FAndroidPlatformStackWalk::ProgramCounterToSymbolInfo(uint64 ProgramCounter, FProgramCounterSymbolInfo&  SymbolInfo)
+void FAndroidPlatformStackWalk::ProgramCounterToSymbolInfo(uint64 ProgramCounter, FProgramCounterSymbolInfo& out_SymbolInfo)
 {
 	Dl_info DylibInfo;
 	int32 Result = dladdr((const void*)ProgramCounter, &DylibInfo);
@@ -78,26 +77,26 @@ void FAndroidPlatformStackWalk::ProgramCounterToSymbolInfo(uint64 ProgramCounter
 	if (DemangledName)
 	{
 		// C++ function
-		FCStringAnsi::Sprintf(SymbolInfo.FunctionName, "%s ", DemangledName);
+		FCStringAnsi::Sprintf(out_SymbolInfo.FunctionName, "%s ", DemangledName);
 	}
 	else if (DylibInfo.dli_sname)
 	{
 		// C function
-		FCStringAnsi::Sprintf(SymbolInfo.FunctionName, "%s() ", DylibInfo.dli_sname);
+		FCStringAnsi::Sprintf(out_SymbolInfo.FunctionName, "%s() ", DylibInfo.dli_sname);
 	}
 	else
 	{
 		// Unknown!
-		FCStringAnsi::Sprintf(SymbolInfo.FunctionName, "[Unknown]() ");
+		FCStringAnsi::Sprintf(out_SymbolInfo.FunctionName, "[Unknown]() ");
 	}
 
 	// No line number available.
 	// TODO open libUE4.so from the apk and get the DWARF-2 data.
-	FCStringAnsi::Strcat(SymbolInfo.Filename, "Unknown");
-	SymbolInfo.LineNumber = 0;
+	FCStringAnsi::Strcat(out_SymbolInfo.Filename, "Unknown");
+	out_SymbolInfo.LineNumber = 0;
 	
 	// Offset of the symbol in the module, eg offset into libUE4.so needed for offline addr2line use.
-	SymbolInfo.OffsetInModule = ProgramCounter - (uint64)DylibInfo.dli_fbase;
+	out_SymbolInfo.OffsetInModule = ProgramCounter - (uint64)DylibInfo.dli_fbase;
 
 	// Write out Module information.
 	ANSICHAR* DylibPath = (ANSICHAR*)DylibInfo.dli_fname;
@@ -110,7 +109,7 @@ void FAndroidPlatformStackWalk::ProgramCounterToSymbolInfo(uint64 ProgramCounter
 	{
 		DylibName = DylibPath;
 	}
-	FCStringAnsi::Strcpy(SymbolInfo.ModuleName, DylibName);
+	FCStringAnsi::Strcpy(out_SymbolInfo.ModuleName, DylibName);
 }
 
 namespace AndroidStackWalkHelpers
diff --git a/Engine/Source/Runtime/Core/Private/Apple/ApplePlatformStackWalk.cpp b/Engine/Source/Runtime/Core/Private/Apple/ApplePlatformStackWalk.cpp
--- a/Engine/Source/Runtime/Core/Private/Apple/ApplePlatformStackWalk.cpp
+++ b/Engine/Source/Runtime/Core/Private/Apple/ApplePlatformStackWalk.cpp
@@ -76,7 +76,7 @@ int32 GetModuleTimeStamp( const struct mach_header* Header )
 	return 0;
 }
 
-static void AsyncSafeProgramCounterToSymbolInfo( uint64 ProgramCounter, FProgramCounterSymbolInfo&  SymbolInfo )
+static void AsyncSafeProgramCounterToSymbolInfo( uint64 ProgramCounter, FProgramCounterSymbolInfo& out_SymbolInfo )
 {
 	Dl_info DylibInfo;
 	int32 Result = dladdr((const void*)ProgramCounter, &DylibInfo);
@@ -89,12 +89,12 @@ static void AsyncSafeProgramCounterToSymbolInfo( uint64 ProgramCounter, FProgram
 	if (DylibInfo.dli_sname)
 	{
 		// Mangled function name
-		FCStringAnsi::Sprintf(SymbolInfo.FunctionName, "%s ", DylibInfo.dli_sname);
+		FCStringAnsi::Sprintf(out_SymbolInfo.FunctionName, "%s ", DylibInfo.dli_sname);
 	}
 	else
 	{
 		// Unknown!
-		FCStringAnsi::Sprintf(SymbolInfo.FunctionName, "[Unknown]() ");
+		FCStringAnsi::Sprintf(out_SymbolInfo.FunctionName, "[Unknown]() ");
 	}
 #else // But on iOS the best we can do is demangle
 	int32 Status = 0;
@@ -109,28 +109,28 @@ static void AsyncSafeProgramCounterToSymbolInfo( uint64 ProgramCounter, FProgram
 	if (DemangledName)
 	{
 		// C++ function
-		FCStringAnsi::Sprintf(SymbolInfo.FunctionName, "%s ", DemangledName);
+		FCStringAnsi::Sprintf(out_SymbolInfo.FunctionName, "%s ", DemangledName);
 	}
 	else if (DylibInfo.dli_sname && strchr(DylibInfo.dli_sname, ']'))
 	{
 		// ObjC function
-		FCStringAnsi::Sprintf(SymbolInfo.FunctionName, "%s ", DylibInfo.dli_sname);
+		FCStringAnsi::Sprintf(out_SymbolInfo.FunctionName, "%s ", DylibInfo.dli_sname);
 	}
 	else if(DylibInfo.dli_sname)
 	{
 		// C function
-		FCStringAnsi::Sprintf(SymbolInfo.FunctionName, "%s() ", DylibInfo.dli_sname);
+		FCStringAnsi::Sprintf(out_SymbolInfo.FunctionName, "%s() ", DylibInfo.dli_sname);
 	}
 	else
 	{
 		// Unknown!
-		FCStringAnsi::Sprintf(SymbolInfo.FunctionName, "[Unknown]() ");
+		FCStringAnsi::Sprintf(out_SymbolInfo.FunctionName, "[Unknown]() ");
 	}
 #endif
 	
 	// No line number found.
-	FCStringAnsi::Strcat(SymbolInfo.Filename, "Unknown");
-	SymbolInfo.LineNumber = 0;
+	FCStringAnsi::Strcat(out_SymbolInfo.Filename, "Unknown");
+	out_SymbolInfo.LineNumber = 0;
 	
 	// Write out Module information.
 	ANSICHAR* DylibPath = (ANSICHAR*)DylibInfo.dli_fname;
@@ -143,7 +143,7 @@ static void AsyncSafeProgramCounterToSymbolInfo( uint64 ProgramCounter, FProgram
 	{
 		DylibName = DylibPath;
 	}
-	FCStringAnsi::Strcpy(SymbolInfo.ModuleName, DylibName);
+	FCStringAnsi::Strcpy(out_SymbolInfo.ModuleName, DylibName);
 }
 
 void FApplePlatformStackWalk::CaptureStackBackTrace( uint64* BackTrace, uint32 MaxDepth, void* Context )
@@ -158,7 +158,7 @@ void FApplePlatformStackWalk::CaptureStackBackTrace( uint64* BackTrace, uint32 M
 	backtrace((void**)BackTrace, MaxDepth);
 }
 
-bool FApplePlatformStackWalk::ProgramCounterToHumanReadableString( int32 CurrentCallDepth, uint64 ProgramCounter, ANSICHAR* HumanReadableString, SIZE_T HumanReadableStringSize, EVerbosityFlags VerbosityFlags, FGenericCrashContext* Context )
+bool FApplePlatformStackWalk::ProgramCounterToHumanReadableString( int32 CurrentCallDepth, uint64 ProgramCounter, ANSICHAR* HumanReadableString, SIZE_T HumanReadableStringSize, FGenericCrashContext* Context )
 {
 	Dl_info DylibInfo;
 	int32 Result = dladdr((const void*)ProgramCounter, &DylibInfo);
@@ -180,7 +180,7 @@ bool FApplePlatformStackWalk::ProgramCounterToHumanReadableString( int32 Current
 	// Write out function name.
 	FCStringAnsi::Strcat(HumanReadableString, HumanReadableStringSize, SymbolInfo.FunctionName);
 
-	if (VerbosityFlags & VF_DISPLAY_FILENAME)
+	// Get filename.
 	{
 		ANSICHAR FileNameLine[MAX_SPRINTF];
 		
@@ -198,7 +198,7 @@ bool FApplePlatformStackWalk::ProgramCounterToHumanReadableString( int32 Current
 		FCStringAnsi::Strcat(HumanReadableString, HumanReadableStringSize, FileNameLine);
 	}
 
-	if (VerbosityFlags & VF_DISPLAY_MODULE)
+	// Get module name.
 	{
 		ANSICHAR ModuleName[MAX_SPRINTF];
 		// Write out Module information if there is sufficient space.
@@ -212,12 +212,12 @@ bool FApplePlatformStackWalk::ProgramCounterToHumanReadableString( int32 Current
 	return true;
 }
 
-void FApplePlatformStackWalk::ProgramCounterToSymbolInfo( uint64 ProgramCounter, FProgramCounterSymbolInfo&  SymbolInfo )
+void FApplePlatformStackWalk::ProgramCounterToSymbolInfo( uint64 ProgramCounter, FProgramCounterSymbolInfo& out_SymbolInfo )
 {
-	bool bOK = FApplePlatformSymbolication::SymbolInfoForAddress(ProgramCounter, SymbolInfo);
+	bool bOK = FApplePlatformSymbolication::SymbolInfoForAddress(ProgramCounter, out_SymbolInfo);
 	if(!bOK)
 	{
-		AsyncSafeProgramCounterToSymbolInfo(ProgramCounter, SymbolInfo);
+		AsyncSafeProgramCounterToSymbolInfo(ProgramCounter, out_SymbolInfo);
 	}
 }
 
diff --git a/Engine/Source/Runtime/Core/Private/Apple/ApplePlatformSymbolication.cpp b/Engine/Source/Runtime/Core/Private/Apple/ApplePlatformSymbolication.cpp
--- a/Engine/Source/Runtime/Core/Private/Apple/ApplePlatformSymbolication.cpp
+++ b/Engine/Source/Runtime/Core/Private/Apple/ApplePlatformSymbolication.cpp
@@ -89,7 +89,7 @@ void FApplePlatformSymbolication::DestroySymbolCache(FApplePlatformSymbolCache*
 #endif
 }
 
-bool FApplePlatformSymbolication::SymbolInfoForAddress(uint64 ProgramCounter, FProgramCounterSymbolInfo& Info)
+bool FApplePlatformSymbolication::SymbolInfoForAddress(uint64 ProgramCounter, FProgramCounterSymbolInfo& out_SymbolInfo)
 {
 #if PLATFORM_MAC
 	bool bOK = false;
@@ -101,19 +101,19 @@ bool FApplePlatformSymbolication::SymbolInfoForAddress(uint64 ProgramCounter, FP
 		
 		if(!CSIsNull(Symbol))
 		{
-			Info.LineNumber = CSSourceInfoGetLineNumber(Symbol);
-			FCStringAnsi::Sprintf(Info.Filename, CSSourceInfoGetPath(Symbol));
-			FCStringAnsi::Sprintf(Info.FunctionName, CSSymbolGetName(CSSourceInfoGetSymbol(Symbol)));
+			out_SymbolInfo.LineNumber = CSSourceInfoGetLineNumber(Symbol);
+			FCStringAnsi::Sprintf(out_SymbolInfo.Filename, CSSourceInfoGetPath(Symbol));
+			FCStringAnsi::Sprintf(out_SymbolInfo.FunctionName, CSSymbolGetName(CSSourceInfoGetSymbol(Symbol)));
 			CSRange CodeRange = CSSourceInfoGetRange(Symbol);
-			Info.SymbolDisplacement = (ProgramCounter - CodeRange.Location);
+			out_SymbolInfo.SymbolDisplacement = (ProgramCounter - CodeRange.Location);
 			
 			CSSymbolOwnerRef Owner = CSSourceInfoGetSymbolOwner(Symbol);
 			if(!CSIsNull(Owner))
 			{
 				ANSICHAR const* DylibName = CSSymbolOwnerGetName(Owner);
-				FCStringAnsi::Strcpy(Info.ModuleName, DylibName);
+				FCStringAnsi::Strcpy(out_SymbolInfo.ModuleName, DylibName);
 				
-				bOK = Info.LineNumber != 0;
+				bOK = out_SymbolInfo.LineNumber != 0;
 			}
 		}
 		
diff --git a/Engine/Source/Runtime/Core/Private/Containers/StackTracker.cpp b/Engine/Source/Runtime/Core/Private/Containers/StackTracker.cpp
--- a/Engine/Source/Runtime/Core/Private/Containers/StackTracker.cpp
+++ b/Engine/Source/Runtime/Core/Private/Containers/StackTracker.cpp
@@ -128,7 +128,7 @@ void FStackTracker::DumpStackTraces( int32 StackThreshold, FOutputDevice& Ar )
 			{
 				ANSICHAR AddressInformation[512];
 				AddressInformation[0] = 0;
-				FPlatformStackWalk::ProgramCounterToHumanReadableString( AddressIndex, CallStack.Addresses[AddressIndex], AddressInformation, ARRAY_COUNT(AddressInformation)-1, VF_DISPLAY_FILENAME );
+				FPlatformStackWalk::ProgramCounterToHumanReadableString( AddressIndex, CallStack.Addresses[AddressIndex], AddressInformation, ARRAY_COUNT(AddressInformation)-1 );
 				CallStackString = CallStackString + LINE_TERMINATOR TEXT(",,,") + FString(AddressInformation);
 			}
 
diff --git a/Engine/Source/Runtime/Core/Private/GenericPlatform/GenericPlatformStackWalk.cpp b/Engine/Source/Runtime/Core/Private/GenericPlatform/GenericPlatformStackWalk.cpp
--- a/Engine/Source/Runtime/Core/Private/GenericPlatform/GenericPlatformStackWalk.cpp
+++ b/Engine/Source/Runtime/Core/Private/GenericPlatform/GenericPlatformStackWalk.cpp
@@ -2,38 +2,54 @@
 
 #include "Core.h"
 
+static const ANSICHAR* GUnknownFunction = "UnknownFunction";
 
-bool FGenericPlatformStackWalk::ProgramCounterToHumanReadableString( int32 CurrentCallDepth, uint64 ProgramCounter, ANSICHAR* HumanReadableString, SIZE_T HumanReadableStringSize, EVerbosityFlags VerbosityFlags, FGenericCrashContext* Context )
+FProgramCounterSymbolInfo::FProgramCounterSymbolInfo() :
+	LineNumber( 0 ),
+	SymbolDisplacement( 0 ),
+	OffsetInModule( 0 ),
+	ProgramCounter( 0 )
+{
+	FCStringAnsi::Strncpy( ModuleName, "UnknownModule", MAX_NAME_LENGHT );
+	FCStringAnsi::Strncpy( FunctionName, GUnknownFunction, MAX_NAME_LENGHT );
+	FCStringAnsi::Strncpy( Filename, "UnknownFile", MAX_NAME_LENGHT );
+}
+
+
+bool FGenericPlatformStackWalk::ProgramCounterToHumanReadableString( int32 CurrentCallDepth, uint64 ProgramCounter, ANSICHAR* HumanReadableString, SIZE_T HumanReadableStringSize, FGenericCrashContext* Context )
 {
 	if (HumanReadableString && HumanReadableStringSize > 0)
 	{
-		ANSICHAR TempArray[MAX_SPRINTF];
-		if (PLATFORM_64BITS)
-		{
-			FCStringAnsi::Sprintf(TempArray, "[Callstack] %p", (void*) ProgramCounter);
-		}
-		else
-		{
-			FCStringAnsi::Sprintf(TempArray, "[Callstack] 0x%.8x", (uint32) ProgramCounter);
-		}
-		FCStringAnsi::Strcat( HumanReadableString, HumanReadableStringSize, TempArray );
+		FProgramCounterSymbolInfo SymbolInfo;
+		FPlatformStackWalk::ProgramCounterToSymbolInfo( ProgramCounter, SymbolInfo );
+
+		// ModuleName!FunctionName (ProgramCounter) + offset bytes [Filename:LineNumber]
+		ANSICHAR StackLine[MAX_SPRINTF];
 
-		if( VerbosityFlags & VF_DISPLAY_FILENAME )
-		{
-			//Append the filename to the string here
-		}
-		return true;
+		// Strip module path.
+		const ANSICHAR* Pos0 = FCStringAnsi::Strrchr( SymbolInfo.ModuleName, '\\' );
+		const ANSICHAR* Pos1 = FCStringAnsi::Strrchr( SymbolInfo.ModuleName, '/' );
+		const UPTRINT RealPos = FMath::Max( (UPTRINT)Pos0, (UPTRINT)Pos1 );
+		const ANSICHAR* StrippedModuleName = RealPos > 0 ? (const ANSICHAR*)(RealPos+1) : SymbolInfo.ModuleName;
+	
+#if	PLATFORM_64BITS
+		FCStringAnsi::Sprintf( StackLine, "%s!%s (0x%016llx) + %i bytes [%s:%i]", StrippedModuleName, (const ANSICHAR*)SymbolInfo.FunctionName, SymbolInfo.ProgramCounter, SymbolInfo.SymbolDisplacement, (const ANSICHAR*)SymbolInfo.Filename, SymbolInfo.LineNumber );
+#else
+		FCStringAnsi::Sprintf( StackLine, "%s!%s (0x%08x) + %i bytes [%s:%i]", StrippedModuleName, (const ANSICHAR*)SymbolInfo.FunctionName, (uint32)SymbolInfo.ProgramCounter, SymbolInfo.OffsetInModule, (const ANSICHAR*)SymbolInfo.Filename, SymbolInfo.LineNumber );
+#endif
+
+		// Append the stack line.
+		FCStringAnsi::Strcat( HumanReadableString, HumanReadableStringSize, StackLine );
+
+		// Return true, if we found a valid function name.
+		return FCStringAnsi::Strncmp( SymbolInfo.FunctionName, GUnknownFunction, FProgramCounterSymbolInfo::MAX_NAME_LENGHT ) != 0;
 	}
-	return true;
+	return false;
 }
 
 void FGenericPlatformStackWalk::CaptureStackBackTrace( uint64* BackTrace, uint32 MaxDepth, void* Context )
 {
-	// nothing we can do generically
-	for( uint32 i=0; i<MaxDepth; i++ )
-	{
-		BackTrace[i] = 0;
-	}
+
 }
 
 void FGenericPlatformStackWalk::StackWalkAndDump( ANSICHAR* HumanReadableString, SIZE_T HumanReadableStringSize, int32 IgnoreCount, void* Context )
@@ -52,8 +68,7 @@ void FGenericPlatformStackWalk::StackWalkAndDump( ANSICHAR* HumanReadableString,
 	// which would mean the top of the callstack is NULL.
 	while( StackTrace[CurrentDepth] || ( CurrentDepth == IgnoreCount ) )
 	{
-		FPlatformStackWalk::ProgramCounterToHumanReadableString( CurrentDepth, StackTrace[CurrentDepth], HumanReadableString, HumanReadableStringSize, 
-			VF_DISPLAY_ALL, reinterpret_cast< FGenericCrashContext* >( Context ) );
+		FPlatformStackWalk::ProgramCounterToHumanReadableString( CurrentDepth, StackTrace[CurrentDepth], HumanReadableString, HumanReadableStringSize, reinterpret_cast< FGenericCrashContext* >( Context ) );
 		FCStringAnsi::Strcat( HumanReadableString, HumanReadableStringSize, LINE_TERMINATOR_ANSI );
 		CurrentDepth++;
 	}
diff --git a/Engine/Source/Runtime/Core/Private/Linux/LinuxPlatformStackWalk.cpp b/Engine/Source/Runtime/Core/Private/Linux/LinuxPlatformStackWalk.cpp
--- a/Engine/Source/Runtime/Core/Private/Linux/LinuxPlatformStackWalk.cpp
+++ b/Engine/Source/Runtime/Core/Private/Linux/LinuxPlatformStackWalk.cpp
@@ -133,7 +133,16 @@ namespace LinuxStackWalkHelpers
 	}
 }
 
-bool FLinuxPlatformStackWalk::ProgramCounterToHumanReadableString( int32 CurrentCallDepth, uint64 ProgramCounter, ANSICHAR* HumanReadableString, SIZE_T HumanReadableStringSize, EVerbosityFlags VerbosityFlags, FGenericCrashContext* Context )
+void FLinuxPlatformStackWalk::ProgramCounterToSymbolInfo( uint64 ProgramCounter, FProgramCounterSymbolInfo& out_SymbolInfo )
+{
+	// Set the program counter.
+	out_SymbolInfo.ProgramCounter = ProgramCounter;
+
+	// Get function, filename and line number.
+	// @TODO
+}
+
+bool FLinuxPlatformStackWalk::ProgramCounterToHumanReadableString( int32 CurrentCallDepth, uint64 ProgramCounter, ANSICHAR* HumanReadableString, SIZE_T HumanReadableStringSize, FGenericCrashContext* Context )
 {
 	if (HumanReadableString && HumanReadableStringSize > 0)
 	{
@@ -164,10 +173,8 @@ bool FLinuxPlatformStackWalk::ProgramCounterToHumanReadableString( int32 Current
 			}
 			LinuxStackWalkHelpers::AppendToString(HumanReadableString, HumanReadableStringSize, Context, TempArray);
 
-			if( VerbosityFlags & VF_DISPLAY_FILENAME )
+			// Get filename.
 			{
-				//Append the filename to the string here
-
 				const char * FunctionName = LinuxStackWalkHelpers::GetFunctionName(Context, CurrentCallDepth);
 				if (FunctionName)
 				{
diff --git a/Engine/Source/Runtime/Core/Private/Windows/WindowsPlatformStackWalk.cpp b/Engine/Source/Runtime/Core/Private/Windows/WindowsPlatformStackWalk.cpp
--- a/Engine/Source/Runtime/Core/Private/Windows/WindowsPlatformStackWalk.cpp
+++ b/Engine/Source/Runtime/Core/Private/Windows/WindowsPlatformStackWalk.cpp
@@ -42,7 +42,7 @@ static TFGetModuleInformation	FGetModuleInformation;
  * This helper function is designed to be called within a structured exception handler.
  *
  * @param	BackTrace			Array to write backtrace to
- * @param	MaxDepth			Maxium depth to walk - needs to be less than or equal to array size
+ * @param	MaxDepth			Maximum depth to walk - needs to be less than or equal to array size
  * @param	Context				Thread context information
  * @return	EXCEPTION_EXECUTE_HANDLER
  */
@@ -66,7 +66,7 @@ static int32 CaptureStackTraceHelper( uint64 *BackTrace, uint32 MaxDepth, CONTEX
 		ThreadHandle	= GetCurrentThread();
 
 		// Zero out stack frame.
-		memset( &StackFrame64, 0, sizeof(StackFrame64) );
+		FMemory::MemZero( StackFrame64 );
 
 		// Initialize the STACKFRAME structure.
 		StackFrame64.AddrPC.Mode         = AddrModeFlat;
@@ -141,10 +141,10 @@ NTSYSAPI uint16 NTAPI RtlCaptureStackBackTrace(
 	);
 
 /** Maximum callstack depth that is supported by the current OS. */
-ULONG GMaxCallstackDepth = 62;
+static ULONG GMaxCallstackDepth = 62;
 
 /** Whether DetermineMaxCallstackDepth() has been called or not. */
-bool GMaxCallstackDepthInitialized = false;
+static bool GMaxCallstackDepthInitialized = false;
 
 /** Maximum callstack depth we support, no matter what OS we're running on. */
 #define MAX_CALLSTACK_DEPTH 128
@@ -169,10 +169,9 @@ void DetermineMaxCallstackDepth()
 void FWindowsPlatformStackWalk::StackWalkAndDump( ANSICHAR* HumanReadableString, SIZE_T HumanReadableStringSize, int32 IgnoreCount, void* Context )
 {
 	InitStackWalking();
-
 	FGenericPlatformStackWalk::StackWalkAndDump(HumanReadableString, HumanReadableStringSize, IgnoreCount, Context);
 }
-
+// @TODO yrx 2014-09-05 Switch to TArray<uint64,TFixedAllocator<100>>>
 /**
  * Capture a stack backtrace and optionally use the passed in exception pointers.
  *
@@ -270,124 +269,25 @@ void FWindowsPlatformStackWalk::CaptureStackBackTrace( uint64* BackTrace, uint32
 
 PRAGMA_ENABLE_OPTIMIZATION
 
-/**
- * Converts the passed in program counter address to a human readable string and appends it to the passed in one.
- * @warning: The code assumes that HumanReadableString is large enough to contain the information.
- *
- * @param	CurrentCallDepth		Depth of the call, if known (-1 if not - note that some platforms may not return meaningful information in the latter case)
- * @param	ProgramCounter			Address to look symbol information up for
- * @param	HumanReadableString		String to concatenate information with
- * @param	HumanReadableStringSize size of string in characters
- * @param	VerbosityFlags			Bit field of requested data for output.
- * @param	Context					Pointer to crash context, if any
- * @return	true if the symbol was found, otherwise false
- */ 
-bool FWindowsPlatformStackWalk::ProgramCounterToHumanReadableString( int32 CurrentCallDepth, uint64 ProgramCounter, ANSICHAR* HumanReadableString, SIZE_T HumanReadableStringSize, EVerbosityFlags VerbosityFlags, FGenericCrashContext* Context )
+void FWindowsPlatformStackWalk::ProgramCounterToSymbolInfo( uint64 ProgramCounter, FProgramCounterSymbolInfo& out_SymbolInfo )
 {
-	ANSICHAR			SymbolBuffer[sizeof(IMAGEHLP_SYMBOL64) + 512];
-	PIMAGEHLP_SYMBOL64	Symbol;
-	uint32				SymbolDisplacement		= 0;
-	DWORD64				SymbolDisplacement64	= 0;
-	uint32				LastError;
-	bool				bSymbolFound = false;
-
-	HANDLE				ProcessHandle = GetCurrentProcess();
-
 	// Initialize stack walking as it loads up symbol information which we require.
 	InitStackWalking();
 
-	// Initialize symbol.
-	Symbol					= (PIMAGEHLP_SYMBOL64) SymbolBuffer;
-	Symbol->SizeOfStruct	= sizeof(SymbolBuffer);
-	Symbol->MaxNameLength	= 512;
+	// Set the program counter.
+	out_SymbolInfo.ProgramCounter = ProgramCounter;
 
-	// Get symbol from address.
-	if( SymGetSymFromAddr64( ProcessHandle, ProgramCounter, &SymbolDisplacement64, Symbol ) )
-	{
-		ANSICHAR			FunctionName[MAX_SPRINTF];
-
-		// Skip any funky chars in the beginning of a function name.
-		int32 Offset = 0;
-		while( Symbol->Name[Offset] < 32 || Symbol->Name[Offset] > 127 )
-		{
-			Offset++;
-		}
-
-		// Write out function name if there is sufficient space.
-		FCStringAnsi::Sprintf( FunctionName,  ("%s() "), (const ANSICHAR*)(Symbol->Name + Offset) );
-		FCStringAnsi::Strcat( HumanReadableString, HumanReadableStringSize, FunctionName );
-		bSymbolFound = true;
-	}
-	else
-	{
-		// No symbol found for this address.
-		LastError = GetLastError( );
-	}
-
-	if( VerbosityFlags & VF_DISPLAY_FILENAME )
-	{
-		IMAGEHLP_LINE64		ImageHelpLine;
-		ANSICHAR			FileNameLine[MAX_SPRINTF];
-
-		// Get Line from address
-		ImageHelpLine.SizeOfStruct = sizeof( ImageHelpLine );
-		if( SymGetLineFromAddr64( ProcessHandle, ProgramCounter, (::DWORD *)&SymbolDisplacement, &ImageHelpLine) )
-		{
-			FCStringAnsi::Sprintf( FileNameLine, ("0x%-8x + %d bytes [File=%s:%d] "), (uint32) ProgramCounter, SymbolDisplacement, (const ANSICHAR*)(ImageHelpLine.FileName), (uint32)ImageHelpLine.LineNumber );
-		}
-		else    
-		{
-			// No line number found.  Print out the logical address instead.
-			FCStringAnsi::Sprintf( FileNameLine, "Address = 0x%-8x (filename not found) ", (uint32) ProgramCounter );
-		}
-		FCStringAnsi::Strcat( HumanReadableString, HumanReadableStringSize, FileNameLine );
-	}
-
-	if( VerbosityFlags & VF_DISPLAY_MODULE )
-	{
-		IMAGEHLP_MODULE64	ImageHelpModule;
-		ANSICHAR			ModuleName[MAX_SPRINTF];
-
-		// Get module information from address.
-		ImageHelpModule.SizeOfStruct = sizeof( ImageHelpModule );
-		if( SymGetModuleInfo64( ProcessHandle, ProgramCounter, &ImageHelpModule) )
-		{
-			// Write out Module information if there is sufficient space.
-			FCStringAnsi::Sprintf( ModuleName, "[in %s]", (const ANSICHAR*)(ImageHelpModule.ImageName) );
-			FCStringAnsi::Strcat( HumanReadableString, HumanReadableStringSize, ModuleName );
-		}
-		else
-		{
-			LastError = GetLastError( );
-		}
-	}
-
-	return bSymbolFound;
-}
-
-void FWindowsPlatformStackWalk::ProgramCounterToSymbolInfo( uint64 ProgramCounter, FProgramCounterSymbolInfo&  SymbolInfo )
-{
-	// Create zeroed out return value.
-	FMemory::Memzero( &SymbolInfo, sizeof(SymbolInfo) );
-
-	ANSICHAR			SymbolBuffer[sizeof(IMAGEHLP_SYMBOL64) + 512];
-	PIMAGEHLP_SYMBOL64	Symbol;
-	uint32				SymbolDisplacement		= 0;
-	DWORD64				SymbolDisplacement64	= 0;
-	uint32				LastError;
-	
-	HANDLE				ProcessHandle = GetCurrentProcess();
-
-	// Initialize stack walking as it loads up symbol information which we require.
-	InitStackWalking();
+	uint32 LastError = 0;
+	HANDLE ProcessHandle = GetCurrentProcess();
 
 	// Initialize symbol.
-	Symbol					= (PIMAGEHLP_SYMBOL64) SymbolBuffer;
-	Symbol->SizeOfStruct	= sizeof(SymbolBuffer);
-	Symbol->MaxNameLength	= 512;
+	ANSICHAR SymbolBuffer[sizeof( IMAGEHLP_SYMBOL64 ) + FProgramCounterSymbolInfo::MAX_NAME_LENGHT] = {0};
+	IMAGEHLP_SYMBOL64* Symbol = (IMAGEHLP_SYMBOL64*)SymbolBuffer;
+	Symbol->SizeOfStruct = sizeof(SymbolBuffer);
+	Symbol->MaxNameLength = FProgramCounterSymbolInfo::MAX_NAME_LENGHT;
 
-	// Get symbol from address.
-	if( SymGetSymFromAddr64( ProcessHandle, ProgramCounter, &SymbolDisplacement64, Symbol ) )
+	// Get function name.
+	if( SymGetSymFromAddr64( ProcessHandle, ProgramCounter, nullptr, Symbol ) )
 	{
 		// Skip any funky chars in the beginning of a function name.
 		int32 Offset = 0;
@@ -397,44 +297,40 @@ void FWindowsPlatformStackWalk::ProgramCounterToSymbolInfo( uint64 ProgramCounte
 		}
 
 		// Write out function name.
-		FCStringAnsi::Strcpy( SymbolInfo.FunctionName, Symbol->Name + Offset );
+		FCStringAnsi::Strncpy( out_SymbolInfo.FunctionName, Symbol->Name + Offset, FProgramCounterSymbolInfo::MAX_NAME_LENGHT ); 
+		FCStringAnsi::Strncat( out_SymbolInfo.FunctionName, "()", FProgramCounterSymbolInfo::MAX_NAME_LENGHT );
 	}
 	else
 	{
 		// No symbol found for this address.
-		LastError = GetLastError( );
+		LastError = GetLastError();
 	}
 
-	// Get Line from address
-	IMAGEHLP_LINE64	ImageHelpLine;
+	// Get filename and line number.
+	IMAGEHLP_LINE64	ImageHelpLine = {0};
 	ImageHelpLine.SizeOfStruct = sizeof( ImageHelpLine );
-	if( SymGetLineFromAddr64( ProcessHandle, ProgramCounter, (::DWORD *)&SymbolDisplacement, &ImageHelpLine) )
+	if( SymGetLineFromAddr64( ProcessHandle, ProgramCounter, (::DWORD *)&out_SymbolInfo.SymbolDisplacement, &ImageHelpLine ) )
 	{
-		FCStringAnsi::Strcpy( SymbolInfo.Filename, ImageHelpLine.FileName );
-		SymbolInfo.LineNumber			= ImageHelpLine.LineNumber;
-		SymbolInfo.SymbolDisplacement	= SymbolDisplacement;
+		FCStringAnsi::Strncpy( out_SymbolInfo.Filename, ImageHelpLine.FileName, FProgramCounterSymbolInfo::MAX_NAME_LENGHT );
+		out_SymbolInfo.LineNumber = ImageHelpLine.LineNumber;
 	}
-	else    
+	else
 	{
-		// No line number found.
-		FCStringAnsi::Strcat( SymbolInfo.Filename, "Unknown" );
-		SymbolInfo.LineNumber			= 0;
-		SymbolDisplacement				= 0;
+		LastError = GetLastError();
 	}
 
-	// Get module information from address.
-	IMAGEHLP_MODULE64 ImageHelpModule;
+	// Get module name.
+	IMAGEHLP_MODULE64 ImageHelpModule = {0};
 	ImageHelpModule.SizeOfStruct = sizeof( ImageHelpModule );
 	if( SymGetModuleInfo64( ProcessHandle, ProgramCounter, &ImageHelpModule) )
 	{
-		// Write out Module information.
-		FCStringAnsi::Strcpy( SymbolInfo.ModuleName, ImageHelpModule.ImageName );
+		// Write out module information.
+		FCStringAnsi::Strncpy( out_SymbolInfo.ModuleName, ImageHelpModule.ImageName, FProgramCounterSymbolInfo::MAX_NAME_LENGHT );
 	}
 	else
 	{
-		LastError = GetLastError( );
+		LastError = GetLastError();
 	}
-
 }
 
 /**
@@ -442,13 +338,12 @@ void FWindowsPlatformStackWalk::ProgramCounterToSymbolInfo( uint64 ProgramCounte
  */ 
 static void LoadProcessModules()
 {
-	int32			ErrorCode = 0;	
-	HANDLE		ProcessHandle = GetCurrentProcess(); 
+	int32 ErrorCode = 0;
+	HANDLE ProcessHandle = GetCurrentProcess();
 	const int32	MAX_MOD_HANDLES = 1024;
-	HMODULE		ModuleHandleArray[MAX_MOD_HANDLES];
-	HMODULE*	ModuleHandlePointer = ModuleHandleArray;
-	uint32		BytesRequired;
-	MODULEINFO	ModuleInfo;
+	HMODULE ModuleHandleArray[MAX_MOD_HANDLES] = {0};
+	HMODULE* ModuleHandlePointer = ModuleHandleArray;
+	uint32 BytesRequired = 0;
 
 	// Enumerate process modules.
 	bool bEnumProcessModulesSucceeded = FEnumProcessModules( ProcessHandle, ModuleHandleArray, sizeof(ModuleHandleArray), (::DWORD *)&BytesRequired );
@@ -464,40 +359,42 @@ static void LoadProcessModules()
 	{
 		// Keep track of the fact that we need to free it again.
 		bNeedToFreeModuleHandlePointer = true;
+		// @TODO yrx 2014-09-05 Calling malloc after a crash is not a good idea, fix it.
 		ModuleHandlePointer = (HMODULE*) FMemory::Malloc( BytesRequired );
 		FEnumProcessModules( ProcessHandle, ModuleHandlePointer, sizeof(ModuleHandleArray), (::DWORD *)&BytesRequired );
 	}
 
 	// Find out how many modules we need to load modules for.
-	int32	ModuleCount = BytesRequired / sizeof( HMODULE );
+	const int32	ModuleCount = BytesRequired / sizeof( HMODULE );
 
 	// Load the modules.
-	for( int32 ModuleIndex=0; ModuleIndex<ModuleCount; ModuleIndex++ )
+	for( int32 ModuleIndex = 0; ModuleIndex < ModuleCount; ModuleIndex++ )
 	{
-		ANSICHAR ModuleName[1024];
-		ANSICHAR ImageName[1024];
+		MODULEINFO ModuleInfo = {0};
+		ANSICHAR ModuleName[FProgramCounterSymbolInfo::MAX_NAME_LENGHT] = {0};
+		ANSICHAR ImageName[FProgramCounterSymbolInfo::MAX_NAME_LENGHT] = {0};
 #if PLATFORM_64BITS
-		static_assert(sizeof(ModuleInfo) == 24, "Broken alignment for 64bit Windows include.");
+		static_assert(sizeof( MODULEINFO ) == 24, "Broken alignment for 64bit Windows include.");
 #else
-		static_assert(sizeof(ModuleInfo) == 12, "Broken alignment for 32bit Windows include.");
+		static_assert(sizeof( MODULEINFO ) == 12, "Broken alignment for 32bit Windows include.");
 #endif
-		FGetModuleInformation( ProcessHandle, ModuleHandleArray[ModuleIndex], &ModuleInfo,sizeof( ModuleInfo ) );
-		FGetModuleFileNameEx( ProcessHandle, ModuleHandleArray[ModuleIndex], ImageName, 1024 );
-		FGetModuleBaseName( ProcessHandle, ModuleHandleArray[ModuleIndex], ModuleName, 1024 );
+		FGetModuleInformation( ProcessHandle, ModuleHandleArray[ModuleIndex], &ModuleInfo, sizeof( ModuleInfo ) );
+		FGetModuleFileNameEx( ProcessHandle, ModuleHandleArray[ModuleIndex], ImageName, FProgramCounterSymbolInfo::MAX_NAME_LENGHT );
+		FGetModuleBaseName( ProcessHandle, ModuleHandleArray[ModuleIndex], ModuleName, FProgramCounterSymbolInfo::MAX_NAME_LENGHT );
 
 		// Set the search path to find PDBs in the same folder as the DLL.
-		ANSICHAR SearchPath[1024];
+		ANSICHAR SearchPath[MAX_PATH] = {0};
 		ANSICHAR* FileName = NULL;
-		const auto Result =	GetFullPathNameA( ImageName, ARRAY_COUNT( SearchPath ), SearchPath, &FileName );
+		const auto Result = GetFullPathNameA( ImageName, MAX_PATH, SearchPath, &FileName );
 
-		if(Result != 0 && Result < ARRAY_COUNT(SearchPath))
+		if( Result != 0 && Result < MAX_PATH )
 		{
 			*FileName = 0;
 			SymSetSearchPath(GetCurrentProcess(), SearchPath);
 		}
 
 		// Load module.
-		DWORD64 BaseAddress = SymLoadModule64( ProcessHandle, ModuleHandleArray[ModuleIndex], ImageName, ModuleName, (DWORD64) ModuleInfo.lpBaseOfDll, (uint32) ModuleInfo.SizeOfImage );
+		const DWORD64 BaseAddress = SymLoadModule64( ProcessHandle, ModuleHandleArray[ModuleIndex], ImageName, ModuleName, (DWORD64) ModuleInfo.lpBaseOfDll, (uint32) ModuleInfo.SizeOfImage );
 		if( !BaseAddress )
 		{
 			ErrorCode = GetLastError();
@@ -515,11 +412,11 @@ int32 FWindowsPlatformStackWalk::GetProcessModuleCount()
 {
 	FPlatformStackWalk::InitStackWalking();
 
-	HANDLE		ProcessHandle = GetCurrentProcess(); 
+	HANDLE ProcessHandle = GetCurrentProcess(); 
 	const int32	MAX_MOD_HANDLES = 1024;
-	HMODULE		ModuleHandleArray[MAX_MOD_HANDLES];
-	HMODULE*	ModuleHandlePointer = ModuleHandleArray;
-	uint32		BytesRequired;
+	HMODULE	ModuleHandleArray[MAX_MOD_HANDLES] = {0};
+	HMODULE* ModuleHandlePointer = ModuleHandleArray;
+	uint32 BytesRequired = 0;
 
 	// Enumerate process modules.
 	bool bEnumProcessModulesSucceeded = FEnumProcessModules( ProcessHandle, ModuleHandleArray, sizeof(ModuleHandleArray), (::DWORD *)&BytesRequired );
@@ -528,25 +425,8 @@ int32 FWindowsPlatformStackWalk::GetProcessModuleCount()
 		return 0;
 	}
 
-	// Static array isn't sufficient so we dynamically allocate one.
-	bool bNeedToFreeModuleHandlePointer = false;
-	if( BytesRequired > sizeof( ModuleHandleArray ) )
-	{
-		// Keep track of the fact that we need to free it again.
-		bNeedToFreeModuleHandlePointer = true;
-		ModuleHandlePointer = (HMODULE*) FMemory::Malloc( BytesRequired );
-		FEnumProcessModules( ProcessHandle, ModuleHandlePointer, sizeof(ModuleHandleArray), (::DWORD *)&BytesRequired );
-	}
-
 	// Find out how many modules we need to load modules for.
-	int32	ModuleCount = BytesRequired / sizeof( HMODULE );
-
-	// Free the module handle pointer allocated in case the static array was insufficient.
-	if( bNeedToFreeModuleHandlePointer )
-	{
-		FMemory::Free( ModuleHandlePointer );
-	}
-
+	const int32 ModuleCount = BytesRequired / sizeof( HMODULE );
 	return ModuleCount;
 }
 
@@ -556,10 +436,9 @@ int32 FWindowsPlatformStackWalk::GetProcessModuleSignatures(FStackWalkModuleInfo
 
 	HANDLE		ProcessHandle = GetCurrentProcess(); 
 	const int32	MAX_MOD_HANDLES = 1024;
-	HMODULE		ModuleHandleArray[MAX_MOD_HANDLES];
+	HMODULE		ModuleHandleArray[MAX_MOD_HANDLES] = {0};
 	HMODULE*	ModuleHandlePointer = ModuleHandleArray;
 	uint32		BytesRequired;
-	MODULEINFO	ModuleInfo;
 
 	// Enumerate process modules.
 	bool bEnumProcessModulesSucceeded = FEnumProcessModules( ProcessHandle, ModuleHandleArray, sizeof(ModuleHandleArray), (::DWORD *)&BytesRequired );
@@ -579,8 +458,8 @@ int32 FWindowsPlatformStackWalk::GetProcessModuleSignatures(FStackWalkModuleInfo
 	}
 
 	// Find out how many modules we need to load modules for.
-	int32	ModuleCount = BytesRequired / sizeof( HMODULE );
-	IMAGEHLP_MODULEW64 Img;
+	const int32 ModuleCount = BytesRequired / sizeof( HMODULE );
+	IMAGEHLP_MODULEW64 Img = {0};
 	Img.SizeOfStruct = sizeof(Img);
 
 	int32 SignatureIndex = 0;
@@ -588,21 +467,22 @@ int32 FWindowsPlatformStackWalk::GetProcessModuleSignatures(FStackWalkModuleInfo
 	// Load the modules.
 	for( int32 ModuleIndex = 0; ModuleIndex < ModuleCount && SignatureIndex < ModuleSignaturesSize; ModuleIndex++ )
 	{
-		ANSICHAR ModuleName[1024];
-		ANSICHAR ImageName[1024];
+		MODULEINFO ModuleInfo = {0};
+		ANSICHAR ModuleName[MAX_PATH] = {0};
+		ANSICHAR ImageName[MAX_PATH] = {0};
 #if PLATFORM_64BITS
-		static_assert(sizeof(ModuleInfo) == 24, "Broken alignment for 64bit Windows include.");
+		static_assert(sizeof( MODULEINFO ) == 24, "Broken alignment for 64bit Windows include.");
 #else
-		static_assert(sizeof(ModuleInfo) == 12, "Broken alignment for 32bit Windows include.");
+		static_assert(sizeof( MODULEINFO ) == 12, "Broken alignment for 32bit Windows include.");
 #endif
 		FGetModuleInformation( ProcessHandle, ModuleHandleArray[ModuleIndex], &ModuleInfo,sizeof( ModuleInfo ) );
-		FGetModuleFileNameEx( ProcessHandle, ModuleHandleArray[ModuleIndex], ImageName, 1024 );
-		FGetModuleBaseName( ProcessHandle, ModuleHandleArray[ModuleIndex], ModuleName, 1024 );
+		FGetModuleFileNameEx( ProcessHandle, ModuleHandleArray[ModuleIndex], ImageName, MAX_PATH );
+		FGetModuleBaseName( ProcessHandle, ModuleHandleArray[ModuleIndex], ModuleName, MAX_PATH );
 
 		// Load module.
 		if(SymGetModuleInfoW64(ProcessHandle, (DWORD64)ModuleInfo.lpBaseOfDll, &Img))
 		{
-			FStackWalkModuleInfo Info;
+			FStackWalkModuleInfo Info = {0};
 			Info.BaseOfImage = Img.BaseOfImage;
 			FCString::Strcpy(Info.ImageName, Img.ImageName);
 			Info.ImageSize = Img.ImageSize;
diff --git a/Engine/Source/Runtime/Core/Public/Android/AndroidPlatformStackWalk.h b/Engine/Source/Runtime/Core/Public/Android/AndroidPlatformStackWalk.h
--- a/Engine/Source/Runtime/Core/Public/Android/AndroidPlatformStackWalk.h
+++ b/Engine/Source/Runtime/Core/Public/Android/AndroidPlatformStackWalk.h
@@ -10,26 +10,6 @@
 #undef PLATFORM_SUPPORTS_STACK_SYMBOLS
 #define PLATFORM_SUPPORTS_STACK_SYMBOLS 1
 
-/**
-* Symbol information associated with a program counter.
-*/
-struct FProgramCounterSymbolInfo
-{
-	/** Module name.					*/
-	ANSICHAR	ModuleName[1024];
-	/** Function name.					*/
-	ANSICHAR	FunctionName[1024];
-	/** Filename.						*/
-	ANSICHAR	Filename[1024];
-	/** Line number in file.			*/
-	int32		LineNumber;
-	/** ProgramCounter offset into module.	*/
-	uint64		OffsetInModule;
-	/** Symbol displacement of address.	*/
-	int32		SymbolDisplacement;
-};
-
-
 /**
 * Android platform stack walking
 */
@@ -37,8 +17,8 @@ struct CORE_API FAndroidPlatformStackWalk : public FGenericPlatformStackWalk
 {
 	typedef FGenericPlatformStackWalk Parent;
 
-	static bool ProgramCounterToHumanReadableString(int32 CurrentCallDepth, uint64 ProgramCounter, ANSICHAR* HumanReadableString, SIZE_T HumanReadableStringSize, EVerbosityFlags VerbosityFlags = VF_DISPLAY_ALL, FGenericCrashContext* Context = NULL);
-	static void ProgramCounterToSymbolInfo(uint64 ProgramCounter, FProgramCounterSymbolInfo&  OutProgramCounterSymbolInfo);
+	static bool ProgramCounterToHumanReadableString(int32 CurrentCallDepth, uint64 ProgramCounter, ANSICHAR* HumanReadableString, SIZE_T HumanReadableStringSize, FGenericCrashContext* Context = NULL);
+	static void ProgramCounterToSymbolInfo(uint64 ProgramCounter, FProgramCounterSymbolInfo& out_SymbolInfo);
 	static void CaptureStackBackTrace(uint64* BackTrace, uint32 MaxDepth, void* Context = NULL);
 };
 
diff --git a/Engine/Source/Runtime/Core/Public/Apple/ApplePlatformStackWalk.h b/Engine/Source/Runtime/Core/Public/Apple/ApplePlatformStackWalk.h
--- a/Engine/Source/Runtime/Core/Public/Apple/ApplePlatformStackWalk.h
+++ b/Engine/Source/Runtime/Core/Public/Apple/ApplePlatformStackWalk.h
@@ -9,32 +9,14 @@
 #undef PLATFORM_SUPPORTS_STACK_SYMBOLS
 #define PLATFORM_SUPPORTS_STACK_SYMBOLS 1
 
-/**
- * Symbol information associated with a program counter.
- */
-struct FProgramCounterSymbolInfo
-{
-	/** Module name.					*/
-	ANSICHAR	ModuleName[1024];
-	/** Function name.					*/
-	ANSICHAR	FunctionName[1024];
-	/** Filename.						*/
-	ANSICHAR	Filename[1024];
-	/** Line number in file.			*/
-	int32			LineNumber;
-	/** Symbol displacement of address.	*/
-	int32			SymbolDisplacement;
-};
-
-
 
 /**
  * Apple platform implementation of the misc OS functions
  */
 struct CORE_API FApplePlatformStackWalk : public FGenericPlatformStackWalk
 {
-	static bool ProgramCounterToHumanReadableString( int32 CurrentCallDepth, uint64 ProgramCounter, ANSICHAR* HumanReadableString, SIZE_T HumanReadableStringSize, EVerbosityFlags VerbosityFlags = VF_DISPLAY_ALL, FGenericCrashContext* Context = NULL );
-	static void ProgramCounterToSymbolInfo( uint64 ProgramCounter, FProgramCounterSymbolInfo&  OutProgramCounterSymbolInfo);
+	static bool ProgramCounterToHumanReadableString( int32 CurrentCallDepth, uint64 ProgramCounter, ANSICHAR* HumanReadableString, SIZE_T HumanReadableStringSize, FGenericCrashContext* Context = NULL );
+	static void ProgramCounterToSymbolInfo( uint64 ProgramCounter, FProgramCounterSymbolInfo&  out_SymbolInfo);
 	static void CaptureStackBackTrace( uint64* BackTrace, uint32 MaxDepth, void* Context = NULL );
 	static int32 GetProcessModuleCount();
 	static int32 GetProcessModuleSignatures(FStackWalkModuleInfo *ModuleSignatures, const int32 ModuleSignaturesSize);
diff --git a/Engine/Source/Runtime/Core/Public/GenericPlatform/GenericPlatformStackWalk.h b/Engine/Source/Runtime/Core/Public/GenericPlatform/GenericPlatformStackWalk.h
--- a/Engine/Source/Runtime/Core/Public/GenericPlatform/GenericPlatformStackWalk.h
+++ b/Engine/Source/Runtime/Core/Public/GenericPlatform/GenericPlatformStackWalk.h
@@ -7,18 +7,6 @@
 
 #pragma once
 
-/** @name ObjectFlags
-* Flags used to control the output from stack tracing
-*/
-typedef uint32 EVerbosityFlags;
-
-#define VF_DISPLAY_BASIC		0x00000000
-#define VF_DISPLAY_FILENAME		0x00000001
-#define VF_DISPLAY_MODULE		0x00000002
-#define VF_DISPLAY_ALL			0xffffffff
-
-#define PLATFORM_SUPPORTS_STACK_SYMBOLS 0
-
 /**
 * This is used to capture all of the module information needed to load pdb's.
 * @todo, the memory profiler should not be using this as platform agnostic
@@ -42,13 +30,49 @@ struct FStackWalkModuleInfo
 	} PdbSig70;
 };
 
+/**
+ * Symbol information associated with a program counter.
+ */
+struct FProgramCounterSymbolInfo /*final*/
+{
+	enum
+	{
+		/** Length of the string used to store the symbol's names, including the trailing character. */
+		MAX_NAME_LENGHT = 1024,
+	};
+
+	/** Module name. */
+	ANSICHAR	ModuleName[MAX_NAME_LENGHT];
+
+	/** Function name. */
+	ANSICHAR	FunctionName[MAX_NAME_LENGHT];
+
+	/** Filename. */
+	ANSICHAR	Filename[MAX_NAME_LENGHT];
+
+	/** Line number in file. */
+	int32		LineNumber;
 
+	/** Symbol displacement of address.	*/
+	int32		SymbolDisplacement;
+
+	/** Program counter offset into module. */
+	uint64		OffsetInModule;
+
+	/** Program counter. */
+	uint64		ProgramCounter;
+
+	/** Default constructor. */
+	FProgramCounterSymbolInfo();
+};
 
 /**
 * Generic implementation for most platforms
 **/
 struct CORE_API FGenericPlatformStackWalk
 {
+	typedef FGenericPlatformStackWalk Base;
+
 	/**
 	 * Initializes stack traversal and symbol. Must be called before any other stack/symbol functions. Safe to reenter.
 	 */
@@ -60,27 +84,35 @@ struct CORE_API FGenericPlatformStackWalk
 	/**
 	 * Converts the passed in program counter address to a human readable string and appends it to the passed in one.
 	 * @warning: The code assumes that HumanReadableString is large enough to contain the information.
+	 * @warning: Please, do not override this method. Can't be modified or altered without notice.
+	 * 
+	 * This method is the same for all platforms to simplify parsing by the crash processor.
+	 * 
+	 * Example formatted line:
+	 *
+	 * UE4Editor_Core!FOutputDeviceWindowsError::Serialize() (0xddf1bae5) + 620 bytes [\engine\source\runtime\core\private\windows\windowsplatformoutputdevices.cpp:110]
+	 * ModuleName!FunctionName (ProgramCounter) + offset bytes [StrippedFilepath:LineNumber]
 	 *
 	 * @param	CurrentCallDepth		Depth of the call, if known (-1 if not - note that some platforms may not return meaningful information in the latter case)
 	 * @param	ProgramCounter			Address to look symbol information up for
 	 * @param	HumanReadableString		String to concatenate information with
 	 * @param	HumanReadableStringSize size of string in characters
-	 * @param	VerbosityFlags			Bit field of requested data for output. -1 for all output.
 	 * @param	Context					Pointer to crash context, if any
 	 * @return	true if the symbol was found, otherwise false
 	 */ 
-	static bool ProgramCounterToHumanReadableString( int32 CurrentCallDepth, uint64 ProgramCounter, ANSICHAR* HumanReadableString, SIZE_T HumanReadableStringSize, EVerbosityFlags VerbosityFlags = VF_DISPLAY_ALL, FGenericCrashContext* Context = NULL );
+	static bool ProgramCounterToHumanReadableString( int32 CurrentCallDepth, uint64 ProgramCounter, ANSICHAR* HumanReadableString, SIZE_T HumanReadableStringSize, FGenericCrashContext* Context = NULL );
 
-#if PLATFORM_SUPPORTS_STACK_SYMBOLS // never true here, but we provide this function signature for reference
 	/**
 	 * Converts the passed in program counter address to a symbol info struct, filling in module and filename, line number and displacement.
 	 * @warning: The code assumes that the destination strings are big enough
 	 *
-	 * @param	ProgramCounter					Address to look symbol information up for
-	 * @param	OutProgramCounterSymbolInfo		symbol information associated with program counter
+	 * @param	ProgramCounter		Address to look symbol information up for
+	 * @param	out_SymbolInfo		Symbol information associated with program counter
 	 */
-	 static void ProgramCounterToSymbolInfo( uint64 ProgramCounter, FProgramCounterSymbolInfo&  OutProgramCounterSymbolInfo);
-#endif
+	static void ProgramCounterToSymbolInfo( uint64 ProgramCounter, FProgramCounterSymbolInfo& out_SymbolInfo)
+	{
+		out_SymbolInfo.ProgramCounter = ProgramCounter;
+	}
 
 	/**
 	 * Capture a stack backtrace and optionally use the passed in exception pointers.
@@ -122,6 +154,4 @@ struct CORE_API FGenericPlatformStackWalk
 	{
 		return 0;
 	}
-
-
 };
diff --git a/Engine/Source/Runtime/Core/Public/Linux/LinuxPlatformStackWalk.h b/Engine/Source/Runtime/Core/Public/Linux/LinuxPlatformStackWalk.h
--- a/Engine/Source/Runtime/Core/Public/Linux/LinuxPlatformStackWalk.h
+++ b/Engine/Source/Runtime/Core/Public/Linux/LinuxPlatformStackWalk.h
@@ -9,9 +9,8 @@
 
 struct CORE_API FLinuxPlatformStackWalk : public FGenericPlatformStackWalk
 {
-	typedef FGenericPlatformStackWalk Parent;
-
-	static bool ProgramCounterToHumanReadableString( int32 CurrentCallDepth, uint64 ProgramCounter, ANSICHAR* HumanReadableString, SIZE_T HumanReadableStringSize, EVerbosityFlags VerbosityFlags = VF_DISPLAY_ALL, FGenericCrashContext* Context = NULL );
+	static void ProgramCounterToSymbolInfo( uint64 ProgramCounter, FProgramCounterSymbolInfo& out_SymbolInfo );
+	static bool ProgramCounterToHumanReadableString( int32 CurrentCallDepth, uint64 ProgramCounter, ANSICHAR* HumanReadableString, SIZE_T HumanReadableStringSize, FGenericCrashContext* Context = NULL );
 	static void CaptureStackBackTrace( uint64* BackTrace, uint32 MaxDepth, void* Context = NULL );
 };
 
diff --git a/Engine/Source/Runtime/Core/Public/Windows/WindowsPlatformStackWalk.h b/Engine/Source/Runtime/Core/Public/Windows/WindowsPlatformStackWalk.h
--- a/Engine/Source/Runtime/Core/Public/Windows/WindowsPlatformStackWalk.h
+++ b/Engine/Source/Runtime/Core/Public/Windows/WindowsPlatformStackWalk.h
@@ -10,35 +10,18 @@
 #undef PLATFORM_SUPPORTS_STACK_SYMBOLS
 #define PLATFORM_SUPPORTS_STACK_SYMBOLS 1
 
-/**
- * Symbol information associated with a program counter.
- */
-struct FProgramCounterSymbolInfo
-{
-	/** Module name.					*/
-	ANSICHAR	ModuleName[1024];
-	/** Function name.					*/
-	ANSICHAR	FunctionName[1024];
-	/** Filename.						*/
-	ANSICHAR	Filename[1024];
-	/** Line number in file.			*/
-	int32		LineNumber;
-	/** Symbol displacement of address.	*/
-	int32		SymbolDisplacement;
-};
-
-
 
 /**
-* Windows implementation of the misc OS functions
-**/
+ * Windows implementation of the stack walking.
+ **/
 struct CORE_API FWindowsPlatformStackWalk : public FGenericPlatformStackWalk
 {
 	static bool InitStackWalking();
-	static bool ProgramCounterToHumanReadableString( int32 CurrentCallDepth, uint64 ProgramCounter, ANSICHAR* HumanReadableString, SIZE_T HumanReadableStringSize, EVerbosityFlags VerbosityFlags = VF_DISPLAY_ALL, FGenericCrashContext* Context = NULL );
-	static void ProgramCounterToSymbolInfo( uint64 ProgramCounter, FProgramCounterSymbolInfo&  OutProgramCounterSymbolInfo);
-	static void StackWalkAndDump( ANSICHAR* HumanReadableString, SIZE_T HumanReadableStringSize, int32 IgnoreCount, void* Context = NULL );
+	
+	static void ProgramCounterToSymbolInfo( uint64 ProgramCounter, FProgramCounterSymbolInfo& out_SymbolInfo );
 	static void CaptureStackBackTrace( uint64* BackTrace, uint32 MaxDepth, void* Context = NULL );
+	static void StackWalkAndDump( ANSICHAR* HumanReadableString, SIZE_T HumanReadableStringSize, int32 IgnoreCount, void* Context = NULL );
+
 	static int32 GetProcessModuleCount();
 	static int32 GetProcessModuleSignatures(FStackWalkModuleInfo *ModuleSignatures, const int32 ModuleSignaturesSize);
 
