#ttp 344974 WEX: Support Android system font

Small change to file rebase so that it checks the start of the path against the returned GFontPath rather than the hardcoded value.

diff --git a/Engine/Source/Runtime/Core/Private/Android/AndroidFile.cpp b/Engine/Source/Runtime/Core/Private/Android/AndroidFile.cpp
--- a/Engine/Source/Runtime/Core/Private/Android/AndroidFile.cpp
+++ b/Engine/Source/Runtime/Core/Private/Android/AndroidFile.cpp
@@ -668,7 +668,7 @@ public:
 		}
 
 		FString Result = Filename;
-		if (!Result.StartsWith(TEXT("/system/fonts")))
+		if (!Result.StartsWith(GFontPathBase))
 		{
 			Result.ReplaceInline(TEXT("../"), TEXT(""));
 			Result.ReplaceInline(TEXT(".."), TEXT(""));
