Remove markdown-link style comments. Causes errors in UnrealDocTool.

diff --git a/Engine/Source/Runtime/Core/Public/Misc/MapErrors.h b/Engine/Source/Runtime/Core/Public/Misc/MapErrors.h
--- a/Engine/Source/Runtime/Core/Public/Misc/MapErrors.h
+++ b/Engine/Source/Runtime/Core/Public/Misc/MapErrors.h
@@ -12,39 +12,39 @@ struct CORE_API FMapErrors
 {
     /**  Lighting */
 
-	/** [LightActor] has same light GUID as [LightActor] (Duplicate and replace the orig with the new one): Duplicate and replace the original with the new one. */
+	/** {LightActor} has same light GUID as {LightActor} (Duplicate and replace the orig with the new one): Duplicate and replace the original with the new one. */
     static FName MatchingLightGUID;
  
     /**  Actor */
     
-	/** [Actor] : Large actor casts a shadow and will cause an extreme performance hit unless bUseBooleanEnvironmentShadowing is set to true: A large actor has been set to cast shadows - this will cause extreme performance issues and should have bUseBooleanEnvironmentShadowing set to true. */
+	/** {Actor} : Large actor casts a shadow and will cause an extreme performance hit unless bUseBooleanEnvironmentShadowing is set to true: A large actor has been set to cast shadows - this will cause extreme performance issues and should have bUseBooleanEnvironmentShadowing set to true. */
     static FName ActorLargeShadowCaster;
 
-	 /** [Volume] causes damage, but has no damagetype defined.: This warning is caused when there is a volume that is set to cause damage but doesn't have a damage type defined.  A damage type is important because it tells the game code how to handle a actor's reaction to taking damage.  This can be solved by going to the actor's Property Window->Physics Volume and setting the 'DamageType' property. */
+	 /** {Volume} causes damage, but has no damagetype defined.: This warning is caused when there is a volume that is set to cause damage but doesn't have a damage type defined.  A damage type is important because it tells the game code how to handle a actor's reaction to taking damage.  This can be solved by going to the actor's Property Window->Physics Volume and setting the 'DamageType' property. */
     static FName NoDamageType;
    
-	/** [Actor] : Brush has non-coplanar polygons: This warning is caused when there is a brush in the level that has non-coplanar polygons.  This is usually caused by using brush editing tools in geometry mode in extreme ways and can cause missing polygons in the level.  This warning can be resolved by deleting the brush and recreating it. */
+	/** {Actor} : Brush has non-coplanar polygons: This warning is caused when there is a brush in the level that has non-coplanar polygons.  This is usually caused by using brush editing tools in geometry mode in extreme ways and can cause missing polygons in the level.  This warning can be resolved by deleting the brush and recreating it. */
     static FName NonCoPlanarPolys;
     
-	/** [Actor] in same location as [Another Actor]: This warning is caused when there is a actor that is in the same exact location as another actor.  This warning is usually the result of a accidental duplication or paste operation.  It can be fixed by deleting one of the actors, or disregarded if the placement was intentional. */
+	/** {Actor} in same location as {Another Actor}: This warning is caused when there is a actor that is in the same exact location as another actor.  This warning is usually the result of a accidental duplication or paste operation.  It can be fixed by deleting one of the actors, or disregarded if the placement was intentional. */
     static FName SameLocation;
     
-	/** [Actor] has invalid DrawScale/ DrawScale3D: This warning is caused when either DrawScale, DrawScale3D X, DrawScale3D Y, or DrawScale 3D Z is equal to zero.  Meaning that the actor will not be shown because it is being scaled to zero on one of its axis.  To solve this problem, change any DrawScale's that are zero to be non-zero by selecting the actor and changing its drawscale at the bottom of the main UnrealEd window. */
+	/** {Actor} has invalid DrawScale/ DrawScale3D: This warning is caused when either DrawScale, DrawScale3D X, DrawScale3D Y, or DrawScale 3D Z is equal to zero.  Meaning that the actor will not be shown because it is being scaled to zero on one of its axis.  To solve this problem, change any DrawScale's that are zero to be non-zero by selecting the actor and changing its drawscale at the bottom of the main UnrealEd window. */
     static FName InvalidDrawscale;
     
-	/** [Actor] is obsolete and must be removed!: This warning is caused when there is an instance of a actor in a level that has been marked deprecated.  This is usually because a actor was marked deprecated after the level was created, but the map was never updated.  This can be fixed by simply deleting the actor. */
+	/** {Actor} is obsolete and must be removed!: This warning is caused when there is an instance of a actor in a level that has been marked deprecated.  This is usually because a actor was marked deprecated after the level was created, but the map was never updated.  This can be fixed by simply deleting the actor. */
     static FName ActorIsObselete;
     
-	/** [Actor] bStatic true, but has Physics set to something other than PHYS_None!: This warning is caused when an actor has its bStatic flag set to true but its Physics is set to PHYS_None.  Since bStatic means that the actor will not be moving, having Physics set to PHYS_None is contradictory.  Actors set with the bStatic flag are also not ticked(updated).  This error can be solved by going to the actor's properties and changing its Physics to PHYS_None. */
+	/** {Actor} bStatic true, but has Physics set to something other than PHYS_None!: This warning is caused when an actor has its bStatic flag set to true but its Physics is set to PHYS_None.  Since bStatic means that the actor will not be moving, having Physics set to PHYS_None is contradictory.  Actors set with the bStatic flag are also not ticked(updated).  This error can be solved by going to the actor's properties and changing its Physics to PHYS_None. */
     static FName StaticPhysNone;
     
-	/** [Actor] : Volume actor has NULL collision component - please delete: The specified volume actor has a NULL collision component and should probably be deleted. */
+	/** {Actor} : Volume actor has NULL collision component - please delete: The specified volume actor has a NULL collision component and should probably be deleted. */
     static FName VolumeActorCollisionComponentNULL;
     
-	/** [Actor] : Volume actor has a collision component with 0 radius - please delete: The specified volume actor has a zero radius for its collision component and should probably be deleted. */
+	/** {Actor} : Volume actor has a collision component with 0 radius - please delete: The specified volume actor has a zero radius for its collision component and should probably be deleted. */
     static FName VolumeActorZeroRadius;
     
-	/** [Actor] (LOD [Index]) has hand-painted vertex colors that no longer match the original StaticMesh [StaticMesh]: It looks like the original mesh was changed since this instance's vertex colors were painted down - this may need a refresh. */
+	/** {Actor} (LOD {Index}) has hand-painted vertex colors that no longer match the original StaticMesh {StaticMesh}: It looks like the original mesh was changed since this instance's vertex colors were painted down - this may need a refresh. */
     static FName VertexColorsNotMatchOriginalMesh;
     
 	/** Actor casts dynamic shadows and has a BoundsScale greater than 1! This will have a large performance hit: Serious performance warning... either reduce BoundsScale to be <= 1 or remove dynamic shadows... */
@@ -52,21 +52,21 @@ struct CORE_API FMapErrors
     
 	static FName MultipleSkyLights;
 
-	/** [ActorName] has WorldTrace blocked. It will be considered to be world geometry */
+	/** {ActorName} has WorldTrace blocked. It will be considered to be world geometry */
 	static FName InvalidTrace;
 
     /**  BSP Brush */
 
-	/** [Brush Actor] : Brush has zero polygons - please delete!: This warning indicates that you have a brush in the level that doesn't have any polygons associated with it.  The brush should be deleted as it isn't doing anything useful. */
+	/** {Brush Actor} : Brush has zero polygons - please delete!: This warning indicates that you have a brush in the level that doesn't have any polygons associated with it.  The brush should be deleted as it isn't doing anything useful. */
     static FName BrushZeroPolygons;
     
-	/** Run 'Clean BSP Materials' to clear [count] references: This warning indicates that there are material references on brush faces that aren't contributing to the BSP, and that applying the Tools->'Clean BSP Materials' operation can clean up these references. */
+	/** Run 'Clean BSP Materials' to clear {count} references: This warning indicates that there are material references on brush faces that aren't contributing to the BSP, and that applying the Tools->'Clean BSP Materials' operation can clean up these references. */
     static FName CleanBSPMaterials;
     
-	/** [Actor] : Brush has NULL BrushComponent property - please delete!: This warning is caused when there is a Brush with a "None" ConstraintInstance component.  It is usually found in older maps where duplication was used to create a Brush and can be fixed by deleting the Brush causing the warning and creating a new one. */
+	/** {Actor} : Brush has NULL BrushComponent property - please delete!: This warning is caused when there is a Brush with a "None" ConstraintInstance component.  It is usually found in older maps where duplication was used to create a Brush and can be fixed by deleting the Brush causing the warning and creating a new one. */
     static FName BrushComponentNull;
     
-	/** [Brush] : Brush is planar: Planar brush used - please note that this may not work well with collision. */
+	/** {Brush} : Brush is planar: Planar brush used - please note that this may not work well with collision. */
     static FName PlanarBrush;
 
 	/** Camera */
@@ -76,10 +76,10 @@ struct CORE_API FMapErrors
 
     /**  Class */
 
-	/** [ ]::[ ] is obsolete and must be removed (Class is abstract): NEEDS DESCRIPTION */
+	/** { }::{ } is obsolete and must be removed (Class is abstract): NEEDS DESCRIPTION */
     static FName AbstractClass;
     
-	/** [ ]::[ ] is obsolete and must be removed (Class is deprecated): NEEDS DESCRIPTION */
+	/** { }::{ } is obsolete and must be removed (Class is deprecated): NEEDS DESCRIPTION */
     static FName DeprecatedClass;
     
 
@@ -93,16 +93,16 @@ struct CORE_API FMapErrors
 
     /**  Landscape */
 
-	/** [LandscapeComponent] : Fixed up deleted layer weightmap: **TODO** */
+	/** {LandscapeComponent} : Fixed up deleted layer weightmap: **TODO** */
     static FName FixedUpDeletedLayerWeightmap;
     
-	/** [LandscapeComponent] : Fixed up incorrect layer weightmap texture index: **TODO** */
+	/** {LandscapeComponent} : Fixed up incorrect layer weightmap texture index: **TODO** */
     static FName FixedUpIncorrectLayerWeightmap;
     
-	/** Fixed up shared weightmap texture for layer [Layer] in component [Component] (shares with [Name]): **TODO** */
+	/** Fixed up shared weightmap texture for layer {Layer} in component {Component} (shares with {Name}): **TODO** */
     static FName FixedUpSharedLayerWeightmap;
 
-	/** Landscape ([ProxyName]) has overlapping render components at location ([X, Y]): **TODO** */
+	/** Landscape ({ProxyName}) has overlapping render components at location ({X, Y}): **TODO** */
 	static FName LandscapeComponentPostLoad_Warning;
 
 
@@ -117,7 +117,7 @@ struct CORE_API FMapErrors
 
     /**  Lighting */
 
-	/** [Actor] : Light actor has NULL LightComponent property - please delete!: This warning is caused when there is a Light actor with a "None" LightComponent.  It is usually found in older maps where duplication was used to create the actor and can be fixed by deleting the actor causing the warning and creating a new one. */
+	/** {Actor} : Light actor has NULL LightComponent property - please delete!: This warning is caused when there is a Light actor with a "None" LightComponent.  It is usually found in older maps where duplication was used to create the actor and can be fixed by deleting the actor causing the warning and creating a new one. */
     static FName LightComponentNull;
     
 	/** Maps need lighting rebuilt: This warning is caused when lighting has been invalidated by moving or modifying a light actor.  This can cause problems because the rendered lighting in the level is not accurately representing the current state of lights in the level.  This error can be solved by going to the Build menu and rebuilding lighting for a map. */
@@ -133,30 +133,30 @@ struct CORE_API FMapErrors
 
     /**  Particle System */
 
-	/** [Actor] : Emitter actor has NULL ParticleSystemComponent property - please delete!: This warning is caused when there is a Emitter actor with a "None" ParticleSystemComponent.  It is usually found in older maps where duplication was used to create the actor and can be fixed by deleting the actor causing the warning and creating a new one. */
+	/** {Actor} : Emitter actor has NULL ParticleSystemComponent property - please delete!: This warning is caused when there is a Emitter actor with a "None" ParticleSystemComponent.  It is usually found in older maps where duplication was used to create the actor and can be fixed by deleting the actor causing the warning and creating a new one. */
     static FName ParticleSystemComponentNull;
     
-	/** PSysComp has an empty parameter actor reference at index [Index] ([Actor]): Param.Actor should not be NULL. */
+	/** PSysComp has an empty parameter actor reference at index {Index} ({Actor}): Param.Actor should not be NULL. */
     static FName PSysCompErrorEmptyActorRef;
     
-	/** PSysComp has an empty parameter material reference at index [Index] ([Actor]): Param.Material should not be NULL. */
+	/** PSysComp has an empty parameter material reference at index {Index} ({Actor}): Param.Material should not be NULL. */
     static FName PSysCompErrorEmptyMaterialRef;
     
 
     /**  Skeletal Mesh */
 
-	/** [Actor] : SkeletalMeshActor has no PhysicsAsset assigned.: In order for a SkeletalMesh to have an accurate bounding box, it needs to have a PhysicsAsset assigned in its SkeletalMeshComponent. An incorrect or inaccurate bounding box can lead to the mesh vanishing when its origin is not in view, or to poor shadow resolution because the bounding box is too big. */
+	/** {Actor} : SkeletalMeshActor has no PhysicsAsset assigned.: In order for a SkeletalMesh to have an accurate bounding box, it needs to have a PhysicsAsset assigned in its SkeletalMeshComponent. An incorrect or inaccurate bounding box can lead to the mesh vanishing when its origin is not in view, or to poor shadow resolution because the bounding box is too big. */
     static FName SkelMeshActorNoPhysAsset;
     
-	/** [Actor] : Skeletal mesh actor has NULL SkeletalMeshComponent property: The specified SkeletalMeshActor has a NULL SkeletalMeshComponent. */
+	/** {Actor} : Skeletal mesh actor has NULL SkeletalMeshComponent property: The specified SkeletalMeshActor has a NULL SkeletalMeshComponent. */
     static FName SkeletalMeshComponent;
     
-	/** [Actor] : Skeletal mesh actor has NULL SkeletalMesh property: The specified SkeletalMeshActor has a NULL SkeletalMesh. */
+	/** {Actor} : Skeletal mesh actor has NULL SkeletalMesh property: The specified SkeletalMeshActor has a NULL SkeletalMesh. */
     static FName SkeletalMeshNull;
 
     /**  Sound */
 
-	/** [Actor] : Ambient sound actor has NULL AudioComponent property - please delete!: This warning is caused when there is a Ambient sound actor with a "None" AudioComponent.  It is usually found in older maps where duplication was used to create a actor and can be fixed by deleting the actor causing the warning and creating a new one. */
+	/** {Actor} : Ambient sound actor has NULL AudioComponent property - please delete!: This warning is caused when there is a Ambient sound actor with a "None" AudioComponent.  It is usually found in older maps where duplication was used to create a actor and can be fixed by deleting the actor causing the warning and creating a new one. */
     static FName AudioComponentNull;
 
 	/** Ambient sound actor's AudioComponent has a NULL SoundCue property!: This warning is caused when there is a AmbientSound actor with a NULL SoundCue property.  This is a problem because the actor won't actually be playing any sounds.  This can be fixed by first choosing a sound cue in the generic browser and then going to the actor's Property Window->Audio Category->Audio Component and setting the 'SoundCue' property. */
@@ -164,19 +164,19 @@ struct CORE_API FMapErrors
 
     /**  Static Mesh */
 
-	/** [Static Mesh Actor] : Static mesh actor has NULL StaticMesh property: This warning is caused when there is a static mesh actor in a level with a NULL StaticMesh property.  This can be a problem because the actor exists and is using memory, but doesn't have a static mesh to actually draw.  This warning is usually the result of creating a StaticMesh actor without first selecting a StaticMesh in the generic browser.  This warning can be fixed by first selecting a static mesh in the generic browser and then going to the StaticMesh actor's Property Window->StaticMeshActor Category->StaticMeshComponent->StaticMeshComponent Category to set the 'StaticMesh' Property. */
+	/** {Static Mesh Actor} : Static mesh actor has NULL StaticMesh property: This warning is caused when there is a static mesh actor in a level with a NULL StaticMesh property.  This can be a problem because the actor exists and is using memory, but doesn't have a static mesh to actually draw.  This warning is usually the result of creating a StaticMesh actor without first selecting a StaticMesh in the generic browser.  This warning can be fixed by first selecting a static mesh in the generic browser and then going to the StaticMesh actor's Property Window->StaticMeshActor Category->StaticMeshComponent->StaticMeshComponent Category to set the 'StaticMesh' Property. */
     static FName StaticMeshNull;
     
-	/** [Actor] : Static mesh actor has NULL StaticMeshComponent property - please delete!: This warning is caused when there is a static mesh actor with a "None" StaticMeshComponent component.  It is usually found in older maps where duplication was used to create the actor and can be fixed by deleting the actor causing the warning and creating a new one. */
+	/** {Actor} : Static mesh actor has NULL StaticMeshComponent property - please delete!: This warning is caused when there is a static mesh actor with a "None" StaticMeshComponent component.  It is usually found in older maps where duplication was used to create the actor and can be fixed by deleting the actor causing the warning and creating a new one. */
     static FName StaticMeshComponent;
     
-	/** [StaticMesh] has simple collision but is being scaled non-uniformly - collision creation will fail: Simple collision cannot be used with non-uniform scale. Please either fix the scale or the collision type. */
+	/** {StaticMesh} has simple collision but is being scaled non-uniformly - collision creation will fail: Simple collision cannot be used with non-uniform scale. Please either fix the scale or the collision type. */
     static FName SimpleCollisionButNonUniformScale;
 
-	/** More overriden materials [Count] on static mesh component than are referenced [Count] in source mesh [StaticMesh]: **TODO** */
+	/** More overriden materials {Count} on static mesh component than are referenced {Count} in source mesh {StaticMesh}: **TODO** */
     static FName MoreMaterialsThanReferenced;
 
-	 /** [Count] element(s) with zero triangles in static mesh [StaticMesh]: **TODO** */
+	 /** {Count} element(s) with zero triangles in static mesh {StaticMesh}: **TODO** */
     static FName ElementsWithZeroTriangles;
   
 
@@ -191,7 +191,7 @@ struct CORE_API FMapErrors
 
     /**  Uncategorized */
 
-	/** Filename [Filename] is too long - this may interfere with cooking for consoles.  Unreal filenames should be no longer than [Length] characters.: Please rename the file to be within the length specified. */
+	/** Filename {Filename} is too long - this may interfere with cooking for consoles.  Unreal filenames should be no longer than {Length} characters.: Please rename the file to be within the length specified. */
     static FName FilenameIsTooLongForCooking;
     
 	/** {ObjectName} : Externally referenced */
@@ -199,7 +199,7 @@ struct CORE_API FMapErrors
 
     /**  Actors */
 
-	/** [Actor] : Repaired painted vertex colors: Painted vertex colors were repaired on this actor. */
+	/** {Actor} : Repaired painted vertex colors: Painted vertex colors were repaired on this actor. */
     static FName RepairedPaintedVertexColors;
 };
 
