Format i18n tables in markdown format; remove percent symbols which UDT has trouble parsing.

diff --git a/Engine/Source/Runtime/Core/Public/Internationalization/NegativeCurrencyOutputFormatting.h b/Engine/Source/Runtime/Core/Public/Internationalization/NegativeCurrencyOutputFormatting.h
--- a/Engine/Source/Runtime/Core/Public/Internationalization/NegativeCurrencyOutputFormatting.h
+++ b/Engine/Source/Runtime/Core/Public/Internationalization/NegativeCurrencyOutputFormatting.h
@@ -9,22 +9,25 @@ namespace ECurrencyNegativePattern
 	 * enum value associated with a pattern used to display negative currency values
 	 * Below X represents the number being converted to display as a currency, $ and - will be replaced
 	 * with the cultural values CurrencySymbol and NegativeSign respectively
-	 * ($X)		BracketCurrencySymbolXBracket
-	 * -$X		NegativeSymbolCurrencySymbolX
-	 * $-X		CurrencySymbolNegativeSymbolX
-	 * $X-		CurrencySymbolXNegativeSymbol
-	 * (X$)		BracketXCurrencySymbolBracket
-	 * -X$		NegativeSymbolXCurrencySymbol
-	 * X-$		XNegativeSymbolCurrencySymbol
-	 * X$-		XCurrencySymbolNegativeSymbol
-	 * -X $		NegativeSymbolXSpaceCurrencySymbol
-	 * -$ X		NegativeSymbolCurrencySymbolSpaceX
-	 * X $-		XSpaceCurrencySymbolNegativeSymbol
-	 * $ X-		CurrencySymbolSpaceXNegativeSymbol
-	 * $ -X		CurrencySymbolSpaceNegativeSymbolX
-	 * X- $		XNegativeSymbolSpaceCurrencySymbol
-	 * ($ X)	BracketCurrencySymbolSpaceXBracket
-	 * (X $)	BracketXSpaceCurrencySymbolBracket
+	 *
+	 * |------------|------------------------------------|
+	 * | ($X)		| BracketCurrencySymbolXBracket      |
+	 * | -$X		| NegativeSymbolCurrencySymbolX      |
+	 * | $-X		| CurrencySymbolNegativeSymbolX      |
+	 * | $X-		| CurrencySymbolXNegativeSymbol      |
+	 * | (X$)		| BracketXCurrencySymbolBracket      |
+	 * | -X$		| NegativeSymbolXCurrencySymbol      |
+	 * | X-$		| XNegativeSymbolCurrencySymbol      |
+	 * | X$-		| XCurrencySymbolNegativeSymbol      |
+	 * | -X $		| NegativeSymbolXSpaceCurrencySymbol |
+	 * | -$ X		| NegativeSymbolCurrencySymbolSpaceX |
+	 * | X $-		| XSpaceCurrencySymbolNegativeSymbol |
+	 * | $ X-		| CurrencySymbolSpaceXNegativeSymbol |
+	 * | $ -X		| CurrencySymbolSpaceNegativeSymbolX |
+	 * | X- $		| XNegativeSymbolSpaceCurrencySymbol |
+	 * | ($ X)	    | BracketCurrencySymbolSpaceXBracket |
+	 * | (X $)		| BracketXSpaceCurrencySymbolBracket |
+	 * |------------|------------------------------------|
 	 */
 	enum Type
 	{
diff --git a/Engine/Source/Runtime/Core/Public/Internationalization/NegativeNumberOutputFormatting.h b/Engine/Source/Runtime/Core/Public/Internationalization/NegativeNumberOutputFormatting.h
--- a/Engine/Source/Runtime/Core/Public/Internationalization/NegativeNumberOutputFormatting.h
+++ b/Engine/Source/Runtime/Core/Public/Internationalization/NegativeNumberOutputFormatting.h
@@ -9,11 +9,14 @@ namespace ENumberNegativePattern
 	 * enum value associated with a pattern used to display negative number values
 	 * Below X represents the number being converted to display, - will be replaced
 	 * with the cultural value NegativeSign
-	 * (X)		BracketXBracket
-	 * -X		NegativeSymbolX
-	 * - X		NegativeSymbolSpaceX
-	 * X-		XNegativeSymbol
-	 * X -		XSpaceNegativeSymbol
+	 *
+	 * |--------|----------------------|
+	 * | (X)	| BracketXBracket      |
+	 * | -X		| NegativeSymbolX      |
+	 * | - X	| NegativeSymbolSpaceX |
+	 * | X-		| XNegativeSymbol      |
+	 * | X -	| XSpaceNegativeSymbol |
+	 * |--------|----------------------|
 	 */
 	enum Type
 	{
diff --git a/Engine/Source/Runtime/Core/Public/Internationalization/NegativePercentOutputFormatting.h b/Engine/Source/Runtime/Core/Public/Internationalization/NegativePercentOutputFormatting.h
--- a/Engine/Source/Runtime/Core/Public/Internationalization/NegativePercentOutputFormatting.h
+++ b/Engine/Source/Runtime/Core/Public/Internationalization/NegativePercentOutputFormatting.h
@@ -7,20 +7,23 @@ namespace EPercentNegativePattern
 {
 	/** 
 	 * enum value associated with a pattern used to display positive percent values
-	 * Below X represents the number being converted to display as a percent, % will be replaced
+	 * Below X represents the number being converted to display as a percent, Y will be replaced
 	 * with the cultural values PercentSymbol
-	 * -X %		NegativeSymbolXSpacePct
-	 * -X%		NegativeSymbolXPct
-	 * -%X		NegativeSymbolPctX
-	 * %-X		PctNegativeSymbolX
-	 * %X-		PctXNegativeSymbol
-	 * X-%		XNegativeSymbolPct
-	 * X%-		XPctNegativeSymbol
-	 * -% X		NegativeSymbolPctSpaceX
-	 * X %-		XSpacePctNegativeSymbol
-	 * % X-		PctSpaceXNegativeSymbol
-	 * % -X		PctSpaceNegativeSymbolX
-	 * X- %		XNegativeSymbolSpacePct
+	 *
+	 * |------------|-------------------------|
+	 * | -X Y		| NegativeSymbolXSpacePct |
+	 * | -XY		| NegativeSymbolXPct      |
+	 * | -YX		| NegativeSymbolPctX      |
+	 * | Y-X		| PctNegativeSymbolX      |
+	 * | YX-		| PctXNegativeSymbol      |
+	 * | X-Y		| XNegativeSymbolPct      |
+	 * | XY-		| XPctNegativeSymbol      |
+	 * | -Y X		| NegativeSymbolPctSpaceX |
+	 * | X Y-		| XSpacePctNegativeSymbol |
+	 * | Y X-		| PctSpaceXNegativeSymbol |
+	 * | Y -X		| PctSpaceNegativeSymbolX |
+	 * | X- Y		| XNegativeSymbolSpacePct |
+	 * |------------|-------------------------|
 	 */
 	enum Type
 	{
diff --git a/Engine/Source/Runtime/Core/Public/Internationalization/PositiveCurrencyOutputFormatting.h b/Engine/Source/Runtime/Core/Public/Internationalization/PositiveCurrencyOutputFormatting.h
--- a/Engine/Source/Runtime/Core/Public/Internationalization/PositiveCurrencyOutputFormatting.h
+++ b/Engine/Source/Runtime/Core/Public/Internationalization/PositiveCurrencyOutputFormatting.h
@@ -9,10 +9,13 @@ namespace ECurrencyPositivePattern
 	 * enum value associated with a pattern used to display positive currency values
 	 * Below X represents the number being converted to display as a currency, $ will be replaced
 	 * with the cultural values CurrencySymbol
-	 * $n	CurrencySymbolX
-	 * n$	XCurrencySymbol
-     * $ n	CurrencySymbolSpaceX
-	 * n $	XSpaceCurrencySymbol
+	 *
+	 * |-----|----------------------|
+	 * | $n	 | CurrencySymbolX      |
+	 * | n$	 | XCurrencySymbol      |
+     * | $ n | CurrencySymbolSpaceX |
+	 * | n $ | XSpaceCurrencySymbol |
+	 * |-----|----------------------|
 	 */
 	enum Type
 	{
diff --git a/Engine/Source/Runtime/Core/Public/Internationalization/PositivePercentOutputFormatting.h b/Engine/Source/Runtime/Core/Public/Internationalization/PositivePercentOutputFormatting.h
--- a/Engine/Source/Runtime/Core/Public/Internationalization/PositivePercentOutputFormatting.h
+++ b/Engine/Source/Runtime/Core/Public/Internationalization/PositivePercentOutputFormatting.h
@@ -7,12 +7,15 @@ namespace EPercentPositivePattern
 {
 	/** 
 	 * enum value associated with a pattern used to display positive percent values
-	 * Below X represents the number being converted to display as a percent, % will be replaced
+	 * Below X represents the number being converted to display as a percent, Y will be replaced
 	 * with the cultural values PercentSymbol
-	 * X %	XSpacePct
-	 * X%	XPct
-     * %X	PctX
-	 * % X	PctSpaceX
+	 *
+	 * |-----|------------|
+	 * | X Y | XSpacePct  |
+	 * | XY	 | XPct       |
+     * | YX	 | PctX       |
+	 * | Y X | PctSpaceX  |
+	 * |-----|------------|
 	 */
 	enum Type
 	{
