Fix another comment that UnrealDocTool cannot parse.

diff --git a/Engine/Source/Runtime/Core/Public/Containers/Array.h b/Engine/Source/Runtime/Core/Public/Containers/Array.h
--- a/Engine/Source/Runtime/Core/Public/Containers/Array.h
+++ b/Engine/Source/Runtime/Core/Public/Containers/Array.h
@@ -15,7 +15,7 @@
  * - A type called ElementType representing the contained type.
  * - A method IndexType Num() const that returns the number of items in the container.
  * - A method bool IsValidIndex(IndexType index) which returns whether a given index is valid in the container.
- * - A method T& operator[](IndexType index) which returns a reference to a contained object by index.
+ * - A method T& operator\[\](IndexType index) which returns a reference to a contained object by index.
  */
 template< typename ContainerType, typename ElementType, typename IndexType>
 class TIndexedContainerIterator
