Don't implement bShowImmediately on Mac OS X - we setup the OpenGL context way too late for that to work.

diff --git a/Engine/Source/Runtime/Core/Private/Mac/MacWindow.cpp b/Engine/Source/Runtime/Core/Private/Mac/MacWindow.cpp
--- a/Engine/Source/Runtime/Core/Private/Mac/MacWindow.cpp
+++ b/Engine/Source/Runtime/Core/Private/Mac/MacWindow.cpp
@@ -171,11 +171,6 @@ void FMacWindow::Initialize( FMacApplication* const Application, const TSharedRe
 	{
 		return;
 	}
-	
-	if ( bShowImmediately )
-	{
-		Show();
-	}
 }
 
 FMacWindow::FMacWindow()
