nicer comment

diff --git a/Engine/Source/Runtime/Core/Public/GenericPlatform/GenericPlatformMath.h b/Engine/Source/Runtime/Core/Public/GenericPlatform/GenericPlatformMath.h
--- a/Engine/Source/Runtime/Core/Public/GenericPlatform/GenericPlatformMath.h
+++ b/Engine/Source/Runtime/Core/Public/GenericPlatform/GenericPlatformMath.h
@@ -153,7 +153,7 @@ struct FGenericPlatformMath
 	* @param F		Floating point value to convert
 	* @return		An integer greater or equal to 'F'.
 	*/
-	DEPRECATED(4.2, "FPlatformMath::Ceil is deprecated, use FPlatformMath::CeilToInt instead.")
+	DEPRECATED(4.2, "FPlatformMath::Ceil is deprecated, use FPlatformMath::CeilToInt or FPlatformMath::CeilToFloat() instead.")
 	static FORCEINLINE int32 Ceil(float F)
 	{
 		return CeilToInt(F);
