Fix GetPrimaryGPU brand for MacBook/Pro's - for multi-GPU setups report both GPUs since there's no way to tell the active one without OpenGL.

diff --git a/Engine/Source/Runtime/Core/Private/Mac/MacPlatformMisc.cpp b/Engine/Source/Runtime/Core/Private/Mac/MacPlatformMisc.cpp
--- a/Engine/Source/Runtime/Core/Private/Mac/MacPlatformMisc.cpp
+++ b/Engine/Source/Runtime/Core/Private/Mac/MacPlatformMisc.cpp
@@ -864,8 +864,6 @@ FString FMacPlatformMisc::GetPrimaryGPUBrand()
 	static FString PrimaryGPU;
 	if(PrimaryGPU.IsEmpty())
 	{
-		PrimaryGPU = FGenericPlatformMisc::GetPrimaryGPUBrand();
-		
 		// Enumerate the GPUs via IOKit to avoid dragging in OpenGL
 		io_iterator_t Iterator;
 		CFMutableDictionaryRef MatchDictionary = IOServiceMatching("IOPCIDevice");
@@ -885,20 +883,24 @@ FString FMacPlatformMisc::GetPrimaryGPUBrand()
 						if(Value && *Value == 0x30000)
 						{
 							// If there are StartupDisplay & hda-gfx entries then this is likely the online display
-							const CFDataRef StartupDisplay = (const CFDataRef)CFDictionaryGetValue(ServiceInfo, CFSTR("StartupDisplay"));
 							const CFDataRef HDAGfx = (const CFDataRef)CFDictionaryGetValue(ServiceInfo, CFSTR("hda-gfx"));
-							if(StartupDisplay && HDAGfx)
+							if(HDAGfx)
 							{
 								const CFDataRef Model = (const CFDataRef)CFDictionaryGetValue(ServiceInfo, CFSTR("model"));
 								if(Model && CFGetTypeID(Model) == CFDataGetTypeID())
 								{
 									CFStringRef ModelName = CFStringCreateFromExternalRepresentation(kCFAllocatorDefault, Model, kCFStringEncodingASCII);
-									PrimaryGPU = FString((NSString*)ModelName);
+									// Append all the GPUs together, watching out for malformed Intel strings...
+									if ( PrimaryGPU.IsEmpty() )
+									{
+										PrimaryGPU = FString((NSString*)ModelName).Trim();
+									}
+									else
+									{
+										PrimaryGPU += TEXT(" - ");
+										PrimaryGPU += FString((NSString*)ModelName).Trim();
+									}
 									CFRelease(ModelName);
-									
-									CFRelease(ServiceInfo);
-									IOObjectRelease(ServiceEntry);
-									break;
 								}
 							}
 						}
@@ -908,6 +910,11 @@ FString FMacPlatformMisc::GetPrimaryGPUBrand()
 				IOObjectRelease(ServiceEntry);
 			}
 			IOObjectRelease(Iterator);
+			
+			if ( PrimaryGPU.IsEmpty() )
+			{
+				PrimaryGPU = FGenericPlatformMisc::GetPrimaryGPUBrand();
+			}
 		}
 	}
 	return PrimaryGPU;
