Fix the project dictionary not recognizing non-foreign projects correctly if they're specified with a relative rather than absolute path.

diff --git a/Engine/Source/Runtime/Core/Private/Misc/UProjectInfo.cpp b/Engine/Source/Runtime/Core/Private/Misc/UProjectInfo.cpp
--- a/Engine/Source/Runtime/Core/Private/Misc/UProjectInfo.cpp
+++ b/Engine/Source/Runtime/Core/Private/Misc/UProjectInfo.cpp
@@ -9,8 +9,7 @@ DEFINE_LOG_CATEGORY(LogUProjectInfo);
 
 FUProjectDictionary::FUProjectDictionary(const FString& InRootDir)
 {
-	FString RootDir = InRootDir;
-	FPaths::NormalizeDirectoryName(RootDir);
+	FString RootDir = FPaths::ConvertRelativePathToFull(InRootDir);
 
 	TArray<FString> DirectoriesToSearch;
 	
@@ -80,8 +79,7 @@ FUProjectDictionary::FUProjectDictionary(const FString& InRootDir)
 
 bool FUProjectDictionary::IsForeignProject(const FString& InProjectFileName) const
 {
-	FString ProjectFileName = InProjectFileName;
-	FPaths::NormalizeFilename(ProjectFileName);
+	FString ProjectFileName = FPaths::ConvertRelativePathToFull(InProjectFileName);
 
 	for(TMap<FString, FString>::TConstIterator Iter(ShortProjectNameDictionary); Iter; ++Iter)
 	{
