Include module names in dialog for recompiling on startup, and don't allow continuing if they are an incorrect version.

diff --git a/Engine/Source/Runtime/Core/Private/Modules/ModuleManager.cpp b/Engine/Source/Runtime/Core/Private/Modules/ModuleManager.cpp
--- a/Engine/Source/Runtime/Core/Private/Modules/ModuleManager.cpp
+++ b/Engine/Source/Runtime/Core/Private/Modules/ModuleManager.cpp
@@ -739,7 +739,6 @@ void FModuleManager::QueryModules( TArray< FModuleStatus >& OutModuleStatuses )
 	}
 }
 
-
 FString FModuleManager::GetModuleFilename(FName ModuleName) const
 {
 	return Modules.FindChecked(ModuleName)->Filename;
@@ -750,28 +749,14 @@ void FModuleManager::SetModuleFilename(FName ModuleName, const FString& Filename
 	Modules.FindChecked(ModuleName)->Filename = Filename;
 }
 
-
-
-void FModuleManager::FindModulePaths(const TCHAR* NamePattern, TMap<FName, FString> &OutModulePaths) const
+FString FModuleManager::GetCleanModuleFilename(FName ModuleName, bool bGameModule)
 {
-	// Search through the engine directory
-	FindModulePathsInDirectory(FPlatformProcess::GetModulesDirectory(), false, NamePattern, OutModulePaths);
-
-	// Search any engine directories
-	for (int Idx = 0; Idx < EngineBinariesDirectories.Num(); Idx++)
-	{
-		FindModulePathsInDirectory(EngineBinariesDirectories[Idx], false, NamePattern, OutModulePaths);
-	}
-
-	// Search any game directories
-	for (int Idx = 0; Idx < GameBinariesDirectories.Num(); Idx++)
-	{
-		FindModulePathsInDirectory(GameBinariesDirectories[Idx], true, NamePattern, OutModulePaths);
-	}
+	FString Prefix, Suffix;
+	GetModuleFilenameFormat(bGameModule, Prefix, Suffix);
+	return Prefix + ModuleName.ToString() + Suffix;
 }
 
-
-void FModuleManager::FindModulePathsInDirectory(const FString& InDirectoryName, bool bIsGameDirectory, const TCHAR* NamePattern, TMap<FName, FString> &OutModulePaths) const
+void FModuleManager::GetModuleFilenameFormat(bool bGameModule, FString& OutPrefix, FString& OutSuffix)
 {
 	// Get the module configuration for this directory type
 	const TCHAR* ConfigSuffix = NULL;
@@ -781,7 +766,7 @@ void FModuleManager::FindModulePathsInDirectory(const FString& InDirectoryName,
 		ConfigSuffix = TEXT("-Debug");
 		break;
 	case EBuildConfigurations::DebugGame:
-		ConfigSuffix = bIsGameDirectory? TEXT("-DebugGame") : NULL;
+		ConfigSuffix = bGameModule? TEXT("-DebugGame") : NULL;
 		break;
 	case EBuildConfigurations::Development:
 		ConfigSuffix = NULL;
@@ -798,26 +783,51 @@ void FModuleManager::FindModulePathsInDirectory(const FString& InDirectoryName,
 	}
 
 	// Get the base name for modules of this application
-	FString ModulePrefix = FPlatformProcess::GetModulePrefix() + FPaths::GetBaseFilename(FPlatformProcess::ExecutableName());
-	if (ModulePrefix.Contains(TEXT("-")))
+	OutPrefix = FPlatformProcess::GetModulePrefix() + FPaths::GetBaseFilename(FPlatformProcess::ExecutableName());
+	if (OutPrefix.Contains(TEXT("-")))
 	{
-		ModulePrefix = ModulePrefix.Left(ModulePrefix.Find(TEXT("-")) + 1);
+		OutPrefix = OutPrefix.Left(OutPrefix.Find(TEXT("-")) + 1);
 	}
 	else
 	{
-		ModulePrefix += TEXT("-");
+		OutPrefix += TEXT("-");
 	}
 
 	// Get the suffix for each module
-	FString ModuleSuffix;
+	OutSuffix.Empty();
 	if (ConfigSuffix != NULL)
 	{
-		ModuleSuffix += TEXT("-");
-		ModuleSuffix += FPlatformProcess::GetBinariesSubdirectory();
-		ModuleSuffix += ConfigSuffix;
+		OutSuffix += TEXT("-");
+		OutSuffix += FPlatformProcess::GetBinariesSubdirectory();
+		OutSuffix += ConfigSuffix;
 	}
-	ModuleSuffix += TEXT(".");
-	ModuleSuffix += FPlatformProcess::GetModuleExtension();
+	OutSuffix += TEXT(".");
+	OutSuffix += FPlatformProcess::GetModuleExtension();
+}
+
+void FModuleManager::FindModulePaths(const TCHAR* NamePattern, TMap<FName, FString> &OutModulePaths) const
+{
+	// Search through the engine directory
+	FindModulePathsInDirectory(FPlatformProcess::GetModulesDirectory(), false, NamePattern, OutModulePaths);
+
+	// Search any engine directories
+	for (int Idx = 0; Idx < EngineBinariesDirectories.Num(); Idx++)
+	{
+		FindModulePathsInDirectory(EngineBinariesDirectories[Idx], false, NamePattern, OutModulePaths);
+	}
+
+	// Search any game directories
+	for (int Idx = 0; Idx < GameBinariesDirectories.Num(); Idx++)
+	{
+		FindModulePathsInDirectory(GameBinariesDirectories[Idx], true, NamePattern, OutModulePaths);
+	}
+}
+
+void FModuleManager::FindModulePathsInDirectory(const FString& InDirectoryName, bool bIsGameDirectory, const TCHAR* NamePattern, TMap<FName, FString> &OutModulePaths) const
+{
+	// Get the prefix and suffix for module filenames
+	FString ModulePrefix, ModuleSuffix;
+	GetModuleFilenameFormat(bIsGameDirectory, ModulePrefix, ModuleSuffix);
 
 	// Find all the files
 	TArray<FString> FullFileNames;
@@ -832,7 +842,7 @@ void FModuleManager::FindModulePathsInDirectory(const FString& InDirectoryName,
 		if (FileName.StartsWith(ModulePrefix) && FileName.EndsWith(ModuleSuffix))
 		{
 			FString ModuleName = FileName.Mid(ModulePrefix.Len(), FileName.Len() - ModulePrefix.Len() - ModuleSuffix.Len());
-			if (ConfigSuffix != NULL || (!ModuleName.EndsWith("-Debug") && !ModuleName.EndsWith("-Shipping") && !ModuleName.EndsWith("-Test") && !ModuleName.EndsWith("-DebugGame")))
+			if (!ModuleName.EndsWith("-Debug") && !ModuleName.EndsWith("-Shipping") && !ModuleName.EndsWith("-Test") && !ModuleName.EndsWith("-DebugGame"))
 			{
 				OutModulePaths.Add(FName(*ModuleName), FullFileName);
 			}
@@ -840,7 +850,6 @@ void FModuleManager::FindModulePathsInDirectory(const FString& InDirectoryName,
 	}
 }
 
-
 void FModuleManager::UnloadOrAbandonModuleWithCallback( const FName InModuleName, FOutputDevice &Ar )
 {
 	TSharedRef< FModuleInfo > Module = Modules.FindChecked( InModuleName );
diff --git a/Engine/Source/Runtime/Core/Public/Modules/ModuleManager.h b/Engine/Source/Runtime/Core/Public/Modules/ModuleManager.h
--- a/Engine/Source/Runtime/Core/Public/Modules/ModuleManager.h
+++ b/Engine/Source/Runtime/Core/Public/Modules/ModuleManager.h
@@ -373,12 +373,15 @@ public:
 	 */
 	static const TCHAR *GetUBTConfiguration( );
 
-	/** Gets the filename for a module */
+	/** Gets the filename for a module. The return value is a full path of a module known to the module manager. */
 	FString GetModuleFilename(FName ModuleName) const;
 
 	/** Sets the filename for a module. The module is not reloaded immediately, but the new name will be used for subsequent unload/load events. */
 	void SetModuleFilename(FName ModuleName, const FString& Filename);
 
+	/** Gets the clean filename for a module, without having added it to the module manager. */
+	static FString GetCleanModuleFilename(FName ModuleName, bool bIsGameModule);
+
 public:
 
 	/**
@@ -488,6 +491,9 @@ private:
 	/** Compares file versions between the current executing engine version and the specified dll */
 	static bool CheckModuleCompatibility(const TCHAR *Filename);
 
+	/** Gets the prefix and suffix for a module file */
+	static void GetModuleFilenameFormat(bool bGameModule, FString& OutPrefix, FString& OutSuffix);
+
 	/** Finds modules matching a given name wildcard. */
 	void FindModulePaths(const TCHAR *NamePattern, TMap<FName, FString> &OutModulePaths) const;
 
diff --git a/Engine/Source/Runtime/Launch/Private/LaunchEngineLoop.cpp b/Engine/Source/Runtime/Launch/Private/LaunchEngineLoop.cpp
--- a/Engine/Source/Runtime/Launch/Private/LaunchEngineLoop.cpp
+++ b/Engine/Source/Runtime/Launch/Private/LaunchEngineLoop.cpp
@@ -2356,29 +2356,45 @@ bool FEngineLoop::AppInit( )
 #if UE_EDITOR
 	if(!GIsBuildMachine && FPaths::IsProjectFilePathSet() && IPluginManager::Get().AreRequiredPluginsAvailable())
 	{
-		if(!IProjectManager::Get().AreProjectModulesUpToDate() || !IPluginManager::Get().AreEnabledPluginModulesUpToDate())
+		TArray<FString> IncompatibleFiles;
+		if((IProjectManager::Get().CheckModuleCompatibility(IncompatibleFiles) & IPluginManager::Get().CheckModuleCompatibility(IncompatibleFiles)) == false)
 		{
-			int32 Result = FPlatformMisc::MessageBoxExt(EAppMsgType::YesNoCancel, TEXT("Project modules are missing or out of date. Would you like to recompile them?"), TEXT("Question"));
-			if(Result == EAppReturnType::Yes)
+			// Log the modules which need to be rebuilt
+			FString ModulesList = TEXT("The following modules are missing or built with a different engine version:\n\n");
+			for(int Idx = 0; Idx < IncompatibleFiles.Num(); Idx++)
 			{
-				FFeedbackContext *Context = FDesktopPlatformModule::Get()->GetNativeFeedbackContext();
+				UE_LOG(LogInit, Warning, TEXT("Incompatible or missing module: %s"), *IncompatibleFiles[Idx]);
+				ModulesList += IncompatibleFiles[Idx] + TEXT("\n");
+			}
+			ModulesList += TEXT("\nWould you like to rebuild them now?");
 
-				// Try to compile it
-				Context->BeginSlowTask(FText::FromString(TEXT("Project is out of date, recompiling...")), true, true);
-				bool bCompileResult = FDesktopPlatformModule::Get()->CompileGameProject(FPaths::RootDir(), FPaths::GetProjectFilePath(), Context);
-				Context->EndSlowTask();
+			// If we're running with -stdout, assume that we're a non interactive process and about to fail
+			if (FApp::IsUnattended() || FParse::Param(FCommandLine::Get(), TEXT("stdout")))
+			{
+				return false;
+			}
 
-				// Check it succeeded
-				if(!bCompileResult || !IProjectManager::Get().AreProjectModulesUpToDate() || !IPluginManager::Get().AreEnabledPluginModulesUpToDate())
-				{
-					if(FPlatformMisc::MessageBoxExt(EAppMsgType::YesNo, TEXT("Game code couldn't be compiled. Continue trying to start anyway?"), TEXT("Error")) == EAppReturnType::No)
-					{
-						return false;
-					}
-				}
+			// Ask whether to compile before continuing
+			if(FPlatformMisc::MessageBoxExt(EAppMsgType::YesNo, *ModulesList, TEXT("Question")) == EAppReturnType::No)
+			{
+				return false;
 			}
-			else if(Result == EAppReturnType::Cancel)
+
+			// Try to compile it
+			FFeedbackContext *Context = (FFeedbackContext*)FDesktopPlatformModule::Get()->GetNativeFeedbackContext();
+			Context->BeginSlowTask(FText::FromString(TEXT("Starting build...")), true, true);
+			bool bCompileResult = FDesktopPlatformModule::Get()->CompileGameProject(FPaths::RootDir(), FPaths::GetProjectFilePath(), Context);
+			Context->EndSlowTask();
+
+			// Check it succeeded
+			TArray<FString> StillIncompatibleFiles;
+			if(!bCompileResult || !IProjectManager::Get().CheckModuleCompatibility(StillIncompatibleFiles) || !IPluginManager::Get().CheckModuleCompatibility(StillIncompatibleFiles))
 			{
+				for(int Idx = 0; Idx < StillIncompatibleFiles.Num(); Idx++)
+				{
+					UE_LOG(LogInit, Warning, TEXT("Still incompatible or missing module: %s"), *StillIncompatibleFiles[Idx]);
+				}
+				FPlatformMisc::MessageBoxExt(EAppMsgType::Ok, *FString::Printf(TEXT("%s could not be compiled. Try rebuilding from source manually."), GGameName), TEXT("Error"));
 				return false;
 			}
 		}
diff --git a/Engine/Source/Runtime/Projects/Private/ModuleDescriptor.cpp b/Engine/Source/Runtime/Projects/Private/ModuleDescriptor.cpp
--- a/Engine/Source/Runtime/Projects/Private/ModuleDescriptor.cpp
+++ b/Engine/Source/Runtime/Projects/Private/ModuleDescriptor.cpp
@@ -360,17 +360,19 @@ void FModuleDescriptor::LoadModulesForPhase(ELoadingPhase::Type LoadingPhase, co
 	}
 }
 
-bool FModuleDescriptor::AreModulesUpToDate(const TArray<FModuleDescriptor>& Modules)
+bool FModuleDescriptor::CheckModuleCompatbility(const TArray<FModuleDescriptor>& Modules, bool bGameModules, TArray<FString>& OutIncompatibleFiles)
 {
+	bool bResult = true;
 	for(int Idx = 0; Idx < Modules.Num(); Idx++)
 	{
 		const FModuleDescriptor &Module = Modules[Idx];
 		if (Module.IsCompiledInCurrentConfiguration() && !FModuleManager::Get().IsModuleUpToDate(Module.Name))
 		{
-			return false;
+			OutIncompatibleFiles.Add(FModuleManager::GetCleanModuleFilename(Module.Name, bGameModules));
+			bResult = false;
 		}
 	}
-	return true;
+	return bResult;
 }
 
 #undef LOCTEXT_NAMESPACE
diff --git a/Engine/Source/Runtime/Projects/Private/PluginManager.cpp b/Engine/Source/Runtime/Projects/Private/PluginManager.cpp
--- a/Engine/Source/Runtime/Projects/Private/PluginManager.cpp
+++ b/Engine/Source/Runtime/Projects/Private/PluginManager.cpp
@@ -412,23 +412,23 @@ bool FPluginManager::AreRequiredPluginsAvailable()
 	return ConfigureEnabledPlugins();
 }
 
-bool FPluginManager::AreEnabledPluginModulesUpToDate()
+bool FPluginManager::CheckModuleCompatibility(TArray<FString>& OutIncompatibleModules)
 {
 	if(!ConfigureEnabledPlugins())
 	{
 		return false;
 	}
 
+	bool bResult = true;
 	for (TArray< TSharedRef< FPluginInstance > >::TConstIterator Iter(AllPlugins); Iter; ++Iter)
 	{
 		const TSharedRef< FPluginInstance > &Plugin = *Iter;
-		if (Plugin->bEnabled && !FModuleDescriptor::AreModulesUpToDate(Plugin->Descriptor.Modules))
+		if (Plugin->bEnabled && !FModuleDescriptor::CheckModuleCompatbility(Plugin->Descriptor.Modules, Plugin->LoadedFrom == EPluginLoadedFrom::GameProject, OutIncompatibleModules))
 		{
-			return false;
+			bResult = false;
 		}
 	}
-
-	return true;
+	return bResult;
 }
 
 IPluginManager& IPluginManager::Get()
diff --git a/Engine/Source/Runtime/Projects/Private/PluginManager.h b/Engine/Source/Runtime/Projects/Private/PluginManager.h
--- a/Engine/Source/Runtime/Projects/Private/PluginManager.h
+++ b/Engine/Source/Runtime/Projects/Private/PluginManager.h
@@ -62,7 +62,7 @@ public:
 	virtual bool LoadModulesForEnabledPlugins( const ELoadingPhase::Type LoadingPhase ) override;
 	virtual void SetRegisterMountPointDelegate( const FRegisterMountPointDelegate& Delegate ) override;
 	virtual bool AreRequiredPluginsAvailable() override;
-	virtual bool AreEnabledPluginModulesUpToDate() override;
+	virtual bool CheckModuleCompatibility( TArray<FString>& OutIncompatibleModules ) override;
 	virtual TArray< FPluginStatus > QueryStatusForAllPlugins() const override;
 	virtual const TArray< FPluginContentFolder >& GetPluginContentFolders() const override;
 
diff --git a/Engine/Source/Runtime/Projects/Private/ProjectManager.cpp b/Engine/Source/Runtime/Projects/Private/ProjectManager.cpp
--- a/Engine/Source/Runtime/Projects/Private/ProjectManager.cpp
+++ b/Engine/Source/Runtime/Projects/Private/ProjectManager.cpp
@@ -96,9 +96,9 @@ bool FProjectManager::LoadModulesForProject( const ELoadingPhase::Type LoadingPh
 	return bSuccess;
 }
 
-bool FProjectManager::AreProjectModulesUpToDate()
+bool FProjectManager::CheckModuleCompatibility(TArray<FString>& OutIncompatibleModules)
 {
-	return !CurrentProject.IsValid() || FModuleDescriptor::AreModulesUpToDate(CurrentProject->Modules);
+	return !CurrentProject.IsValid() || FModuleDescriptor::CheckModuleCompatbility(CurrentProject->Modules, true, OutIncompatibleModules);
 }
 
 const FString& FProjectManager::GetAutoLoadProjectFileName()
diff --git a/Engine/Source/Runtime/Projects/Private/ProjectManager.h b/Engine/Source/Runtime/Projects/Private/ProjectManager.h
--- a/Engine/Source/Runtime/Projects/Private/ProjectManager.h
+++ b/Engine/Source/Runtime/Projects/Private/ProjectManager.h
@@ -16,7 +16,7 @@ public:
 	virtual const FProjectDescriptor *GetCurrentProject() const override;
 	virtual bool LoadProjectFile( const FString& ProjectFile ) override;
 	virtual bool LoadModulesForProject( const ELoadingPhase::Type LoadingPhase ) override;
-	virtual bool AreProjectModulesUpToDate( ) override;
+	virtual bool CheckModuleCompatibility( TArray<FString>& OutIncompatibleModules ) override;
 	virtual const FString& GetAutoLoadProjectFileName() override;
 	virtual bool GenerateNewProjectFile(const FString& NewProjectFilename, const TArray<FString>& ModuleNames, const FString& EngineIdentifier, FText& OutFailReason) override;
 	virtual bool SignSampleProject(const FString& FilePath, const FString& Category, FText& OutFailReason) override;
diff --git a/Engine/Source/Runtime/Projects/Public/Interfaces/IPluginManager.h b/Engine/Source/Runtime/Projects/Public/Interfaces/IPluginManager.h
--- a/Engine/Source/Runtime/Projects/Public/Interfaces/IPluginManager.h
+++ b/Engine/Source/Runtime/Projects/Public/Interfaces/IPluginManager.h
@@ -121,9 +121,10 @@ public:
 	/** 
 	 * Checks whether modules for the enabled plugins are up to date
 	 *
+	 * @param OutIncompatibleNames	Array to receive a list of incompatible module names
 	 * @returns true if the enabled plugin modules are up to date
 	 */
-	virtual bool AreEnabledPluginModulesUpToDate( ) = 0;
+	virtual bool CheckModuleCompatibility( TArray<FString>& OutIncompatibleModules ) = 0;
 
 	/**
 	 * Gets status about all currently known plugins
diff --git a/Engine/Source/Runtime/Projects/Public/Interfaces/IProjectManager.h b/Engine/Source/Runtime/Projects/Public/Interfaces/IProjectManager.h
--- a/Engine/Source/Runtime/Projects/Public/Interfaces/IProjectManager.h
+++ b/Engine/Source/Runtime/Projects/Public/Interfaces/IProjectManager.h
@@ -106,9 +106,9 @@ public:
 	/**
 	 * Checks if the modules for a project are up to date
 	 *
-	 * @return	true if the UBT needs to be run to recompile modules for a project.
+	 * @return	false if UBT needs to be run to recompile modules for a project.
 	 */
-	virtual bool AreProjectModulesUpToDate( ) = 0;
+	virtual bool CheckModuleCompatibility(TArray<FString>& OutIncompatibleModules) = 0;
 
 	/**
 	 * Gets the name of the text file that contains the most recently loaded filename.
diff --git a/Engine/Source/Runtime/Projects/Public/ModuleDescriptor.h b/Engine/Source/Runtime/Projects/Public/ModuleDescriptor.h
--- a/Engine/Source/Runtime/Projects/Public/ModuleDescriptor.h
+++ b/Engine/Source/Runtime/Projects/Public/ModuleDescriptor.h
@@ -132,6 +132,6 @@ struct PROJECTS_API FModuleDescriptor
 	/** Loads all the modules for a given loading phase. Returns a map of module names to load errors */
 	static void LoadModulesForPhase(ELoadingPhase::Type LoadingPhase, const TArray<FModuleDescriptor>& Modules, TMap<FName, EModuleLoadResult>& ModuleLoadErrors);
 
-	/** Checks whether all the modules in an array are up to date */
-	static bool AreModulesUpToDate(const TArray<FModuleDescriptor>& Modules);
+	/** Checks that all modules are compatible with the current engine version. Returns false and appends a list of names to OutIncompatibleFiles if not. */
+	static bool CheckModuleCompatbility(const TArray<FModuleDescriptor>& Modules, bool bGameModules, TArray<FString>& OutIncompatibleFiles);
 };
