Core: Improved handling of delegate invocation list compaction

diff --git a/Engine/Source/Runtime/Core/Public/Delegates/DelegateSignatureImpl.inl b/Engine/Source/Runtime/Core/Public/Delegates/DelegateSignatureImpl.inl
--- a/Engine/Source/Runtime/Core/Public/Delegates/DelegateSignatureImpl.inl
+++ b/Engine/Source/Runtime/Core/Public/Delegates/DelegateSignatureImpl.inl
@@ -1839,10 +1839,12 @@ protected:
 	 */
 	void Broadcast( FUNC_PARAM_LIST ) const
 	{
-		const TArray<IDelegateInstance*>& InvocationList = GetInvocationList();
+		bool NeedsCompaction = false;
 
 		LockInvocationList();
 		{
+			const TArray<IDelegateInstance*>& InvocationList = GetInvocationList();
+
 			// call bound functions in reverse order, so we ignore any instances that may be added by callees
 			for (int32 InvocationListIndex = InvocationList.Num() - 1; InvocationListIndex >= 0; --InvocationListIndex)
 			{
@@ -1851,16 +1853,16 @@ protected:
 
 				if ((DelegateInstanceInterface == nullptr) || !DelegateInstanceInterface->ExecuteIfSafe(FUNC_PARAM_PASSTHRU))
 				{
-					const_cast<BASE_MULTICAST_DELEGATE_CLASS*>(this)->RequestCompaction();
+					NeedsCompaction = true;
 				}
 			}
 		}
 		UnlockInvocationList();
 
-		// NOTE: The implementation of multicast delegates relies on the assumption
-		// that this is the only place where invocation list compaction takes place.
-
-		const_cast<BASE_MULTICAST_DELEGATE_CLASS*>(this)->CompactInvocationList();
+		if (NeedsCompaction)
+		{
+			const_cast<BASE_MULTICAST_DELEGATE_CLASS*>(this)->CompactInvocationList();
+		}
 	}
 
 	/**
@@ -1894,6 +1896,8 @@ protected:
 				break;	// no need to continue, as we never allow the same delegate to be bound twice
 			}
 		}
+
+		const_cast<BASE_MULTICAST_DELEGATE_CLASS*>(this)->CompactInvocationList();
 	}
 };
 
diff --git a/Engine/Source/Runtime/Core/Public/Delegates/MulticastDelegateBase.h b/Engine/Source/Runtime/Core/Public/Delegates/MulticastDelegateBase.h
--- a/Engine/Source/Runtime/Core/Public/Delegates/MulticastDelegateBase.h
+++ b/Engine/Source/Runtime/Core/Public/Delegates/MulticastDelegateBase.h
@@ -11,9 +11,7 @@ class FMulticastDelegateBase
 {
 public:
 
-	/**
-	 * Removes all functions from this delegate's invocation list
-	 */
+	/** Removes all functions from this delegate's invocation list. */
 	void Clear( )
 	{
 		for (IDelegateInstance*& DelegateInstanceRef : InvocationList)
@@ -21,6 +19,8 @@ public:
 			delete DelegateInstanceRef;
 			DelegateInstanceRef = nullptr;
 		}
+
+		CompactInvocationList();
 	}
 
 	/**
@@ -76,16 +76,16 @@ public:
 				DelegateInstanceRef = nullptr;
 			}
 		}
+
+		CompactInvocationList();
 	}
 
 protected:
 
-	/**
-	 * Hidden default constructor.
-	 */
+	/** Hidden default constructor. */
 	inline FMulticastDelegateBase( )
-		: InvocationListLockCount(0)
-		, NeedsCompaction(false)
+		: CompactionThreshold(2)
+		, InvocationListLockCount(0)
 	{ }
 
 protected:
@@ -107,7 +107,7 @@ protected:
 	 */
 	void CompactInvocationList( )
 	{
-		if (!NeedsCompaction || (InvocationListLockCount != 0))
+		if ((InvocationList.Num() < CompactionThreshold) || (InvocationListLockCount != 0))
 		{
 			return;
 		}
@@ -127,7 +127,7 @@ protected:
 			}
 		}
 
-		NeedsCompaction = false;
+		CompactionThreshold = FMath::Max(2, 2 * InvocationList.Num());
 	}
 
 	/**
@@ -140,27 +140,13 @@ protected:
 		return InvocationList;
 	}
 
-	/**
-	 * Requests that the invocation list be compacted.
-	 *
-	 * @see CompactInvocationList
-	 */
-	inline void RequestCompaction( )
-	{
-		NeedsCompaction = true;
-	}
-
-	/**
-	 * Increments the lock counter for the invocation list.
-	 */
+	/** Increments the lock counter for the invocation list. */
 	inline void LockInvocationList( ) const
 	{
 		++InvocationListLockCount;
 	}
 
-	/**
-	 * Decrements the lock counter for the invocation list.
-	 */
+	/** Decrements the lock counter for the invocation list. */
 	inline void UnlockInvocationList( ) const
 	{
 		--InvocationListLockCount;
@@ -168,12 +154,12 @@ protected:
 
 private:
 
-	// Holds the collection of delegate instances to invoke.
+	/** Used to determine when a compaction should happen. */
+	int32 CompactionThreshold;
+
+	/** Holds the collection of delegate instances to invoke. */
 	TArray<IDelegateInstance*> InvocationList;
 
-	// Holds a lock counter for the invocation list.
+	/** Holds a lock counter for the invocation list. */
 	mutable int32 InvocationListLockCount;
-
-	// Holds a flag indicating whether the invocation list needs to be compacted.
-	bool NeedsCompaction;
 };
