AutomationTest.cpp
--Added in a check to make sure that there is a module name when loading test modules.

diff --git a/Engine/Source/Runtime/Core/Private/Misc/AutomationTest.cpp b/Engine/Source/Runtime/Core/Private/Misc/AutomationTest.cpp
--- a/Engine/Source/Runtime/Core/Private/Misc/AutomationTest.cpp
+++ b/Engine/Source/Runtime/Core/Private/Misc/AutomationTest.cpp
@@ -314,8 +314,15 @@ void FAutomationTestFramework::LoadTestModules( )
 		for( int32 EngineModuleId = 0; EngineModuleId < EngineTestModules.Num(); ++EngineModuleId)
 		{
 			const FName ModuleName = FName(*EngineTestModules[EngineModuleId]);
+			//Make sure that there is a name available.  This can happen if a name is left blank in the Engine.ini
+			if (ModuleName == NAME_None || ModuleName == TEXT("None"))
+			{
+				UE_LOG(LogAutomationTest, Warning, TEXT("The automation test module ('%s') doesn't have a valid name."), *ModuleName.ToString());
+				continue;
+			}
 			if (!FModuleManager::Get().IsModuleLoaded(ModuleName))
 			{
+				UE_LOG(LogAutomationTest, Log, TEXT("Loading automation test module: '%s'."), *ModuleName.ToString());
 				FModuleManager::Get().LoadModule(ModuleName);
 			}
 		}
@@ -327,8 +334,15 @@ void FAutomationTestFramework::LoadTestModules( )
 			for( int32 EditorModuleId = 0; EditorModuleId < EditorTestModules.Num(); ++EditorModuleId )
 			{
 				const FName ModuleName = FName(*EditorTestModules[EditorModuleId]);
+				//Make sure that there is a name available.  This can happen if a name is left blank in the Engine.ini
+				if (ModuleName == NAME_None || ModuleName == TEXT("None"))
+				{
+					UE_LOG(LogAutomationTest, Warning, TEXT("The automation test module ('%s') doesn't have a valid name."), *ModuleName.ToString());
+					continue;
+				}
 				if (!FModuleManager::Get().IsModuleLoaded(ModuleName))
 				{
+					UE_LOG(LogAutomationTest, Log, TEXT("Loading automation test module: '%s'."), *ModuleName.ToString());
 					FModuleManager::Get().LoadModule(ModuleName);
 				}
 			}
