HashCombine function.

diff --git a/Engine/Source/Runtime/Core/Public/Templates/TypeHash.h b/Engine/Source/Runtime/Core/Public/Templates/TypeHash.h
--- a/Engine/Source/Runtime/Core/Public/Templates/TypeHash.h
+++ b/Engine/Source/Runtime/Core/Public/Templates/TypeHash.h
@@ -7,38 +7,39 @@
 #pragma once
 
 /**
- * A hashing function that works well for pointers.
+ * Combines two hash values to get a third.
+ * Note - this function is not commutative.
  */
-inline uint32 PointerHash(const void* Key,uint32 C = 0)
+inline uint32 HashCombine(uint32 A, uint32 C)
 {
-#define mix(a,b,c) \
-{ \
-  a -= b; a -= c; a ^= (c>>13); \
-  b -= c; b -= a; b ^= (a<<8); \
-  c -= a; c -= b; c ^= (b>>13); \
-  a -= b; a -= c; a ^= (c>>12);  \
-  b -= c; b -= a; b ^= (a<<16); \
-  c -= a; c -= b; c ^= (b>>5); \
-  a -= b; a -= c; a ^= (c>>3);  \
-  b -= c; b -= a; b ^= (a<<10); \
-  c -= a; c -= b; c ^= (b>>15); \
+	uint32 B = 0x9e3779b9;
+	A += B;
+
+	A -= B; A -= C; A ^= (C>>13);
+	B -= C; B -= A; B ^= (A<<8);
+	C -= A; C -= B; C ^= (B>>13);
+	A -= B; A -= C; A ^= (C>>12);
+	B -= C; B -= A; B ^= (A<<16);
+	C -= A; C -= B; C ^= (B>>5);
+	A -= B; A -= C; A ^= (C>>3);
+	B -= C; B -= A; B ^= (A<<10);
+	C -= A; C -= B; C ^= (B>>15);
+
+	return C;
 }
 
-	uint32 A;
-	uint32 B;
-	A = B = 0x9e3779b9;
+inline uint32 PointerHash(const void* Key,uint32 C = 0)
+{
 	// Avoid LHS stalls on PS3 and Xbox 360
 #if PLATFORM_64BITS
 	// Ignoring the lower 4 bits since they are likely zero anyway.
 	// Higher bits are more significant in 64 bit builds.
-	A += (reinterpret_cast<UPTRINT>(Key) >> 4);
+	auto PtrInt = reinterpret_cast<UPTRINT>(Key) >> 4;
 #else
-	A += reinterpret_cast<UPTRINT>(Key);
+	auto PtrInt = reinterpret_cast<UPTRINT>(Key);
 #endif
-	mix(A,B,C);
-	return C;
 
-#undef mix
+	return HashCombine(PtrInt, C);
 }
 
 //
