Fix popup menus in UT.

diff --git a/Engine/Source/Runtime/Core/Private/Linux/LinuxApplication.cpp b/Engine/Source/Runtime/Core/Private/Linux/LinuxApplication.cpp
--- a/Engine/Source/Runtime/Core/Private/Linux/LinuxApplication.cpp
+++ b/Engine/Source/Runtime/Core/Private/Linux/LinuxApplication.cpp
@@ -937,15 +937,38 @@ void FLinuxApplication::SetHighPrecisionMouseMode( const bool Enable, const TSha
 
 FPlatformRect FLinuxApplication::GetWorkArea( const FPlatformRect& CurrentWindow ) const
 {
-	SDL_Rect sdl_rect;
-	SDL_GetDisplayBounds(0, &sdl_rect);
+	// loop over all monitors to determine which one is the best
+	int NumDisplays = SDL_GetNumVideoDisplays();
+	if (NumDisplays <= 0)
+	{
+		// fake something
+		return CurrentWindow;
+	}
 	
-	FPlatformRect WorkArea;
+	SDL_Rect BestDisplayBounds;
+	SDL_GetDisplayBounds(0, &BestDisplayBounds);
 	
-	WorkArea.Top = sdl_rect.y;
-	WorkArea.Bottom = sdl_rect.y + sdl_rect.h;
-	WorkArea.Left = sdl_rect.x;
-	WorkArea.Right = sdl_rect.x + sdl_rect.w;
+	// see if any other are better (i.e. cover top left)
+	for (int DisplayIdx = 1; DisplayIdx < NumDisplays; ++DisplayIdx)
+	{
+		SDL_Rect DisplayBounds;
+		SDL_GetDisplayBounds(DisplayIdx, &DisplayBounds);
+		
+		// only check top left corner for "bestness"
+		if (DisplayBounds.x <= CurrentWindow.Left && DisplayBounds.x + DisplayBounds.w > CurrentWindow.Left &&
+			DisplayBounds.y <= CurrentWindow.Top && DisplayBounds.y + DisplayBounds.h > CurrentWindow.Bottom)
+		{
+			BestDisplayBounds = DisplayBounds;
+			// there can be only one, as we don't expect overlapping displays
+			break;
+		}
+	}
+	
+	FPlatformRect WorkArea;
+	WorkArea.Left	= BestDisplayBounds.x;
+	WorkArea.Top	= BestDisplayBounds.y;
+	WorkArea.Right	= BestDisplayBounds.x + BestDisplayBounds.w;
+	WorkArea.Bottom	= BestDisplayBounds.y + BestDisplayBounds.h;
 
 	return WorkArea;
 }
