#UE4 Critical fix to FText Serialization. When FText was being serialized for cases other than saving and loading (like reference gathering), the FText would be corrupted and data would be lost. Now it only clears and recreates the text history on load

diff --git a/Engine/Source/Runtime/Core/Private/Internationalization/Text.cpp b/Engine/Source/Runtime/Core/Private/Internationalization/Text.cpp
--- a/Engine/Source/Runtime/Core/Private/Internationalization/Text.cpp
+++ b/Engine/Source/Runtime/Core/Private/Internationalization/Text.cpp
@@ -574,7 +574,7 @@ CORE_API FArchive& operator<<( FArchive& Ar, FText& Value )
 				Ar << NoHistory;
 			}
 		}
-		else
+		else if (Ar.IsLoading())
 		{
 			// The type is serialized during the serialization of the history, during deserialization we need to deserialize it and create the correct history
 			int8 HistoryType = INDEX_NONE;
