Back out revision 22 from //depot/UE4/Engine/Source/Runtime/Core/Private/Math/Transform.cpp

diff --git a/Engine/Source/Runtime/Core/Private/Math/Transform.cpp b/Engine/Source/Runtime/Core/Private/Math/Transform.cpp
--- a/Engine/Source/Runtime/Core/Private/Math/Transform.cpp
+++ b/Engine/Source/Runtime/Core/Private/Math/Transform.cpp
@@ -46,7 +46,7 @@ FString FTransform::ToString() const
 	const FVector T(GetTranslation());
 	const FVector S(GetScale3D());
 
-	return FString::Printf(TEXT("%f,%f,%f|P(%f),Y(%f),R(%f)|%f,%f,%f"), T.X, T.Y, T.Z, R.Pitch, R.Yaw, R.Roll, S.X, S.Y, S.Z);
+	return FString::Printf(TEXT("%f,%f,%f|%f,%f,%f|%f,%f,%f"), T.X, T.Y, T.Z, R.Pitch, R.Yaw, R.Roll, S.X, S.Y, S.Z);
 }
 
 bool FTransform::InitFromString( const FString & Source )
