- Added support for 0 contentScaleFactor (Zen using it for iPhone5S and up)

diff --git a/Engine/Source/Programs/UnrealBuildTool/IOS/UEDeployIOS.cs b/Engine/Source/Programs/UnrealBuildTool/IOS/UEDeployIOS.cs
--- a/Engine/Source/Programs/UnrealBuildTool/IOS/UEDeployIOS.cs
+++ b/Engine/Source/Programs/UnrealBuildTool/IOS/UEDeployIOS.cs
@@ -1,4 +1,4 @@
-// Copyright 1998-2014 Epic Games, Inc. All Rights Reserved.
+﻿// Copyright 1998-2014 Epic Games, Inc. All Rights Reserved.
 
 using System;
 using System.Collections.Generic;
@@ -98,8 +98,14 @@ namespace UnrealBuildTool.IOS
 			}
 		}
 
+
 		public override bool PrepForUATPackageOrDeploy(string InProjectName, string InProjectDirectory, string InExecutablePath, string InEngineDir, bool bForDistribution, string CookFlavor)
 		{
+			if (BuildHostPlatform.Current.Platform != UnrealTargetPlatform.Mac)
+			{
+				throw new BuildException("UEDeployIOS.PrepForUATPackageOrDeploy only supports running on the Mac");
+			}
+
 			bool bIsUE4Game = InExecutablePath.Contains ("UE4Game");
 			string BinaryPath = Path.GetDirectoryName (InExecutablePath);
 			string GameExeName = Path.GetFileName (InExecutablePath);
@@ -276,9 +282,6 @@ namespace UnrealBuildTool.IOS
 				PListAdditionalLines.Add("\t<key>UILaunchStoryboardName</key>");
 				PListAdditionalLines.Add("\t<string>LaunchScreen</string>");
 
-				PListAdditionalLines.Add("\t<key>NibLaunch</key>");
-				PListAdditionalLines.Add("\t<string>" + CommandLine + "</string>");
-
 				bSkipDefaultPNGs = true;
 			}
 
diff --git a/Engine/Source/Runtime/Core/Private/Android/AndroidWindow.cpp b/Engine/Source/Runtime/Core/Private/Android/AndroidWindow.cpp
--- a/Engine/Source/Runtime/Core/Private/Android/AndroidWindow.cpp
+++ b/Engine/Source/Runtime/Core/Private/Android/AndroidWindow.cpp
@@ -103,22 +103,31 @@ FPlatformRect FAndroidWindow::GetScreenRect()
 	static IConsoleVariable* CVar = IConsoleManager::Get().FindConsoleVariable(TEXT("r.MobileContentScaleFactor")); 
 	float RequestedContentScaleFactor = CVar->GetFloat();
 
-	int32 Width, Height;
-	if (GAndroidIsPortrait)
+	// 0 means to use native size
+	if (RequestedContentScaleFactor == 0.0f)
 	{
-		Height = 1280 * RequestedContentScaleFactor;
+		Width = MaxWidth;
+		Height = MaxHeight;
 	}
 	else
 	{
-		Height = 720 * RequestedContentScaleFactor;
-	}
+		int32 Width, Height;
+		if (GAndroidIsPortrait)
+		{
+			Height = 1280 * RequestedContentScaleFactor;
+		}
+		else
+		{
+			Height = 720 * RequestedContentScaleFactor;
+		}
 
-	// apply the aspect ration to get the width
-	Width = Height * AspectRatio;
+		// apply the aspect ration to get the width
+		Width = Height * AspectRatio;
 
-	// clamp to native resolution
-	Width = FPlatformMath::Min(Width, MaxWidth);
-	Height = FPlatformMath::Min(Height, MaxHeight);
+		// clamp to native resolution
+		Width = FPlatformMath::Min(Width, MaxWidth);
+		Height = FPlatformMath::Min(Height, MaxHeight);
+	}
 
 	FPlatformRect ScreenRect;
 	ScreenRect.Left = 0;
diff --git a/Engine/Source/Runtime/Core/Private/IOS/IOSView.cpp b/Engine/Source/Runtime/Core/Private/IOS/IOSView.cpp
--- a/Engine/Source/Runtime/Core/Private/IOS/IOSView.cpp
+++ b/Engine/Source/Runtime/Core/Private/IOS/IOSView.cpp
@@ -145,9 +145,22 @@ id<MTLDevice> GMetalDevice = nil;
 		static IConsoleVariable* CVar = IConsoleManager::Get().FindConsoleVariable(TEXT("r.MobileContentScaleFactor"));
 		float RequestedContentScaleFactor = CVar->GetFloat();
 
-		// for TV screens, always use scale factor of 1
-		self.contentScaleFactor = bIsForOnDevice ? RequestedContentScaleFactor : 1.0f;
-		UE_LOG(LogIOS, Log, TEXT("Setting contentScaleFactor to %0.4f (optimal = %0.4f)"), self.contentScaleFactor, NativeScale);
+		// 0 means to leave the scale alone, use native
+		if (RequestedContentScaleFactor == 0.0f)
+		{
+#ifdef __IPHONE_8_0
+			UE_LOG(LogIOS, Log, TEXT("Leaving view scale factor alone, with scale = %f, nativeScale = %f"), [[UIScreen mainScreen] scale], [[UIScreen mainScreen] nativeScale]);
+#else
+			UE_LOG(LogIOS, Log, TEXT("Leaving view scale factor alone, with scale = %f"), NativeScale);
+#endif
+		}
+		else
+		{
+			// for TV screens, always use scale factor of 1
+			self.contentScaleFactor = bIsForOnDevice ? RequestedContentScaleFactor : 1.0f;
+			UE_LOG(LogIOS, Log, TEXT("Setting contentScaleFactor to %0.4f (optimal = %0.4f)"), self.contentScaleFactor, NativeScale);
+		}
+
 
 		// metal has no further action needed
 		if (!bIsUsingMetal)
