Core: Fixed broken formatting codes in FTimespan::ToString()

diff --git a/Engine/Source/Runtime/Core/Private/Misc/Timespan.cpp b/Engine/Source/Runtime/Core/Private/Misc/Timespan.cpp
--- a/Engine/Source/Runtime/Core/Private/Misc/Timespan.cpp
+++ b/Engine/Source/Runtime/Core/Private/Misc/Timespan.cpp
@@ -59,9 +59,9 @@ FString FTimespan::ToString( const TCHAR* Format ) const
 			case TCHAR('f'): Result += FString::Printf(TEXT("%03i"), FMath::Abs(GetMilliseconds())); break;
 			case TCHAR('D'): Result += FString::Printf(TEXT("%f"), FMath::Abs(GetTotalDays())); break;
 			case TCHAR('H'): Result += FString::Printf(TEXT("%f"), FMath::Abs(GetTotalHours())); break;
-			case TCHAR('M'): Result += FString::Printf(TEXT("%02i"), FMath::Abs(GetTotalMinutes())); break;
-			case TCHAR('S'): Result += FString::Printf(TEXT("%02i"), FMath::Abs(GetTotalSeconds())); break;
-			case TCHAR('F'): Result += FString::Printf(TEXT("%03i"), FMath::Abs(GetTotalMilliseconds())); break;
+			case TCHAR('M'): Result += FString::Printf(TEXT("%f"), FMath::Abs(GetTotalMinutes())); break;
+			case TCHAR('S'): Result += FString::Printf(TEXT("%f"), FMath::Abs(GetTotalSeconds())); break;
+			case TCHAR('F'): Result += FString::Printf(TEXT("%f"), FMath::Abs(GetTotalMilliseconds())); break;
 
 			default:
 
diff --git a/Engine/Source/Runtime/Core/Public/Misc/Timespan.h b/Engine/Source/Runtime/Core/Public/Misc/Timespan.h
--- a/Engine/Source/Runtime/Core/Public/Misc/Timespan.h
+++ b/Engine/Source/Runtime/Core/Public/Misc/Timespan.h
@@ -440,11 +440,11 @@ public:
 	 *		%m - prints the time span's minutes part (0..59)
 	 *		%s - prints the time span's seconds part (0..59)
 	 *		%f - prints the time span's milliseconds part (0..999)
-	 *		%D - prints the total number of days
-	 *		%H - prints the total number of hours (0..23)
-	 *		%M - prints the total number of minutes (0..59)
-	 *		%S - prints the total number of seconds (0..59)
-	 *		%F - prints the total number of milliseconds (0..999)
+	 *		%D - prints the total number of days (without minus sign)
+	 *		%H - prints the total number of hours (without minus sign)
+	 *		%M - prints the total number of minutes (without minus sign)
+	 *		%S - prints the total number of seconds (without minus sign)
+	 *		%F - prints the total number of milliseconds (without minus sign)
 	 *
 	 * @param Format - The format of the returned string.
 	 * @return String representation.
