Repairs to source code, localization archives, and locres files based on errors made by translators.

diff --git a/Engine/Source/Editor/UnrealEd/Private/Kismet2/BlueprintEditorUtils.cpp b/Engine/Source/Editor/UnrealEd/Private/Kismet2/BlueprintEditorUtils.cpp
--- a/Engine/Source/Editor/UnrealEd/Private/Kismet2/BlueprintEditorUtils.cpp
+++ b/Engine/Source/Editor/UnrealEd/Private/Kismet2/BlueprintEditorUtils.cpp
@@ -4494,7 +4494,7 @@ void FBlueprintEditorUtils::ConformImplementedEvents(UBlueprint* Blueprint)
 								Args.Add(TEXT("EventNodeName"), FText::FromString(EventNode->GetName()));
 
 								// Emit something to the log to indicate that we've made a change
-								Blueprint->Message_Note( FText::Format(LOCTEXT("EventSignatureFixed_Note", "{NodeTitle ({EventNodeName}) had an invalid function signature - it has now been fixed."), Args).ToString() );
+								Blueprint->Message_Note( FText::Format(LOCTEXT("EventSignatureFixed_Note", "{NodeTitle} ({EventNodeName}) had an invalid function signature - it has now been fixed."), Args).ToString() );
 
 								// Fix up the event signature
 								EventNode->EventSignatureClass = Blueprint->GeneratedClass;
diff --git a/Engine/Source/Runtime/Core/Private/Internationalization/TextLocalizationResourceGenerator.cpp b/Engine/Source/Runtime/Core/Private/Internationalization/TextLocalizationResourceGenerator.cpp
--- a/Engine/Source/Runtime/Core/Private/Internationalization/TextLocalizationResourceGenerator.cpp
+++ b/Engine/Source/Runtime/Core/Private/Internationalization/TextLocalizationResourceGenerator.cpp
@@ -208,7 +208,7 @@ bool FTextLocalizationResourceGenerator::Generate(const FString& SourcePath, con
 			const FString& Namespace = ManifestEntry->Namespace;
 			const FLocItem& Source = ManifestEntry->Source;
 			const FString& SourceString = Source.Text;
-			const FString UnescapedSourceString = SourceString.ReplaceEscapedCharWithChar();
+			const FString UnescapedSourceString = SourceString;
 			const uint32 SourceStringHash = FCrc::StrCrc32(*UnescapedSourceString);
 
 			FLocalizationEntryTracker::FKeyTable& KeyTable = LocalizationEntryTracker.Namespaces.FindOrAdd(*Namespace);
@@ -236,7 +236,7 @@ bool FTextLocalizationResourceGenerator::Generate(const FString& SourcePath, con
 				FString UnescapedTranslatedString;
 				if( ArchiveEntry.IsValid() )
 				{
-					UnescapedTranslatedString = ArchiveEntry->Translation.Text.ReplaceEscapedCharWithChar();
+					UnescapedTranslatedString = ArchiveEntry->Translation.Text;
 
 					if(UnescapedTranslatedString.IsEmpty())
 					{
