UE4: Illegal characters in names are now only reported once, even if they show up in the illegal characters or test string list multiple times.

diff --git a/Engine/Source/Runtime/Core/Private/UObject/UnrealNames.cpp b/Engine/Source/Runtime/Core/Private/UObject/UnrealNames.cpp
--- a/Engine/Source/Runtime/Core/Private/UObject/UnrealNames.cpp
+++ b/Engine/Source/Runtime/Core/Private/UObject/UnrealNames.cpp
@@ -837,12 +837,15 @@ bool FName::IsValidXName( FString InvalidChars/*=INVALID_NAME_CHARACTERS*/, FTex
 	// See if the name contains invalid characters.
 	TCHAR CharString[] = { '\0', '\0' };
 	FString MatchedInvalidChars;
+	TSet<TCHAR> AlreadyMatchedInvalidChars;
 	for( int32 x = 0; x < InvalidChars.Len() ; ++x )
 	{
-		CharString[0] = InvalidChars[x];
-		if( Name.Contains( CharString ) )
+		TCHAR CharToTest = InvalidChars[x];
+		CharString[0] = CharToTest;
+		if( !AlreadyMatchedInvalidChars.Contains( CharToTest ) && Name.Contains( CharString ) )
 		{
 			MatchedInvalidChars += CharString;
+			AlreadyMatchedInvalidChars.Add( CharToTest );
 		}
 	}
 
