Fix crash reporting failing on Mac OS X 10.10 because one of the module names has no path (cl_kernels).

diff --git a/Engine/Source/Runtime/Core/Private/Mac/MacPlatformMisc.cpp b/Engine/Source/Runtime/Core/Private/Mac/MacPlatformMisc.cpp
--- a/Engine/Source/Runtime/Core/Private/Mac/MacPlatformMisc.cpp
+++ b/Engine/Source/Runtime/Core/Private/Mac/MacPlatformMisc.cpp
@@ -1537,7 +1537,14 @@ void FMacCrashContext::GenerateMinidump(char const* Path) const
 			WriteUTF16String(ReportFile, TEXT(" "));
 			
 			char const* Name = FCStringAnsi::Strrchr(ModulePath, '/');
-			Name++;
+            if ( Name )
+            {
+                Name++;
+            }
+            else
+            {
+                Name = ModulePath;
+            }
 			FMemory::Memzero(Line, PATH_MAX * sizeof(TCHAR));
 			FUTF8ToTCHAR_Convert::Convert(Line, PATH_MAX, Name, FCStringAnsi::Strlen(Name));
 			WriteUTF16String(ReportFile, Line);
