Core: Improve error message when GGameName (derived from the .uproject name) mismatches the GameName key in the [URL] section of DefaultGame.ini

diff --git a/Engine/Source/Runtime/Core/Private/Misc/ConfigCacheIni.cpp b/Engine/Source/Runtime/Core/Private/Misc/ConfigCacheIni.cpp
--- a/Engine/Source/Runtime/Core/Private/Misc/ConfigCacheIni.cpp
+++ b/Engine/Source/Runtime/Core/Private/Misc/ConfigCacheIni.cpp
@@ -2711,7 +2711,10 @@ void FConfigCacheIni::InitializeConfigSystem()
 			}
 			else
 			{
-				const FText Message = FText::Format( NSLOCTEXT("Core", "IncorrectGameNameCase", "Game name case not detected for '{0}'. The engine ini should contain an entry [URL]GameName with the name of the game in the correct case."), FText::FromString( GGameName ) );
+				const FText Message = FText::Format(
+					NSLOCTEXT("Core", "MismatchedGameNames", "The name of the .uproject file ('{0}') must match the GameName key in the [URL] section of Config/DefaultGame.ini (currently '{1}').  Please either change the ini or rename the .uproject to match each other (case-insensitive match)."),
+					FText::FromString(GGameName),
+					FText::FromString(GameName));
 				if (!GIsBuildMachine)
 				{
 					UE_LOG(LogInit, Warning, TEXT("%s"), *Message.ToString());
