fix HTML5 Win32 build.

diff --git a/Engine/Source/Runtime/Core/Private/HTML5/HTML5PlatformMisc.cpp b/Engine/Source/Runtime/Core/Private/HTML5/HTML5PlatformMisc.cpp
--- a/Engine/Source/Runtime/Core/Private/HTML5/HTML5PlatformMisc.cpp
+++ b/Engine/Source/Runtime/Core/Private/HTML5/HTML5PlatformMisc.cpp
@@ -6,7 +6,11 @@
 
 #include "Core.h"
 #include "HTML5Application.h"
+
+#if PLATFORM_HTML5_BROWSER
 #include "HTML5JavaScriptFx.h"
+#endif 
+
 #include <ctime>
 
 void FHTML5Misc::PlatformInit()
