Made FName case-preserving by storing case-variant strings in its string table

This is controlled by the macro "WITH_CASE_PRESERVING_NAME", which is currently just set to "WITH_EDITORONLY_DATA" so that it works in editor builds (and UHT).

Added an extra NAME_INDEX entry to FName to store a second string table index for a case-variant string. The previous Index value (now called ComparisonIndex) is still used for comparison purposes (as FNames are still case-insensitive).

The Init process for an FName now works like this:
 1) It will find or add a string table entry for the given string (not matching case) - this entry index is stored in ComparisonIndex.
 2) It will then compare the string table entry string against the given string (matching case) to see if it also needs to add a case-variant entry for the FName.
 3) If it does, it finds or adds a second string table entry (matching case) for the string - this entry index is stored in DisplayIndex.

Hard-coded FNames (those listed in UnrealNames.h) do not support case-variants (due to existing network replication rules for hard-coded FNames), so they skip steps 2 and 3.

I added FMinimalName, which is the same size as FName was previously. This shouldn't really be used (and as such, is deliberately awkward to make/use) as it loses the case-preserving behaviour of FName, however it was required for some things (like stats) that had a hard-coded upper limit on FName size.

I added FScriptName, which always contains the extra display index (even when WITH_CASE_PRESERVING_NAME is disabled). This is used by Blueprint bytecode, as the types used by Blueprint bytecode must be a consistent size between all build configurations.

Other changes:

 - Fixed up any places that were passing an Index into the FName constructor which was supposed to take an EName.
     - Some places were doing this to make the number unique when replicating an object, but this was losing the case-variant information, so I had to fix them.
     - FName will now assert if the EName constructor is used with an value outside the range of hard-coded FNames.

 - Ensured that assets, actors, and blueprint components could all be renamed in a way that only changed their case, and that these changes were correctly persisted.

 - Added FLinkerNamePairKeyFuncs and TLinkerNameMapKeyFuncs for use with TSet and TMap.
     - These allow ULinkerSave and ULinkerLoad to correctly write out case-variants for FNames, and also fixes an issue where the linker would erroneously write out duplicate string table entries for FNames which had a different number (causing package bloat).

 - Bumped VER_MIN_SCRIPTVM_UE4 so that all Blueprint bytecode is recompiled using FScriptName.

diff --git a/Engine/Source/Developer/ProfilerService/Public/Interfaces/IProfilerServiceManager.h b/Engine/Source/Developer/ProfilerService/Public/Interfaces/IProfilerServiceManager.h
--- a/Engine/Source/Developer/ProfilerService/Public/Interfaces/IProfilerServiceManager.h
+++ b/Engine/Source/Developer/ProfilerService/Public/Interfaces/IProfilerServiceManager.h
@@ -284,7 +284,7 @@ FORCEINLINE FArchive& operator<<(FArchive& Ar, FStatNameAndInfo& NameAndInfo)
 	else
 	{
 		FName RawName = NameAndInfo.GetRawName();
-		int32 Index = RawName.GetIndex();
+		int32 Index = RawName.GetComparisonIndex();
 		Ar << Index;
 		int32 Number = NameAndInfo.GetRawNumber();
 		Number |= EStatMetaFlags::SendingFName << (EStatMetaFlags::Shift + EStatAllFields::StartShift);
diff --git a/Engine/Source/Editor/BlueprintGraph/Private/BasicTokenParser.cpp b/Engine/Source/Editor/BlueprintGraph/Private/BasicTokenParser.cpp
--- a/Engine/Source/Editor/BlueprintGraph/Private/BasicTokenParser.cpp
+++ b/Engine/Source/Editor/BlueprintGraph/Private/BasicTokenParser.cpp
@@ -93,7 +93,7 @@ void FBasicToken::SetConstFloat(float InFloat)
 void FBasicToken::SetConstName(FName InName)
 {
 	ConstantType = CPT_Name;
-	*(FName*)NameBytes = InName;
+	*(FScriptName*)NameBytes = NameToScriptName(InName);
 	TokenType = TOKEN_Const;
 }
 
@@ -126,7 +126,7 @@ FString FBasicToken::GetConstantValue() const
 		case CPT_Float:
 			return FString::Printf(TEXT("%f"), Float);
 		case CPT_Name:
-			return FString::Printf(TEXT("%s"), *(*(FName*)NameBytes).ToString());
+			return FString::Printf(TEXT("%s"), *ScriptNameToName(*(FScriptName*)NameBytes).ToString());
 		case CPT_String:
 			return String;
 
diff --git a/Engine/Source/Editor/BlueprintGraph/Private/BasicTokenParser.h b/Engine/Source/Editor/BlueprintGraph/Private/BasicTokenParser.h
--- a/Engine/Source/Editor/BlueprintGraph/Private/BasicTokenParser.h
+++ b/Engine/Source/Editor/BlueprintGraph/Private/BasicTokenParser.h
@@ -108,12 +108,12 @@ public:
 	/* TOKEN_Const values */
 	union
 	{
-		uint8 Byte;							// If CPT_Byte.
-		int32 Int;							// If CPT_Int.
-		bool  NativeBool;					// If CPT_Bool
-		float Float;						// If CPT_Float.
-		uint8 NameBytes[sizeof(FName)];		// If CPT_Name.
-		TCHAR String[MAX_STRING_CONST_SIZE];// If CPT_String
+		uint8 Byte;								// If CPT_Byte.
+		int32 Int;								// If CPT_Int.
+		bool  NativeBool;						// If CPT_Bool
+		float Float;							// If CPT_Float.
+		uint8 NameBytes[sizeof(FScriptName)];	// If CPT_Name.
+		TCHAR String[MAX_STRING_CONST_SIZE];	// If CPT_String
 	};
 };
 
diff --git a/Engine/Source/Editor/ContentBrowser/Private/SAssetView.cpp b/Engine/Source/Editor/ContentBrowser/Private/SAssetView.cpp
--- a/Engine/Source/Editor/ContentBrowser/Private/SAssetView.cpp
+++ b/Engine/Source/Editor/ContentBrowser/Private/SAssetView.cpp
@@ -3225,7 +3225,7 @@ void SAssetView::AssetRenameCommit(const TSharedPtr<FAssetViewItem>& Item, const
 		const TSharedPtr<FAssetViewAsset>& ItemAsAsset = StaticCastSharedPtr<FAssetViewAsset>(Item);
 
 		// Check if the name is different
-		if( NewName == ItemAsAsset->Data.AssetName.ToString() )
+		if( NewName.Equals(ItemAsAsset->Data.AssetName.ToString(), ESearchCase::CaseSensitive) )
 		{
 			RenamingAsset.Reset();
 			return;
diff --git a/Engine/Source/Editor/DetailCustomizations/Private/InputSettingsDetails.cpp b/Engine/Source/Editor/DetailCustomizations/Private/InputSettingsDetails.cpp
--- a/Engine/Source/Editor/DetailCustomizations/Private/InputSettingsDetails.cpp
+++ b/Engine/Source/Editor/DetailCustomizations/Private/InputSettingsDetails.cpp
@@ -141,13 +141,12 @@ void FActionMappingsNodeBuilder::AddActionMappingButton_OnClick()
 		InputSettings->Modify();
 		ActionMappingsPropertyHandle->NotifyPreChange();
 
-		EName BaseActionMappingIndex = (EName)BaseActionMappingName.GetIndex();
 		FName NewActionMappingName;
 		bool bFoundUniqueName;
 		do
 		{
 			// Create a numbered name and check whether it's already been used
-			NewActionMappingName = FName(BaseActionMappingIndex, ++NewMappingCount);
+			NewActionMappingName = FName(BaseActionMappingName, ++NewMappingCount);
 			bFoundUniqueName = true;
 			for (int32 Index = 0; Index < InputSettings->ActionMappings.Num(); ++Index)
 			{
@@ -444,13 +443,12 @@ void FAxisMappingsNodeBuilder::AddAxisMappingButton_OnClick()
 		InputSettings->Modify();
 		AxisMappingsPropertyHandle->NotifyPreChange();
 
-		EName BaseAxisMappingIndex = (EName)BaseAxisMappingName.GetIndex();
 		FName NewAxisMappingName;
 		bool bFoundUniqueName;
 		do
 		{
 			// Create a numbered name and check whether it's already been used
-			NewAxisMappingName = FName(BaseAxisMappingIndex, ++NewMappingCount);
+			NewAxisMappingName = FName(BaseAxisMappingName, ++NewMappingCount);
 			bFoundUniqueName = true;
 			for (int32 Index = 0; Index < InputSettings->AxisMappings.Num(); ++Index)
 			{
diff --git a/Engine/Source/Editor/Kismet/Private/SBlueprintPalette.cpp b/Engine/Source/Editor/Kismet/Private/SBlueprintPalette.cpp
--- a/Engine/Source/Editor/Kismet/Private/SBlueprintPalette.cpp
+++ b/Engine/Source/Editor/Kismet/Private/SBlueprintPalette.cpp
@@ -1126,6 +1126,9 @@ bool SBlueprintPaletteItem::OnNameTextVerifyChanged(const FText& InNewText, FTex
 //------------------------------------------------------------------------------
 void SBlueprintPaletteItem::OnNameTextCommitted(const FText& NewText, ETextCommit::Type InTextCommit)
 {
+	const FString NewNameString = NewText.ToString();
+	const FName NewName = *NewNameString;
+
 	if(ActionPtr.Pin()->GetTypeId() == FEdGraphSchemaAction_K2Graph::StaticGetTypeId())
 	{
 		FEdGraphSchemaAction_K2Graph* GraphAction = (FEdGraphSchemaAction_K2Graph*)ActionPtr.Pin().Get();
@@ -1145,12 +1148,11 @@ void SBlueprintPaletteItem::OnNameTextCommitted(const FText& NewText, ETextCommi
 				}
 			}
 
-			// Make sure we aren't renaming the graph into something
-			UEdGraph* ExistingGraph = FindObject<UEdGraph>(Graph->GetOuter(), *NewText.ToString() );
-			if (ExistingGraph == NULL)
+			// Make sure we aren't renaming the graph into something that already exists
+			UEdGraph* ExistingGraph = FindObject<UEdGraph>(Graph->GetOuter(), *NewNameString );
+			if (ExistingGraph == NULL || ExistingGraph == Graph)
 			{
 				const FScopedTransaction Transaction( LOCTEXT( "Rename Function", "Rename Function" ) );
-				FString NewNameString = NewText.ToString();
 
 				// Search through all function entry nodes for local variables to update their scope name
 				TArray<UK2Node_Variable*> VariableNodes;
@@ -1165,7 +1167,7 @@ void SBlueprintPaletteItem::OnNameTextCommitted(const FText& NewText, ETextCommi
 					}
 				}
 
-				FBlueprintEditorUtils::RenameGraph(Graph, *NewNameString );
+				FBlueprintEditorUtils::RenameGraph(Graph, NewNameString );
 			}
 		}
 	}
@@ -1188,13 +1190,11 @@ void SBlueprintPaletteItem::OnNameTextCommitted(const FText& NewText, ETextCommi
 				}
 			}
 
-			// Make sure we aren't renaming the graph into something
-			UEdGraph* ExistingGraph = FindObject<UEdGraph>(Graph->GetOuter(), *NewText.ToString() );
-			if (ExistingGraph == NULL)
+			// Make sure we aren't renaming the graph into something that already exists
+			UEdGraph* ExistingGraph = FindObject<UEdGraph>(Graph->GetOuter(), *NewNameString );
+			if (ExistingGraph == NULL || ExistingGraph == Graph)
 			{
 				const FScopedTransaction Transaction( LOCTEXT( "Rename Delegate", "Rename Event Dispatcher" ) );
-				const FString NewNameString = NewText.ToString();
-				const FName NewName = FName(*NewNameString);
 				const FName OldName =  Graph->GetFName();
 
 				UBlueprint* Blueprint = BlueprintEditorPtr.Pin()->GetBlueprintObj();
@@ -1212,7 +1212,7 @@ void SBlueprintPaletteItem::OnNameTextCommitted(const FText& NewText, ETextCommi
 					}
 				}
 
-				FBlueprintEditorUtils::RenameGraph(Graph, *NewNameString );
+				FBlueprintEditorUtils::RenameGraph(Graph, NewNameString );
 			}
 		}
 	}
@@ -1221,7 +1221,7 @@ void SBlueprintPaletteItem::OnNameTextCommitted(const FText& NewText, ETextCommi
 		FEdGraphSchemaAction_K2Var* VarAction = (FEdGraphSchemaAction_K2Var*)ActionPtr.Pin().Get();
 
 		// Check if the name is unchanged
-		if(NewText.ToString() == VarAction->GetVariableName().ToString())
+		if (NewName.IsEqual(VarAction->GetVariableName(), ENameCase::CaseSensitive))
 		{
 			return;
 		}
@@ -1230,8 +1230,6 @@ void SBlueprintPaletteItem::OnNameTextCommitted(const FText& NewText, ETextCommi
 
 		BlueprintEditorPtr.Pin()->GetBlueprintObj()->Modify();
 
-		FName NewVarName = FName(*NewText.ToString());
-
 		// Double check we're not renaming a timeline disguised as a variable
 		bool bIsTimeline = false;
 		UProperty* VariableProperty = VarAction->GetProperty();
@@ -1248,11 +1246,11 @@ void SBlueprintPaletteItem::OnNameTextCommitted(const FText& NewText, ETextCommi
 		// Rename as a timeline if required
 		if (bIsTimeline)
 		{
-			FBlueprintEditorUtils::RenameTimeline(BlueprintEditorPtr.Pin()->GetBlueprintObj(), VarAction->GetVariableName(), NewVarName);
+			FBlueprintEditorUtils::RenameTimeline(BlueprintEditorPtr.Pin()->GetBlueprintObj(), VarAction->GetVariableName(), NewName);
 		}
 		else
 		{
-			FBlueprintEditorUtils::RenameMemberVariable(BlueprintEditorPtr.Pin()->GetBlueprintObj(), VarAction->GetVariableName(), NewVarName);
+			FBlueprintEditorUtils::RenameMemberVariable(BlueprintEditorPtr.Pin()->GetBlueprintObj(), VarAction->GetVariableName(), NewName);
 		}
 	}
 	else if (ActionPtr.Pin()->GetTypeId() == FEdGraphSchemaAction_K2LocalVar::StaticGetTypeId())
@@ -1260,7 +1258,7 @@ void SBlueprintPaletteItem::OnNameTextCommitted(const FText& NewText, ETextCommi
 		FEdGraphSchemaAction_K2LocalVar* LocalVarAction = (FEdGraphSchemaAction_K2LocalVar*)ActionPtr.Pin().Get();
 
 		// Check if the name is unchanged
-		if(NewText.ToString() == LocalVarAction->GetVariableName().ToString())
+		if (NewName.IsEqual(LocalVarAction->GetVariableName(), ENameCase::CaseSensitive))
 		{
 			return;
 		}
@@ -1269,10 +1267,9 @@ void SBlueprintPaletteItem::OnNameTextCommitted(const FText& NewText, ETextCommi
 
 		BlueprintEditorPtr.Pin()->GetBlueprintObj()->Modify();
 
-		FName NewVarName = FName(*NewText.ToString());
-		FBlueprintEditorUtils::RenameLocalVariable(BlueprintEditorPtr.Pin()->GetBlueprintObj(), LocalVarAction->GetVariableScope(), LocalVarAction->GetVariableName(), NewVarName);
+		FBlueprintEditorUtils::RenameLocalVariable(BlueprintEditorPtr.Pin()->GetBlueprintObj(), LocalVarAction->GetVariableScope(), LocalVarAction->GetVariableName(), NewName);
 	}
-	BlueprintEditorPtr.Pin()->GetMyBlueprintWidget()->SelectItemByName(FName(*NewText.ToString()), ESelectInfo::OnMouseClick);
+	BlueprintEditorPtr.Pin()->GetMyBlueprintWidget()->SelectItemByName(NewName, ESelectInfo::OnMouseClick);
 }
 
 //------------------------------------------------------------------------------
diff --git a/Engine/Source/Editor/Kismet/Private/SMyBlueprint.cpp b/Engine/Source/Editor/Kismet/Private/SMyBlueprint.cpp
--- a/Engine/Source/Editor/Kismet/Private/SMyBlueprint.cpp
+++ b/Engine/Source/Editor/Kismet/Private/SMyBlueprint.cpp
@@ -643,7 +643,7 @@ void SMyBlueprint::GetLocalVariables(FGraphActionListBuilderBase& OutAllActions)
 				}
 
 				TSharedPtr<FEdGraphSchemaAction_K2LocalVar> NewVarAction = MakeShareable(new FEdGraphSchemaAction_K2LocalVar(Category, FText::FromName(Variable.VarName), TEXT(""), 0));
-				UFunction* Func = FindField<UFunction>(BlueprintEditorPtr.Pin()->GetBlueprintObj()->SkeletonGeneratedClass, *EdGraph->GetName());
+				UFunction* Func = FindField<UFunction>(BlueprintEditorPtr.Pin()->GetBlueprintObj()->SkeletonGeneratedClass, EdGraph->GetFName());
 				NewVarAction->SetVariableInfo(Variable.VarName, Func);
 				OutAllActions.AddAction(NewVarAction);
 			}
diff --git a/Engine/Source/Editor/KismetCompiler/Private/KismetCompilerVMBackend.cpp b/Engine/Source/Editor/KismetCompiler/Private/KismetCompilerVMBackend.cpp
--- a/Engine/Source/Editor/KismetCompiler/Private/KismetCompilerVMBackend.cpp
+++ b/Engine/Source/Editor/KismetCompiler/Private/KismetCompilerVMBackend.cpp
@@ -36,17 +36,20 @@ public:
 		FMemory::Memcpy( &(ScriptBuffer[iStart]), V, Length );
 	}
 
-	FArchive& operator<<(class FName& Name)
+	FArchive& operator<<(FName& Name)
 	{
-		// We can't call Serialize directly as we need to store the data endian clean.
-		NAME_INDEX W = Name.GetIndex(); 
-		int32 Num = Name.GetNumber(); 
 		FArchive& Ar = *this;
-		Ar << W << Num;
+
+		// We can't call Serialize directly as we need to store the data endian clean.
+		FScriptName ScriptName = NameToScriptName(Name);
+		Ar << ScriptName.ComparisonIndex;
+		Ar << ScriptName.DisplayIndex;
+		Ar << ScriptName.Number;
+
 		return Ar;
 	}
 
-	FArchive& operator<<(class UObject*& Res)
+	FArchive& operator<<(UObject*& Res)
 	{
 		ScriptPointerType D = (ScriptPointerType)Res; 
 		FArchive& Ar = *this;
@@ -55,12 +58,12 @@ public:
 		return Ar;
 	}
 
-	FArchive& operator<<(class FLazyObjectPtr& LazyObjectPtr)
+	FArchive& operator<<(FLazyObjectPtr& LazyObjectPtr)
 	{
 		return FArchive::operator<<(LazyObjectPtr);
 	}
 
-	FArchive& operator<<(class FAssetPtr& AssetPtr)
+	FArchive& operator<<(FAssetPtr& AssetPtr)
 	{
 		return FArchive::operator<<(AssetPtr);
 	}
@@ -76,7 +79,7 @@ public:
 		return *this;
 	}
 
-	FArchive& operator<<(enum EExprToken E)
+	FArchive& operator<<(EExprToken E)
 	{
 		checkSlow(E < 0xFF);
 
@@ -85,14 +88,14 @@ public:
 		return *this;
 	}
 
-	FArchive& operator<<(enum ECastToken E)
+	FArchive& operator<<(ECastToken E)
 	{
 		uint8 B = E; 
 		Serialize(&B, 1); 
 		return *this;
 	}
 
-	FArchive& operator<<(enum EPropertyType E)
+	FArchive& operator<<(EPropertyType E)
 	{
 		uint8 B = E; 
 		Serialize(&B, 1); 
diff --git a/Engine/Source/Editor/Persona/Private/Persona.cpp b/Engine/Source/Editor/Persona/Private/Persona.cpp
--- a/Engine/Source/Editor/Persona/Private/Persona.cpp
+++ b/Engine/Source/Editor/Persona/Private/Persona.cpp
@@ -1946,17 +1946,16 @@ FName FPersona::GenerateUniqueSocketName( FName InName )
 {
 	USkeletalMesh* Mesh = GetMesh();
 
-	EName OriginalNameIndex = (EName)InName.GetIndex();
 	int32 NewNumber = InName.GetNumber();
 
 	// Increment NewNumber until we have a unique name (potential infinite loop if *all* int32 values are used!)
-	while ( DoesSocketAlreadyExist( NULL, FText::FromName( FName( OriginalNameIndex, NewNumber ) ), TargetSkeleton->Sockets ) ||
-		( Mesh ? DoesSocketAlreadyExist( NULL, FText::FromName( FName( OriginalNameIndex, NewNumber ) ), Mesh->GetMeshOnlySocketList() ) : false ) )
+	while ( DoesSocketAlreadyExist( NULL, FText::FromName( FName( InName, NewNumber ) ), TargetSkeleton->Sockets ) ||
+		( Mesh ? DoesSocketAlreadyExist( NULL, FText::FromName( FName( InName, NewNumber ) ), Mesh->GetMeshOnlySocketList() ) : false ) )
 	{
 		++NewNumber;
 	}
 
-	return FName( OriginalNameIndex, NewNumber );
+	return FName( InName, NewNumber );
 }
 
 bool FPersona::DoesSocketAlreadyExist( const USkeletalMeshSocket* InSocket, const FText& InSocketName, const TArray< USkeletalMeshSocket* >& InSocketArray ) const
diff --git a/Engine/Source/Editor/PropertyEditor/Private/Presentation/PropertyTable/PropertyTableColumn.cpp b/Engine/Source/Editor/PropertyEditor/Private/Presentation/PropertyTable/PropertyTableColumn.cpp
--- a/Engine/Source/Editor/PropertyEditor/Private/Presentation/PropertyTable/PropertyTableColumn.cpp
+++ b/Engine/Source/Editor/PropertyEditor/Private/Presentation/PropertyTable/PropertyTableColumn.cpp
@@ -326,7 +326,7 @@ FPropertyTableColumn::FPropertyTableColumn( const TSharedRef< IPropertyTable >&
 	: Cells()
 	, DataSource( MakeShareable( new UObjectDataSource( InObject.Get() ) ) )
 	, Table( InTable )
-	, Id( FName( EName::NAME_None ) )
+	, Id( NAME_None )
 	, DisplayName()
 	, Width( 1.0f )
 	, bIsHidden( false )
@@ -342,7 +342,7 @@ FPropertyTableColumn::FPropertyTableColumn( const TSharedRef< IPropertyTable >&
 	: Cells()
 	, DataSource( MakeShareable( new PropertyPathDataSource( InPropertyPath ) ) )
 	, Table( InTable )
-	, Id( FName( EName::NAME_None ) )
+	, Id( NAME_None )
 	, DisplayName()
 	, Width( 1.0f )
 	, bIsHidden( false )
@@ -358,7 +358,7 @@ FPropertyTableColumn::FPropertyTableColumn( const TSharedRef< class IPropertyTab
 	: Cells()
 	, DataSource( MakeShareable( new UObjectDataSource( InObject.Get() ) ) )
 	, Table( InTable )
-	, Id( FName( EName::NAME_None ) )
+	, Id( NAME_None )
 	, DisplayName()
 	, Width( 1.0f )
 	, bIsHidden( false )
@@ -390,7 +390,7 @@ void FPropertyTableColumn::GenerateColumnId()
 	}
 	else
 	{
-		Id = FName( EName::NAME_None );
+		Id = NAME_None;
 	}
 }
 
diff --git a/Engine/Source/Editor/SceneOutliner/Private/SSceneOutliner.cpp b/Engine/Source/Editor/SceneOutliner/Private/SSceneOutliner.cpp
--- a/Engine/Source/Editor/SceneOutliner/Private/SSceneOutliner.cpp
+++ b/Engine/Source/Editor/SceneOutliner/Private/SSceneOutliner.cpp
@@ -1537,7 +1537,7 @@ namespace SceneOutliner
 			auto ActorTreeItem = StaticCastSharedPtr<TOutlinerActorTreeItem>(TreeItem);
 			auto* Actor = ActorTreeItem->Actor.Get();
 
-			if (Actor && InLabel.ToString() != Actor->GetActorLabel() && Actor->IsActorLabelEditable())
+			if (Actor && Actor->IsActorLabelEditable() && !InLabel.ToString().Equals(Actor->GetActorLabel(), ESearchCase::CaseSensitive))
 			{
 				const FScopedTransaction Transaction( LOCTEXT( "SceneOutlinerRenameActorTransaction", "Rename Actor" ) );
 				Actor->SetActorLabel( InLabel.ToString() );
@@ -1551,7 +1551,7 @@ namespace SceneOutliner
 		else
 		{
 			auto FolderTreeItem = StaticCastSharedPtr<TOutlinerFolderTreeItem>(TreeItem);
-			if (!InLabel.EqualTo(FText::FromName(FolderTreeItem->LeafName)))
+			if (!InLabel.ToString().Equals(FolderTreeItem->LeafName.ToString(), ESearchCase::CaseSensitive))
 			{
 				// Rename the item
 				FName NewPath = FolderTreeItem->GetParentPath();
@@ -1570,7 +1570,6 @@ namespace SceneOutliner
 				}
 			}
 		}
-		
 	}
 
 	bool SSceneOutliner::OnVerifyItemLabelChanged( const FText& InLabel, FText& OutErrorMessage, FOutlinerTreeItemPtr TreeItem )
diff --git a/Engine/Source/Editor/UnrealEd/Private/Commandlets/PackageUtilities.cpp b/Engine/Source/Editor/UnrealEd/Private/Commandlets/PackageUtilities.cpp
--- a/Engine/Source/Editor/UnrealEd/Private/Commandlets/PackageUtilities.cpp
+++ b/Engine/Source/Editor/UnrealEd/Private/Commandlets/PackageUtilities.cpp
@@ -886,7 +886,7 @@ void FPkgInfoReporter_Log::GeneratePackageReport( ULinkerLoad* InLinker/*=NULL*/
 		for( int32 i = 0; i < Linker->NameMap.Num(); ++i )
 		{
 			FName& name = Linker->NameMap[ i ];
-			UE_LOG(LogPackageUtilities, Warning, TEXT("\t%d: Name '%s' Index %d [Internal: %s, %d]"), i, *name.ToString(), name.GetIndex(), *name.GetPlainNameString(), name.GetNumber() );
+			UE_LOG(LogPackageUtilities, Warning, TEXT("\t%d: Name '%s' Comparison Index %d Display Index %d [Internal: %s, %d]"), i, *name.ToString(), name.GetComparisonIndex(), name.GetDisplayIndex(), *name.GetPlainNameString(), name.GetNumber() );
 		}
 	}
 
diff --git a/Engine/Source/Editor/UnrealEd/Private/Kismet2/BlueprintEditorUtils.cpp b/Engine/Source/Editor/UnrealEd/Private/Kismet2/BlueprintEditorUtils.cpp
--- a/Engine/Source/Editor/UnrealEd/Private/Kismet2/BlueprintEditorUtils.cpp
+++ b/Engine/Source/Editor/UnrealEd/Private/Kismet2/BlueprintEditorUtils.cpp
@@ -3343,9 +3343,9 @@ void FBlueprintEditorUtils::RemoveVariableNodes(UBlueprint* Blueprint, const FNa
 void FBlueprintEditorUtils::RenameComponentMemberVariable(UBlueprint* Blueprint, USCS_Node* Node, const FName& NewName)
 {
 	// Should not allow renaming to "none" (UI should prevent this)
-	check(NewName != NAME_None);
+	check(!NewName.IsNone());
 
-	if (Node->VariableName != NewName)
+	if (!NewName.IsEqual(Node->VariableName, ENameCase::CaseSensitive))
 	{
 		Blueprint->Modify();
 		
@@ -3372,7 +3372,7 @@ void FBlueprintEditorUtils::RenameComponentMemberVariable(UBlueprint* Blueprint,
 
 void FBlueprintEditorUtils::RenameMemberVariable(UBlueprint* Blueprint, const FName& OldName, const FName& NewName)
 {
-	if ((OldName != NewName) && (NewName != NAME_None))
+	if (!NewName.IsNone() && !NewName.IsEqual(OldName, ENameCase::CaseSensitive))
 	{
 		const int32 VarIndex = FBlueprintEditorUtils::FindNewVariableIndex(Blueprint, OldName);
 		if (VarIndex != INDEX_NONE)
@@ -3677,17 +3677,19 @@ UEdGraph* FBlueprintEditorUtils::FindScopeGraph(const UBlueprint* InBlueprint, c
 
 void FBlueprintEditorUtils::RenameLocalVariable(UBlueprint* InBlueprint, const UStruct* InScope, const FName& InOldName, const FName& InNewName)
 {
-	if ((InOldName != InNewName) && (InNewName != NAME_None))
+	if (!InNewName.IsNone() && !InNewName.IsEqual(InOldName, ENameCase::CaseSensitive))
 	{
 		UK2Node_FunctionEntry* FunctionEntry = nullptr;
 		FBPVariableDescription* LocalVariable = FindLocalVariable(InBlueprint, InScope, InOldName, &FunctionEntry);
+		const auto OldProperty = FindField<const UProperty>(InScope, InOldName);
 		const auto ExistingProperty = FindField<const UProperty>(InScope, InNewName);
-		if (ExistingProperty)
+		const bool bHasExistingProperty = ExistingProperty && ExistingProperty != OldProperty;
+		if (bHasExistingProperty)
 		{
 			UE_LOG(LogBlueprint, Warning, TEXT("Cannot name local variable '%s'. The name is already used."), *InNewName.ToString());
 		}
 
-		if (LocalVariable && !ExistingProperty)
+		if (LocalVariable && !bHasExistingProperty)
 		{
 			const FScopedTransaction Transaction( LOCTEXT("RenameVariable", "Rename Local Variable") );
 			InBlueprint->Modify();
@@ -3696,7 +3698,7 @@ void FBlueprintEditorUtils::RenameLocalVariable(UBlueprint* InBlueprint, const U
 			// Update the name
 			const UEdGraphSchema_K2* K2Schema = GetDefault<UEdGraphSchema_K2>();
 			LocalVariable->VarName = InNewName;
-			LocalVariable->FriendlyName = FName::NameToDisplayString( InNewName.ToString(), (LocalVariable->VarType.PinCategory == K2Schema->PC_Boolean) ? true : false );
+			LocalVariable->FriendlyName = FName::NameToDisplayString( InNewName.ToString(), LocalVariable->VarType.PinCategory == K2Schema->PC_Boolean );
 
 			// Update any existing references to the old name
 			FBlueprintEditorUtils::ReplaceVariableReferences(InBlueprint, InOldName, InNewName);
@@ -5094,10 +5096,10 @@ bool FBlueprintEditorUtils::RenameTimeline(UBlueprint* Blueprint, const FName& O
 	// NewName should be already validated. But one must make sure that NewTemplateName is also unique.
 	const bool bUniqueNameForTemplate = (EValidatorResult::Ok == NameValidator->IsValid(NewTemplateName));
 
-	UTimelineTemplate* Template = Blueprint->FindTimelineTemplateByVariableName(NewName);
-	if ((Template == NULL) && bUniqueNameForTemplate)
+	UTimelineTemplate* Template = Blueprint->FindTimelineTemplateByVariableName(OldName);
+	UTimelineTemplate* NewTemplate = Blueprint->FindTimelineTemplateByVariableName(NewName);
+	if ((!NewTemplate && bUniqueNameForTemplate) || NewTemplate == Template)
 	{
-		Template = Blueprint->FindTimelineTemplateByVariableName(OldName);
 		if (Template)
 		{
 			Blueprint->Modify();
diff --git a/Engine/Source/Editor/UnrealEd/Private/ScriptDisassembler.cpp b/Engine/Source/Editor/UnrealEd/Private/ScriptDisassembler.cpp
--- a/Engine/Source/Editor/UnrealEd/Private/ScriptDisassembler.cpp
+++ b/Engine/Source/Editor/UnrealEd/Private/ScriptDisassembler.cpp
@@ -116,10 +116,10 @@ uint8 FKismetBytecodeDisassembler::ReadBYTE(int32& ScriptIndex)
 
 FString FKismetBytecodeDisassembler::ReadName(int32& ScriptIndex)
 {
-	FName ConstValue = *(FName*)(Script.GetTypedData() + ScriptIndex);
-	ScriptIndex += sizeof(FName);
+	const FScriptName ConstValue = *(FScriptName*)(Script.GetTypedData() + ScriptIndex);
+	ScriptIndex += sizeof(FScriptName);
 
-	return ConstValue.ToString();
+	return ScriptNameToName(ConstValue).ToString();
 }
 
 uint16 FKismetBytecodeDisassembler::ReadWORD(int32& ScriptIndex)
diff --git a/Engine/Source/Runtime/AssetRegistry/Private/NameTableArchive.cpp b/Engine/Source/Runtime/AssetRegistry/Private/NameTableArchive.cpp
--- a/Engine/Source/Runtime/AssetRegistry/Private/NameTableArchive.cpp
+++ b/Engine/Source/Runtime/AssetRegistry/Private/NameTableArchive.cpp
@@ -88,7 +88,7 @@ int64 FNameTableArchiveReader::TotalSize()
 
 FArchive& FNameTableArchiveReader::operator<<( FName& Name )
 {
-	NAME_INDEX NameIndex;
+	int32 NameIndex;
 	FArchive& Ar = *this;
 	Ar << NameIndex;
 
@@ -98,7 +98,8 @@ FArchive& FNameTableArchiveReader::operator<<( FName& Name )
 	}
 
 	// if the name wasn't loaded (because it wasn't valid in this context)
-	if (NameMap[NameIndex] == NAME_None)
+	const FName& MappedName = NameMap[NameIndex];
+	if (MappedName.IsNone())
 	{
 		int32 TempNumber;
 		Ar << TempNumber;
@@ -108,8 +109,8 @@ FArchive& FNameTableArchiveReader::operator<<( FName& Name )
 	{
 		int32 Number;
 		Ar << Number;
-		// simply create the name from the NameMap's name index and the serialized instance number
-		Name = FName((EName)NameMap[NameIndex].GetIndex(), Number);
+		// simply create the name from the NameMap's name and the serialized instance number
+		Name = FName(MappedName, Number);
 	}
 
 	return *this;
@@ -163,7 +164,7 @@ void FNameTableArchiveWriter::SerializeNameMap()
 		for ( int32 NameMapIdx = 0; NameMapIdx < NameCount; ++NameMapIdx )
 		{
 			// Read the name entry
-			*this << *const_cast<FNameEntry*>(FName::GetEntry( NameMap[NameMapIdx].GetIndex() ));
+			*this << *const_cast<FNameEntry*>(NameMap[NameMapIdx].GetDisplayNameEntry());
 		}
 	}
 }
@@ -199,8 +200,11 @@ FArchive& FNameTableArchiveWriter::operator<<( FName& Name )
 	int32 NameIndex = NameIndexPtr ? *NameIndexPtr : INDEX_NONE;
 	if ( NameIndex == INDEX_NONE )
 	{
-		NameIndex = NameMap.Add(Name);
-		NameMapLookup.Add(Name, NameIndex);
+		// We need to store the FName without the number, as the number is stored separately and we don't 
+		// want duplicate entries in the name table just because of the number
+		const FName NameNoNumber(Name, 0);
+		NameIndex = NameMap.Add(NameNoNumber);
+		NameMapLookup.Add(NameNoNumber, NameIndex);
 	}
 
 	FArchive& Ar = *this;
@@ -218,4 +222,4 @@ FArchive& FNameTableArchiveWriter::operator<<( FName& Name )
 	}
 
 	return *this;
-}
\ No newline at end of file
+}
diff --git a/Engine/Source/Runtime/AssetRegistry/Private/NameTableArchive.h b/Engine/Source/Runtime/AssetRegistry/Private/NameTableArchive.h
--- a/Engine/Source/Runtime/AssetRegistry/Private/NameTableArchive.h
+++ b/Engine/Source/Runtime/AssetRegistry/Private/NameTableArchive.h
@@ -46,6 +46,6 @@ public:
 private:
 	FArrayWriter Writer;
 	TArray<FName> NameMap;
-	TMap<FName, int32> NameMapLookup;
+	TMap<FName, int32, FDefaultSetAllocator, TLinkerNameMapKeyFuncs<int32>> NameMapLookup;
 	int64 NameOffsetLoc;
 };
\ No newline at end of file
diff --git a/Engine/Source/Runtime/AssetRegistry/Private/PackageReader.cpp b/Engine/Source/Runtime/AssetRegistry/Private/PackageReader.cpp
--- a/Engine/Source/Runtime/AssetRegistry/Private/PackageReader.cpp
+++ b/Engine/Source/Runtime/AssetRegistry/Private/PackageReader.cpp
@@ -374,8 +374,8 @@ FArchive& FPackageReader::operator<<( FName& Name )
 	{
 		int32 Number;
 		Ar << Number;
-		// simply create the name from the NameMap's name index and the serialized instance number
-		Name = FName((EName)NameMap[NameIndex].GetIndex(), Number);
+		// simply create the name from the NameMap's name and the serialized instance number
+		Name = FName(NameMap[NameIndex], Number);
 	}
 
 	return *this;
diff --git a/Engine/Source/Runtime/Core/Private/Stats/Stats2.cpp b/Engine/Source/Runtime/Core/Private/Stats/Stats2.cpp
--- a/Engine/Source/Runtime/Core/Private/Stats/Stats2.cpp
+++ b/Engine/Source/Runtime/Core/Private/Stats/Stats2.cpp
@@ -133,7 +133,7 @@ DECLARE_DWORD_COUNTER_STAT(TEXT("Frame Packets Received"),STAT_StatFramePacketsR
 DECLARE_MEMORY_STAT( TEXT( "Stats descriptions (ansi+wide)" ), STAT_StatDescMemory, STATGROUP_StatSystem );
 
 
-FName TStatId::TStatId_NAME_None;
+TStatIdData TStatId::TStatId_NAME_None;
 
 void FStartupMessages::AddThreadMetadata( const FName InThreadName, uint32 InThreadID )
 {
@@ -171,8 +171,8 @@ void FThreadSafeStaticStatBase::DoSetup(const char* InStatName, const TCHAR* InS
 	// send meta data, we don't use normal messages because the stats thread might not be running yet
 	FStartupMessages::Get().AddMetadata(TempName, InStatDesc, InGroupName, InGroupCategory, InGroupDesc, bCanBeDisabled, InStatType, bCycleStat, InMemoryRegion);
 
-	FName const* LocalHighPerformanceEnable(IStatGroupEnableManager::Get().GetHighPerformanceEnableForStat(FName(InStatName), InGroupName, InGroupCategory, bDefaultEnable, bCanBeDisabled, InStatType, InStatDesc, bCycleStat, InMemoryRegion).GetRawPointer());
-	FName const* OldHighPerformanceEnable = (FName const*)FPlatformAtomics::InterlockedCompareExchangePointer((void**)&HighPerformanceEnable, (void*)LocalHighPerformanceEnable, NULL);
+	TStatIdData const* LocalHighPerformanceEnable(IStatGroupEnableManager::Get().GetHighPerformanceEnableForStat(FName(InStatName), InGroupName, InGroupCategory, bDefaultEnable, bCanBeDisabled, InStatType, InStatDesc, bCycleStat, InMemoryRegion).GetRawPointer());
+	TStatIdData const* OldHighPerformanceEnable = (TStatIdData const*)FPlatformAtomics::InterlockedCompareExchangePointer((void**)&HighPerformanceEnable, (void*)LocalHighPerformanceEnable, NULL);
 	check(!OldHighPerformanceEnable || HighPerformanceEnable == OldHighPerformanceEnable); // we are assigned two different groups?
 }
 
@@ -182,8 +182,8 @@ class FStatGroupEnableManager : public IStatGroupEnableManager
 {
 	struct FGroupEnable
 	{
-		TMap<FName, FName *> NamesInThisGroup;
-		TMap<FName, FName *> AlwaysEnabledNamesInThisGroup;
+		TMap<FName, TStatIdData *> NamesInThisGroup;
+		TMap<FName, TStatIdData *> AlwaysEnabledNamesInThisGroup;
 		bool DefaultEnable; 
 		bool CurrentEnable; 
 
@@ -200,10 +200,10 @@ class FStatGroupEnableManager : public IStatGroupEnableManager
 		 *	Increment between two successive long names. 
 		 *	@see TStatId
 		 */
-		FNAME_INCREMENT = 3,
+		STATID_INCREMENT = 3,
 
 		/** Number of stats pointer allocated per block. */
-		NUM_PER_BLOCK = 16384 * FNAME_INCREMENT,
+		NUM_PER_BLOCK = 16384 * STATID_INCREMENT,
 	};
 
 
@@ -213,7 +213,7 @@ class FStatGroupEnableManager : public IStatGroupEnableManager
 	FCriticalSection SynchronizationObject;
 	
 	/** Pointer to the long name in the names block. */
-	FName* PendingFNames;
+	TStatIdData* PendingStatIds;
 
 	/** Pending count of the name in the names block. */
 	int32 PendingCount;
@@ -226,32 +226,32 @@ class FStatGroupEnableManager : public IStatGroupEnableManager
 	bool EnableForNewGroups;
 	bool UseEnableForNewGroups;
 
-	void EnableStat(FName StatName, FName *DisablePtr)
+	void EnableStat(FName StatName, TStatIdData* DisablePtr)
 	{
 		// This is all complicated to ensure an atomic 8 byte write
-		static_assert(sizeof(FName) == sizeof(uint64), "FName should have the same size of uint64.");
-		check(UPTRINT(DisablePtr) % sizeof(FName) == 0);
+		static_assert(sizeof(FMinimalName) == sizeof(uint64), "FMinimalName should have the same size of uint64.");
+		check(UPTRINT(&DisablePtr->Name) % sizeof(FMinimalName) == 0);
 		MS_ALIGN(8) struct FAligner
 		{
-			FName Temp;
+			FMinimalName Temp;
 		} GCC_ALIGN(8);
 		FAligner Align;
-		check(UPTRINT(&Align.Temp) % sizeof(FName) == 0);
+		check(UPTRINT(&Align.Temp) % sizeof(FMinimalName) == 0);
 
-		Align.Temp = StatName;
-		*(uint64*)DisablePtr = *(uint64 const*)&Align.Temp;
+		Align.Temp = NameToMinimalName(StatName);
+		*(uint64*)&DisablePtr->Name = *(uint64 const*)&Align.Temp;
 	}
 
-	void DisableStat(FName *DisablePtr)
+	void DisableStat(TStatIdData* DisablePtr)
 	{
-		static_assert(sizeof(FName) == sizeof(uint64), "FName should have the same size of uint64.");
-		check(UPTRINT(DisablePtr) % sizeof(FName) == 0);
-		*(uint64*)DisablePtr = 0;
+		static_assert(sizeof(FMinimalName) == sizeof(uint64), "FMinimalName should have the same size of uint64.");
+		check(UPTRINT(&DisablePtr->Name) % sizeof(FMinimalName) == 0);
+		*(uint64*)&DisablePtr->Name = 0;
 	}
 
 public:
 	FStatGroupEnableManager()
-		: PendingFNames(NULL)
+		: PendingStatIds(NULL)
 		, PendingCount(0)
 		, EnableForNewGroups(false)
 		, UseEnableForNewGroups(false)
@@ -357,8 +357,8 @@ public:
 			{
 				UE_LOG(LogStatGroupEnableManager, Fatal, TEXT("Stat group %s was was defined both on and off by default."), *Group.ToString());
 			}
-			FName** StatFound = Found->NamesInThisGroup.Find(Stat);
-			FName** StatFoundAlways = Found->AlwaysEnabledNamesInThisGroup.Find(Stat);
+			TStatIdData** StatFound = Found->NamesInThisGroup.Find(Stat);
+			TStatIdData** StatFoundAlways = Found->AlwaysEnabledNamesInThisGroup.Find(Stat);
 			if (StatFound)
 			{
 				if (StatFoundAlways)
@@ -389,30 +389,29 @@ public:
 		}
 		if (PendingCount < 1)
 		{
-			PendingFNames = new FName[NUM_PER_BLOCK];
-			FMemory::Memzero( PendingFNames, NUM_PER_BLOCK *sizeof( FName ) );
+			PendingStatIds = new TStatIdData[NUM_PER_BLOCK];
+			FMemory::Memzero( PendingStatIds, NUM_PER_BLOCK * sizeof( TStatIdData ) );
 			PendingCount = NUM_PER_BLOCK;
 		}
-		PendingCount -= FNAME_INCREMENT;
-		FName* Result = &PendingFNames[TStatId::INDEX_FNAME];
+		PendingCount -= STATID_INCREMENT;
+		TStatIdData* Result = PendingStatIds;
 
 		// Get the wide stat description.
 		const int32 StatDescLen = FCString::Strlen( InDescription ) + 1;
 		// We are leaking this. @see STAT_StatDescMemory
 		WIDECHAR* StatDescWide = new WIDECHAR[StatDescLen];
 		TCString<WIDECHAR>::Strcpy( StatDescWide, StatDescLen, StringCast<WIDECHAR>( InDescription ).Get() );
+		Result->WideString = reinterpret_cast<uint64>(StatDescWide);
 
 		// Get the ansi stat description.
 		// We are leaking this. @see STAT_StatDescMemory
 		ANSICHAR* StatDescAnsi = new ANSICHAR[StatDescLen];
 		TCString<ANSICHAR>::Strcpy( StatDescAnsi, StatDescLen, StringCast<ANSICHAR>( InDescription ).Get() );
-
-		*(UPTRINT*)&PendingFNames[TStatId::INDEX_WIDE_STRING] = (UPTRINT)(UPTRINT*)StatDescWide;
-		*(UPTRINT*)&PendingFNames[TStatId::INDEX_ANSI_STRING] = (UPTRINT)(UPTRINT*)StatDescAnsi;
+		Result->AnsiString = reinterpret_cast<uint64>(StatDescAnsi);
 
 		MemoryCounter.Add( StatDescLen*(sizeof( ANSICHAR ) + sizeof( WIDECHAR )) );
 		
-		PendingFNames += FNAME_INCREMENT;
+		PendingStatIds += STATID_INCREMENT;
 
 		if (Found->CurrentEnable)
 		{
diff --git a/Engine/Source/Runtime/Core/Private/Stats/StatsFile.cpp b/Engine/Source/Runtime/Core/Private/Stats/StatsFile.cpp
--- a/Engine/Source/Runtime/Core/Private/Stats/StatsFile.cpp
+++ b/Engine/Source/Runtime/Core/Private/Stats/StatsFile.cpp
@@ -345,7 +345,7 @@ void FStatsWriteFile::Finalize()
 	for( const auto& It : Stats.ShortNameToLongName )
 	{
 		const FStatMessage& StatMessage = It.Value;
-		FNamesSent.Add( StatMessage.NameAndInfo.GetRawName().GetIndex() );
+		FNamesSent.Add( StatMessage.NameAndInfo.GetRawName().GetComparisonIndex() );
 	}
 
 	// Create a copy of names.
@@ -357,7 +357,7 @@ void FStatsWriteFile::Finalize()
 	Header.NumFNames = FNamesToSent.Num();
 	for( const int32 It : FNamesToSent )
 	{
-		WriteFName( Ar, FStatNameAndInfo(FName(EName(It)),false) );
+		WriteFName( Ar, FStatNameAndInfo(FName(It, It, 0),false) );
 	}
 
 	// Serialize metadata messages.
@@ -374,7 +374,7 @@ void FStatsWriteFile::Finalize()
 			TArray<FName> Result;
 			for( const int32 NameIndex : NameIndices )
 			{
-				new(Result) FName( EName( NameIndex ) );
+				new(Result) FName( NameIndex, NameIndex, 0 );
 			}
 			return Result;
 		}
diff --git a/Engine/Source/Runtime/Core/Private/UObject/UnrealNames.cpp b/Engine/Source/Runtime/Core/Private/UObject/UnrealNames.cpp
--- a/Engine/Source/Runtime/Core/Private/UObject/UnrealNames.cpp
+++ b/Engine/Source/Runtime/Core/Private/UObject/UnrealNames.cpp
@@ -25,7 +25,7 @@ const TCHAR* DebugFName(int32 Index)
 {
 	// Hardcoded static array. This function is only used inside the debugger so it should be fine to return it.
 	static TCHAR TempName[256];
-	FCString::Strcpy(TempName, *FName::SafeString((EName)Index));
+	FCString::Strcpy(TempName, *FName::SafeString(Index));
 	return TempName;
 }
 
@@ -40,7 +40,7 @@ const TCHAR* DebugFName(int32 Index, int32 Number)
 {
 	// Hardcoded static array. This function is only used inside the debugger so it should be fine to return it.
 	static TCHAR TempName[256];
-	FCString::Strcpy(TempName, *FName::SafeString((EName)Index, Number));
+	FCString::Strcpy(TempName, *FName::SafeString(Index, Number));
 	return TempName;
 }
 
@@ -54,7 +54,7 @@ const TCHAR* DebugFName(FName& Name)
 {
 	// Hardcoded static array. This function is only used inside the debugger so it should be fine to return it.
 	static TCHAR TempName[256];
-	FCString::Strcpy(TempName, *FName::SafeString((EName)Name.Index, Name.Number));
+	FCString::Strcpy(TempName, *FName::SafeString(Name.GetDisplayIndex(), Name.GetNumber()));
 	return TempName;
 }
 
@@ -96,21 +96,6 @@ void FNameEntry::AppendNameToString( FString& String ) const
 	}
 }
 
-/**
- * @return case insensitive hash of name
- */
-uint32 FNameEntry::GetNameHash() const
-{
-	if( IsWide() )
-	{
-		return FCrc::Strihash_DEPRECATED(WideName);
-	}
-	else
-	{
-		return FCrc::Strihash_DEPRECATED(AnsiName);
-	}
-}
-
 /**
  * @return length of name
  */
@@ -127,12 +112,12 @@ int32 FNameEntry::GetNameLength() const
 }
 
 /**
- * Compares name without looking at case.
+ * Compares name using the compare method provided.
  *
  * @param	InName	Name to compare to
  * @return	true if equal, false otherwise
  */
-bool FNameEntry::IsEqual( const ANSICHAR* InName ) const
+bool FNameEntry::IsEqual( const ANSICHAR* InName, const ENameCase CompareMethod ) const
 {
 	if( IsWide() )
 	{
@@ -141,17 +126,17 @@ bool FNameEntry::IsEqual( const ANSICHAR* InName ) const
 	}
 	else
 	{
-		return FCStringAnsi::Stricmp( AnsiName, InName ) == 0;
+		return ( (CompareMethod == ENameCase::CaseSensitive) ? FCStringAnsi::Strcmp( AnsiName, InName ) : FCStringAnsi::Stricmp( AnsiName, InName ) ) == 0;
 	}
 }
 
 /**
- * Compares name without looking at case.
+ * Compares name using the compare method provided.
  *
  * @param	InName	Name to compare to
  * @return	true if equal, false otherwise
  */
-bool FNameEntry::IsEqual( const WIDECHAR* InName ) const
+bool FNameEntry::IsEqual( const WIDECHAR* InName, const ENameCase CompareMethod ) const
 {
 	if( !IsWide() )
 	{
@@ -160,7 +145,7 @@ bool FNameEntry::IsEqual( const WIDECHAR* InName ) const
 	}
 	else
 	{
-		return FCStringWide::Stricmp( WideName, InName ) == 0;
+		return ( (CompareMethod == ENameCase::CaseSensitive) ? FCStringWide::Strcmp( WideName, InName ) : FCStringWide::Stricmp( WideName, InName ) ) == 0;
 	}
 }
 
@@ -415,7 +400,7 @@ FName::FName( ELinkerNameTableConstructor, const ANSICHAR* Name )
 	Init(Name, NAME_NO_NUMBER_INTERNAL, FNAME_Add, false);
 }
 
-FName::FName( enum EName HardcodedIndex, const TCHAR* Name )
+FName::FName( EName HardcodedIndex, const TCHAR* Name )
 {
 	check(HardcodedIndex >= 0);
 	Init(Name, NAME_NO_NUMBER_INTERNAL, FNAME_Add, false, HardcodedIndex);
@@ -431,10 +416,7 @@ bool FName::operator==( const TCHAR * Other ) const
 {
 	// Find name entry associated with this FName.
 	check( Other );
-	TNameEntryArray& Names = GetNames();
-	check( Index < Names.Num() );
-	FNameEntry const* Entry = Names[Index];
-	check( Entry );
+	const FNameEntry* const Entry = GetComparisonNameEntry();
 
 	// Temporary buffer to hold split name in case passed in name is of Name_Number format.
 	WIDECHAR TempBuffer[NAME_SIZE];
@@ -452,7 +434,7 @@ bool FName::operator==( const TCHAR * Other ) const
 
 	// Report a match if both the number and string portion match.
 	bool bAreNamesMatching = false;
-	if( InNumber == Number && !FCStringWide::Stricmp( WideOtherPtr, Entry->IsWide() ? Entry->GetWideName() : StringCast<WIDECHAR>(Entry->GetAnsiName()).Get() ) )
+	if( InNumber == GetNumber() && !FCStringWide::Stricmp( WideOtherPtr, Entry->IsWide() ? Entry->GetWideName() : StringCast<WIDECHAR>(Entry->GetAnsiName()).Get() ) )
 	{
 		bAreNamesMatching = true;
 	}
@@ -469,7 +451,7 @@ bool FName::operator==( const TCHAR * Other ) const
 int32 FName::Compare( const FName& Other ) const
 {
 	// Names match, check whether numbers match.
-	if( GetIndex() == Other.GetIndex() )
+	if( GetComparisonIndexFast() == Other.GetComparisonIndexFast() )
 	{
 		return GetNumber() - Other.GetNumber();
 	}
@@ -477,8 +459,8 @@ int32 FName::Compare( const FName& Other ) const
 	else
 	{
 		TNameEntryArray& Names = GetNames();
-		FNameEntry const* ThisEntry = Names[GetIndex()];
-		FNameEntry const* OtherEntry = Names[Other.GetIndex()];
+		const FNameEntry* const ThisEntry = GetComparisonNameEntry();
+		const FNameEntry* const OtherEntry = Other.GetComparisonNameEntry();
 
 		// Ansi/Wide mismatch, convert to wide
 		if( ThisEntry->IsWide() != OtherEntry->IsWide() )
@@ -528,7 +510,10 @@ void FName::Init(const WIDECHAR* InName, int32 InNumber, EFindName FindType, boo
 	if( !InName[0] )
 	{
 		check(HardcodeIndex < 1); // if this is hardcoded, it better be zero 
-		Index = NAME_None;
+		ComparisonIndex = NAME_None;
+#if WITH_CASE_PRESERVING_NAME
+		DisplayIndex = NAME_None;
+#endif
 		Number = NAME_NO_NUMBER_INTERNAL;
 		return;
 	}
@@ -536,36 +521,131 @@ void FName::Init(const WIDECHAR* InName, int32 InNumber, EFindName FindType, boo
 
 	//!!!! Caution, since these are set by static initializers from multiple threads, we must use local variables for this stuff until just before we return.
 
-	int32 OutIndex = HardcodeIndex;
+	bool bWasFoundOrAdded = true;
+	int32 OutComparisonIndex = HardcodeIndex;
+	int32 OutDisplayIndex = HardcodeIndex;
 
-	// set the number
-	int32 OutNumber = InNumber;
+	const bool bIsPureAnsi = FCStringWide::IsPureAnsi( InName );
+	if(bIsPureAnsi)
+	{
+		// Convert to an ansi display name
+		ANSICHAR AnsiDisplayName[NAME_SIZE];
+		FCStringAnsi::Strncpy(AnsiDisplayName, StringCast<ANSICHAR>(InName).Get(), ARRAY_COUNT(AnsiDisplayName));
 
-	// Hash value of string. Depends on whether the name is going to be ansi or wide.
-	int32 iHash;
+		bWasFoundOrAdded = InitInternal_FindOrAdd<ANSICHAR>(AnsiDisplayName, FindType, HardcodeIndex, OutComparisonIndex, OutDisplayIndex);
+	}
+	else
+	{
+		bWasFoundOrAdded = InitInternal_FindOrAdd<WIDECHAR>(InName, FindType, HardcodeIndex, OutComparisonIndex, OutDisplayIndex);
+	}
 
-	// Figure out whether we have a pure ansi name or not.
-	ANSICHAR AnsiName[NAME_SIZE];
-	bool bIsPureAnsi = FCStringWide::IsPureAnsi( InName );
-	if( bIsPureAnsi )
+	if(bWasFoundOrAdded)
+	{
+		ComparisonIndex = OutComparisonIndex;
+#if WITH_CASE_PRESERVING_NAME
+		DisplayIndex = OutDisplayIndex;
+#endif
+		Number = InNumber;
+	}
+	else
+	{
+		ComparisonIndex = NAME_None;
+#if WITH_CASE_PRESERVING_NAME
+		DisplayIndex = NAME_None;
+#endif
+		Number = NAME_NO_NUMBER_INTERNAL;
+	}
+}
+
+template <typename TCharType>
+struct FNameInitHelper
+{
+};
+
+template <>
+struct FNameInitHelper<ANSICHAR>
+{
+	static const bool IsAnsi = true;
+
+	static const ANSICHAR* GetNameString(const FNameEntry* const NameEntry)
+	{
+		return NameEntry->GetAnsiName();
+	}
+
+	static void SetNameString(FNameEntry* const DestNameEntry, const ANSICHAR* SrcName)
+	{
+		// Can't rely on the template override for static arrays since the safe crt version of strcpy will fill in
+		// the remainder of the array of NAME_SIZE with 0xfd.  So, we have to pass in the length of the dynamically allocated array instead.
+		FCStringAnsi::Strcpy(const_cast<ANSICHAR*>(DestNameEntry->GetAnsiName()), DestNameEntry->GetNameLength()+1, SrcName);
+	}
+};
+
+template <>
+struct FNameInitHelper<WIDECHAR>
+{
+	static const bool IsAnsi = false;
+
+	static const WIDECHAR* GetNameString(const FNameEntry* const NameEntry)
+	{
+		return NameEntry->GetWideName();
+	}
+
+	static void SetNameString(FNameEntry* const DestNameEntry, const WIDECHAR* SrcName)
 	{
-		FCStringAnsi::Strncpy( AnsiName, StringCast<ANSICHAR>(InName).Get(), ARRAY_COUNT(AnsiName) );
-		iHash = FCrc::Strihash_DEPRECATED( AnsiName ) & (ARRAY_COUNT(NameHash)-1);
+		// Can't rely on the template override for static arrays since the safe crt version of strcpy will fill in
+		// the remainder of the array of NAME_SIZE with 0xfd.  So, we have to pass in the length of the dynamically allocated array instead.
+		FCStringWide::Strcpy(const_cast<WIDECHAR*>(DestNameEntry->GetWideName()), DestNameEntry->GetNameLength()+1, SrcName);
+	}
+};
+
+template <typename TCharType>
+bool FName::InitInternal_FindOrAdd(const TCharType* InName, const EFindName FindType, const int32 HardcodeIndex, int32& OutComparisonIndex, int32& OutDisplayIndex)
+{
+	const bool bWasFoundOrAdded = InitInternal_FindOrAddNameEntry<TCharType>(InName, FindType, ENameCase::IgnoreCase, OutComparisonIndex);
+	
+#if WITH_CASE_PRESERVING_NAME
+	if(bWasFoundOrAdded && HardcodeIndex < 0)
+	{
+		TNameEntryArray& Names = GetNames();
+		const FNameEntry* const NameEntry = Names[OutComparisonIndex];
+
+		// If the string we got back doesn't match the case of the string we provided, also add a case variant version for display purposes
+		if(TCString<TCharType>::Strcmp(InName, FNameInitHelper<TCharType>::GetNameString(NameEntry)) != 0)
+		{
+			if(!InitInternal_FindOrAddNameEntry<TCharType>(InName, FindType, ENameCase::CaseSensitive, OutDisplayIndex))
+			{
+				// We don't consider failing to find/add the case variant a full failure
+				OutDisplayIndex = OutComparisonIndex;
+			}
+		}
+		else
+		{
+			OutDisplayIndex = OutComparisonIndex;
+		}
 	}
 	else
+#endif
 	{
-		iHash = FCrc::Strihash_DEPRECATED( InName ) & (ARRAY_COUNT(NameHash)-1);
+		OutDisplayIndex = OutComparisonIndex;
 	}
 
+	return bWasFoundOrAdded;
+}
+
+template <typename TCharType>
+bool FName::InitInternal_FindOrAddNameEntry(const TCharType* InName, const EFindName FindType, const ENameCase ComparisonMode, int32& OutIndex)
+{
+	// Hash value of string
+	const int32 iHash = ( (ComparisonMode == ENameCase::IgnoreCase) ? FCrc::Strihash_DEPRECATED( InName ) : FCrc::StrCrc32( InName ) ) & (ARRAY_COUNT(NameHash)-1);
+
 	if (OutIndex < 0)
 	{
 		// Try to find the name in the hash.
 		for( FNameEntry* Hash=NameHash[iHash]; Hash; Hash=Hash->HashNext )
 		{
 			FPlatformMisc::Prefetch( Hash->HashNext );
-			// Compare the passed in string, either ANSI or TCHAR.
-			if( ( bIsPureAnsi && Hash->IsEqual( AnsiName )) 
-			||  (!bIsPureAnsi && Hash->IsEqual( InName )) )
+			// Compare the passed in string
+			if( Hash->IsEqual( InName, ComparisonMode ) )
 			{
 				// Found it in the hash.
 				OutIndex = Hash->GetIndex();
@@ -577,25 +657,15 @@ void FName::Init(const WIDECHAR* InName, int32 InNumber, EFindName FindType, boo
 				if (FindType == FNAME_Replace_Not_Safe_For_Threading)
 				{
 					check(IsInGameThread());
-					// This should be impossible due to the compare above
+
 					// This *must* be true, or we'll overwrite memory when the
 					// copy happens if it is longer
-					check(FCString::Strlen(InName) == Hash->GetNameLength());
-					// Can't rely on the template override for static arrays since the safe crt version of strcpy will fill in
-					// the remainder of the array of NAME_SIZE with 0xfd.  So, we have to pass in the length of the dynamically allocated array instead.
-					if( bIsPureAnsi )
-					{
-						FCStringAnsi::Strcpy(const_cast<ANSICHAR*>(Hash->GetAnsiName()),Hash->GetNameLength()+1,AnsiName);
-					}
-					else
-					{
-						FCStringWide::Strcpy(const_cast<WIDECHAR*>(Hash->GetWideName()),Hash->GetNameLength()+1,InName);
-					}
+					check(TCString<TCharType>::Strlen(InName) == Hash->GetNameLength());
+
+					FNameInitHelper<TCharType>::SetNameString(Hash, InName);
 				}
 				check(OutIndex >= 0);
-				Index = OutIndex;
-				Number = OutNumber;
-				return;
+				return true;
 			}
 		}
 
@@ -603,9 +673,7 @@ void FName::Init(const WIDECHAR* InName, int32 InNumber, EFindName FindType, boo
 		if( FindType==FNAME_Find )
 		{
 			// Not found.
-			Index = NAME_None;
-			Number = NAME_NO_NUMBER_INTERNAL;
-			return;
+			return false;
 		}
 	}
 	// acquire the lock
@@ -615,16 +683,13 @@ void FName::Init(const WIDECHAR* InName, int32 InNumber, EFindName FindType, boo
 		// Try to find the name in the hash. AGAIN...we might have been adding from a different thread and we just missed it
 		for( FNameEntry* Hash=NameHash[iHash]; Hash; Hash=Hash->HashNext )
 		{
-			// Compare the passed in string, either ANSI or TCHAR.
-			if( ( bIsPureAnsi && Hash->IsEqual( AnsiName )) 
-				||  (!bIsPureAnsi && Hash->IsEqual( InName )) )
+			// Compare the passed in string
+			if( Hash->IsEqual( InName, ComparisonMode ) )
 			{
 				// Found it in the hash.
 				OutIndex = Hash->GetIndex();
 				check(FindType == FNAME_Add);  // if this was a replace, well it isn't safe for threading. Find should have already been handled
-				Index = OutIndex;
-				Number = OutNumber;
-				return;
+				return true;
 			}
 		}
 	}
@@ -636,9 +701,9 @@ void FName::Init(const WIDECHAR* InName, int32 InNumber, EFindName FindType, boo
 	}
 	else
 	{
-		Names.AddZeroed(OutIndex + 1 - Names.Num());
+		check(OutIndex < Names.Num());
 	}
-	FNameEntry* NewEntry = AllocateNameEntry( bIsPureAnsi ? (void const*)AnsiName : (void const*)InName, OutIndex, OldHash, bIsPureAnsi );
+	FNameEntry* NewEntry = AllocateNameEntry( InName, OutIndex, OldHash, FNameInitHelper<TCharType>::IsAnsi );
 	if (FPlatformAtomics::InterlockedCompareExchangePointer((void**)&Names[OutIndex], NewEntry, NULL) != NULL) // we use an atomic operation to check for unexpected concurrency, verify alignment, etc
 	{
 		UE_LOG(LogUnrealNames, Fatal, TEXT("Hardcoded name '%s' at index %i was duplicated (or unexpected concurrency). Existing entry is '%s'."), *NewEntry->GetPlainNameString(), NewEntry->GetIndex(), *Names[OutIndex]->GetPlainNameString() );
@@ -648,10 +713,22 @@ void FName::Init(const WIDECHAR* InName, int32 InNumber, EFindName FindType, boo
 		check(0); // someone changed this while we were changing it
 	}
 	check(OutIndex >= 0);
-	Index = OutIndex;
-	Number = OutNumber;
+	return true;
+}
+
+const FNameEntry* FName::GetComparisonNameEntry() const
+{
+	TNameEntryArray& Names = GetNames();
+	const NAME_INDEX Index = GetComparisonIndex();
+	return Names[Index];
 }
 
+const FNameEntry* FName::GetDisplayNameEntry() const
+{
+	TNameEntryArray& Names = GetNames();
+	const NAME_INDEX Index = GetDisplayIndex();
+	return Names[Index];
+}
 
 FString FName::ToString() const
 {
@@ -662,26 +739,20 @@ FString FName::ToString() const
 
 void FName::ToString(FString& Out) const
 {
-	TNameEntryArray& Names = GetNames();
 	// a version of ToString that saves at least one string copy
-	checkName(Index < Names.Num());
-	checkName(Names[Index]);
-	FNameEntry const* NameEntry = Names[Index];
+	const FNameEntry* const NameEntry = GetDisplayNameEntry();
 	Out.Empty( NameEntry->GetNameLength() + 6);
 	AppendString(Out);
 }
 
 void FName::AppendString(FString& Out) const
 {
-	TNameEntryArray& Names = GetNames();
-	checkName(Index < Names.Num());
-	checkName(Names[Index]);
-	FNameEntry const* NameEntry = Names[Index];
+	const FNameEntry* const NameEntry = GetDisplayNameEntry();
 	NameEntry->AppendNameToString( Out );
-	if (Number != NAME_NO_NUMBER_INTERNAL)
+	if (GetNumber() != NAME_NO_NUMBER_INTERNAL)
 	{
 		Out += TEXT("_");
-		Out.AppendInt(NAME_INTERNAL_TO_EXTERNAL(Number));
+		Out.AppendInt(NAME_INTERNAL_TO_EXTERNAL(GetNumber()));
 	}
 }
 
@@ -709,6 +780,13 @@ void FName::StaticInit()
 		NameHash[HashIndex] = NULL;
 	}
 
+	{
+		FScopeLock ScopeLock(GetCriticalSection());
+
+		TNameEntryArray& Names = GetNames();
+		Names.AddZeroed(NAME_MaxHardcodedNameIndex + 1);
+	}
+
 	{
 		// Register all hardcoded names.
 		#define REGISTER_NAME(num,namestr) FName Temp_##namestr(EName(num), TEXT(#namestr));
@@ -865,17 +943,23 @@ bool FName::IsValidXName( FString InvalidChars/*=INVALID_NAME_CHARACTERS*/, FTex
 
 void FName::AutoTest()
 {
-	FName AutoTest_1("AutoTest_1");
-	FName autoTest_1("autoTest_1");
-	FName AutoTest_2(TEXT("AutoTest_2"));
-	FName AutoTestB_2(TEXT("AutoTestB_2"));
+	const FName AutoTest_1("AutoTest_1");
+	const FName autoTest_1("autoTest_1");
+	const FName autoTeSt_1("autoTeSt_1");
+	const FName AutoTest1Find("autoTEST_1", EFindName::FNAME_Find);
+	const FName AutoTest_2(TEXT("AutoTest_2"));
+	const FName AutoTestB_2(TEXT("AutoTestB_2"));
 
 	check(AutoTest_1 != AutoTest_2);
 	check(AutoTest_1 == autoTest_1);
+	check(AutoTest_1 == autoTeSt_1);
+#if WITH_CASE_PRESERVING_NAME
 	check(!FCString::Strcmp(*AutoTest_1.ToString(), TEXT("AutoTest_1")));
-	check(!FCString::Strcmp(*autoTest_1.ToString(), TEXT("AutoTest_1")));
+	check(!FCString::Strcmp(*autoTest_1.ToString(), TEXT("autoTest_1")));
+	check(!FCString::Strcmp(*autoTeSt_1.ToString(), TEXT("autoTeSt_1")));
 	check(!FCString::Strcmp(*AutoTestB_2.ToString(), TEXT("AutoTestB_2")));
-	check(autoTest_1.GetIndex() == AutoTest_2.GetIndex());
+#endif
+	check(autoTest_1.GetComparisonIndex() == AutoTest_2.GetComparisonIndex());
 	check(autoTest_1.GetPlainNameString() == AutoTest_1.GetPlainNameString());
 	check(autoTest_1.GetPlainNameString() == AutoTest_2.GetPlainNameString());
 	check(*AutoTestB_2.GetPlainNameString() != *AutoTest_2.GetPlainNameString());
@@ -1112,14 +1196,14 @@ static class FFNameExec: private FSelfRegisteringExec
 							check(Table[Index].NameNumber == Temp.GetNumber());
 							if (Table[Index].NameIndex == 0)
 							{
-								Table[Index].NameIndex = Temp.GetIndex();
+								Table[Index].NameIndex = Temp.GetComparisonIndex();
 							}
-							check(Table[Index].NameIndex == Temp.GetIndex());
+							check(Table[Index].NameIndex == Temp.GetComparisonIndex());
 
 							FName Temp2(*Table[Index].TheString);
 							check(Temp2.ToString() == Table[Index].TheString);
 							check(Table[Index].NameNumber == Temp2.GetNumber());
-							check(Table[Index].NameIndex == Temp2.GetIndex());
+							check(Table[Index].NameIndex == Temp2.GetComparisonIndex());
 							TestCounter.Increment();
 						}
 
diff --git a/Engine/Source/Runtime/Core/Public/Async/TaskGraphInterfaces.h b/Engine/Source/Runtime/Core/Public/Async/TaskGraphInterfaces.h
--- a/Engine/Source/Runtime/Core/Public/Async/TaskGraphInterfaces.h
+++ b/Engine/Source/Runtime/Core/Public/Async/TaskGraphInterfaces.h
@@ -670,7 +670,7 @@ private:
 				{
 					LogPossiblyInvalidSubsequentsTask(
 #if STATS
-						*Task.GetStatId()->ToString()
+						*Task.GetStatId().GetName().ToString()
 #else
 						TEXT("Unknown")
 #endif
diff --git a/Engine/Source/Runtime/Core/Public/Stats/Stats2.h b/Engine/Source/Runtime/Core/Public/Stats/Stats2.h
--- a/Engine/Source/Runtime/Core/Public/Stats/Stats2.h
+++ b/Engine/Source/Runtime/Core/Public/Stats/Stats2.h
@@ -60,46 +60,56 @@ struct CORE_API FStats
 
 #if STATS
 
-struct TStatId
+struct TStatIdData
 {
-	enum
+	FORCEINLINE TStatIdData()
+		: Name(NameToMinimalName(NAME_None))
+		, AnsiString(0)
+		, WideString(0)
 	{
-		/**
-		 *	Index of the long name.
-		 *	@see FStatGroupEnableManager
-		 */
-		INDEX_FNAME = 0,
+	}
 
-		/** Index of the stat desc as an ansi string. */
-		INDEX_ANSI_STRING = 1,
+	FORCEINLINE bool IsNone() const
+	{
+		return MinimalNameToName(Name).IsNone();
+	}
 
-		/** Index of the stat desc as a wide string. */
-		INDEX_WIDE_STRING = 2,
-	};
+	/** Name of the active stat; stored as a minimal name to minimize the data size */
+	FMinimalName Name;
 
+	/** const ANSICHAR* pointer to a string; stored as a uint64 so it doesn't change size and affect TStatIdData alignment between 32 and 64-bit builds) */
+	uint64 AnsiString;
+
+	/** const WIDECHAR* pointer to a string; stored as a uint64 so it doesn't change size and affect TStatIdData alignment between 32 and 64-bit builds) */
+	uint64 WideString;
+};
+
+struct TStatId
+{
 	FORCEINLINE TStatId()
 		: StatIdPtr(&TStatId_NAME_None)
 	{
 	}
-	FORCEINLINE TStatId(FName const* InStatIdPtr)
+	FORCEINLINE TStatId(TStatIdData const* InStatIdPtr)
 		: StatIdPtr(InStatIdPtr)
 	{
 	}
-	FORCEINLINE FName operator*() const
+	FORCEINLINE bool IsValidStat() const
 	{
-		return *StatIdPtr;
+		return StatIdPtr != &TStatId_NAME_None;
 	}
-	FORCEINLINE FName const* operator->() const
+	FORCEINLINE bool IsNone() const
 	{
-		return StatIdPtr;
+		return StatIdPtr->IsNone();
 	}
-	FORCEINLINE bool IsValidStat() const
+	FORCEINLINE TStatIdData const* GetRawPointer() const
 	{
-		return StatIdPtr != &TStatId_NAME_None;
+		return StatIdPtr;
 	}
-	FORCEINLINE FName const* GetRawPointer() const
+
+	FORCEINLINE FName GetName() const
 	{
-		return &StatIdPtr[INDEX_FNAME];
+		return MinimalNameToName(StatIdPtr->Name);
 	}
 
 	/**
@@ -109,7 +119,7 @@ struct TStatId
 	 */
 	FORCEINLINE const ANSICHAR* GetStatDescriptionANSI() const
 	{
-		return (ANSICHAR*)*(uint64*)(&StatIdPtr[INDEX_ANSI_STRING]);
+		return reinterpret_cast<const ANSICHAR*>(StatIdPtr->AnsiString);
 	}
 
 	/**
@@ -119,12 +129,12 @@ struct TStatId
 	 */
 	FORCEINLINE const WIDECHAR* GetStatDescriptionWIDE() const
 	{
-		return (WIDECHAR*)*(uint64*)(&StatIdPtr[INDEX_WIDE_STRING]);
+		return reinterpret_cast<const WIDECHAR*>(StatIdPtr->WideString);
 	}
 
 private:
 	/** NAME_None. */
-	CORE_API static FName TStatId_NAME_None;
+	CORE_API static TStatIdData TStatId_NAME_None;
 
 	/**
 	 *	Holds a pointer to the stat long name if enabled, or to the NAME_None if disabled.
@@ -135,7 +145,7 @@ private:
 	 *	Next pointer points to the wide string with a stat description
 	 *	@see FStatGroupEnableManager::GetHighPerformanceEnableForStat 
 	 */
-	FName const* StatIdPtr;
+	TStatIdData const* StatIdPtr;
 };
 
 /**
@@ -285,7 +295,7 @@ class FStatNameAndInfo
 	/**
 	 * An FName, but the high bits of the Number are used for other fields.
 	 */
-	FName NameAndInfo;
+	FMinimalName NameAndInfo;
 public:
 	FORCEINLINE_STATS FStatNameAndInfo()
 	{
@@ -305,15 +315,15 @@ public:
 	 * Build from a raw FName
 	 */
 	FORCEINLINE_STATS FStatNameAndInfo(FName Other, bool bAlreadyHasMeta)
-		: NameAndInfo(Other)
+		: NameAndInfo(NameToMinimalName(Other))
 	{
 		if (!bAlreadyHasMeta)
 		{
-			int32 Number = NameAndInfo.GetNumber();
+			int32 Number = NameAndInfo.Number;
 			// ok, you can't have numbered stat FNames too large
 			checkStats(!(Number >> EStatAllFields::StartShift));
 			Number |= EStatMetaFlags::DummyAlwaysOne << (EStatMetaFlags::Shift + EStatAllFields::StartShift);
-			NameAndInfo.SetNumber(Number);
+			NameAndInfo.Number = Number;
 		}
 		CheckInvariants();
 	}
@@ -322,13 +332,13 @@ public:
 	 * Build with stat metadata
 	 */
 	FORCEINLINE_STATS FStatNameAndInfo(FName InStatName, char const* InGroup, char const* InCategory, TCHAR const* InDescription, EStatDataType::Type InStatType, bool bShouldClearEveryFrame, bool bCycleStat, FPlatformMemory::EMemoryCounterRegion MemoryRegion = FPlatformMemory::MCR_Invalid)
-		: NameAndInfo(ToLongName(InStatName, InGroup, InCategory, InDescription))
+		: NameAndInfo(NameToMinimalName(ToLongName(InStatName, InGroup, InCategory, InDescription)))
 	{
-		int32 Number = NameAndInfo.GetNumber();
+		int32 Number = NameAndInfo.Number;
 		// ok, you can't have numbered stat FNames too large
 		checkStats(!(Number >> EStatAllFields::StartShift));
 		Number |= (EStatMetaFlags::DummyAlwaysOne | EStatMetaFlags::HasLongNameAndMetaInfo) << (EStatMetaFlags::Shift + EStatAllFields::StartShift);
-		NameAndInfo.SetNumber(Number);
+		NameAndInfo.Number = Number;
 
 		SetField<EStatDataType>(InStatType);
 		SetFlag(EStatMetaFlags::ShouldClearEveryFrame, bShouldClearEveryFrame);
@@ -347,7 +357,7 @@ public:
 	 */
 	FORCEINLINE_STATS void SetNumberDirect(int32 Number)
 	{
-		NameAndInfo.SetNumber(Number);
+		NameAndInfo.Number = Number;
 	}
 
 	/**
@@ -356,8 +366,7 @@ public:
 	FORCEINLINE_STATS int32 GetRawNumber() const
 	{
 		CheckInvariants();
-		int32 Number = NameAndInfo.GetNumber();
-		return Number;
+		return NameAndInfo.Number;
 	}
 
 	/**
@@ -368,10 +377,10 @@ public:
 		// ok, you can't have numbered stat FNames too large
 		checkStats(!(RawName.GetNumber() >> EStatAllFields::StartShift));
 		CheckInvariants();
-		int32 Number = NameAndInfo.GetNumber();
+		int32 Number = NameAndInfo.Number;
 		Number &= ~((1 << EStatAllFields::StartShift) - 1);
-		NameAndInfo = RawName;
-		NameAndInfo.SetNumber(Number | RawName.GetNumber());
+		NameAndInfo = NameToMinimalName(RawName);
+		NameAndInfo.Number = (Number | RawName.GetNumber());
 	}
 
 	/**
@@ -380,11 +389,11 @@ public:
 	FORCEINLINE_STATS FName GetRawName() const
 	{
 		CheckInvariants();
-		FName Result(NameAndInfo);
-		int32 Number = NameAndInfo.GetNumber();
+		FMinimalName Result(NameAndInfo);
+		int32 Number = NameAndInfo.Number;
 		Number &= ((1 << EStatAllFields::StartShift) - 1);
-		Result.SetNumber(Number);
-		return Result;
+		Result.Number = Number;
+		return MinimalNameToName(Result);
 	}
 
 	/**
@@ -393,7 +402,7 @@ public:
 	FORCEINLINE_STATS FName GetEncodedName() const
 	{
 		CheckInvariants();
-		return NameAndInfo;
+		return MinimalNameToName(NameAndInfo);
 	}
 
 	/**
@@ -437,8 +446,8 @@ public:
 	 */
 	FORCEINLINE_STATS void CheckInvariants() const
 	{
-		checkStats((NameAndInfo.GetNumber() & (EStatMetaFlags::DummyAlwaysOne << (EStatAllFields::StartShift + EStatMetaFlags::Shift)))
-			&& NameAndInfo.GetIndex());
+		checkStats((NameAndInfo.Number & (EStatMetaFlags::DummyAlwaysOne << (EStatAllFields::StartShift + EStatMetaFlags::Shift)))
+			&& NameAndInfo.Index);
 	}
 
 	/**
@@ -449,7 +458,7 @@ public:
 	typename TField::Type GetField() const
 	{
 		CheckInvariants();
-		int32 Number = NameAndInfo.GetNumber();
+		int32 Number = NameAndInfo.Number;
 		Number = (Number >> (EStatAllFields::StartShift + TField::Shift)) & TField::Mask;
 		checkStats(Number != TField::Invalid && Number < TField::Num);
 		return typename TField::Type(Number);
@@ -462,12 +471,12 @@ public:
 	template<typename TField>
 	void SetField(typename TField::Type Value)
 	{
-		int32 Number = NameAndInfo.GetNumber();
+		int32 Number = NameAndInfo.Number;
 		CheckInvariants();
 		checkStats(Value < TField::Num && Value != TField::Invalid);
 		Number &= ~(TField::Mask << (EStatAllFields::StartShift + TField::Shift));
 		Number |= Value << (EStatAllFields::StartShift + TField::Shift);
-		NameAndInfo.SetNumber(Number);
+		NameAndInfo.Number = Number;
 		CheckInvariants();
 	}
 
@@ -477,7 +486,7 @@ public:
 	 */
 	bool GetFlag(EStatMetaFlags::Type Bit) const
 	{
-		int32 Number = NameAndInfo.GetNumber();
+		int32 Number = NameAndInfo.Number;
 		CheckInvariants();
 		checkStats(Bit < EStatMetaFlags::Num && Bit != EStatMetaFlags::Invalid);
 		return !!((Number >> (EStatAllFields::StartShift + EStatMetaFlags::Shift)) & Bit);
@@ -490,7 +499,7 @@ public:
 	 */
 	void SetFlag(EStatMetaFlags::Type Bit, bool Value)
 	{
-		int32 Number = NameAndInfo.GetNumber();
+		int32 Number = NameAndInfo.Number;
 		CheckInvariants();
 		checkStats(Bit < EStatMetaFlags::Num && Bit != EStatMetaFlags::Invalid);
 		if (Value)
@@ -501,7 +510,7 @@ public:
 		{
 			Number &= ~(Bit << (EStatAllFields::StartShift + EStatMetaFlags::Shift));
 		}
-		NameAndInfo.SetNumber(Number);
+		NameAndInfo.Number = Number;
 		CheckInvariants();
 	}
 
@@ -697,17 +706,23 @@ struct FStatMessage
 		return *(double const*)&StatData;
 	}
 
-	FORCEINLINE_STATS FName& GetValue_FName()
+	FORCEINLINE_STATS FMinimalName& GetValue_FMinimalName()
+	{
+		static_assert(sizeof(FMinimalName) <= DATA_SIZE && ALIGNOF(FMinimalName) <= DATA_ALIGN, "Bad data for stat message.");
+		checkStats(NameAndInfo.GetField<EStatDataType>() == EStatDataType::ST_FName);
+		return *(FMinimalName*)&StatData;
+	}
+
+	FORCEINLINE_STATS FMinimalName GetValue_FMinimalName() const
 	{
-		static_assert(sizeof(FName) <= DATA_SIZE && ALIGNOF(FName) <= DATA_ALIGN, "Bad data for stat message.");
 		checkStats(NameAndInfo.GetField<EStatDataType>() == EStatDataType::ST_FName);
-		return *(FName*)&StatData;
+		return *(FMinimalName const*)&StatData;
 	}
 
 	FORCEINLINE_STATS FName GetValue_FName() const
 	{
 		checkStats(NameAndInfo.GetField<EStatDataType>() == EStatDataType::ST_FName);
-		return *(FName const*)&StatData;
+		return MinimalNameToName(*(FMinimalName const*)&StatData);
 	}
 };
 template<> struct TIsPODType<FStatMessage> { enum { Value = true }; };
@@ -1102,7 +1117,7 @@ public:
 	static FORCEINLINE_STATS bool IsCollectingData(TStatId StatId)
 	{
 		// we don't test StatId for NULL here because we assume it is non-null. If it is NULL, that indicates a problem with higher level code.
-		return !StatId->IsNone() && IsCollectingData();
+		return !StatId.IsNone() && IsCollectingData();
 	}
 
 	/** Return true if we are currently collecting data **/
@@ -1202,8 +1217,8 @@ public:
 	{
 		if( (bAlways && InStatId.IsValidStat()) || FThreadStats::IsCollectingData( InStatId ) )
 		{
-			StatId = *InStatId;
-			FThreadStats::AddMessage( *InStatId, EStatOperation::CycleScopeStart );
+			StatId = InStatId.GetName();
+			FThreadStats::AddMessage( StatId, EStatOperation::CycleScopeStart );
 
 			// Emit named event for active cycle stat.
 			if( GCycleStatsShouldEmitNamedEvents > 0 )
@@ -1314,7 +1329,7 @@ public:
 struct FThreadSafeStaticStatBase
 {
 protected:
-	mutable FName* HighPerformanceEnable; // must be uninitialized, because we need atomic initialization
+	mutable TStatIdData* HighPerformanceEnable; // must be uninitialized, because we need atomic initialization
 	CORE_API void DoSetup(const char* InStatName, const TCHAR* InStatDesc, const char* InGroupName, const char* InGroupCategory, const TCHAR* InGroupDesc, bool bDefaultEnable, bool bCanBeDisabled, EStatDataType::Type InStatType, bool bCycleStat, FPlatformMemory::EMemoryCounterRegion InMemoryRegion) const;
 };
 
@@ -1332,7 +1347,7 @@ struct FThreadSafeStaticStatInner : public FThreadSafeStaticStatBase
 	}
 	FORCEINLINE FName GetStatFName() const
 	{
-		return *GetStatId();
+		return GetStatId().GetName();
 	}
 };
 
diff --git a/Engine/Source/Runtime/Core/Public/Stats/StatsData.h b/Engine/Source/Runtime/Core/Public/Stats/StatsData.h
--- a/Engine/Source/Runtime/Core/Public/Stats/StatsData.h
+++ b/Engine/Source/Runtime/Core/Public/Stats/StatsData.h
@@ -729,7 +729,7 @@ struct FGroupSort
 				return false;
 			}
 			// otherwise, reverse order of definition
-			return A.NameAndInfo.GetRawName().GetIndex() > B.NameAndInfo.GetRawName().GetIndex();
+			return A.NameAndInfo.GetRawName().GetComparisonIndex() > B.NameAndInfo.GetRawName().GetComparisonIndex();
 		}
 		if (GroupA == NAME_None)
 		{
@@ -739,7 +739,7 @@ struct FGroupSort
 		{
 			return true;
 		}
-		return GroupA.GetIndex() > GroupB.GetIndex();
+		return GroupA.GetComparisonIndex() > GroupB.GetComparisonIndex();
 	}
 };
 
diff --git a/Engine/Source/Runtime/Core/Public/Stats/StatsFile.h b/Engine/Source/Runtime/Core/Public/Stats/StatsFile.h
--- a/Engine/Source/Runtime/Core/Public/Stats/StatsFile.h
+++ b/Engine/Source/Runtime/Core/Public/Stats/StatsFile.h
@@ -407,13 +407,13 @@ protected:
 	FORCEINLINE_STATS void WriteFName( FArchive& Ar, FStatNameAndInfo NameAndInfo )
 	{
 		FName RawName = NameAndInfo.GetRawName();
-		bool bSendFName = !FNamesSent.Contains( RawName.GetIndex() );
-		int32 Index = RawName.GetIndex();
+		bool bSendFName = !FNamesSent.Contains( RawName.GetComparisonIndex() );
+		int32 Index = RawName.GetComparisonIndex();
 		Ar << Index;
 		int32 Number = NameAndInfo.GetRawNumber();
 		if( bSendFName )
 		{
-			FNamesSent.Add( RawName.GetIndex() );
+			FNamesSent.Add( RawName.GetComparisonIndex() );
 			Number |= EStatMetaFlags::SendingFName << (EStatMetaFlags::Shift + EStatAllFields::StartShift);
 		}
 		Ar << Number;
@@ -581,7 +581,7 @@ public:
 				Ar << Name;
 
 				TheFName = FName( *Name );
-				FNamesIndexMap.Add( Index, TheFName.GetIndex() );
+				FNamesIndexMap.Add( Index, TheFName.GetComparisonIndex() );
 				Number &= ~(EStatMetaFlags::SendingFName << (EStatMetaFlags::Shift + EStatAllFields::StartShift));
 			}
 			else
@@ -589,7 +589,7 @@ public:
 				if( FNamesIndexMap.Contains( Index ) )
 				{
 					int32 MyIndex = FNamesIndexMap.FindChecked( Index );
-					TheFName = FName( EName( MyIndex ) );
+					TheFName = FName( MyIndex, MyIndex, 0 );
 				}
 				else
 				{
@@ -610,7 +610,7 @@ public:
 			if( FNamesIndexMap.Contains( Index ) )
 			{
 				int32 MyIndex = FNamesIndexMap.FindChecked( Index );
-				TheFName = FName( EName( MyIndex ) );
+				TheFName = FName( MyIndex, MyIndex, 0 );
 			}
 			else
 			{
diff --git a/Engine/Source/Runtime/Core/Public/UObject/NameTypes.h b/Engine/Source/Runtime/Core/Public/UObject/NameTypes.h
--- a/Engine/Source/Runtime/Core/Public/UObject/NameTypes.h
+++ b/Engine/Source/Runtime/Core/Public/UObject/NameTypes.h
@@ -6,6 +6,16 @@
 	Definitions.
 ----------------------------------------------------------------------------*/
 
+/** 
+ * Do we want to support case-variants for FName?
+ * This will add an extra NAME_INDEX variable to FName, but means that ToString() will return you the exact same 
+ * string that FName::Init was called with (which is useful if your FNames are shown to the end user)
+ * Currently this is enabled for the Editor and any Programs (such as UHT), but not the Runtime
+ */
+#ifndef WITH_CASE_PRESERVING_NAME
+	#define WITH_CASE_PRESERVING_NAME WITH_EDITORONLY_DATA
+#endif
+
 /** Maximum size of name. */
 enum {NAME_SIZE	= 1024};
 
@@ -45,8 +55,11 @@ typedef int32 NAME_INDEX;
 /** These characters cannot be used in long package names */
 #define INVALID_LONGPACKAGE_CHARACTERS	TEXT("\\:*?\"<>|' ,.&!\n\r\t@#")
 
-
-
+enum class ENameCase : uint8
+{
+	CaseSensitive,
+	IgnoreCase,
+};
 
 namespace FNameDefs
 {
@@ -168,31 +181,26 @@ public:
 	 */
 	void AppendNameToString( FString& String ) const;
 
-	/**
-	 * @return case insensitive hash of name
-	 */
-	uint32 GetNameHash() const;
-
 	/**
 	 * @return length of name
 	 */
 	int32 GetNameLength() const;
 
 	/**
-	 * Compares name without looking at case.
+	 * Compares name using the compare method provided.
 	 *
 	 * @param	InName	Name to compare to
 	 * @return	true if equal, false otherwise
 	 */
-	bool IsEqual( const ANSICHAR* InName ) const;
+	bool IsEqual( const ANSICHAR* InName, const ENameCase CompareMethod ) const;
 
 	/**
-	 * Compares name without looking at case.
+	 * Compares name using the compare method provided.
 	 *
 	 * @param	InName	Name to compare to
 	 * @return	true if equal, false otherwise
 	 */
-	bool IsEqual( const WIDECHAR* InName ) const;
+	bool IsEqual( const WIDECHAR* InName, const ENameCase CompareMethod ) const;
 
 	/**
 	 * @return direct access to ANSI name if stored in ANSI
@@ -368,17 +376,92 @@ public:
 // CAUTION: If you change those constants, you probably need to update the debug visualizers.
 typedef TStaticIndirectArrayThreadSafeRead<FNameEntry, 2 * 1024 * 1024 /* 2M unique FNames */, 16384 /* allocated in 64K/128K chunks */ > TNameEntryArray;
 
+/**
+ * The minimum amount of data required to reconstruct a name
+ * This is smaller than FName, but you lose the case-preserving behavior
+ */
+struct CORE_API FMinimalName
+{
+	FORCEINLINE FMinimalName()
+		: Index(0)
+		, Number(NAME_NO_NUMBER_INTERNAL)
+	{
+	}
+
+	FORCEINLINE FMinimalName(const EName N)
+		: Index(N)
+		, Number(NAME_NO_NUMBER_INTERNAL)
+	{
+	}
+
+	FORCEINLINE FMinimalName(const NAME_INDEX InIndex, const int32 InNumber)
+		: Index(InIndex)
+		, Number(InNumber)
+	{
+	}
+
+	/** Index into the Names array (used to find String portion of the string/number pair) */
+	NAME_INDEX		Index;
+	/** Number portion of the string/number pair (stored internally as 1 more than actual, so zero'd memory will be the default, no-instance case) */
+	int32			Number;
+};
+
+/**
+ * The full amount of data required to reconstruct a case-preserving name
+ * This will be the same size as FName when WITH_CASE_PRESERVING_NAME is 1, and is used to store an FName in cases where 
+ * the size of FName must be constant between build configurations (eg, blueprint bytecode)
+ */
+struct CORE_API FScriptName
+{
+	FORCEINLINE FScriptName()
+		: ComparisonIndex(0)
+		, DisplayIndex(0)
+		, Number(NAME_NO_NUMBER_INTERNAL)
+	{
+	}
+
+	FORCEINLINE FScriptName(const EName N)
+		: ComparisonIndex(N)
+		, DisplayIndex(N)
+		, Number(NAME_NO_NUMBER_INTERNAL)
+	{
+	}
+
+	FORCEINLINE FScriptName(const NAME_INDEX InComparisonIndex, const NAME_INDEX InDisplayIndex, const int32 InNumber)
+		: ComparisonIndex(InComparisonIndex)
+		, DisplayIndex(InDisplayIndex)
+		, Number(InNumber)
+	{
+	}
+
+	/** Index into the Names array (used to find String portion of the string/number pair used for comparison) */
+	NAME_INDEX		ComparisonIndex;
+	/** Index into the Names array (used to find String portion of the string/number pair used for display) */
+	NAME_INDEX		DisplayIndex;
+	/** Number portion of the string/number pair (stored internally as 1 more than actual, so zero'd memory will be the default, no-instance case) */
+	uint32			Number;
+};
+
 /**
  * Public name, available to the world.  Names are stored as a combination of
  * an index into a table of unique strings and an instance number.
- * Names are case-insensitive.
+ * Names are case-insensitive, but case-preserving (when WITH_CASE_PRESERVING_NAME is 1)
  */
-class CORE_API FName 
+class CORE_API FName
 {
 public:
 
-	FORCEINLINE NAME_INDEX GetIndex() const
+	FORCEINLINE NAME_INDEX GetComparisonIndex() const
 	{
+		const NAME_INDEX Index = GetComparisonIndexFast();
+		checkName(Index >= 0 && Index < GetNames().Num());
+		checkName(GetNames()[Index]);
+		return Index;
+	}
+
+	FORCEINLINE NAME_INDEX GetDisplayIndex() const
+	{
+		const NAME_INDEX Index = GetDisplayIndexFast();
 		checkName(Index >= 0 && Index < GetNames().Num());
 		checkName(GetNames()[Index]);
 		return Index;
@@ -389,7 +472,7 @@ public:
 		return Number;
 	}
 
-	FORCEINLINE void SetNumber(int32 NewNumber)
+	FORCEINLINE void SetNumber(const int32 NewNumber)
 	{
 		Number = NewNumber;
 	}
@@ -397,7 +480,7 @@ public:
 	/** Returns the pure name string without any trailing numbers */
 	FString GetPlainNameString() const
 	{
-		return GetNames()[Index]->GetPlainNameString();
+		return GetDisplayNameEntry()->GetPlainNameString();
 	}
 
 	/**
@@ -405,7 +488,7 @@ public:
 	 */
 	FORCEINLINE ANSICHAR const* GetPlainANSIString() const
 	{
-		return GetNames()[Index]->GetAnsiName();		
+		return GetDisplayNameEntry()->GetAnsiName();		
 	}
 
 	/**
@@ -413,9 +496,12 @@ public:
 	 */
 	FORCEINLINE WIDECHAR const* GetPlainWIDEString() const
 	{
-		return GetNames()[Index]->GetWideName();
+		return GetDisplayNameEntry()->GetWideName();
 	}
 
+	const FNameEntry* GetComparisonNameEntry() const;
+	const FNameEntry* GetDisplayNameEntry() const;
+
 	/**
 	 * Converts an FName to a readable format
 	 *
@@ -426,7 +512,7 @@ public:
 	/**
 	 * Converts an FName to a readable format, in place
 	 * 
-	 * @param Out String to fill ot with the string representation of the name
+	 * @param Out String to fill with the string representation of the name
 	 */
 	void ToString(FString& Out) const;
 
@@ -437,13 +523,22 @@ public:
 	 */
 	void AppendString(FString& Out) const;
 
-	FORCEINLINE bool operator==(FName Other) const
+	/**
+	 * Check to see if this FName matches the other FName, potentially also checking for any case variations
+	 */
+	FORCEINLINE bool IsEqual(const FName& Other, const ENameCase CompareMethod = ENameCase::IgnoreCase, const bool bCompareNumber = true ) const
 	{
-		return Index == Other.Index && Number == Other.Number;
+		return ((CompareMethod == ENameCase::IgnoreCase) ? GetComparisonIndexFast() == Other.GetComparisonIndexFast() : GetDisplayIndexFast() == Other.GetDisplayIndexFast())
+			&& (!bCompareNumber || GetNumber() == Other.GetNumber());
 	}
-	FORCEINLINE bool operator!=(FName Other) const
+
+	FORCEINLINE bool operator==(const FName& Other) const
+	{
+		return IsEqual(Other, ENameCase::IgnoreCase);
+	}
+	FORCEINLINE bool operator!=(const FName& Other) const
 	{
-		return Index != Other.Index || Number != Other.Number;
+		return !IsEqual(Other, ENameCase::IgnoreCase);
 	}
 
 	/**
@@ -456,13 +551,14 @@ public:
 
 	FORCEINLINE bool IsNone() const
 	{
-		return Index == 0 && Number == 0;
+		return GetComparisonIndexFast() == 0 && GetNumber() == 0;
 	}
 
 	FORCEINLINE bool IsValid() const
 	{
 		TNameEntryArray& Names = GetNames();
-		return Index>=0 && Index<Names.Num() && Names[Index]!=NULL;
+		return GetComparisonIndexFast()>=0 && GetComparisonIndexFast()<Names.Num() && Names[GetComparisonIndexFast()]!=NULL
+			&& GetDisplayIndexFast()>=0 && GetDisplayIndexFast()<Names.Num() && Names[GetDisplayIndexFast()]!=NULL;
 	}
 
 	/**
@@ -470,7 +566,8 @@ public:
 	 */
 	FORCEINLINE bool IsValidIndexFast() const
 	{
-		return Index >= 0 && Index < GetNames().Num();
+		return GetComparisonIndexFast() >= 0 && GetComparisonIndexFast() < GetNames().Num()
+			&& GetDisplayIndexFast() >= 0 && GetDisplayIndexFast() < GetNames().Num();
 	}
 
 	/**
@@ -525,7 +622,10 @@ public:
 	// Assignment operator
 	FORCEINLINE FName& operator=(const FName& Other)
 	{
-		this->Index = Other.Index;
+		this->ComparisonIndex = Other.ComparisonIndex;
+#if WITH_CASE_PRESERVING_NAME
+		this->DisplayIndex = Other.DisplayIndex;
+#endif
 		this->Number = Other.Number;
 
 		return *this;
@@ -545,10 +645,15 @@ public:
 	 *
 	 * @param N The harcdcoded value the string portion of the name will have. The number portion will be NAME_NO_NUMBER
 	 */
-	FORCEINLINE FName( enum EName N )
-	: Index( N )
-	, Number( NAME_NO_NUMBER_INTERNAL )
-	{}
+	FORCEINLINE FName( EName N )
+		: ComparisonIndex( N )
+#if WITH_CASE_PRESERVING_NAME
+		, DisplayIndex( N )
+#endif
+		, Number( NAME_NO_NUMBER_INTERNAL )
+	{
+		check(N < NAME_MaxHardcodedNameIndex);
+	}
 
 	/**
 	 * Create an FName with a hardcoded string index and (instance).
@@ -556,9 +661,44 @@ public:
 	 * @param N The harcdcoded value the string portion of the name will have
 	 * @param InNumber The hardcoded value for the number portion of the name
 	 */
-	FORCEINLINE FName( enum EName N, int32 InNumber )
-	: Index( N )
-	, Number( InNumber )
+	FORCEINLINE FName( EName N, int32 InNumber )
+		: ComparisonIndex( N )
+#if WITH_CASE_PRESERVING_NAME
+		, DisplayIndex( N )
+#endif
+		, Number( InNumber )
+	{
+		// If this fires the enum was out of bounds - did you pass an index instead?
+		// If you want to clone an FName with a new number, FName( const FName& Other, int32 InNumber ) is the thing you want
+		check(N < NAME_MaxHardcodedNameIndex);
+		check(InNumber >= 0 && InNumber <= 0xffffff);
+	}
+
+	/**
+	 * Create an FName from an existing string, but with a different instance.
+	 *
+	 * @param Other The FName to take the string values from
+	 * @param InNumber The hardcoded value for the number portion of the name
+	 */
+	FORCEINLINE FName( const FName& Other, int32 InNumber )
+		: ComparisonIndex( Other.ComparisonIndex )
+#if WITH_CASE_PRESERVING_NAME
+		, DisplayIndex( Other.DisplayIndex )
+#endif
+		, Number( InNumber )
+	{
+	}
+
+	/**
+	 * Create an FName from its component parts
+	 * Only call this if you *really* know what you're doing
+	 */
+	FORCEINLINE FName( const NAME_INDEX InComparisonIndex, const NAME_INDEX InDisplayIndex, const int32 InNumber )
+		: ComparisonIndex( InComparisonIndex )
+#if WITH_CASE_PRESERVING_NAME
+		, DisplayIndex( InDisplayIndex )
+#endif
+		, Number( InNumber )
 	{
 	}
 
@@ -566,8 +706,11 @@ public:
 	 * Default constructor, initialized to None
 	 */
 	FORCEINLINE FName()
-		: Index( 0 )
-		, Number( 0 )
+		: ComparisonIndex( 0 )
+#if WITH_CASE_PRESERVING_NAME
+		, DisplayIndex( 0 )
+#endif
+		, Number( NAME_NO_NUMBER_INTERNAL )
 	{
 	}
 
@@ -614,13 +757,13 @@ public:
 	/**
 	 * Create an FName with a hardcoded string index.
 	 *
-	 * @param HardcodedIndex	The harcdcoded value the string portion of the name will have. 
-	 * @param Name				The harcdcoded name to intialize
+	 * @param HardcodedIndex	The hardcoded value the string portion of the name will have. 
+	 * @param Name				The hardcoded name to intialize
 	 */
-	explicit FName( enum EName HardcodedIndex, const TCHAR* Name );
+	explicit FName( EName HardcodedIndex, const TCHAR* Name );
 
 	/**
-	 * Comparision operator.
+	 * Comparison operator.
 	 *
 	 * @param	Other	String to compare this name to
 	 * @return true if name matches the string, false otherwise
@@ -629,12 +772,12 @@ public:
 
 	static void StaticInit();
 	static void DisplayHash( class FOutputDevice& Ar );
-	static FString SafeString( EName Index, int32 InstanceNumber=NAME_NO_NUMBER_INTERNAL )
+	static FString SafeString( int32 InDisplayIndex, int32 InstanceNumber=NAME_NO_NUMBER_INTERNAL )
 	{
 		TNameEntryArray& Names = GetNames();
 		return GetIsInitialized()
-			? (Names.IsValidIndex(Index) && Names[Index])
-				? FName(Index, InstanceNumber).ToString()
+			? (Names.IsValidIndex(InDisplayIndex) && Names[InDisplayIndex])
+				? FName(InDisplayIndex, InDisplayIndex, InstanceNumber).ToString()
 				: FString(TEXT("*INVALID*"))
 			: FString(TEXT("*UNINITIALIZED*"));
 	}
@@ -683,7 +826,7 @@ public:
 	 * returns true.
 	 *
 	 * @param OldName		Old-style name
-	 * @param NewName		Ouput string portion of the name/number pair
+	 * @param NewName		Output string portion of the name/number pair
 	 * @param NewNameLen	Size of NewName buffer (in TCHAR units)
 	 * @param NewNumber		Number portion of the name/number pair
 	 *
@@ -711,10 +854,14 @@ public:
 
 private:
 
-	/** Index into the Names array (used to find String portion of the string/number pair) */
-	NAME_INDEX	Index;
+	/** Index into the Names array (used to find String portion of the string/number pair used for comparison) */
+	NAME_INDEX		ComparisonIndex;
+#if WITH_CASE_PRESERVING_NAME
+	/** Index into the Names array (used to find String portion of the string/number pair used for display) */
+	NAME_INDEX		DisplayIndex;
+#endif
 	/** Number portion of the string/number pair (stored internally as 1 more than actual, so zero'd memory will be the default, no-instance case) */
-	int32		Number;
+	uint32			Number;
 
 	/** Name hash.												*/
 	static FNameEntry*						NameHash[ FNameDefs::NameHashBucketCount ];
@@ -762,6 +909,26 @@ private:
 		Init(StringCast<WIDECHAR>(InName).Get(), InNumber, FindType, bSplitName, HardcodeIndex);
 	}
 
+	template <typename TCharType>
+	static bool InitInternal_FindOrAdd(const TCharType* InName, const EFindName FindType, const int32 HardcodeIndex, int32& OutComparisonIndex, int32& OutDisplayIndex);
+
+	template <typename TCharType>
+	static bool InitInternal_FindOrAddNameEntry(const TCharType* InName, const EFindName FindType, const ENameCase ComparisonMode, int32& OutIndex);
+
+	FORCEINLINE NAME_INDEX GetComparisonIndexFast() const
+	{
+		return ComparisonIndex;
+	}
+
+	FORCEINLINE NAME_INDEX GetDisplayIndexFast() const
+	{
+#if WITH_CASE_PRESERVING_NAME
+		return DisplayIndex;
+#else
+		return ComparisonIndex;
+#endif
+	}
+
 	/** Singleton to retrieve the critical section. */
 	static FCriticalSection* GetCriticalSection();
 
@@ -773,9 +940,32 @@ Expose_TNameOf(FName)
 
 inline uint32 GetTypeHash( const FName N )
 {
-	return N.GetIndex();
+	return N.GetComparisonIndex();
 }
 
+
+FORCEINLINE FMinimalName NameToMinimalName(const FName& InName)
+{
+	return FMinimalName(InName.GetComparisonIndex(), InName.GetNumber());
+}
+
+FORCEINLINE FName MinimalNameToName(const FMinimalName& InName)
+{
+	return FName(InName.Index, InName.Index, InName.Number);
+}
+
+
+FORCEINLINE FScriptName NameToScriptName(const FName& InName)
+{
+	return FScriptName(InName.GetComparisonIndex(), InName.GetDisplayIndex(), InName.GetNumber());
+}
+
+FORCEINLINE FName ScriptNameToName(const FScriptName& InName)
+{
+	return FName(InName.ComparisonIndex, InName.DisplayIndex, InName.Number);
+}
+
+
 /**
  * Comparison operator with TCHAR* on left hand side and FName on right hand side
  * 
diff --git a/Engine/Source/Runtime/Core/Public/UObject/ObjectVersion.h b/Engine/Source/Runtime/Core/Public/UObject/ObjectVersion.h
--- a/Engine/Source/Runtime/Core/Public/UObject/ObjectVersion.h
+++ b/Engine/Source/Runtime/Core/Public/UObject/ObjectVersion.h
@@ -613,6 +613,8 @@ enum EUnrealEngineObjectUE4Version
 	// Rename USpringArmComponent::bUseControllerViewRotation to bUsePawnViewRotation,
 	// Rename UCameraComponent::bUseControllerViewRotation to bUsePawnViewRotation (and change the default value)
 	VER_UE4_RENAME_CAMERA_COMPONENT_VIEW_ROTATION,
+	// Changed FName to be case preserving
+	VER_UE4_CASE_PRESERVING_FNAME,
 
 	// -----<new versions can be added before this line>-------------------------------------------------
 	// - this needs to be the last line (see note below)
@@ -636,7 +638,7 @@ enum EUnrealEngineObjectLicenseeUE4Version
 #define VER_LATEST_ENGINE_LICENSEEUE4   PREPROCESSOR_ENUM_PROTECT(VER_LIC_AUTOMATIC_VERSION)
 
 // Minimum package version that contains legal bytecode
-#define VER_MIN_SCRIPTVM_UE4									(VER_UE4_DISABLED_SCRIPT_LIMIT_BYTECODE)
+#define VER_MIN_SCRIPTVM_UE4									(VER_UE4_CASE_PRESERVING_FNAME)
 #define VER_MIN_SCRIPTVM_LICENSEEUE4							(VER_LIC_NONE)
 
 // Version access.
diff --git a/Engine/Source/Runtime/Core/Public/UObject/UnrealNames.h b/Engine/Source/Runtime/Core/Public/UObject/UnrealNames.h
--- a/Engine/Source/Runtime/Core/Public/UObject/UnrealNames.h
+++ b/Engine/Source/Runtime/Core/Public/UObject/UnrealNames.h
@@ -204,6 +204,7 @@ REGISTER_NAME(500,PerformanceCapture)
 //
 
 #ifdef REGISTERING_ENUM
+	NAME_MaxHardcodedNameIndex,
 	};
 	#undef REGISTER_NAME
 	#undef REGISTERING_ENUM
diff --git a/Engine/Source/Runtime/CoreUObject/Private/Serialization/DuplicateDataReader.cpp b/Engine/Source/Runtime/CoreUObject/Private/Serialization/DuplicateDataReader.cpp
--- a/Engine/Source/Runtime/CoreUObject/Private/Serialization/DuplicateDataReader.cpp
+++ b/Engine/Source/Runtime/CoreUObject/Private/Serialization/DuplicateDataReader.cpp
@@ -27,6 +27,18 @@ void FDuplicateDataReader::SerializeFail()
 	UE_LOG(LogObj, Fatal, TEXT("FDuplicateDataReader Overread. GSerializedObject = %s GSerializedProperty = %s"), *GetFullNameSafe(GSerializedObject), *GetFullNameSafe(GSerializedProperty));
 }
 
+FArchive& FDuplicateDataReader::operator<<(FName& N)
+{
+	NAME_INDEX ComparisonIndex;
+	NAME_INDEX DisplayIndex;
+	int32 Number;
+	ByteOrderSerialize(&ComparisonIndex, sizeof(ComparisonIndex));
+	ByteOrderSerialize(&DisplayIndex, sizeof(DisplayIndex));
+	ByteOrderSerialize(&Number, sizeof(Number));
+	// copy over the name with a name made from the name index and number
+	N = FName(ComparisonIndex, DisplayIndex, Number);
+	return *this;
+}
 
 FArchive& FDuplicateDataReader::operator<<( UObject*& Object )
 {
diff --git a/Engine/Source/Runtime/CoreUObject/Private/Serialization/DuplicateDataWriter.cpp b/Engine/Source/Runtime/CoreUObject/Private/Serialization/DuplicateDataWriter.cpp
--- a/Engine/Source/Runtime/CoreUObject/Private/Serialization/DuplicateDataWriter.cpp
+++ b/Engine/Source/Runtime/CoreUObject/Private/Serialization/DuplicateDataWriter.cpp
@@ -36,6 +36,17 @@ FDuplicateDataWriter::FDuplicateDataWriter( FUObjectAnnotationSparse<FDuplicated
 /**
  * I/O function
  */
+FArchive& FDuplicateDataWriter::operator<<(FName& N)
+{
+	NAME_INDEX ComparisonIndex = N.GetComparisonIndex();
+	NAME_INDEX DisplayIndex = N.GetDisplayIndex();
+	int32 Number = N.GetNumber();
+	ByteOrderSerialize(&ComparisonIndex, sizeof(ComparisonIndex));
+	ByteOrderSerialize(&DisplayIndex, sizeof(DisplayIndex));
+	ByteOrderSerialize(&Number, sizeof(Number));
+	return *this;
+}
+
 FArchive& FDuplicateDataWriter::operator<<(UObject*& Object)
 {
 	GetDuplicatedObject(Object);
diff --git a/Engine/Source/Runtime/CoreUObject/Private/Serialization/ObjectReader.cpp b/Engine/Source/Runtime/CoreUObject/Private/Serialization/ObjectReader.cpp
--- a/Engine/Source/Runtime/CoreUObject/Private/Serialization/ObjectReader.cpp
+++ b/Engine/Source/Runtime/CoreUObject/Private/Serialization/ObjectReader.cpp
@@ -7,12 +7,14 @@
 
 FArchive& FObjectReader::operator<<( class FName& N )
 {
-	NAME_INDEX Name;
+	NAME_INDEX ComparisonIndex;
+	NAME_INDEX DisplayIndex;
 	int32 Number;
-	ByteOrderSerialize(&Name, sizeof(Name));
+	ByteOrderSerialize(&ComparisonIndex, sizeof(ComparisonIndex));
+	ByteOrderSerialize(&DisplayIndex, sizeof(DisplayIndex));
 	ByteOrderSerialize(&Number, sizeof(Number));
 	// copy over the name with a name made from the name index and number
-	N = FName((EName)Name, Number);
+	N = FName(ComparisonIndex, DisplayIndex, Number);
 	return *this;
 }
 
diff --git a/Engine/Source/Runtime/CoreUObject/Private/Serialization/ObjectWriter.cpp b/Engine/Source/Runtime/CoreUObject/Private/Serialization/ObjectWriter.cpp
--- a/Engine/Source/Runtime/CoreUObject/Private/Serialization/ObjectWriter.cpp
+++ b/Engine/Source/Runtime/CoreUObject/Private/Serialization/ObjectWriter.cpp
@@ -7,9 +7,11 @@
 
 FArchive& FObjectWriter::operator<<( class FName& N )
 {
-	NAME_INDEX Name = N.GetIndex();
+	NAME_INDEX ComparisonIndex = N.GetComparisonIndex();
+	NAME_INDEX DisplayIndex = N.GetDisplayIndex();
 	int32 Number = N.GetNumber();
-	ByteOrderSerialize(&Name, sizeof(Name));
+	ByteOrderSerialize(&ComparisonIndex, sizeof(ComparisonIndex));
+	ByteOrderSerialize(&DisplayIndex, sizeof(DisplayIndex));
 	ByteOrderSerialize(&Number, sizeof(Number));
 	return *this;
 }
diff --git a/Engine/Source/Runtime/CoreUObject/Private/Serialization/ReloadObjectArc.cpp b/Engine/Source/Runtime/CoreUObject/Private/Serialization/ReloadObjectArc.cpp
--- a/Engine/Source/Runtime/CoreUObject/Private/Serialization/ReloadObjectArc.cpp
+++ b/Engine/Source/Runtime/CoreUObject/Private/Serialization/ReloadObjectArc.cpp
@@ -151,21 +151,23 @@ void FReloadObjectArc::Reset()
  */
 FArchive& FReloadObjectArc::operator<<( class FName& Name )
 {
-	NAME_INDEX NameIndex;
+	NAME_INDEX NameComparisonIndex;
+	NAME_INDEX NameDisplayIndex;
 	int32 NameInstance;
 	if ( IsLoading() )
 	{
-		Reader << NameIndex << NameInstance;
+		Reader << NameComparisonIndex << NameDisplayIndex << NameInstance;
 
 		// recreate the FName using the serialized index and instance number
-		Name = FName((EName)NameIndex, NameInstance);
+		Name = FName(NameComparisonIndex, NameDisplayIndex, NameInstance);
 	}
 	else if ( IsSaving() )
 	{
-		NameIndex = Name.GetIndex();
+		NameComparisonIndex = Name.GetComparisonIndex();
+		NameDisplayIndex = Name.GetDisplayIndex();
 		NameInstance = Name.GetNumber();
 
-		Writer << NameIndex << NameInstance;
+		Writer << NameComparisonIndex << NameDisplayIndex << NameInstance;
 	}
 	return *this;
 }
diff --git a/Engine/Source/Runtime/CoreUObject/Private/UObject/Class.cpp b/Engine/Source/Runtime/CoreUObject/Private/UObject/Class.cpp
--- a/Engine/Source/Runtime/CoreUObject/Private/UObject/Class.cpp
+++ b/Engine/Source/Runtime/CoreUObject/Private/UObject/Class.cpp
@@ -1417,7 +1417,7 @@ void UStruct::AddReferencedObjects(UObject* InThis, FReferenceCollector& Collect
 		int32 iCode = 0;
 		while( iCode < This->Script.Num() )
 		{	
-			const_cast<UStruct*>(This)->SerializeExpr( iCode, ObjectReferenceCollector );
+			This->SerializeExpr( iCode, ObjectReferenceCollector );
 		}
 		for( int32 Index = 0; Index < ScriptObjectReferences.Num(); Index++ )
 		{
diff --git a/Engine/Source/Runtime/CoreUObject/Private/UObject/CoreNet.cpp b/Engine/Source/Runtime/CoreUObject/Private/UObject/CoreNet.cpp
--- a/Engine/Source/Runtime/CoreUObject/Private/UObject/CoreNet.cpp
+++ b/Engine/Source/Runtime/CoreUObject/Private/UObject/CoreNet.cpp
@@ -95,13 +95,13 @@ bool UPackageMap::SerializeName(FArchive& Ar, FName& Name)
 	}
 	else if (Ar.IsSaving())
 	{
-		uint8 bHardcoded = Name.GetIndex() <= MAX_NETWORKED_HARDCODED_NAME;
+		uint8 bHardcoded = Name.GetComparisonIndex() <= MAX_NETWORKED_HARDCODED_NAME;
 		Ar.SerializeBits(&bHardcoded, 1);
 		if (bHardcoded)
 		{
 			// send by hardcoded index
 			checkSlow(Name.GetNumber() <= 0); // hardcoded names should never have a Number
-			uint32 NameIndex = uint32(Name.GetIndex());
+			uint32 NameIndex = uint32(Name.GetComparisonIndex());
 			Ar.SerializeInt(NameIndex, MAX_NETWORKED_HARDCODED_NAME + 1);
 		}
 		else
diff --git a/Engine/Source/Runtime/CoreUObject/Private/UObject/LinkerLoad.cpp b/Engine/Source/Runtime/CoreUObject/Private/UObject/LinkerLoad.cpp
--- a/Engine/Source/Runtime/CoreUObject/Private/UObject/LinkerLoad.cpp
+++ b/Engine/Source/Runtime/CoreUObject/Private/UObject/LinkerLoad.cpp
@@ -295,7 +295,7 @@ static bool IgnoreMissingReferencedClass( FName ClassName )
 
 static inline int32 HashNames( FName A, FName B, FName C )
 {
-	return A.GetIndex() + 7 * B.GetIndex() + 31 * FPackageName::GetShortFName(C).GetIndex();
+	return A.GetComparisonIndex() + 7 * B.GetComparisonIndex() + 31 * FPackageName::GetShortFName(C).GetComparisonIndex();
 }
 
 static FORCEINLINE bool IsCoreUObjectPackage(const FName& PackageName)
@@ -3624,7 +3624,8 @@ FArchive& ULinkerLoad::operator<<( FName& Name )
 	}
 
 	// if the name wasn't loaded (because it wasn't valid in this context)
-	if (NameMap[NameIndex] == NAME_None)
+	const FName& MappedName = NameMap[NameIndex];
+	if (MappedName.IsNone())
 	{
 		int32 TempNumber;
 		Ar << TempNumber;
@@ -3634,12 +3635,12 @@ FArchive& ULinkerLoad::operator<<( FName& Name )
 	{
 		int32 Number;
 		Ar << Number;
-		// simply create the name from the NameMap's name index and the serialized instance number
+		// simply create the name from the NameMap's name and the serialized instance number
 #ifndef __clang__
-		Name = FName((EName)NameMap[NameIndex].GetIndex(), Number);
+		Name = FName(MappedName, Number);
 #else
 		// @todo-mobile: IOS crashes on the assignment; need to do a memcpy manually...
-		FName TempName = FName((EName)NameMap[NameIndex].GetIndex(), Number);
+		FName TempName = FName(MappedName, Number);
 		FMemory::Memcpy(&Name, &TempName, sizeof(FName));
 #endif
 	}
diff --git a/Engine/Source/Runtime/CoreUObject/Private/UObject/LinkerSave.cpp b/Engine/Source/Runtime/CoreUObject/Private/UObject/LinkerSave.cpp
--- a/Engine/Source/Runtime/CoreUObject/Private/UObject/LinkerSave.cpp
+++ b/Engine/Source/Runtime/CoreUObject/Private/UObject/LinkerSave.cpp
@@ -96,14 +96,14 @@ void ULinkerSave::BeginDestroy()
 
 int32 ULinkerSave::MapName(const FName& Name) const
 {
-	const int32* IndexPtr = NameIndices.Find(Name.GetIndex());
+	const int32* IndexPtr = NameIndices.Find(Name);
 
 	if (IndexPtr)
 	{
 		return *IndexPtr;
 	}
 
-	return 0;
+	return INDEX_NONE;
 }
 
 FPackageIndex ULinkerSave::MapObject( const UObject* Object ) const
@@ -134,6 +134,14 @@ void ULinkerSave::Serialize( void* V, int64 Length )
 	Saver->Serialize( V, Length );
 }
 	
+FArchive& ULinkerSave::operator<<( FName& InName )
+{
+	int32 Save = MapName(InName);
+	int32 Number = InName.GetNumber();
+	FArchive& Ar = *this;
+	return Ar << Save << Number;
+}
+
 FArchive& ULinkerSave::operator<<( UObject*& Obj )
 {
 	FPackageIndex Save;
diff --git a/Engine/Source/Runtime/CoreUObject/Private/UObject/Obj.cpp b/Engine/Source/Runtime/CoreUObject/Private/UObject/Obj.cpp
--- a/Engine/Source/Runtime/CoreUObject/Private/UObject/Obj.cpp
+++ b/Engine/Source/Runtime/CoreUObject/Private/UObject/Obj.cpp
@@ -70,12 +70,22 @@ bool UObject::Rename( const TCHAR* InName, UObject* NewOuter, ERenameFlags Flags
 	UObject* NameScopeOuter = (Flags & REN_ForceGlobalUnique) ? ANY_PACKAGE : NewOuter;
 
 	// find an object with the same name and same class in the new outer
+	bool bIsCaseOnlyChange = false;
 	if (InName)
 	{
 		UObject* ExistingObject = StaticFindObject(/*Class=*/ NULL, NameScopeOuter ? NameScopeOuter : GetOuter(), InName, true);
 		if (ExistingObject == this)
 		{
-			return true;
+			if (ExistingObject->GetName().Equals(InName, ESearchCase::CaseSensitive))
+			{
+				// The name is exactly the same - there's nothing to change
+				return true;
+			}
+			else
+			{
+				// This rename has only changed the case, so we need to allow it to continue, but won't create a redirector (since the internal FName comparison ignores case)
+				bIsCaseOnlyChange = true;
+			}
 		}
 		else if (ExistingObject)
 		{
@@ -130,12 +140,12 @@ bool UObject::Rename( const TCHAR* InName, UObject* NewOuter, ERenameFlags Flags
 	if ( HasAnyFlags(RF_Public) )
 	{
 		const bool bUniquePathChanged	= ((NewOuter != NULL && OldOuter != NewOuter) || (OldName != NewName));
-		const bool bRootPackage		= GetClass() == UPackage::StaticClass() && OldOuter == NULL;
+		const bool bRootPackage			= GetClass() == UPackage::StaticClass() && OldOuter == NULL;
 		const bool bRedirectionAllowed = !FApp::IsGame() && ((Flags & REN_DontCreateRedirectors) == 0);
 
 		// We need to create a redirector if we changed the Outer or Name of an object that can be referenced from other packages
 		// [i.e. has the RF_Public flag] so that references to this object are not broken.
-		bCreateRedirector = bRootPackage == false && bUniquePathChanged == true && bRedirectionAllowed == true;
+		bCreateRedirector = bRootPackage == false && bUniquePathChanged == true && bRedirectionAllowed == true && bIsCaseOnlyChange == false;
 	}
 
 	if( NewOuter )
diff --git a/Engine/Source/Runtime/CoreUObject/Private/UObject/SavePackage.cpp b/Engine/Source/Runtime/CoreUObject/Private/UObject/SavePackage.cpp
--- a/Engine/Source/Runtime/CoreUObject/Private/UObject/SavePackage.cpp
+++ b/Engine/Source/Runtime/CoreUObject/Private/UObject/SavePackage.cpp
@@ -212,25 +212,29 @@ public:
 	 * Mark a FName as referenced.
 	 * @param Name name to mark as referenced
 	 */
-	void MarkNameAsReferenced(FName Name)
+	void MarkNameAsReferenced(const FName& Name)
 	{
-		ReferencedNames.Add(Name);
+		// We need to store the FName without the number, as the number is stored separately by ULinkerSave 
+		// and we don't want duplicate entries in the name table just because of the number
+		const FName NameNoNumber(Name, 0);
+		ReferencedNames.Add(NameNoNumber);
 	}
+
 	/** 
 	 * Add the marked names to a linker
 	 * @param Linker linker to save the marked names to
 	 */
 	void UpdateLinkerWithMarkedNames(ULinkerSave* Linker)
 	{
-		for (TSet<FName>::TIterator It(ReferencedNames); It; ++It)
+		Linker->NameMap.Reserve(Linker->NameMap.Num() + ReferencedNames.Num());
+		for (const FName& Name : ReferencedNames)
 		{
-			Linker->NameMap.Add(*It);
+			Linker->NameMap.Add(Name);
 		}
 	}
 
-
 private:
-	TSet<FName> ReferencedNames;
+	TSet<FName, FLinkerNamePairKeyFuncs> ReferencedNames;
 };
 
 static FSavePackageState* SavePackageState = NULL;
@@ -2982,8 +2986,8 @@ bool UPackage::SavePackage( UPackage* InOuter, UObject* Base, EObjectFlags TopLe
 				Linker->Summary.NameCount = Linker->NameMap.Num();
 				for( int32 i=0; i<Linker->NameMap.Num(); i++ )
 				{
-					*Linker << *const_cast<FNameEntry*>(FName::GetEntry( Linker->NameMap[i].GetIndex() ));
-					Linker->NameIndices.Add(Linker->NameMap[i].GetIndex(), i);
+					*Linker << *const_cast<FNameEntry*>(Linker->NameMap[i].GetDisplayNameEntry());
+					Linker->NameIndices.Add(Linker->NameMap[i], i);
 				}
 
 				
diff --git a/Engine/Source/Runtime/CoreUObject/Private/UObject/UObjectBase.cpp b/Engine/Source/Runtime/CoreUObject/Private/UObject/UObjectBase.cpp
--- a/Engine/Source/Runtime/CoreUObject/Private/UObject/UObjectBase.cpp
+++ b/Engine/Source/Runtime/CoreUObject/Private/UObject/UObjectBase.cpp
@@ -805,7 +805,7 @@ const TCHAR* DebugFName(UObject* Object)
 	// Hardcoded static array. This function is only used inside the debugger so it should be fine to return it.
 	static TCHAR TempName[256];
 	FName Name = Object->GetFName();
-	FCString::Strcpy(TempName,Object ? *FName::SafeString((EName)Name.GetIndex(), Name.GetNumber()) : TEXT("NULL"));
+	FCString::Strcpy(TempName,Object ? *FName::SafeString(Name.GetDisplayIndex(), Name.GetNumber()) : TEXT("NULL"));
 	return TempName;
 }
 
diff --git a/Engine/Source/Runtime/CoreUObject/Private/UObject/UObjectGlobals.cpp b/Engine/Source/Runtime/CoreUObject/Private/UObject/UObjectGlobals.cpp
--- a/Engine/Source/Runtime/CoreUObject/Private/UObject/UObjectGlobals.cpp
+++ b/Engine/Source/Runtime/CoreUObject/Private/UObject/UObjectGlobals.cpp
@@ -950,7 +950,7 @@ struct FCompareUObjectByLinkerAndOffset
 			// Sort by linker name.
 			else
 			{
-				return LinkerA->GetFName().GetIndex() < LinkerB->GetFName().GetIndex();
+				return LinkerA->GetFName().GetComparisonIndex() < LinkerB->GetFName().GetComparisonIndex();
 			}
 		}
 		// Neither objects have a linker, don't do anything.
@@ -1221,7 +1221,6 @@ FName MakeUniqueObjectName( UObject* Parent, UClass* Class, FName BaseName/*=NAM
 	}
 
 	// cache the class's name's index for faster name creation later
-	EName BaseNameIndex = (EName)BaseName.GetIndex();
 	FName TestName;
 	if (!FPlatformProperties::HasEditorOnlyData() && GFastPathUniqueNameGeneration)
 	{
@@ -1242,7 +1241,7 @@ FName MakeUniqueObjectName( UObject* Parent, UClass* Class, FName BaseName/*=NAM
 		*   impossible.
 		*/
 		static int32 UniqueIndex = MAX_int32 - 1000;
-		TestName = FName(BaseNameIndex, --UniqueIndex);
+		TestName = FName(BaseName, --UniqueIndex);
 		checkSlow(Parent); 
 		checkSlow(Parent!=ANY_PACKAGE); 
 		checkSlow(!StaticFindObjectFastInternal( NULL, Parent, TestName ));
@@ -1254,7 +1253,7 @@ FName MakeUniqueObjectName( UObject* Parent, UClass* Class, FName BaseName/*=NAM
 		do
 		{
 			// create the next name in the sequence for this class
-			if (BaseNameIndex == NAME_Package)
+			if (BaseName.GetComparisonIndex() == NAME_Package)
 			{
 				if ( Parent == NULL )
 				{
@@ -1269,7 +1268,7 @@ FName MakeUniqueObjectName( UObject* Parent, UClass* Class, FName BaseName/*=NAM
 			}
 			else
 			{
-				TestName = FName(BaseNameIndex, ++Class->ClassUnique);
+				TestName = FName(BaseName, ++Class->ClassUnique);
 			}
 
 			if (Parent == ANY_PACKAGE)
diff --git a/Engine/Source/Runtime/CoreUObject/Private/UObject/UObjectHash.cpp b/Engine/Source/Runtime/CoreUObject/Private/UObject/UObjectHash.cpp
--- a/Engine/Source/Runtime/CoreUObject/Private/UObject/UObjectHash.cpp
+++ b/Engine/Source/Runtime/CoreUObject/Private/UObject/UObjectHash.cpp
@@ -32,7 +32,7 @@ static TMultiMap<int32,class UObjectBase*> ObjectHashOuter;
  */
 static FORCEINLINE int32 GetObjectHash(FName ObjName)
 {
-	return (ObjName.GetIndex() ^ ObjName.GetNumber()) & (OBJECT_HASH_BINS - 1);
+	return (ObjName.GetComparisonIndex() ^ ObjName.GetNumber()) & (OBJECT_HASH_BINS - 1);
 }
 
 /**
@@ -46,7 +46,7 @@ static FORCEINLINE int32 GetObjectHash(FName ObjName)
  */
 static FORCEINLINE int32 GetObjectOuterHash(FName ObjName,PTRINT Outer)
 {
-	return ((ObjName.GetIndex() ^ ObjName.GetNumber()) ^ (Outer >> 4)) & (OBJECT_HASH_BINS - 1);
+	return ((ObjName.GetComparisonIndex() ^ ObjName.GetNumber()) ^ (Outer >> 4)) & (OBJECT_HASH_BINS - 1);
 }
 
 /**
diff --git a/Engine/Source/Runtime/CoreUObject/Public/Serialization/ArchiveUObject.h b/Engine/Source/Runtime/CoreUObject/Public/Serialization/ArchiveUObject.h
--- a/Engine/Source/Runtime/CoreUObject/Public/Serialization/ArchiveUObject.h
+++ b/Engine/Source/Runtime/CoreUObject/Public/Serialization/ArchiveUObject.h
@@ -842,17 +842,7 @@ private:
 
 	// FArchive interface.
 
-	virtual FArchive& operator<<(FName& N)
-	{
-		NAME_INDEX Name;
-		int32 Number;
-		ByteOrderSerialize(&Name, sizeof(Name));
-		ByteOrderSerialize(&Number, sizeof(Number));
-		// copy over the name with a name made from the name index and number
-		N = FName((EName)Name, Number);
-		return *this;
-	}
-
+	virtual FArchive& operator<<(FName& N);
 	virtual FArchive& operator<<(UObject*& Object);
 	virtual FArchive& operator<<(FLazyObjectPtr& LazyObjectPtr);
 	virtual FArchive& operator<<(FAssetPtr& AssetPtr);
@@ -928,15 +918,7 @@ private:
 
 	// FArchive interface.
 
-	virtual FArchive& operator<<(FName& N)
-	{
-		NAME_INDEX Name = N.GetIndex();
-		int32 Number = N.GetNumber();
-		ByteOrderSerialize(&Name, sizeof(Name));
-		ByteOrderSerialize(&Number, sizeof(Number));
-		return *this;
-	}
-
+	virtual FArchive& operator<<(FName& N);
 	virtual FArchive& operator<<(UObject*& Object);
 	virtual FArchive& operator<<(FLazyObjectPtr& LazyObjectPtr);
 	virtual FArchive& operator<<(FAssetPtr& AssetPtr);
diff --git a/Engine/Source/Runtime/CoreUObject/Public/UObject/Linker.h b/Engine/Source/Runtime/CoreUObject/Public/UObject/Linker.h
--- a/Engine/Source/Runtime/CoreUObject/Public/UObject/Linker.h
+++ b/Engine/Source/Runtime/CoreUObject/Public/UObject/Linker.h
@@ -888,6 +888,41 @@ public:
 	}
 
 };
+
+
+struct FLinkerNamePairKeyFuncs : DefaultKeyFuncs<FName, false>
+{
+	static FORCEINLINE bool Matches(FName A, FName B)
+	{
+		// The linker requires that FNames preserve case, but the numeric suffix can be ignored since
+		// that is stored separately for each FName instance saved
+		return A.IsEqual(B, ENameCase::CaseSensitive, false/*bCompareNumber*/);
+	}
+
+	static FORCEINLINE uint32 GetKeyHash(FName Key)
+	{
+		return Key.GetComparisonIndex();
+	}
+};
+
+
+template<typename ValueType>
+struct TLinkerNameMapKeyFuncs : TDefaultMapKeyFuncs<FName, ValueType, false>
+{
+	static FORCEINLINE bool Matches(FName A, FName B)
+	{
+		// The linker requires that FNames preserve case, but the numeric suffix can be ignored since
+		// that is stored separately for each FName instance saved
+		return A.IsEqual(B, ENameCase::CaseSensitive, false/*bCompareNumber*/);
+	}
+
+	static FORCEINLINE uint32 GetKeyHash(FName Key)
+	{
+		return Key.GetComparisonIndex();
+	}
+};
+
+
 /*----------------------------------------------------------------------------
 	ULinker.
 ----------------------------------------------------------------------------*/
@@ -1815,7 +1850,7 @@ class ULinkerSave : public ULinker, public FArchiveUObject
 	TMap<UObject *,FPackageIndex> ObjectIndicesMap;
 
 	/** Index array - location of the name in the NameMap array for each FName is stored in the NameIndices array using the FName's Index */
-	TMap<int32, int32> NameIndices;
+	TMap<FName, int32, FDefaultSetAllocator, TLinkerNameMapKeyFuncs<int32>> NameIndices;
 
 	/** List of bulkdata that needs to be stored at the end of the file */
 	struct FBulkDataStorageInfo
@@ -1845,13 +1880,7 @@ class ULinkerSave : public ULinker, public FArchiveUObject
 	FPackageIndex MapObject(const UObject* Object) const;
 
 	// FArchive interface.
-	FArchive& operator<<(FName& InName)
-	{
-		int32 Save = MapName(InName);
-		int32 Number = InName.GetNumber();
-		FArchive& Ar = *this;
-		return Ar << Save << Number;
-	}
+	FArchive& operator<<( FName& InName );
 	FArchive& operator<<( UObject*& Obj );
 	FArchive& operator<<( FLazyObjectPtr& LazyObjectPtr );
 	FArchive& operator<<( FAssetPtr& AssetPtr );
diff --git a/Engine/Source/Runtime/CoreUObject/Public/UObject/ScriptSerialization.h b/Engine/Source/Runtime/CoreUObject/Public/UObject/ScriptSerialization.h
--- a/Engine/Source/Runtime/CoreUObject/Public/UObject/ScriptSerialization.h
+++ b/Engine/Source/Runtime/CoreUObject/Public/UObject/ScriptSerialization.h
@@ -33,8 +33,25 @@
 #endif
 #endif
 
+//FScriptName
 #ifndef XFERNAME
-	#define XFERNAME() XFER(FName)
+	#define XFERNAME() \
+	{ \
+   	    FName Name; \
+		FScriptName ScriptName; \
+        if (!Ar.IsLoading()) \
+		{ \
+			FMemory::Memcpy( &ScriptName, &Script[iCode], sizeof(FScriptName) ); \
+			Name = ScriptNameToName(ScriptName); \
+		} \
+		Ar << Name; \
+		if (!Ar.IsSaving()) \
+		{ \
+			ScriptName = NameToScriptName(Name); \
+			FMemory::Memcpy( &Script[iCode], &ScriptName, sizeof(FScriptName) ); \
+		} \
+		iCode += sizeof(FScriptName); \
+	}
 #endif	//XFERNAME
 
 #ifndef XFERPTR 
diff --git a/Engine/Source/Runtime/CoreUObject/Public/UObject/Stack.h b/Engine/Source/Runtime/CoreUObject/Public/UObject/Stack.h
--- a/Engine/Source/Runtime/CoreUObject/Public/UObject/Stack.h
+++ b/Engine/Source/Runtime/CoreUObject/Public/UObject/Stack.h
@@ -338,7 +338,7 @@ inline VariableSizeType FFrame::ReadVariableSize( UField** ExpressionField/*=NUL
 			break;
 		case CPT_Double: Result = sizeof(double);
 			break;
-		case CPT_Name: Result = sizeof(FName);
+		case CPT_Name: Result = sizeof(FScriptName);
 			break;
 		case CPT_Vector: Result = sizeof(FVector);
 			break;
@@ -372,14 +372,14 @@ inline VariableSizeType FFrame::ReadVariableSize( UField** ExpressionField/*=NUL
 
 inline FName FFrame::ReadName()
 {
-	FName Result;
+	FScriptName Result;
 #ifdef REQUIRES_ALIGNED_ACCESS
-	FMemory::Memcpy(&Result, Code, sizeof(FName));
+	FMemory::Memcpy(&Result, Code, sizeof(FScriptName));
 #else
-	Result = *(FName*)Code;
+	Result = *(FScriptName*)Code;
 #endif
-	Code += sizeof(FName);
-	return Result;
+	Code += sizeof(FScriptName);
+	return ScriptNameToName(Result);
 }
 
 COREUOBJECT_API void GInitRunaway();
diff --git a/Engine/Source/Runtime/CoreUObject/Public/UObject/UObjectBaseUtility.h b/Engine/Source/Runtime/CoreUObject/Public/UObject/UObjectBaseUtility.h
--- a/Engine/Source/Runtime/CoreUObject/Public/UObject/UObjectBaseUtility.h
+++ b/Engine/Source/Runtime/CoreUObject/Public/UObject/UObjectBaseUtility.h
@@ -469,7 +469,7 @@ struct FScopeCycleCounterUObject : public FCycleCounter
 		if (FThreadStats::IsCollectingData(OtherStat) && Object)
 		{
 			TStatId StatId = Object->GetStatID();
-			if (!StatId->IsNone())
+			if (!StatId.IsNone())
 			{
 				Start(StatId);
 			}
diff --git a/Engine/Source/Runtime/CoreUObject/Public/UObject/UnrealType.h b/Engine/Source/Runtime/CoreUObject/Public/UObject/UnrealType.h
--- a/Engine/Source/Runtime/CoreUObject/Public/UObject/UnrealType.h
+++ b/Engine/Source/Runtime/CoreUObject/Public/UObject/UnrealType.h
@@ -2967,27 +2967,14 @@ struct TFieldRange
 //
 // Find a typed field in a struct.
 //
-template <class T> T* FindField( const UStruct* Owner, const TCHAR* FieldName )
+template <class T> T* FindField( const UStruct* Owner, FName FieldName )
 {
-	// lookup the string name in the Name hash
-	FName Name(FieldName, FNAME_Find);
-	// If we didn't find it, we know the field won't exist in this Struct
-	if (Name == NAME_None)
-		return NULL;
-	// Search by comparing FNames (INTs), not strings
-	for( TFieldIterator<T>It( Owner ); It; ++It )
+	// We know that a "none" field won't exist in this Struct
+	if( FieldName.IsNone() )
 	{
-		if( It->GetFName() == Name )
-		{
-			return *It;
-		}
+		return nullptr;
 	}
-	// If we didn't find it, return no field
-	return NULL;
-}
 
-template <class T> T* FindField( const UStruct* Owner, FName FieldName )
-{
 	// Search by comparing FNames (INTs), not strings
 	for( TFieldIterator<T>It( Owner ); It; ++It )
 	{
@@ -2998,7 +2985,14 @@ template <class T> T* FindField( const UStruct* Owner, FName FieldName )
 	}
 
 	// If we didn't find it, return no field
-	return NULL;
+	return nullptr;
+}
+
+template <class T> T* FindField( const UStruct* Owner, const TCHAR* FieldName )
+{
+	// lookup the string name in the Name hash
+	FName Name(FieldName, FNAME_Find);
+	return FindField<T>(Owner, Name);
 }
 
 /**
diff --git a/Engine/Source/Runtime/Engine/Private/Materials/MaterialExpressions.cpp b/Engine/Source/Runtime/Engine/Private/Materials/MaterialExpressions.cpp
--- a/Engine/Source/Runtime/Engine/Private/Materials/MaterialExpressions.cpp
+++ b/Engine/Source/Runtime/Engine/Private/Materials/MaterialExpressions.cpp
@@ -1487,7 +1487,7 @@ int32 UMaterialExpressionTextureSampleParameter::Compile(class FMaterialCompiler
 		}
 	}
 
-	if (!ParameterName.IsValid() || (ParameterName.GetIndex() == NAME_None))
+	if (!ParameterName.IsValid() || ParameterName.IsNone())
 	{
 		return UMaterialExpressionTextureSample::Compile(Compiler, OutputIndex, MultiplexIndex);
 	}
@@ -5514,7 +5514,7 @@ int32 UMaterialExpressionFontSampleParameter::Compile(class FMaterialCompiler* C
 {
 	int32 Result = -1;
 	if( !ParameterName.IsValid() || 
-		ParameterName.GetIndex() == NAME_None || 
+		ParameterName.IsNone() || 
 		!Font ||
 		!Font->Textures.IsValidIndex(FontTexturePage) )
 	{
@@ -8464,7 +8464,7 @@ int32 UMaterialExpressionAntialiasedTextureMask::Compile(class FMaterialCompiler
 
 	int32 TextureCodeIndex;
 
-	if (!ParameterName.IsValid() || (ParameterName.GetIndex() == NAME_None))
+	if (!ParameterName.IsValid() || ParameterName.IsNone())
 	{
 		TextureCodeIndex = Compiler->Texture(Texture);
 	}
diff --git a/Engine/Source/Runtime/Engine/Private/UnrealEngine.cpp b/Engine/Source/Runtime/Engine/Private/UnrealEngine.cpp
--- a/Engine/Source/Runtime/Engine/Private/UnrealEngine.cpp
+++ b/Engine/Source/Runtime/Engine/Private/UnrealEngine.cpp
@@ -2831,7 +2831,7 @@ bool UEngine::HandleViewnamesCommand( const TCHAR* Cmd, FOutputDevice& Ar )
 	{
 		for (int32 NameIndex = FMath::Max<int32>(FName::GetMaxNames() - NumNames, 0); NameIndex < FName::GetMaxNames(); NameIndex++)
 		{
-			Ar.Logf(TEXT("%d->%s"), NameIndex, *FName::SafeString(EName(NameIndex)));
+			Ar.Logf(TEXT("%d->%s"), NameIndex, *FName::SafeString(NameIndex));
 		}
 	}
 	return true;
diff --git a/Engine/Source/Runtime/ShaderCore/Private/Shader.cpp b/Engine/Source/Runtime/ShaderCore/Private/Shader.cpp
--- a/Engine/Source/Runtime/ShaderCore/Private/Shader.cpp
+++ b/Engine/Source/Runtime/ShaderCore/Private/Shader.cpp
@@ -306,7 +306,7 @@ FShaderResource::FShaderResource(const FShaderCompilerOutput& Output)
 	checkSlow(OutputHash != FSHAHash());
 
 	ShaderResourceIdMap.Add(GetId(), this);
-	INC_DWORD_STAT_BY_FName(*GetMemoryStatType((EShaderFrequency)Target.Frequency), Code.Num());
+	INC_DWORD_STAT_BY_FName(GetMemoryStatType((EShaderFrequency)Target.Frequency).GetName(), Code.Num());
 	INC_DWORD_STAT_BY(STAT_Shaders_ShaderResourceMemory, GetSizeBytes());
 	INC_DWORD_STAT_BY(STAT_Shaders_NumShaderResourcesLoaded, 1);
 }
@@ -314,7 +314,7 @@ FShaderResource::FShaderResource(const FShaderCompilerOutput& Output)
 
 FShaderResource::~FShaderResource()
 {
-	DEC_DWORD_STAT_BY_FName(*GetMemoryStatType((EShaderFrequency)Target.Frequency), Code.Num());
+	DEC_DWORD_STAT_BY_FName(GetMemoryStatType((EShaderFrequency)Target.Frequency).GetName(), Code.Num());
 	DEC_DWORD_STAT_BY(STAT_Shaders_ShaderResourceMemory, GetSizeBytes());
 	DEC_DWORD_STAT_BY(STAT_Shaders_NumShaderResourcesLoaded, 1);
 }
@@ -336,7 +336,7 @@ void FShaderResource::Serialize(FArchive& Ar)
 	
 	if (Ar.IsLoading())
 	{
-		INC_DWORD_STAT_BY_FName(*GetMemoryStatType((EShaderFrequency)Target.Frequency), (int64)Code.Num());
+		INC_DWORD_STAT_BY_FName(GetMemoryStatType((EShaderFrequency)Target.Frequency).GetName(), (int64)Code.Num());
 		INC_DWORD_STAT_BY(STAT_Shaders_ShaderResourceMemory, GetSizeBytes());
 	}
 }
@@ -470,7 +470,7 @@ void FShaderResource::InitRHI()
 
 	if (!FPlatformProperties::HasEditorOnlyData())
 	{
-		DEC_DWORD_STAT_BY_FName(*GetMemoryStatType((EShaderFrequency)Target.Frequency), Code.Num());
+		DEC_DWORD_STAT_BY_FName(GetMemoryStatType((EShaderFrequency)Target.Frequency).GetName(), Code.Num());
 		DEC_DWORD_STAT_BY(STAT_Shaders_ShaderResourceMemory, Code.GetAllocatedSize());
 		Code.Empty();
 	}
diff --git a/Engine/Source/Runtime/Sockets/Private/BSDIPv6Sockets/SocketSubsystemBSDIPv6.cpp b/Engine/Source/Runtime/Sockets/Private/BSDIPv6Sockets/SocketSubsystemBSDIPv6.cpp
--- a/Engine/Source/Runtime/Sockets/Private/BSDIPv6Sockets/SocketSubsystemBSDIPv6.cpp
+++ b/Engine/Source/Runtime/Sockets/Private/BSDIPv6Sockets/SocketSubsystemBSDIPv6.cpp
@@ -18,7 +18,7 @@ FSocket* FSocketSubsystemBSDIPv6::CreateSocket(const FName& SocketType, const FS
 {
 	SOCKET Socket = INVALID_SOCKET;
 	FSocket* NewSocket = NULL;
-	switch (SocketType.GetIndex())
+	switch (SocketType.GetComparisonIndex())
 	{
 	case NAME_DGram:
 		// Creates a data gram (UDP) socket
diff --git a/Engine/Source/Runtime/Sockets/Private/BSDSockets/SocketSubsystemBSD.cpp b/Engine/Source/Runtime/Sockets/Private/BSDSockets/SocketSubsystemBSD.cpp
--- a/Engine/Source/Runtime/Sockets/Private/BSDSockets/SocketSubsystemBSD.cpp
+++ b/Engine/Source/Runtime/Sockets/Private/BSDSockets/SocketSubsystemBSD.cpp
@@ -18,7 +18,7 @@ FSocket* FSocketSubsystemBSD::CreateSocket(const FName& SocketType, const FStrin
 {
 	SOCKET Socket = INVALID_SOCKET;
 	FSocket* NewSocket = NULL;
-	switch (SocketType.GetIndex())
+	switch (SocketType.GetComparisonIndex())
 	{
 	case NAME_DGram:
 		// Creates a data gram (UDP) socket
diff --git a/Engine/Source/Runtime/Windows/D3D11RHI/Private/D3D11Texture.cpp b/Engine/Source/Runtime/Windows/D3D11RHI/Private/D3D11Texture.cpp
--- a/Engine/Source/Runtime/Windows/D3D11RHI/Private/D3D11Texture.cpp
+++ b/Engine/Source/Runtime/Windows/D3D11RHI/Private/D3D11Texture.cpp
@@ -85,7 +85,7 @@ void UpdateD3D11TextureStats(uint32 BindFlags, uint32 MiscFlags, int64 TextureSi
 
 	bool bCubeMap = (MiscFlags & D3D11_RESOURCE_MISC_TEXTURECUBE) != 0;
 
-	INC_MEMORY_STAT_BY_FName(*GetD3D11StatEnum(BindFlags, bCubeMap, b3D), TextureSize);
+	INC_MEMORY_STAT_BY_FName(GetD3D11StatEnum(BindFlags, bCubeMap, b3D).GetName(), TextureSize);
 
 	if(TextureSize > 0)
 	{
