Fixed incorrect value of STATID_INCREMENT

STATID_INCREMENT (formerly FNAME_INCREMENT) was still set to 3, but CL# 2300730 changed the PendingStatIds array use a single TStatIdData per-stat, rather than 3 FNames.

This wasn't causing issues, but was allocating 3x the amount of memory actually required (and then ignoring 2/3 of it).

diff --git a/Engine/Source/Runtime/Core/Private/Stats/Stats2.cpp b/Engine/Source/Runtime/Core/Private/Stats/Stats2.cpp
--- a/Engine/Source/Runtime/Core/Private/Stats/Stats2.cpp
+++ b/Engine/Source/Runtime/Core/Private/Stats/Stats2.cpp
@@ -196,14 +196,8 @@ class FStatGroupEnableManager : public IStatGroupEnableManager
 
 	enum 
 	{
-		/**
-		 *	Increment between two successive long names. 
-		 *	@see TStatId
-		 */
-		STATID_INCREMENT = 3,
-
 		/** Number of stats pointer allocated per block. */
-		NUM_PER_BLOCK = 16384 * STATID_INCREMENT,
+		NUM_PER_BLOCK = 16384,
 	};
 
 
@@ -393,7 +387,7 @@ public:
 			FMemory::Memzero( PendingStatIds, NUM_PER_BLOCK * sizeof( TStatIdData ) );
 			PendingCount = NUM_PER_BLOCK;
 		}
-		PendingCount -= STATID_INCREMENT;
+		--PendingCount;
 		TStatIdData* Result = PendingStatIds;
 
 		// Get the wide stat description.
@@ -411,7 +405,7 @@ public:
 
 		MemoryCounter.Add( StatDescLen*(sizeof( ANSICHAR ) + sizeof( WIDECHAR )) );
 		
-		PendingStatIds += STATID_INCREMENT;
+		++PendingStatIds;
 
 		if (Found->CurrentEnable)
 		{
