Reimplemented RemoveSingle and RemoveSingleSwap in terms of Find to reduce code and avoid an incorrect check.

diff --git a/Engine/Source/Runtime/Core/Public/Containers/Array.h b/Engine/Source/Runtime/Core/Public/Containers/Array.h
--- a/Engine/Source/Runtime/Core/Public/Containers/Array.h
+++ b/Engine/Source/Runtime/Core/Public/Containers/Array.h
@@ -1327,31 +1327,24 @@ public:
 	 */
 	int32 RemoveSingle( const ElementType& Item )
 	{
-		CheckAddress(&Item);
-
-		for( int32 Index=0; Index<ArrayNum; Index++ )
+		int32 Index = Find(Item);
+		if (Index == INDEX_NONE)
 		{
-			if( GetTypedData()[Index] == Item )
-			{
-				// Destruct items that match the specified Item.
-				DestructItems(GetTypedData() + Index, 1);
-				const int32 NextIndex = Index + 1;
-				if( NextIndex < ArrayNum )
-				{
-					const int32 NumElementsToMove = ArrayNum - NextIndex;
-					FMemory::Memmove(&GetTypedData()[Index],&GetTypedData()[NextIndex],sizeof(ElementType) * NumElementsToMove);
-				}
+			return 0;
+		}
 
-				// Update the array count
-				--ArrayNum;
+		auto* RemovePtr = GetTypedData() + Index;
 
-				// Removed one item
-				return 1;
-			}
-		}
+		// Destruct items that match the specified Item.
+		DestructItems(RemovePtr, 1);
+		const int32 NextIndex = Index + 1;
+		MoveConstructItems(RemovePtr, RemovePtr + 1, ArrayNum - (Index + 1));
 
-		// Specified item was not found.  Removed zero items.
-		return 0;
+		// Update the array count
+		--ArrayNum;
+
+		// Removed one item
+		return 1;
 	}
 
 	/** Removes as many instances of Item as there are in the array, maintaining order but not indices. */
@@ -1443,21 +1436,16 @@ public:
 	 */
 	int32 RemoveSingleSwap( const ElementType& Item )
 	{
-		CheckAddress(&Item);
-
-		for( int32 Index=0; Index<ArrayNum; Index++ )
+		int32 Index = Find(Item);
+		if (Index == INDEX_NONE)
 		{
-			if( (*this)[Index]==Item )
-			{
-				RemoveAtSwap(Index);
-
-				// Removed one item
-				return 1;
-			}
+			return 0;
 		}
 
-		// Specified item was not found.  Removed zero items.
-		return 0;
+		RemoveAtSwap(Index);
+
+		// Removed one item
+		return 1;
 	}
 
 	/** RemoveItemSwap, this version is much more efficient O(Count) instead of O(ArrayNum), but does not preserve the order */
