Change FMacWindow::ReshapeWindow to not assert but guard against nil WindowHandle to allow come UI actions to function in the editor.

diff --git a/Engine/Source/Runtime/Core/Private/Mac/MacWindow.cpp b/Engine/Source/Runtime/Core/Private/Mac/MacWindow.cpp
--- a/Engine/Source/Runtime/Core/Private/Mac/MacWindow.cpp
+++ b/Engine/Source/Runtime/Core/Private/Mac/MacWindow.cpp
@@ -188,39 +188,41 @@ FCocoaWindow* FMacWindow::GetWindowHandle() const
 
 void FMacWindow::ReshapeWindow( int32 X, int32 Y, int32 Width, int32 Height )
 {
-	SCOPED_AUTORELEASE_POOL;
-	check(WindowHandle);
-	
-	const TSharedRef<FGenericApplicationMessageHandler> MessageHandler = OwningApplication->MessageHandler;
-	MessageHandler->BeginReshapingWindow( SharedThis( this ) );
-
-	if(WindowMode == EWindowMode::Windowed || WindowMode == EWindowMode::WindowedFullscreen)
+	if (WindowHandle)
 	{
-		MainThreadCall(^{
-			BOOL DisplayIfNeeded = (WindowMode == EWindowMode::Windowed);
-			
-			const int32 InvertedY = FPlatformMisc::ConvertSlateYPositionToCocoa(Y) - Height + 1;
-			if (Definition->HasOSWindowBorder)
-			{
-				[WindowHandle setFrame: [WindowHandle frameRectForContentRect: NSMakeRect(X, InvertedY, FMath::Max(Width, 1), FMath::Max(Height, 1))] display:DisplayIfNeeded];
-			}
-			else
-			{
-				[WindowHandle setFrame: NSMakeRect(X, InvertedY, FMath::Max(Width, 1), FMath::Max(Height, 1)) display:DisplayIfNeeded];
-			}
-			
-			// Force resize back to screen size in fullscreen - not ideally pretty but means we don't
-			// have to subvert the OS X or UE fullscreen handling events elsewhere.
-			if(WindowMode != EWindowMode::Windowed)
-			{
-				[WindowHandle setFrame: [WindowHandle screen].frame display:YES];
-			}
-			
-			WindowHandle->bZoomed = [WindowHandle isZoomed];
-		}, UE4ResizeEventMode, true);
+		SCOPED_AUTORELEASE_POOL;
+		
+		const TSharedRef<FGenericApplicationMessageHandler> MessageHandler = OwningApplication->MessageHandler;
+		MessageHandler->BeginReshapingWindow( SharedThis( this ) );
+		
+		if(WindowMode == EWindowMode::Windowed || WindowMode == EWindowMode::WindowedFullscreen)
+		{
+			MainThreadCall(^{
+				BOOL DisplayIfNeeded = (WindowMode == EWindowMode::Windowed);
+				
+				const int32 InvertedY = FPlatformMisc::ConvertSlateYPositionToCocoa(Y) - Height + 1;
+				if (Definition->HasOSWindowBorder)
+				{
+					[WindowHandle setFrame: [WindowHandle frameRectForContentRect: NSMakeRect(X, InvertedY, FMath::Max(Width, 1), FMath::Max(Height, 1))] display:DisplayIfNeeded];
+				}
+				else
+				{
+					[WindowHandle setFrame: NSMakeRect(X, InvertedY, FMath::Max(Width, 1), FMath::Max(Height, 1)) display:DisplayIfNeeded];
+				}
+				
+				// Force resize back to screen size in fullscreen - not ideally pretty but means we don't
+				// have to subvert the OS X or UE fullscreen handling events elsewhere.
+				if(WindowMode != EWindowMode::Windowed)
+				{
+					[WindowHandle setFrame: [WindowHandle screen].frame display:YES];
+				}
+				
+				WindowHandle->bZoomed = [WindowHandle isZoomed];
+			}, UE4ResizeEventMode, true);
+		}
+		
+		MessageHandler->FinishedReshapingWindow( SharedThis( this ) );
 	}
-	
-	MessageHandler->FinishedReshapingWindow( SharedThis( this ) );
 }
 
 bool FMacWindow::GetFullScreenInfo( int32& X, int32& Y, int32& Width, int32& Height ) const
