When the Mac splash text is nil don't try and draw it or we'll crash.

diff --git a/Engine/Source/Runtime/Core/Private/Mac/MacPlatformSplash.cpp b/Engine/Source/Runtime/Core/Private/Mac/MacPlatformSplash.cpp
--- a/Engine/Source/Runtime/Core/Private/Mac/MacPlatformSplash.cpp
+++ b/Engine/Source/Runtime/Core/Private/Mac/MacPlatformSplash.cpp
@@ -152,10 +152,25 @@ static void StartSetSplashText( const SplashTextType::Type InType, const FText&
 				align = NSRightTextAlignment;
 				break;
 			}
+			
+			// Cope with the possibility that users disable or remove fonts so they cannot be used anymore.
+			NSFont* Font = [NSFont fontWithName:FontName size:FontSize];
+			if (Font == nil)
+			{
+				// Fallback to the system font for the given size, which must exist
+				Font = [NSFont systemFontOfSize:FontSize];
+				check(Font);
+				
+				// Use this font's name
+				FontName = [Font fontName];
+			}
 
 			NSString *Text = (NSString *)FPlatformString::TCHARToCFString(*SplashText.ToString());
-			[self drawText: Text inRect: TextRect withAlignment: align withColor: TextColor fontName: FontName fontSize: FontSize];
-			[Text release];
+			if ( Text )
+			{
+				[self drawText: Text inRect: TextRect withAlignment: align withColor: TextColor fontName: FontName fontSize: FontSize];
+				[Text release];
+			}
 		}
 	}
 }
