TranslationEditor: Preview in Editor: Change FTextLocalizationManager::RegenerateResources() function to not filter updates by table name when calling UpdateLiveTable() because we couldn't live preview strings that had no translation originally since their table name was blank and never matched anything

diff --git a/Engine/Source/Runtime/Core/Private/Internationalization/TextLocalizationManager.cpp b/Engine/Source/Runtime/Core/Private/Internationalization/TextLocalizationManager.cpp
--- a/Engine/Source/Runtime/Core/Private/Internationalization/TextLocalizationManager.cpp
+++ b/Engine/Source/Runtime/Core/Private/Internationalization/TextLocalizationManager.cpp
@@ -752,5 +752,6 @@ void FTextLocalizationManager::RegenerateResources(const FString& ConfigFilePath
 		MemoryReader.Close();
 	}
 
-	UpdateLiveTable(LocalizationEntryTrackers, true);
+	// Don't filter updates by table name, or we can't live preview strings that had no translation originally
+	UpdateLiveTable(LocalizationEntryTrackers);
 }
