Username switched to alpha numeric to fix "This Application" not showing up.

diff --git a/Engine/Source/Runtime/Core/Private/Misc/App.cpp b/Engine/Source/Runtime/Core/Private/Misc/App.cpp
--- a/Engine/Source/Runtime/Core/Private/Misc/App.cpp
+++ b/Engine/Source/Runtime/Core/Private/Misc/App.cpp
@@ -98,7 +98,7 @@ void FApp::InitializeSession( )
 
 	if (!FParse::Value(FCommandLine::Get(), TEXT("-SessionOwner="), SessionOwner))
 	{
-		SessionOwner = FPlatformProcess::UserName(false);
+		SessionOwner = FPlatformProcess::UserName(true);
 	}
 }
 
