Removed logging message left over from testing force feedback
#ue4
#android

diff --git a/Engine/Source/Runtime/Core/Private/Android/AndroidInputInterface.cpp b/Engine/Source/Runtime/Core/Private/Android/AndroidInputInterface.cpp
--- a/Engine/Source/Runtime/Core/Private/Android/AndroidInputInterface.cpp
+++ b/Engine/Source/Runtime/Core/Private/Android/AndroidInputInterface.cpp
@@ -139,7 +139,6 @@ void FAndroidInputInterface::UpdateVibeMotors(FAndroidControllerData &State)
 		{
 			AndroidThunkCpp_Vibrate(0);
 			State.VibeIsOn = false;
-			UE_LOG(LogTemp, Display, TEXT("UpdateVibeMotors: OFF"));
 		}
 	}
 	else {
@@ -148,7 +147,6 @@ void FAndroidInputInterface::UpdateVibeMotors(FAndroidControllerData &State)
 			// Turn it on for 10 seconds (or until below threshold)
 			AndroidThunkCpp_Vibrate(10000);
 			State.VibeIsOn = true;
-			UE_LOG(LogTemp, Display, TEXT("UpdateVibeMotors: ON"));
 		}
 	}
 }
