When using the main thread as the game thread on Mac OS X window release has to be deferred until after the next round of events to prevent crashes. Some events pointing to invalid windows may still occur but they can be safely filtered out.

diff --git a/Engine/Source/Runtime/Core/Private/Mac/CocoaThread.cpp b/Engine/Source/Runtime/Core/Private/Mac/CocoaThread.cpp
--- a/Engine/Source/Runtime/Core/Private/Mac/CocoaThread.cpp
+++ b/Engine/Source/Runtime/Core/Private/Mac/CocoaThread.cpp
@@ -205,7 +205,12 @@ void ProcessGameThreadEvents(void)
 	SCOPED_AUTORELEASE_POOL;
 	while( NSEvent *Event = [NSApp nextEventMatchingMask: NSAnyEventMask untilDate: nil inMode: NSDefaultRunLoopMode dequeue: true] )
 	{
-		[NSApp sendEvent: Event];
+		// Either the windowNumber is 0 or the window must be valid for the event to be processed.
+		// Processing events with a windowNumber but no window will crash inside sendEvent as it will try to send to a destructed window.
+		if ([Event windowNumber] == 0 || [Event window] != nil)
+		{
+			[NSApp sendEvent: Event];
+		}
 	}
 #endif
 }
diff --git a/Engine/Source/Runtime/Core/Private/Mac/MacWindow.cpp b/Engine/Source/Runtime/Core/Private/Mac/MacWindow.cpp
--- a/Engine/Source/Runtime/Core/Private/Mac/MacWindow.cpp
+++ b/Engine/Source/Runtime/Core/Private/Mac/MacWindow.cpp
@@ -14,7 +14,10 @@ FMacWindow::~FMacWindow()
 	// So instead we release the window here.
 	if(WindowHandle)
 	{
-		[WindowHandle release];
+		NSWindow* Window = WindowHandle;
+		dispatch_async(dispatch_get_main_queue(), ^{
+			[Window release];
+		});
 		WindowHandle = nil;
 	}
 }
@@ -90,6 +93,7 @@ void FMacWindow::Initialize( FMacApplication* const Application, const TSharedRe
 		
 		if( WindowHandle != nullptr )
 		{
+			[WindowHandle setReleasedWhenClosed:NO];
 			[WindowHandle setWindowMode: WindowMode];
 			[WindowHandle setAcceptsInput: Definition->AcceptsInput];
 			[WindowHandle setDisplayReconfiguring: false];
@@ -270,18 +274,19 @@ void FMacWindow::Destroy()
 		if( MacApplication->OnWindowDestroyed( Window ) )
 		{
 			// This FMacWindow may have been destructed by now & so the WindowHandle will probably be invalid memory.
+			bool const bIsKey = [Window isKeyWindow];
 			
 			// Then change the focus to something useful, either the previous in the stack
 			TSharedPtr<FMacWindow> KeyWindow = MacApplication->GetKeyWindow();
-			if( KeyWindow.IsValid() && [Window isKeyWindow] && KeyWindow->GetOSWindowHandle() && ![(FCocoaWindow*)KeyWindow->GetOSWindowHandle() isMiniaturized] )
+			if( KeyWindow.IsValid() && bIsKey && KeyWindow->GetOSWindowHandle() && ![(FCocoaWindow*)KeyWindow->GetOSWindowHandle() isMiniaturized] )
 			{
 				// Activate specified previous window if still present, provided it isn't minimized
 				KeyWindow->SetWindowFocus();
 			}
 			
+			// Close the window
 			MainThreadCall(^{
 				[Window performClose:nil];
-				WindowHandle = nil;
 			}, UE4CloseEventMode, true);
 		}
 	}
