Disable dependency on Steam controller temporarily (for 4.5).

diff --git a/Engine/Source/Runtime/Core/Core.Build.cs b/Engine/Source/Runtime/Core/Core.Build.cs
--- a/Engine/Source/Runtime/Core/Core.Build.cs
+++ b/Engine/Source/Runtime/Core/Core.Build.cs
@@ -128,6 +128,7 @@ public class Core : ModuleRules
 				AddThirdPartyPrivateStaticDependencies(Target, "OpenGL");
 			}
 
+			/* [RCL] 2014-09-18: disabling SteamController temporarily until redist/deployment problems are cleared
             // add steam controller dependency for game and client only
             if (UEBuildConfiguration.bCompileAgainstEngine == true && 
                 (Target.Type == TargetRules.TargetType.Game || Target.Type == TargetRules.TargetType.Client) &&
@@ -135,6 +136,7 @@ public class Core : ModuleRules
             {
                 AddThirdPartyPrivateStaticDependencies(Target, "SteamController");
             }
+			*/
         }
 		else if (Target.Platform == UnrealTargetPlatform.HTML5 && Target.Architecture == "-win32")
 		{
