Restore the command+key-up workaround for Mac OS X that was removed in CL #2284973. Fixes TTP #346886.

diff --git a/Engine/Source/Runtime/Core/Private/Mac/MacApplication.cpp b/Engine/Source/Runtime/Core/Private/Mac/MacApplication.cpp
--- a/Engine/Source/Runtime/Core/Private/Mac/MacApplication.cpp
+++ b/Engine/Source/Runtime/Core/Private/Mac/MacApplication.cpp
@@ -20,7 +20,7 @@ static void HandleNSEvent(NSEvent* Event)
 	const bool bIsMouseClickOrKeyEvent = [Event type] == NSLeftMouseDown || [Event type] == NSLeftMouseUp
 									  || [Event type] == NSRightMouseDown || [Event type] == NSRightMouseUp
 									  || [Event type] == NSOtherMouseDown || [Event type] == NSOtherMouseUp
-									  || [Event type] == NSKeyDown || [Event type] == NSKeyUp;
+									  || [Event type] == NSKeyDown || ([Event type] == NSKeyUp && !([Event modifierFlags] & NSCommandKeyMask));
 
 	if (MacApplication)
 	{
