Core: Moved message log related files into their own subfolder

diff --git a/Engine/Source/Runtime/Core/Core.Build.cs b/Engine/Source/Runtime/Core/Core.Build.cs
--- a/Engine/Source/Runtime/Core/Core.Build.cs
+++ b/Engine/Source/Runtime/Core/Core.Build.cs
@@ -20,6 +20,7 @@ public class Core : ModuleRules
 				"Runtime/Core/Public/Delegates",
 				"Runtime/Core/Public/GenericPlatform",
 				"Runtime/Core/Public/HAL",
+				"Runtime/Core/Public/Logging",
 				"Runtime/Core/Public/Math",
 				"Runtime/Core/Public/Misc",
 				"Runtime/Core/Public/Modules",
diff --git a/Engine/Source/Runtime/Core/Private/Logging/MessageLog.cpp b/Engine/Source/Runtime/Core/Private/Logging/MessageLog.cpp
new file mode 100644
--- /dev/null
+++ b/Engine/Source/Runtime/Core/Private/Logging/MessageLog.cpp
@@ -0,0 +1,228 @@
+// Copyright 1998-2014 Epic Games, Inc. All Rights Reserved.
+
+#include "Core.h"
+#include "MessageLog.h"
+#include "IMessageLog.h"
+
+FMessageLog::FGetLog FMessageLog::GetLog;
+
+FMessageLog::FMessageSelectionChanged FMessageLog::MessageSelectionChanged;
+
+#define LOCTEXT_NAMESPACE "MessageLog"
+
+class FBasicMessageLog : public IMessageLog, public TSharedFromThis<FBasicMessageLog>
+{
+public:
+	FBasicMessageLog( const FName& InLogName )
+		: LogName( InLogName )
+	{
+	}
+
+	/** Begin IMessageLog interface */
+	virtual void AddMessage( const TSharedRef<FTokenizedMessage>& NewMessage ) override
+	{
+		AddMessageInternal(NewMessage);
+	}
+
+	virtual void AddMessages( const TArray< TSharedRef<FTokenizedMessage> >& NewMessages ) override
+	{
+		for(TArray< TSharedRef<FTokenizedMessage> >::TConstIterator It = NewMessages.CreateConstIterator(); It; It++)
+		{
+			AddMessageInternal(*It);
+		}
+	}
+
+	virtual void NewPage( const FText& Title ) override
+	{
+		FFormatNamedArguments Arguments;
+		Arguments.Add(TEXT("PageTitle"), Title);
+		FMsg::Logf(__FILE__, __LINE__, LogName, ELogVerbosity::Log, *FText::Format(LOCTEXT("BasicMessageLog_NewPage", "New Page: {PageTitle}"), Arguments).ToString());
+	}
+
+	virtual void NotifyIfAnyMessages( const FText& Message, EMessageSeverity::Type SeverityFilter, bool bForce ) override
+	{
+		FFormatNamedArguments Arguments;
+		Arguments.Add(TEXT("Message"), Message);
+		FMsg::Logf(__FILE__, __LINE__, LogName, ELogVerbosity::Log, *FText::Format(LOCTEXT("BasicMessageLog_Notify", "Notify: {Message}"), Arguments).ToString());
+	}
+
+	virtual void Open() override
+	{
+		FMsg::Logf(__FILE__, __LINE__, LogName, ELogVerbosity::Log, *LOCTEXT("BasicMessageLog_Open", "Open Log").ToString());
+	}
+
+	virtual int32 NumMessages( EMessageSeverity::Type SeverityFilter ) override
+	{
+		return 0;
+	}
+	/** End IMessageLog interface */
+
+private:
+	void AddMessageInternal( const TSharedRef<FTokenizedMessage>& Message )
+	{
+		const TCHAR* const LogColor = FMessageLog::GetLogColor(Message->GetSeverity());
+		if(LogColor)
+		{
+			SET_WARN_COLOR(LogColor);
+		}
+		FMsg::Logf(__FILE__, __LINE__, LogName, FMessageLog::GetLogVerbosity(Message->GetSeverity()), *Message->ToText().ToString());
+		CLEAR_WARN_COLOR();
+	}
+
+private:
+	/** The name of this log */
+	FName LogName;
+};
+
+FMessageLog::FMessageLog( const FName& InLogName )
+{
+	if(GetLog.IsBound())
+	{
+		MessageLog = GetLog.Execute(InLogName);
+	}
+	else
+	{
+		MessageLog = MakeShareable( new FBasicMessageLog(InLogName) );
+	}
+}
+
+FMessageLog::~FMessageLog()
+{
+	Flush();
+}
+
+const TSharedRef<FTokenizedMessage>& FMessageLog::AddMessage( const TSharedRef<FTokenizedMessage>& InMessage )
+{
+	Messages.Add(InMessage);
+
+	return InMessage;
+}
+
+void FMessageLog::AddMessages( const TArray< TSharedRef<FTokenizedMessage> >& InMessages )
+{
+	Messages.Append(InMessages);
+}
+
+TSharedRef<FTokenizedMessage> FMessageLog::Message( EMessageSeverity::Type InSeverity, const FText& InMessage )
+{
+	TSharedRef<FTokenizedMessage> Message = FTokenizedMessage::Create(InSeverity, InMessage);
+	Messages.Add(Message);
+	return Message;
+}
+
+TSharedRef<FTokenizedMessage> FMessageLog::CriticalError( const FText& InMessage )
+{
+	TSharedRef<FTokenizedMessage> Message = FTokenizedMessage::Create(EMessageSeverity::CriticalError, InMessage);
+	Messages.Add(Message);
+	return Message;
+}
+
+TSharedRef<FTokenizedMessage> FMessageLog::Error( const FText& InMessage )
+{
+	TSharedRef<FTokenizedMessage> Message = FTokenizedMessage::Create(EMessageSeverity::Error, InMessage);
+	Messages.Add(Message);
+	return Message;
+}
+
+TSharedRef<FTokenizedMessage> FMessageLog::PerformanceWarning( const FText& InMessage )
+{
+	TSharedRef<FTokenizedMessage> Message = FTokenizedMessage::Create(EMessageSeverity::PerformanceWarning, InMessage);
+#if !PLATFORM_LINUX // @todo: these are too spammy for now on Linux
+	Messages.Add(Message);
+#endif // !PLATFORM_LINUX
+	return Message;
+}
+
+TSharedRef<FTokenizedMessage> FMessageLog::Warning( const FText& InMessage )
+{
+	TSharedRef<FTokenizedMessage> Message = FTokenizedMessage::Create(EMessageSeverity::Warning, InMessage);
+#if !PLATFORM_LINUX // @todo: these are too spammy for now on Linux
+	Messages.Add(Message); // TODO These are too spammy for now
+#endif // !PLATFORM_LINUX
+	return Message;
+}
+
+TSharedRef<FTokenizedMessage> FMessageLog::Info( const FText& InMessage )
+{
+	TSharedRef<FTokenizedMessage> Message = FTokenizedMessage::Create(EMessageSeverity::Info, InMessage);
+	Messages.Add(Message);
+	return Message;
+}
+
+int32 FMessageLog::NumMessages( EMessageSeverity::Type InSeverityFilter )
+{
+	Flush();
+	return MessageLog->NumMessages(InSeverityFilter);
+}
+
+void FMessageLog::Open( EMessageSeverity::Type InSeverityFilter, bool bOpenEvenIfEmpty )
+{
+	Flush();
+	if(bOpenEvenIfEmpty)
+	{
+		MessageLog->Open();
+	}
+	else if(MessageLog->NumMessages(InSeverityFilter) > 0)
+	{
+		MessageLog->Open();
+	}
+}
+
+void FMessageLog::Notify( const FText& InMessage, EMessageSeverity::Type InSeverityFilter, bool bForce )
+{
+	Flush();
+	MessageLog->NotifyIfAnyMessages(InMessage, InSeverityFilter, bForce);
+}
+
+void FMessageLog::NewPage( const FText& InLabel )
+{
+	Flush();
+	MessageLog->NewPage(InLabel);
+}
+
+void FMessageLog::Flush()
+{
+	if(Messages.Num() > 0)
+	{
+		MessageLog->AddMessages(Messages);
+		Messages.Empty();
+	}
+}
+
+ELogVerbosity::Type FMessageLog::GetLogVerbosity( EMessageSeverity::Type InSeverity )
+{
+	switch(InSeverity)
+	{
+	case EMessageSeverity::CriticalError:
+		return ELogVerbosity::Fatal;
+	case EMessageSeverity::Error:
+		return ELogVerbosity::Error;
+	case EMessageSeverity::PerformanceWarning:
+	case EMessageSeverity::Warning:
+		return ELogVerbosity::Warning;
+	case EMessageSeverity::Info:
+	default:
+		return ELogVerbosity::Log;
+	}
+}
+
+const TCHAR* const FMessageLog::GetLogColor( EMessageSeverity::Type InSeverity )
+{
+#if !PLATFORM_DESKTOP
+	return NULL;
+#else
+	switch(InSeverity)
+	{
+	case EMessageSeverity::CriticalError:
+		return OutputDeviceColor::COLOR_RED;
+	case EMessageSeverity::PerformanceWarning:
+	case EMessageSeverity::Warning:
+		return OutputDeviceColor::COLOR_YELLOW;
+	case EMessageSeverity::Info:
+	default:
+		return NULL;
+	}
+#endif
+}
+
+#undef LOCTEXT_NAMESPACE
\ No newline at end of file
diff --git a/Engine/Source/Runtime/Core/Private/Logging/TokenizedMessage.cpp b/Engine/Source/Runtime/Core/Private/Logging/TokenizedMessage.cpp
new file mode 100644
--- /dev/null
+++ b/Engine/Source/Runtime/Core/Private/Logging/TokenizedMessage.cpp
@@ -0,0 +1,194 @@
+// Copyright 1998-2014 Epic Games, Inc. All Rights Reserved.
+
+#include "Core.h"
+#include "TokenizedMessage.h"
+
+#define LOCTEXT_NAMESPACE "Core.MessageLog"
+
+TSharedRef<FTokenizedMessage> FTokenizedMessage::Create(EMessageSeverity::Type InSeverity, const FText& InMessageText)
+{
+	TSharedRef<FTokenizedMessage> Message = MakeShareable(new FTokenizedMessage());
+	Message->SetSeverity( InSeverity );
+	Message->AddToken( FSeverityToken::Create(InSeverity) );
+	if(!InMessageText.IsEmpty())
+	{
+		Message->AddToken( FTextToken::Create(InMessageText) );
+	}
+	return Message;
+}
+
+FText FTokenizedMessage::ToText() const
+{
+	FText OutMessage;
+
+	// Dont duplicate starting severity when displaying as string - but display it if it differs (for whatever reason)
+	if(MessageTokens.Num() > 0 && MessageTokens[0]->GetType() == EMessageToken::Severity)
+	{
+		TSharedRef<FSeverityToken> SeverityToken = StaticCastSharedRef<FSeverityToken>(MessageTokens[0]);
+		if(SeverityToken->GetSeverity() != Severity)
+		{
+			OutMessage = FTokenizedMessage::GetSeverityText( Severity );
+		}
+	}
+	else
+	{
+		OutMessage = FTokenizedMessage::GetSeverityText( Severity );
+	}
+	
+	//@todo This is bad and not safe for localization, this needs to be refactored once rich text is implemented [9/24/2013 justin.sargent]
+	for(int TokenIndex = 0; TokenIndex < MessageTokens.Num(); TokenIndex++)
+	{
+		if ( !OutMessage.IsEmpty() )
+		{
+			OutMessage = FText::Format( LOCTEXT("AggregateMessageTokenFormatter", "{0} {1}"), OutMessage, MessageTokens[TokenIndex]->ToText() );
+		}
+		else
+		{
+			OutMessage = MessageTokens[TokenIndex]->ToText();
+	}
+	}
+
+	return OutMessage;
+}
+
+FText FTokenizedMessage::GetSeverityText( EMessageSeverity::Type InSeverity )
+{
+	switch (InSeverity)
+	{
+	case EMessageSeverity::CriticalError:		
+		return LOCTEXT("CritError", "Critical Error");			
+	case EMessageSeverity::Error:				
+		return LOCTEXT("Error", "Error");						
+	case EMessageSeverity::PerformanceWarning:	
+		return LOCTEXT("PerfWarning", "Performance Warning");	
+	case EMessageSeverity::Warning:				
+		return LOCTEXT("Warning", "Warning");					
+	case EMessageSeverity::Info:				
+		return LOCTEXT("Info", "Info");							
+	default:		
+		return FText::GetEmpty();					
+	}
+}
+
+FName FTokenizedMessage::GetSeverityIconName(EMessageSeverity::Type InSeverity)
+{
+	FName SeverityIconName;
+	switch (InSeverity)
+	{
+	case EMessageSeverity::CriticalError:		
+		SeverityIconName = "MessageLog.Error";		
+		break;
+	case EMessageSeverity::Error:				
+		SeverityIconName = "MessageLog.Error";		
+		break;
+	case EMessageSeverity::PerformanceWarning:	
+		SeverityIconName = "MessageLog.Warning";	
+		break;
+	case EMessageSeverity::Warning:				
+		SeverityIconName = "MessageLog.Warning";	
+		break;
+	case EMessageSeverity::Info:				
+		SeverityIconName = "MessageLog.Note";		
+		break;
+	default:		
+		/* No icon for this type */						
+		break;
+	}
+	return SeverityIconName;
+}
+
+TSharedRef<FTokenizedMessage> FTokenizedMessage::AddToken( const TSharedRef<IMessageToken>& InToken )
+{
+	MessageTokens.Add( InToken );
+	return AsShared();
+}
+
+void FTokenizedMessage::SetSeverity( const EMessageSeverity::Type InSeverity )
+{
+	Severity = InSeverity;
+}
+
+EMessageSeverity::Type FTokenizedMessage::GetSeverity() const
+{
+	return Severity;
+}
+
+TSharedRef<FTokenizedMessage> FTokenizedMessage::SetMessageData( FTokenizedMiscData* InMessageData )
+{
+	MessageData = InMessageData;
+	return AsShared();
+}
+
+FTokenizedMiscData* FTokenizedMessage::GetMessageData() const
+{
+	return MessageData;
+}
+
+const TArray<TSharedRef<IMessageToken> >& FTokenizedMessage::GetMessageTokens() const
+{
+	return MessageTokens;
+}
+
+FURLToken::FGenerateURL FURLToken::GenerateURL;
+
+void FURLToken::VisitURL(const TSharedRef<IMessageToken>& Token, FString InURL)
+{	
+	FPlatformProcess::LaunchURL(*InURL, NULL, NULL);
+}
+
+FURLToken::FURLToken( const FString& InURL, const FText& InMessage )
+{
+	if(GenerateURL.IsBound())
+	{
+		URL = GenerateURL.Execute(InURL);
+	}
+	else
+	{
+		URL = InURL;
+	}
+	
+	if ( !InMessage.IsEmpty() )
+	{
+		CachedText = InMessage;
+	}
+	else
+	{
+		CachedText = NSLOCTEXT("Core.MessageLog", "DefaultHelpURLLabel", "Help");
+	}
+
+	MessageTokenActivated = FOnMessageTokenActivated::CreateStatic(&FURLToken::VisitURL, URL);
+}
+
+FAssetNameToken::FOnGotoAsset FAssetNameToken::GotoAsset;
+
+void FAssetNameToken::FindAsset(const TSharedRef<IMessageToken>& Token, FString InAssetName)
+{
+	if(GotoAsset.IsBound())
+	{
+		GotoAsset.Execute(InAssetName);
+	}
+}
+
+FAssetNameToken::FAssetNameToken(const FString& InAssetName, const FText& InMessage)
+	: AssetName(InAssetName)
+{
+	if ( !InMessage.IsEmpty() )
+	{
+		CachedText = InMessage;
+	}
+	else
+	{
+		CachedText = FText::FromString( AssetName );
+	}
+
+	MessageTokenActivated = FOnMessageTokenActivated::CreateStatic(&FAssetNameToken::FindAsset, AssetName);
+}
+
+FDocumentationToken::FDocumentationToken( const FString& InDocumentationLink, const FString& InPreviewExcerptLink, const FString& InPreviewExcerptName )
+	: DocumentationLink(InDocumentationLink)
+	, PreviewExcerptLink(InPreviewExcerptLink)
+	, PreviewExcerptName(InPreviewExcerptName)
+{
+}
+
+#undef LOCTEXT_NAMESPACE
\ No newline at end of file
diff --git a/Engine/Source/Runtime/Core/Private/Misc/MessageLog.cpp b/Engine/Source/Runtime/Core/Private/Misc/MessageLog.cpp
deleted file mode 100644
diff --git a/Engine/Source/Runtime/Core/Private/Misc/TokenizedMessage.cpp b/Engine/Source/Runtime/Core/Private/Misc/TokenizedMessage.cpp
deleted file mode 100644
diff --git a/Engine/Source/Runtime/Core/Public/Logging/IMessageLog.h b/Engine/Source/Runtime/Core/Public/Logging/IMessageLog.h
new file mode 100644
--- /dev/null
+++ b/Engine/Source/Runtime/Core/Public/Logging/IMessageLog.h
@@ -0,0 +1,66 @@
+// Copyright 1998-2014 Epic Games, Inc. All Rights Reserved.
+
+#pragma once
+
+#include "TokenizedMessage.h"
+
+/**
+ * A message log.
+ * Message logs can be written to from any module, incorporating rich tokenized information as well
+ * as text. Messages will be displayed to the standard log and the rich MessageLogListing if it is available.
+ */
+class IMessageLog
+{
+public:
+
+	/**
+	 * Virtual destructor
+	 */
+	virtual ~IMessageLog() {}
+
+	/** 
+	 * Appends a single message
+	 * 
+	 * @param	NewMessage	The message to append
+	 */
+	virtual void AddMessage( const TSharedRef<class FTokenizedMessage>& NewMessage ) = 0;
+
+	/** 
+	 * Appends multiple messages
+	 * 
+	 * @param	NewMessages	The messages to append.
+	 */
+	virtual void AddMessages( const TArray< TSharedRef<class FTokenizedMessage> >& NewMessages ) = 0;
+
+	/** 
+	 * Adds a new page to the log.
+	 * Old pages are only kept around if they contain messages, so if the current
+	 * page is empty, this call does nothing.
+	 * @param	Title	the title of the new page
+	 */
+	virtual void NewPage( const FText& Title ) = 0;
+
+	/** 
+	 * Notify the user if there are any messages on the current page for this log.
+	 * If there are no messages present, this call does nothing.
+	 * @param	Message			The message to display in the notification.
+	 * @param	SeverityFilter	Notifications will only be displayed if there are messages present
+	 *							that are of equal or greater severity than this.
+	 * @param	bForce			Notify anyway, even if the filters gives us no messages.
+	 */
+	virtual void NotifyIfAnyMessages( const FText& Message, EMessageSeverity::Type SeverityFilter = EMessageSeverity::Info, bool bForce = false ) = 0;
+
+	/** 
+	 * Opens up the message log to this listing.
+	 */
+	virtual void Open() = 0;
+
+	/** 
+	 * Checks to see if there are any messages according to the passed-in severity
+	 *
+	 * @param	SeverityFilter	Function will only return true if all the messages
+	 *							are of equal or greater severity than this.
+	 * @return The number of messages that pass our filter
+	 */
+	virtual int32 NumMessages( EMessageSeverity::Type SeverityFilter ) = 0;
+};
\ No newline at end of file
diff --git a/Engine/Source/Runtime/Core/Public/Logging/MessageLog.h b/Engine/Source/Runtime/Core/Public/Logging/MessageLog.h
new file mode 100644
--- /dev/null
+++ b/Engine/Source/Runtime/Core/Public/Logging/MessageLog.h
@@ -0,0 +1,120 @@
+// Copyright 1998-2014 Epic Games, Inc. All Rights Reserved.
+
+#pragma once
+
+#include "TokenizedMessage.h"
+
+class FMessageLog
+{
+public:
+	/**
+	 * Constructor
+	 * @param	InLogName	the name of the log to use
+	 */
+	CORE_API FMessageLog( const FName& InLogName );
+
+	/**
+	 * Destructor
+	 * When this object goes out of scope, its buffered messages will be flushed.
+	 */
+	CORE_API ~FMessageLog();
+
+	/**
+	 * Add a message to the log.
+	 * @param	InMessage	The message to add
+	 * @return The message for chaining calls
+	 */
+	CORE_API const TSharedRef<FTokenizedMessage>& AddMessage( const TSharedRef<FTokenizedMessage>& InMessage );
+
+	/**
+	 * Add an array of messages to the message log
+	 * @param	InMessages	The messages to add
+	 */
+	CORE_API void AddMessages( const TArray< TSharedRef<FTokenizedMessage> >& InMessages );
+
+	/** 
+	 * Helper functions to allow easy adding of messages to the log. Each one will add a token
+	 * for the severity (if valid) and the textual message (if any). Messages are buffered.
+	 * @param	InSeverity	The severity of the message
+	 * @param	InMessage	The message to use
+	 * @return the message for chaining calls.
+	 */
+	CORE_API TSharedRef<FTokenizedMessage> Message( EMessageSeverity::Type InSeverity, const FText& InMessage = FText() );
+	CORE_API TSharedRef<FTokenizedMessage> CriticalError( const FText& InMessage = FText() );
+	CORE_API TSharedRef<FTokenizedMessage> Error( const FText& InMessage = FText() );
+	CORE_API TSharedRef<FTokenizedMessage> PerformanceWarning( const FText& InMessage = FText() );
+	CORE_API TSharedRef<FTokenizedMessage> Warning( const FText& InMessage = FText() );
+	CORE_API TSharedRef<FTokenizedMessage> Info( const FText& InMessage = FText() );
+
+	/**
+	 * Check whether there are any messages present for this log.
+	 * If the log is paged, this will check if there are any messages on the current page.
+	 * This call will cause a flush so that the logs state is properly reflected.
+	 * @param	InSeverityFilter	Only messages of higher severity than this filter will be considered when checking.
+	 * @return The number of messages that pass our filter
+	 */
+	CORE_API int32 NumMessages( EMessageSeverity::Type InSeverityFilter = EMessageSeverity::Info );
+
+	/**
+	 * Opens the log for display to the user given certain conditions.
+	 * This call will cause a flush so that the logs state is properly reflected.
+	 * @param	InSeverityFilter	Only messages of higher severity than this filter will be considered when checking.
+	 * @param	bOpenEvenIfEmpty	Override the filter & log status & force the log to open.
+	 */
+	CORE_API void Open( EMessageSeverity::Type InSeverityFilter = EMessageSeverity::Info, bool bOpenEvenIfEmpty = false );
+
+	/**
+	 * Notify the user with a message if there are messages present.
+	 * This call will cause a flush so that the logs state is properly reflected.
+	 * @param	InMessage			The notification message.
+	 * @param	InSeverityFilter	Only messages of higher severity than this filter will be considered when checking.
+	 * @param	bForce				Notify anyway, even if the filters gives us no messages.
+	 */
+	CORE_API void Notify( const FText& InMessage = FText(), EMessageSeverity::Type InSeverityFilter = EMessageSeverity::Info, bool bForce = false );
+
+	/**
+	 * Add a new page to the log.
+	 * Pages will not be created until messages are added to them, so we dont generate lots of empty pages.
+	 * This call will cause a flush so that the logs state is properly reflected.
+	 * @param	InLabel		The label for the page.
+	 */
+	CORE_API void NewPage( const FText& InLabel );
+
+	/**
+	 * Delegate to retrieve a log interface.
+	 * This allows systems that implement IMessageLog to receive messages.
+	 */
+	DECLARE_DELEGATE_RetVal_OneParam( TSharedRef<class IMessageLog>, FGetLog, const FName& );
+	CORE_API static FGetLog& OnGetLog() { return GetLog; }
+
+	/**
+	 * Delegate used when message selection changes.
+	 * This is used to select object in the scene according to objects referenced in message tokens.
+	 */
+	DECLARE_DELEGATE_OneParam( FMessageSelectionChanged, TArray< TSharedRef<FTokenizedMessage> >& );
+	CORE_API static FMessageSelectionChanged& OnMessageSelectionChanged() { return MessageSelectionChanged; }
+
+	/** Helper function to convert message log severity to log verbosity */
+	CORE_API static ELogVerbosity::Type GetLogVerbosity( EMessageSeverity::Type InSeverity );
+
+	/** Helper function to retrieve a log color for a specified message severity */
+	CORE_API static const TCHAR* const GetLogColor( EMessageSeverity::Type InSeverity );
+
+private:
+	/** Send the currently buffered messages to the log & clear the buffer */
+	void Flush();
+
+private:
+	/** Buffer for messages */
+	TArray< TSharedRef<FTokenizedMessage> > Messages;
+
+	/** The message log we use for output */
+	TSharedPtr<class IMessageLog> MessageLog;
+
+	/** Delegate for retrieving the message log */
+	static FGetLog GetLog;
+
+	/** Delegate for message selection */
+	static FMessageSelectionChanged MessageSelectionChanged;
+};
+
diff --git a/Engine/Source/Runtime/Core/Public/Logging/TokenizedMessage.h b/Engine/Source/Runtime/Core/Public/Logging/TokenizedMessage.h
new file mode 100644
--- /dev/null
+++ b/Engine/Source/Runtime/Core/Public/Logging/TokenizedMessage.h
@@ -0,0 +1,501 @@
+// Copyright 1998-2014 Epic Games, Inc. All Rights Reserved.
+
+#pragma once
+
+#include "SharedPointer.h"
+#include "UnrealString.h"
+
+/** The severity of the message type */
+namespace EMessageSeverity
+{
+	/** Ordered according to their severity */
+	enum Type
+	{
+		CriticalError		= 0,
+		Error				= 1,
+		PerformanceWarning	= 2,
+		Warning				= 3,
+		Info				= 4,	// Should be last
+	};
+}
+
+/** Delgate used when clocking a message token */
+DECLARE_DELEGATE_OneParam(FOnMessageTokenActivated, const TSharedRef<class IMessageToken>&);
+
+namespace EMessageToken
+{
+	enum Type
+	{
+		Action,
+		Text,
+		Image,
+		Severity,
+		Object,
+		URL,
+		AssetName,
+		Documentation,
+	};
+}
+
+/** A single message token for a FTokenizedMessage instance */
+class IMessageToken : public TSharedFromThis<IMessageToken>
+{
+public:
+
+	/**
+	 * Virtual destructor
+	 */
+	virtual ~IMessageToken() {}
+
+	/** 
+	 * Get the type of this message token
+	 * 
+	 * @returns the type of the token
+	 */
+	virtual EMessageToken::Type GetType() const = 0;
+
+	/** 
+	 * Get a string representation of this token
+	 * 
+	 * @returns a string representing this token
+	 */
+	virtual const FText& ToText() const
+	{
+		return CachedText;
+	}
+
+	/** 
+	 * Get the activated delegate associated with this token, if any
+	 * 
+	 * @returns a reference to the delegate
+	 */
+	virtual const FOnMessageTokenActivated& GetOnMessageTokenActivated() const
+	{
+		return MessageTokenActivated;
+	}
+
+	/** 
+	 * Set the activated delegate associated with this token
+	 * 
+	 * @returns a reference to this token, for chaining
+	 */
+	virtual TSharedRef<IMessageToken> OnMessageTokenActivated( FOnMessageTokenActivated InMessageTokenActivated )
+	{
+		MessageTokenActivated = InMessageTokenActivated;
+		return AsShared();
+	}
+
+protected:
+	/** A delegate for when this token is activated (e.g. hyperlink clicked) */
+	FOnMessageTokenActivated MessageTokenActivated;
+
+	/** Cached string representation of this token */
+	FText CachedText;
+};
+
+
+/** Dummy struct that other structs need to inherit from if they wish to assign to MessageData */
+struct FTokenizedMiscData {};
+
+/** This class represents a rich tokenized message, such as would be used for compiler output with 'hyperlinks' to source file locations */
+class FTokenizedMessage : public TSharedFromThis<FTokenizedMessage>
+{
+public:
+
+	/** 
+	 * Creates a new FTokenizedMessage
+	 * 
+	 * @param	InSeverity			The severity of the message.
+	 * @param	InMessageString		The string to display for this message. If this is not empty, then a string token will be added to this message.
+	 * @returns	the generated message for temporary storage or so other tokens can be added to it if required.
+	 */
+	CORE_API static TSharedRef<FTokenizedMessage> Create(EMessageSeverity::Type InSeverity, const FText& InMessageText = FText());
+
+	/** 
+	 * Get this tokenized message as a string
+	 * 
+	 * @returns a string representation of this message
+	 */
+	CORE_API FText ToText() const;
+	
+	/** 
+	 * Adds a token to a message.
+	 * @param	InMessage	The message to insert a token into
+	 * @param	InToken		The token to insert
+	 * @returns this message, for chaining calls.
+	 */
+	CORE_API TSharedRef<FTokenizedMessage> AddToken( const TSharedRef<IMessageToken>& InToken );
+
+	/** 
+	 * Sets the severity of this message
+	 * 
+	 * @param	InSeverity	The severity to set this message to
+
+	 */
+	CORE_API void SetSeverity( const EMessageSeverity::Type InSeverity );
+
+	/** 
+	 * Gets the severity of this message
+	 * 
+	 * @returns the severity of this message
+	 */
+	CORE_API EMessageSeverity::Type GetSeverity() const;
+
+	/** 
+	 * Sets the custom message data of this message
+	 * 
+	 * @param	InMessageData	The custom message data to set
+	 * @returns this message, so calls can be chained together
+	 */
+	CORE_API TSharedRef<FTokenizedMessage> SetMessageData( FTokenizedMiscData* InMessageData );
+
+	/** 
+	 * Gets the custom message data of this message
+	 * 
+	 * @returns the custom message data contained in this message
+	 */
+	CORE_API FTokenizedMiscData* GetMessageData() const;
+
+	/**
+	 * Get the tokens in this message
+	 * 
+	 * @returns an array of tokens
+	 */
+	CORE_API const TArray< TSharedRef<IMessageToken> >& GetMessageTokens() const;
+
+	/** 
+	 * Helper function for getting a severity as text
+	 *
+	 * @param	InSeverity	the severity to use
+	 * @returns a string representation of the severity
+	 */
+	CORE_API static FText GetSeverityText(EMessageSeverity::Type InSeverity);
+
+	/** 
+	 * Helper function for getting a severity as an icon name
+	 *
+	 * @param	InSeverity	the severity to use
+	 * @returns the name of the icon for the severity
+	 */
+	CORE_API static FName GetSeverityIconName(EMessageSeverity::Type InSeverity);
+
+private:
+	/** Private constructor - we want to only create these structures as shared references via Create() */
+	FTokenizedMessage()
+		: Severity( EMessageSeverity::Info )
+		, MessageData( nullptr )
+	{ }
+
+protected:
+
+	/** the array of message tokens this message contains */
+	TArray< TSharedRef<IMessageToken> > MessageTokens;
+
+private:
+
+	/** The severity of this message */
+	EMessageSeverity::Type Severity;
+
+	/** A payload for this message */
+	FTokenizedMiscData* MessageData;	// @todo Don't use TSharedPtr as it'll get free'd mid-sort! (refcount == 0)
+};
+
+/** Basic message token with a localized text payload */
+class FTextToken : public IMessageToken
+{
+public:
+	/** Factory method, tokens can only be constructed as shared refs */
+	CORE_API static TSharedRef<FTextToken> Create( const FText& InMessage )
+	{
+		return MakeShareable(new FTextToken(InMessage));
+	}
+
+	/** Begin IMessageToken interface */
+	virtual EMessageToken::Type GetType() const override
+	{
+		return EMessageToken::Text;
+	}
+	/** End IMessageToken interface */
+
+private:
+	/** Private constructor */
+	FTextToken( const FText& InMessage )
+	{
+		CachedText = InMessage;
+	}
+};
+
+/** Basic message token with an icon/image payload */
+class FImageToken : public IMessageToken
+{
+public:
+	/** Factory method, tokens can only be constructed as shared refs */
+	CORE_API static TSharedRef<FImageToken> Create( const FName& InImageName )
+	{
+		return MakeShareable(new FImageToken(InImageName));
+	}
+
+	/** Begin IMessageToken interface */
+	virtual EMessageToken::Type GetType() const override
+	{
+		return EMessageToken::Image;
+	}
+	/** End IMessageToken interface */
+
+	/** Get the name of the image for this token */
+	const FName& GetImageName() const
+	{
+		return ImageName;
+	}
+
+private:
+	/** Private constructor */
+	FImageToken( const FName& InImageName )
+		: ImageName(InImageName)
+	{
+		CachedText = FText::FromName( InImageName );
+	}
+
+	/** A name to be used as a brush in this message */
+	FName ImageName;
+};
+
+/** Basic message token with a severity payload */
+class FSeverityToken : public IMessageToken
+{
+public:
+	/** Factory method, tokens can only be constructed as shared refs */
+	CORE_API static TSharedRef<FSeverityToken> Create( EMessageSeverity::Type InSeverity )
+	{
+		return MakeShareable(new FSeverityToken(InSeverity));
+	}
+
+	/** Begin IMessageToken interface */
+	virtual EMessageToken::Type GetType() const override
+	{
+		return EMessageToken::Severity;
+	}
+	/** End IMessageToken interface */
+
+	/** Get the severity of this token */
+	EMessageSeverity::Type GetSeverity() const
+	{
+		return Severity;
+	}
+
+private:
+	/** Private constructor */
+	FSeverityToken( EMessageSeverity::Type InSeverity )
+		: Severity(InSeverity)
+	{
+		CachedText = FTokenizedMessage::GetSeverityText( InSeverity );
+	}
+
+	/** A severity for this token */
+	EMessageSeverity::Type Severity;
+};
+
+/** Basic message token that defaults is activated method to traverse a URL */
+class FURLToken : public IMessageToken
+{
+public:
+	/** Factory method, tokens can only be constructed as shared refs */
+	CORE_API static TSharedRef<FURLToken> Create( const FString& InURL, const FText& InMessage = FText() )
+	{
+		return MakeShareable(new FURLToken(InURL, InMessage));
+	}
+
+	/** Begin IMessageToken interface */
+	virtual EMessageToken::Type GetType() const override
+	{
+		return EMessageToken::URL;
+	}
+	/** End IMessageToken interface */
+
+	/** Get the URL used by this token */
+	const FString& GetURL() const
+	{
+		return URL;
+	}
+
+	DECLARE_DELEGATE_RetVal_OneParam(FString, FGenerateURL, const FString&);
+	CORE_API static FGenerateURL& OnGenerateURL()
+	{
+		return GenerateURL;
+	}
+
+private:
+	/** Private constructor */
+	FURLToken( const FString& InURL, const FText& InMessage );
+
+	/**
+	 * Delegate used to visit a URL
+	 * @param	Token		The token that was clicked
+	 * @param	InURL		The URL to visit
+	 */
+	static void VisitURL(const TSharedRef<IMessageToken>& Token, FString InURL);
+
+	/** The URL we will follow */
+	FString URL;
+
+	/** The delegate we will use to generate our URL */
+	static FGenerateURL GenerateURL;
+};
+
+/** 
+ * Basic message token that defaults its activated method to find a file
+ * Intended to hook into things like the content browser.
+ */
+class FAssetNameToken : public IMessageToken
+{
+public:
+	/** Factory method, tokens can only be constructed as shared refs */
+	CORE_API static TSharedRef<FAssetNameToken> Create( const FString& InAssetName, const FText& InMessage = FText() )
+	{
+		return MakeShareable(new FAssetNameToken(InAssetName, InMessage));
+	}
+
+	/** Begin IMessageToken interface */
+	virtual EMessageToken::Type GetType() const override
+	{
+		return EMessageToken::AssetName;
+	}
+	/** End IMessageToken interface */
+
+	/** Get the filename used by this token */
+	const FString& GetAssetName() const
+	{
+		return AssetName;
+	}
+
+	DECLARE_DELEGATE_OneParam(FOnGotoAsset, const FString&);
+	CORE_API static FOnGotoAsset& OnGotoAsset()
+	{
+		return GotoAsset;
+	}
+
+private:
+	/** Private constructor */
+	FAssetNameToken( const FString& InAssetName, const FText& InMessage );
+
+	/**
+	 * Delegate used to find a file
+	 * @param	Token		The token that was clicked
+	 * @param	InURL		The file to find
+	 */
+	static void FindAsset(const TSharedRef<IMessageToken>& Token, FString InAssetName);
+
+	/** The asset name we will find */
+	FString AssetName;
+
+	/** The delegate we will use to go to our file */
+	static FOnGotoAsset GotoAsset;
+};
+
+/** 
+ * Basic message token that defaults is activated method to access UDN documentation.
+ */
+class FDocumentationToken : public IMessageToken
+{
+public:
+	/** Factory method, tokens can only be constructed as shared refs */
+	CORE_API static TSharedRef<FDocumentationToken> Create( const FString& InDocumentationLink, const FString& InPreviewExcerptLink = FString(), const FString& InPreviewExcerptName = FString() )
+	{
+		return MakeShareable(new FDocumentationToken(InDocumentationLink, InPreviewExcerptLink, InPreviewExcerptName));
+	}
+
+	/** Begin IMessageToken interface */
+	virtual EMessageToken::Type GetType() const override
+	{
+		return EMessageToken::Documentation;
+	}
+	/** End IMessageToken interface */
+
+	/** Get the documentation link used by this token */
+	const FString& GetDocumentationLink() const
+	{
+		return DocumentationLink;
+	}
+
+	/** Get the documentation excerpt link used by this token */
+	const FString& GetPreviewExcerptLink() const
+	{
+		return PreviewExcerptLink;
+	}
+
+	/** Get the documentation excerpt name used by this token */
+	const FString& GetPreviewExcerptName() const
+	{
+		return PreviewExcerptName;
+	}
+
+protected:
+	/** Protected constructor */
+	FDocumentationToken( const FString& InDocumentationLink, const FString& InPreviewExcerptLink, const FString& InPreviewExcerptName );
+
+private:
+	/** The documentation path we link to when clicked */
+	FString DocumentationLink;
+
+	/** The link we display an excerpt from */
+	FString PreviewExcerptLink;
+
+	/** The excerpt to display */
+	FString PreviewExcerptName;
+};
+
+
+DECLARE_DELEGATE(FOnActionTokenExecuted);
+
+/**
+ * Message token that performs an action when activated.
+ */
+class FActionToken
+	: public IMessageToken
+{
+	/** Factory method, tokens can only be constructed as shared refs */
+	CORE_API static TSharedRef<FActionToken> Create( const FText& ActionName, const FText& ActionDescription, const FOnActionTokenExecuted& Action )
+	{
+		return MakeShareable(new FActionToken(ActionName, ActionDescription, Action));
+	}
+
+public:
+
+	/** Executes the assigned action delegate. */
+	void ExecuteAction( )
+	{
+		Action.ExecuteIfBound();
+	}
+
+	/** Gets the action's description text. */
+	const FText& GetActionDescription( )
+	{
+		return ActionDescription;
+	}
+
+public:
+
+	// IMessageToken interface
+	virtual EMessageToken::Type GetType() const override
+	{
+		return EMessageToken::Action;
+	}
+
+protected:
+
+	/** Hidden constructor. */
+	FActionToken( const FText& InActionName, const FText& InActionDescription, const FOnActionTokenExecuted& InAction )
+		: Action(InAction)
+		, ActionDescription(InActionDescription)
+	{
+		CachedText = InActionName;
+	}
+
+private:
+
+	/** Holds a delegate that is executed when this token is activated. */
+	FOnActionTokenExecuted Action;
+
+	/** The action's description text. */
+	const FText ActionDescription;
+};
\ No newline at end of file
diff --git a/Engine/Source/Runtime/Core/Public/Misc/IMessageLog.h b/Engine/Source/Runtime/Core/Public/Misc/IMessageLog.h
deleted file mode 100644
diff --git a/Engine/Source/Runtime/Core/Public/Misc/MessageLog.h b/Engine/Source/Runtime/Core/Public/Misc/MessageLog.h
deleted file mode 100644
diff --git a/Engine/Source/Runtime/Core/Public/Misc/TokenizedMessage.h b/Engine/Source/Runtime/Core/Public/Misc/TokenizedMessage.h
deleted file mode 100644
