Remove stale comments on GetTypeHash for FVectors.

diff --git a/Engine/Source/Runtime/Core/Public/Math/Vector.h b/Engine/Source/Runtime/Core/Public/Math/Vector.h
--- a/Engine/Source/Runtime/Core/Public/Math/Vector.h
+++ b/Engine/Source/Runtime/Core/Public/Math/Vector.h
@@ -848,9 +848,7 @@ FORCEINLINE FVector operator*( float Scale, const FVector& V )
 }
 
 /**
- * Creates a hash value from a FVector. Uses pointers to the elements to
- * bypass any type conversion. This is a simple hash that just ORs the
- * raw 32bit data together
+ * Creates a hash value from a FVector. 
  *
  * @param Vector the vector to create a hash value for
  *
@@ -859,13 +857,11 @@ FORCEINLINE FVector operator*( float Scale, const FVector& V )
 FORCEINLINE uint32 GetTypeHash(const FVector& Vector)
 {
 	// Note: this assumes there's no padding in FVector that could contain uncompared data.
-	return FCrc::MemCrc_DEPRECATED(&Vector,sizeof(FVector));
+	return FCrc::MemCrc_DEPRECATED(&Vector,sizeof(Vector));
 }
 
 /**
- * Creates a hash value from a FVector2D. Uses pointers to the elements to
- * bypass any type conversion. This is a simple hash that just ORs the
- * raw 32bit data together
+ * Creates a hash value from a FVector2D. 
  *
  * @param Vector the vector to create a hash value for
  *
@@ -874,7 +870,7 @@ FORCEINLINE uint32 GetTypeHash(const FVector& Vector)
 FORCEINLINE uint32 GetTypeHash(const FVector2D& Vector)
 {
 	// Note: this assumes there's no padding in FVector2D that could contain uncompared data.
-	return FCrc::MemCrc_DEPRECATED(&Vector,sizeof(FVector2D));
+	return FCrc::MemCrc_DEPRECATED(&Vector,sizeof(Vector));
 }
 
 
