Implement GetTypeHash(FVector4) for use in hashing functiond (and GetTypeHash(FLineaerColor) to get around ambiguities that introduced).
* Doing so exposed an ambiguity in GetTypeHash(FLinearColor) because previously it was relying on an implicit conversion to FColor. Once an overload for FVector4 was provided, it became ambiguous because there is ALSO an implicit conversion to FVector4!
* Resolved this by implementing GetTypeHash(FLinearColor) directly instead of relying on implicit conversion to FColor.
* This exposed a situation where SDataGraph inlined a local definition of GetTypehash(FLinearColor), and was now in conflict with the real, global one.
  * This was a nono in the first place because depending on how UBT decided to combine cpp files, this could CHANGE THE MEANING of GetTypeHash(FLinearColor)!!!
  * Turns out that the function was completely unused in that file anyway, so never needed to be added in the first place.Thus, I removed it.
  * Arguably, the implementation in SDataGraph is better, but I left the old way so as not to change the meaning of any existing serialized hashes of colors that might be in data files.

diff --git a/Engine/Source/Developer/Profiler/Private/Widgets/SDataGraph.cpp b/Engine/Source/Developer/Profiler/Private/Widgets/SDataGraph.cpp
--- a/Engine/Source/Developer/Profiler/Private/Widgets/SDataGraph.cpp
+++ b/Engine/Source/Developer/Profiler/Private/Widgets/SDataGraph.cpp
@@ -4,14 +4,6 @@
 
 #define LOCTEXT_NAMESPACE "SDataGraph"
 
-/**
- * @return hash value for the specified linear color
- */
-FORCEINLINE uint32 GetTypeHash( const FLinearColor& LinearColor )
-{
-	return LinearColor.ToFColor(true).DWColor();
-}
-
 /*-----------------------------------------------------------------------------
 	STrackedStatSummary/SDataGraphSummary
 -----------------------------------------------------------------------------*/
diff --git a/Engine/Source/Runtime/Core/Public/Math/Color.h b/Engine/Source/Runtime/Core/Public/Math/Color.h
--- a/Engine/Source/Runtime/Core/Public/Math/Color.h
+++ b/Engine/Source/Runtime/Core/Public/Math/Color.h
@@ -514,6 +514,12 @@ FORCEINLINE uint32 GetTypeHash( const FColor& Color )
 	return Color.DWColor();
 }
 
+FORCEINLINE uint32 GetTypeHash( const FLinearColor& Color )
+{
+	return GetTypeHash(FColor(Color));
+}
+
+
 /** Computes a brightness and a fixed point color from a floating point color. */
 extern CORE_API void ComputeAndFixedColorAndIntensity(const FLinearColor& InLinearColor,FColor& OutColor,float& OutIntensity);
 
diff --git a/Engine/Source/Runtime/Core/Public/Math/Vector4.h b/Engine/Source/Runtime/Core/Public/Math/Vector4.h
--- a/Engine/Source/Runtime/Core/Public/Math/Vector4.h
+++ b/Engine/Source/Runtime/Core/Public/Math/Vector4.h
@@ -371,6 +371,19 @@ public:
 
 } GCC_ALIGN(16);
 
+/**
+ * Creates a hash value from a FVector4.
+ *
+ * @param Vector the vector to create a hash value for
+ *
+ * @return The hash value from the components
+ */
+FORCEINLINE uint32 GetTypeHash(const FVector4& Vector)
+{
+	// Note: this assumes there's no padding in FVector that could contain uncompared data.
+	return FCrc::MemCrc_DEPRECATED(&Vector,sizeof(Vector));
+}
+
 
 /* FVector4 inline functions
  *****************************************************************************/
