Fixed engine not compiling when USE_MALLOC_PROFILER was defined to 1

diff --git a/Engine/Source/Runtime/Core/Private/ProfilingDebugging/MallocProfiler.cpp b/Engine/Source/Runtime/Core/Private/ProfilingDebugging/MallocProfiler.cpp
--- a/Engine/Source/Runtime/Core/Private/ProfilingDebugging/MallocProfiler.cpp
+++ b/Engine/Source/Runtime/Core/Private/ProfilingDebugging/MallocProfiler.cpp
@@ -12,7 +12,6 @@
 #include "ModuleManager.h"
 #include "MemoryMisc.h"
 
-DEFINE_LOG_CATEGORY(LogProfilingDebugging);
 
 CORE_API FMallocProfiler* GMallocProfiler;
 
