Fixing COLOR_RED collision.

- Taken from PR #452 by salamanderrake & amigo.

diff --git a/Engine/Source/Runtime/Core/Private/Linux/LinuxPlatformOutputDevices.cpp b/Engine/Source/Runtime/Core/Private/Linux/LinuxPlatformOutputDevices.cpp
--- a/Engine/Source/Runtime/Core/Private/Linux/LinuxPlatformOutputDevices.cpp
+++ b/Engine/Source/Runtime/Core/Private/Linux/LinuxPlatformOutputDevices.cpp
@@ -12,11 +12,11 @@
 #include "LinuxPlatformOutputDevicesPrivate.h"
 #include "LinuxPlatformFeedbackContextPrivate.h"
 
-#define COLOR_RED     "\x1b[31m"
-#define COLOR_GREEN   "\x1b[32m"
-#define COLOR_YELLOW  "\x1b[33m"
-#define COLOR_BLUE    "\x1b[34m"
-#define COLOR_NONE   "\x1b[0m"
+#define CONSOLE_RED		"\x1b[31m"
+#define CONSOLE_GREEN	"\x1b[32m"
+#define CONSOLE_YELLOW	"\x1b[33m"
+#define CONSOLE_BLUE	"\x1b[34m"
+#define CONSOLE_NONE	"\x1b[0m"
 
 //////////////////////////////////
 // FLinuxOutputDevices
@@ -196,12 +196,12 @@ void FOutputDeviceConsoleLinux::Serialize(const TCHAR* Data, ELogVerbosity::Type
 			{
 				if (Verbosity == ELogVerbosity::Error)
 				{
-					printf(COLOR_RED);
+					printf(CONSOLE_RED);
 					bNeedToResetColor = true;
 				}
 				else if (Verbosity == ELogVerbosity::Warning)
 				{
-					printf(COLOR_YELLOW);
+					printf(CONSOLE_YELLOW);
 					bNeedToResetColor = true;
 				}
 			}
@@ -211,7 +211,7 @@ void FOutputDeviceConsoleLinux::Serialize(const TCHAR* Data, ELogVerbosity::Type
 
 			if (bNeedToResetColor)
 			{
-				printf(COLOR_NONE);
+				printf(CONSOLE_NONE);
 			}
 		}
 	}
