implementing missing platformprocess functionality

diff --git a/Engine/Source/Runtime/Core/Private/Linux/LinuxPlatformProcess.cpp b/Engine/Source/Runtime/Core/Private/Linux/LinuxPlatformProcess.cpp
--- a/Engine/Source/Runtime/Core/Private/Linux/LinuxPlatformProcess.cpp
+++ b/Engine/Source/Runtime/Core/Private/Linux/LinuxPlatformProcess.cpp
@@ -242,6 +242,20 @@ const TCHAR* FLinuxPlatformProcess::ExecutableName(bool bRemoveExtension)
 	return CachedResult;
 }
 
+
+FString FLinuxPlatformProcess::GenerateApplicationPath( const FString& AppName, EBuildConfigurations::Type BuildConfiguration)
+{
+	FString PlatformName = GetBinariesSubdirectory();
+	FString ExecutablePath = FString::Printf(TEXT("../%s/%s"), *PlatformName, *AppName);
+	
+	if (BuildConfiguration != EBuildConfigurations::Development && BuildConfiguration != EBuildConfigurations::DebugGame)
+	{
+		ExecutablePath += FString::Printf(TEXT("-%s-%s"), *PlatformName, EBuildConfigurations::ToString(BuildConfiguration));
+	}
+	return ExecutablePath;
+}
+
+
 FString FLinuxPlatformProcess::GetApplicationName( uint32 ProcessId )
 {
 	FString Output = TEXT("");
diff --git a/Engine/Source/Runtime/Core/Public/Linux/LinuxPlatformProcess.h b/Engine/Source/Runtime/Core/Public/Linux/LinuxPlatformProcess.h
--- a/Engine/Source/Runtime/Core/Public/Linux/LinuxPlatformProcess.h
+++ b/Engine/Source/Runtime/Core/Public/Linux/LinuxPlatformProcess.h
@@ -141,6 +141,7 @@ struct CORE_API FLinuxPlatformProcess : public FGenericPlatformProcess
 	static const TCHAR* UserDir();
 	static const TCHAR* UserSettingsDir();
 	static const TCHAR* ApplicationSettingsDir();
+	static FString GenerateApplicationPath( const FString& AppName, EBuildConfigurations::Type BuildConfiguration);
 	static FString GetApplicationName( uint32 ProcessId );
 	static bool SetProcessLimits(EProcessResource::Type Resource, uint64 Limit);
 	static const TCHAR* ExecutableName(bool bRemoveExtension = true);
