Invert keydown/keyup message handling order on Mac OS X so that Slate gets the chance to process the message before the OS. This ensures that we always get the proper key-down/key-up pairs so there's no odd behaviour in viewports & also prevents the possibility of activating menu key-equivalents when setting up new key bindings in the preferences. In order to highlight the appropriate menu we customise NSMenu as FCocoaMenu and FMacMenu now inherits from that, this provides the -[highlightKeyEquivalent:(NSEvent*)] method which performs all the steps to highlight the menu item, but doesn't perform the action (as it already has been by Slate).

diff --git a/Engine/Source/Runtime/Core/Private/Mac/CocoaMenu.cpp b/Engine/Source/Runtime/Core/Private/Mac/CocoaMenu.cpp
new file mode 100644
--- /dev/null
+++ b/Engine/Source/Runtime/Core/Private/Mac/CocoaMenu.cpp
@@ -0,0 +1,40 @@
+// Copyright 1998-2014 Epic Games, Inc. All Rights Reserved.
+
+#include "Core.h"
+#include "CocoaMenu.h"
+#include "MacApplication.h"
+
+@implementation FCocoaMenu
+
+- (id)initWithTitle:(NSString *)Title
+{
+	id Self = [super initWithTitle:Title];
+	if (Self)
+	{
+		bHighlightingKeyEquivalent = false;
+	}
+	return Self;
+}
+
+- (bool)isHighlightingKeyEquivalent
+{
+	FCocoaMenu* SuperMenu = [[self supermenu] isKindOfClass:[FCocoaMenu class]] ? (FCocoaMenu*)[self supermenu] : nil;
+	if ( SuperMenu )
+	{
+		return [SuperMenu isHighlightingKeyEquivalent];
+	}
+	else
+	{
+		return bHighlightingKeyEquivalent;
+	}
+}
+
+- (bool)highlightKeyEquivalent:(NSEvent *)Event
+{
+	bHighlightingKeyEquivalent = true;
+	bool bHighlighted = [super performKeyEquivalent:Event];
+	bHighlightingKeyEquivalent = false;
+	return bHighlighted;
+}
+
+@end
diff --git a/Engine/Source/Runtime/Core/Private/Mac/CocoaWindow.cpp b/Engine/Source/Runtime/Core/Private/Mac/CocoaWindow.cpp
--- a/Engine/Source/Runtime/Core/Private/Mac/CocoaWindow.cpp
+++ b/Engine/Source/Runtime/Core/Private/Mac/CocoaWindow.cpp
@@ -310,18 +310,12 @@ NSString* NSPerformDragOperation = @"NSPerformDragOperation";
 
 - (void)keyDown:(NSEvent *)Event
 {
-	if(self.bForwardEvents)
-	{
-		FMacEvent::SendToGameRunLoop(Event, EMacEventSendMethod::Async);
-	}
+	// @note Deliberately empty - we don't want OS X to handle keyboard input as it will recursively re-add events we aren't handling
 }
 
 - (void)keyUp:(NSEvent *)Event
 {
-	if(self.bForwardEvents)
-	{
-		FMacEvent::SendToGameRunLoop(Event, EMacEventSendMethod::Async);
-	}
+	// @note Deliberately empty - we don't want OS X to handle keyboard input as it will recursively re-add events we aren't handling
 }
 
 - (void)windowWillEnterFullScreen:(NSNotification *)notification
diff --git a/Engine/Source/Runtime/Core/Private/Mac/MacApplication.cpp b/Engine/Source/Runtime/Core/Private/Mac/MacApplication.cpp
--- a/Engine/Source/Runtime/Core/Private/Mac/MacApplication.cpp
+++ b/Engine/Source/Runtime/Core/Private/Mac/MacApplication.cpp
@@ -6,6 +6,7 @@
 #include "MacWindow.h"
 #include "MacCursor.h"
 #include "MacEvent.h"
+#include "CocoaMenu.h"
 #include "GenericApplicationMessageHandler.h"
 #include "HIDInputInterface.h"
 #include "AnalyticsEventAttribute.h"
@@ -15,28 +16,6 @@
 
 static NSString* NSWindowDraggingFinished = @"NSWindowDraggingFinished";
 
-static void HandleNSEvent(NSEvent* Event)
-{
-	const bool bIsMouseClickOrKeyEvent = [Event type] == NSLeftMouseDown || [Event type] == NSLeftMouseUp
-									  || [Event type] == NSRightMouseDown || [Event type] == NSRightMouseUp
-									  || [Event type] == NSOtherMouseDown || [Event type] == NSOtherMouseUp
-									  || [Event type] == NSKeyDown || ([Event type] == NSKeyUp && !([Event modifierFlags] & NSCommandKeyMask));
-
-	if (MacApplication)
-	{
-		if (!bIsMouseClickOrKeyEvent || [Event window] == NULL)
-		{
-			FMacEvent::SendToGameRunLoop(Event, EMacEventSendMethod::Async);
-		}
-
-		if ([Event type] == NSLeftMouseUp)
-		{
-			NSNotification* Notification = [NSNotification notificationWithName:NSWindowDraggingFinished object:[Event window]];
-			FMacEvent::SendToGameRunLoop(Notification, [Event window], EMacEventSendMethod::Async);
-		}
-	}
-}
-
 FMacApplication* MacApplication = NULL;
 
 #if WITH_EDITOR
@@ -63,6 +42,41 @@ FMacApplication* FMacApplication::CreateMacApplication()
 	return MacApplication;
 }
 
+NSEvent* FMacApplication::HandleNSEvent(NSEvent* Event)
+{
+	NSEvent* ReturnEvent = Event;
+	const bool bIsMouseClickOrKeyEvent = [Event type] == NSLeftMouseDown || [Event type] == NSLeftMouseUp
+		|| [Event type] == NSRightMouseDown || [Event type] == NSRightMouseUp
+		|| [Event type] == NSOtherMouseDown || [Event type] == NSOtherMouseUp;
+	const bool bIsResentEvent = [Event type] == NSApplicationDefined && [Event subtype] == FMacApplication::ResentEvent;
+	
+	if (MacApplication)
+	{
+		if ( bIsResentEvent )
+		{
+			ReturnEvent = (NSEvent*)[Event data1];
+		}
+		
+		if ( !bIsResentEvent && ( !bIsMouseClickOrKeyEvent || [Event window] == NULL ) )
+		{
+			FMacEvent::SendToGameRunLoop(Event, EMacEventSendMethod::Async);
+			
+			if ( [Event type] == NSKeyDown || [Event type] == NSKeyUp )
+			{
+				ReturnEvent = nil;
+			}
+		}
+		
+		if ([Event type] == NSLeftMouseUp)
+		{
+			NSNotification* Notification = [NSNotification notificationWithName:NSWindowDraggingFinished object:[Event window]];
+			FMacEvent::SendToGameRunLoop(Notification, [Event window], EMacEventSendMethod::Async);
+		}
+	}
+	
+	return ReturnEvent;
+}
+
 void FMacApplication::OnDisplayReconfiguration(CGDirectDisplayID Display, CGDisplayChangeSummaryFlags Flags, void* UserInfo)
 {
 	FMacApplication* App = (FMacApplication*)UserInfo;
@@ -102,8 +116,8 @@ FMacApplication::FMacApplication()
 	
 	EventMonitor = (void*)[NSEvent addLocalMonitorForEventsMatchingMask:NSAnyEventMask handler:^(NSEvent* IncomingEvent)
 	{
-		HandleNSEvent(IncomingEvent);
-		return IncomingEvent;
+		NSEvent* ReturnEvent = HandleNSEvent(IncomingEvent);
+		return ReturnEvent;
 	}];
 
 #if WITH_EDITOR
@@ -713,9 +727,11 @@ void FMacApplication::ProcessNSEvent(NSEvent* const Event, TSharedPtr< FMacWindo
 		case NSKeyDown:
 		{
 			NSString *Characters = [Event characters];
+			bool bHandled = false;
 			if( [Characters length] && CurrentEventWindow.IsValid() )
 			{
-				if(!CurrentEventWindow->OnIMKKeyDown(Event))
+				bHandled = CurrentEventWindow->OnIMKKeyDown(Event);
+				if(!bHandled)
 				{
 					const bool IsRepeat = [Event isARepeat];
 					const TCHAR Character = ConvertChar( [Characters characterAtIndex:0] );
@@ -723,7 +739,7 @@ void FMacApplication::ProcessNSEvent(NSEvent* const Event, TSharedPtr< FMacWindo
 					const uint32 KeyCode = [Event keyCode];
 					const bool IsPrintable = IsPrintableKey( Character );
 					
-					MessageHandler->OnKeyDown( KeyCode, TranslateCharCode( CharCode, KeyCode ), IsRepeat );
+					bHandled = MessageHandler->OnKeyDown( KeyCode, TranslateCharCode( CharCode, KeyCode ), IsRepeat );
 					
 					// First KeyDown, then KeyChar. This is important, as in-game console ignores first character otherwise
 					
@@ -734,12 +750,25 @@ void FMacApplication::ProcessNSEvent(NSEvent* const Event, TSharedPtr< FMacWindo
 					}
 				}
 			}
+			if (bHandled)
+			{
+				FCocoaMenu* MainMenu = [[NSApp mainMenu] isKindOfClass:[FCocoaMenu class]] ? (FCocoaMenu*)[NSApp mainMenu]: nil;
+				if ( MainMenu )
+				{
+					[MainMenu highlightKeyEquivalent:Event];
+				}
+			}
+			else
+			{
+				ResendEvent(Event);
+			}
 			break;
 		}
 
 		case NSKeyUp:
 		{
 			NSString *Characters = [Event characters];
+			bool bHandled = false;
 			if( [Characters length] && CurrentEventWindow.IsValid() )
 			{
 				const bool IsRepeat = [Event isARepeat];
@@ -748,7 +777,11 @@ void FMacApplication::ProcessNSEvent(NSEvent* const Event, TSharedPtr< FMacWindo
 				const uint32 KeyCode = [Event keyCode];
 				const bool IsPrintable = IsPrintableKey( Character );
 
-				MessageHandler->OnKeyUp( KeyCode, TranslateCharCode( CharCode, KeyCode ), IsRepeat );
+				bHandled = MessageHandler->OnKeyUp( KeyCode, TranslateCharCode( CharCode, KeyCode ), IsRepeat );
+			}
+			if (!bHandled)
+			{
+				ResendEvent(Event);
 			}
 			break;
 		}
@@ -769,6 +802,15 @@ void FMacApplication::ProcessEvent( NSEvent* Event )
 	ProcessNSEvent(Event, CurrentEventWindow, MacCursor->GetPosition());
 }
 
+void FMacApplication::ResendEvent(NSEvent* Event)
+{
+	MainThreadCall(^{
+		NSEvent* Wrapper = [NSEvent otherEventWithType:NSApplicationDefined location:[Event locationInWindow] modifierFlags:[Event modifierFlags] timestamp:[Event timestamp] windowNumber:[Event windowNumber] context:[Event context] subtype:FMacApplication::ResentEvent data1:(NSInteger)Event data2:0];
+		
+		[NSApp sendEvent:Wrapper];
+	}, NSDefaultRunLoopMode, true);
+}
+
 FCocoaWindow* FMacApplication::FindEventWindow( NSEvent* Event )
 {
 	SCOPED_AUTORELEASE_POOL;
diff --git a/Engine/Source/Runtime/Core/Private/Mac/MacPlatformMisc.cpp b/Engine/Source/Runtime/Core/Private/Mac/MacPlatformMisc.cpp
--- a/Engine/Source/Runtime/Core/Private/Mac/MacPlatformMisc.cpp
+++ b/Engine/Source/Runtime/Core/Private/Mac/MacPlatformMisc.cpp
@@ -10,6 +10,7 @@
 #include "VarargsHelper.h"
 #include "MacApplication.h"
 #include "MacCursor.h"
+#include "CocoaMenu.h"
 #include "CocoaThread.h"
 #include "Runtime/Launch/Resources/Version.h"
 #include "EngineVersion.h"
@@ -269,12 +270,12 @@ void FMacPlatformMisc::PlatformPostInit(bool ShowSplashScreen)
 		NSMenuItem* QuitItem = [[[NSMenuItem alloc] initWithTitle:[NSString stringWithFormat:@"Quit %@", AppName] action:RequestQuitSelector keyEquivalent:@"q"] autorelease];
 
 		NSMenuItem* ServicesItem = [[NSMenuItem new] autorelease];
-		NSMenu* ServicesMenu = [[NSMenu new] autorelease];
+		FCocoaMenu* ServicesMenu = [[FCocoaMenu new] autorelease];
 		[ServicesItem setTitle:@"Services"];
 		[ServicesItem setSubmenu:ServicesMenu];
 		[NSApp setServicesMenu:ServicesMenu];
 
-		NSMenu* AppMenu = [[NSMenu new] autorelease];
+		FCocoaMenu* AppMenu = [[FCocoaMenu new] autorelease];
 		[AppMenu addItem:AboutItem];
 		[AppMenu addItem:[NSMenuItem separatorItem]];
 		if (PreferencesItem)
@@ -290,7 +291,7 @@ void FMacPlatformMisc::PlatformPostInit(bool ShowSplashScreen)
 		[AppMenu addItem:[NSMenuItem separatorItem]];
 		[AppMenu addItem:QuitItem];
 
-		NSMenu* MenuBar = [[NSMenu new] autorelease];
+		FCocoaMenu* MenuBar = [[FCocoaMenu new] autorelease];
 		NSMenuItem* AppMenuItem = [[NSMenuItem new] autorelease];
 		[MenuBar addItem:AppMenuItem];
 		[NSApp setMainMenu:MenuBar];
@@ -305,7 +306,7 @@ void FMacPlatformMisc::UpdateWindowMenu()
 	NSMenu* WindowMenu = [NSApp windowsMenu];
 	if (!WindowMenu)
 	{
-		WindowMenu = [[NSMenu new] autorelease];
+		WindowMenu = [[FCocoaMenu new] autorelease];
 		[WindowMenu setTitle:@"Window"];
 		NSMenuItem* WindowMenuItem = [[NSMenuItem new] autorelease];
 		[WindowMenuItem setSubmenu:WindowMenu];
diff --git a/Engine/Source/Runtime/Core/Public/Mac/CocoaMenu.h b/Engine/Source/Runtime/Core/Public/Mac/CocoaMenu.h
new file mode 100644
--- /dev/null
+++ b/Engine/Source/Runtime/Core/Public/Mac/CocoaMenu.h
@@ -0,0 +1,13 @@
+// Copyright 1998-2014 Epic Games, Inc. All Rights Reserved.
+#pragma once
+
+#include "Core.h"
+
+@interface FCocoaMenu : NSMenu
+{
+@private
+	bool bHighlightingKeyEquivalent;
+}
+- (bool)isHighlightingKeyEquivalent;
+- (bool)highlightKeyEquivalent:(NSEvent *)Event;
+@end
diff --git a/Engine/Source/Runtime/Core/Public/Mac/MacApplication.h b/Engine/Source/Runtime/Core/Public/Mac/MacApplication.h
--- a/Engine/Source/Runtime/Core/Public/Mac/MacApplication.h
+++ b/Engine/Source/Runtime/Core/Public/Mac/MacApplication.h
@@ -109,6 +109,12 @@ public:
 	static void ProcessEvent(FMacEvent const* const Event);
 
 private:
+	enum FMacApplicationEventTypes
+	{
+		ResentEvent = 0x0f00
+	};
+
+	static NSEvent* HandleNSEvent(NSEvent* Event);
 	static void OnDisplayReconfiguration(CGDirectDisplayID Display, CGDisplayChangeSummaryFlags Flags, void* UserInfo);
 
 	FMacApplication();
@@ -117,6 +123,7 @@ private:
 	TCHAR ConvertChar( TCHAR Character );
 	TCHAR TranslateCharCode( TCHAR CharCode, uint32 KeyCode );
 
+	void ResendEvent( NSEvent* Event );
 	FCocoaWindow* FindEventWindow( NSEvent* CocoaEvent );
 	TSharedPtr<FGenericWindow> LocateWindowUnderCursor( const FVector2D& CursorPos );
 
diff --git a/Engine/Source/Runtime/Slate/Private/Framework/MultiBox/Mac/MacMenu.cpp b/Engine/Source/Runtime/Slate/Private/Framework/MultiBox/Mac/MacMenu.cpp
--- a/Engine/Source/Runtime/Slate/Private/Framework/MultiBox/Mac/MacMenu.cpp
+++ b/Engine/Source/Runtime/Slate/Private/Framework/MultiBox/Mac/MacMenu.cpp
@@ -20,7 +20,11 @@
 
 - (void)performAction
 {
-	FSlateMacMenu::ExecuteMenuItemAction(self.MenuEntryBlock.ToSharedRef());
+	FCocoaMenu* CocoaMenu = [[self menu] isKindOfClass:[FCocoaMenu class]] ? (FCocoaMenu*)[self menu] : nil;
+	if ( !CocoaMenu || ![CocoaMenu isHighlightingKeyEquivalent] )
+	{
+		FSlateMacMenu::ExecuteMenuItemAction(self.MenuEntryBlock.ToSharedRef());
+	}
 }
 
 @end
diff --git a/Engine/Source/Runtime/Slate/Private/Framework/MultiBox/Mac/MacMenu.h b/Engine/Source/Runtime/Slate/Private/Framework/MultiBox/Mac/MacMenu.h
--- a/Engine/Source/Runtime/Slate/Private/Framework/MultiBox/Mac/MacMenu.h
+++ b/Engine/Source/Runtime/Slate/Private/Framework/MultiBox/Mac/MacMenu.h
@@ -4,8 +4,9 @@
 
 #include "MultiBox.h"
 #include "SMenuEntryBlock.h"
+#include "CocoaMenu.h"
 
-@interface FMacMenu : NSMenu <NSMenuDelegate>
+@interface FMacMenu : FCocoaMenu <NSMenuDelegate>
 @property (assign) TSharedPtr<const FMenuEntryBlock> MenuEntryBlock;
 @property (assign) TSharedPtr<const FMultiBox> MultiBox;
 @end
