#docs Added link for FString docs. Documented ESearchCase, ESearchDir.

diff --git a/Engine/Source/Runtime/Core/Public/Containers/UnrealString.h b/Engine/Source/Runtime/Core/Public/Containers/UnrealString.h
--- a/Engine/Source/Runtime/Core/Public/Containers/UnrealString.h
+++ b/Engine/Source/Runtime/Core/Public/Containers/UnrealString.h
@@ -11,26 +11,36 @@
 #include "Array.h"
 
 
+/** Determines case sensitivity options for string comparisons. */
 namespace ESearchCase
 {
 	enum Type
 	{
+		/** Case sensitive. Upper/lower casing must match for strings to be considered equal. */
 		CaseSensitive,
+
+		/** Ignore case. Upper/lower casing does not matter when making a comparison. */
 		IgnoreCase,
 	};
 };
 
+/** Determines search direction for string operations. */
 namespace ESearchDir
 {
 	enum Type
 	{
+		/** Search from the start, moving forward through the string. */
 		FromStart,
+
+		/** Search from the end, moving backward through the string. */
 		FromEnd,
 	};
 }
-//
-// A dynamically sizeable string.
-//
+
+/**
+ * A dynamically sizeable string.
+ * @see https://docs.unrealengine.com/latest/INT/Programming/UnrealArchitecture/StringHandling/FString/
+ */
 class CORE_API FString
 {
 private:
