#docs Improved docs for FVector, FVector2D.

diff --git a/Engine/Source/Runtime/Core/Public/Math/Vector.h b/Engine/Source/Runtime/Core/Public/Math/Vector.h
--- a/Engine/Source/Runtime/Core/Public/Math/Vector.h
+++ b/Engine/Source/Runtime/Core/Public/Math/Vector.h
@@ -8,19 +8,19 @@
 
 
 /**
- * A 3x1 of FLOATs.
+ * A vector in 3-D space composed of components (X, Y, Z) with floating point precision.
  */
 class FVector 
 {
 public:
 
-	/** Holds the vector's X-component. */
+	/** Vector's X component. */
 	float X;
 
-	/** Holds the vector's Y-component. */
+	/** Vector's Y component. */
 	float Y;
 
-	/** Holds the vector's Z-component. */
+	/** Vector's Z component. */
 	float Z;
 
 public:
@@ -31,7 +31,7 @@ public:
 	/** World up vector (0,0,1) */
 	static CORE_API const FVector UpVector;
 
-	/** Local Unreal forward vector (1,0,0) */
+	/** Unreal forward vector (1,0,0) */
 	static CORE_API const FVector ForwardVector;
 
 public:
@@ -45,18 +45,20 @@ public:
 	FORCEINLINE void DiagnosticCheckNaN() const {}
 #endif
 
-	/** Constructor */
+	/**
+	 * Default constructor (no initialization).
+	 */
 	FORCEINLINE FVector();
 
 	/**
-	 * Constructor
+	 * Constructor initializing all components to a single float value.
 	 *
 	 * @param InF Value to set all components to.
 	 */
 	explicit FORCEINLINE FVector(float InF);
 
 	/**
-	 * Constructor
+	 * Constructor using initial values for each component.
 	 *
 	 * @param InX X Coordinate.
 	 * @param InY Y Coordinate.
@@ -65,7 +67,7 @@ public:
 	FORCEINLINE FVector( float InX, float InY, float InZ );
 
 	/**
-	 * Constructor
+	 * Constructs a vector from an FVector2D and Z value.
 	 * 
 	 * @param V Vector to copy from.
 	 * @param InZ Z Coordinate.
@@ -73,37 +75,37 @@ public:
 	explicit FORCEINLINE FVector( const FVector2D V, float InZ );
 
 	/**
-	 * Constructor
+	 * Constructor using the XYZ components from a 4D vector.
 	 *
 	 * @param V 4D Vector to copy from.
 	 */
 	FORCEINLINE FVector( const FVector4& V );
 
 	/**
-	 * Constructor
+	 * Constructs a vector from an FLinearColor.
 	 *
-	 * @param InColour Colour to copy from.
+	 * @param InColor Color to copy from.
 	 */
 	explicit FVector(const FLinearColor& InColor);
 
 	/**
-	 * Constructor
+	 * Constructs a vector from an FIntVector.
 	 *
 	 * @param InVector FIntVector to copy from.
 	 */
 	explicit FVector(FIntVector InVector);
 
 	/**
-	 * Constructor
+	 * Constructs a vector from an FIntPoint.
 	 *
-	 * @param A Int Point used to set X and Y coords, Z is set to zero.
+	 * @param A Int Point used to set X and Y coordinates, Z is set to zero.
 	 */
 	explicit FVector( FIntPoint A );
 
 	/**
-	 * Constructor
+	 * Constructor which initializes all components to zero.
 	 *
-	 * @param EForceInit Force Init Enum
+	 * @param EForceInit Force init enum
 	 */
 	explicit FORCEINLINE FVector(EForceInit);
 
@@ -119,51 +121,51 @@ public:
 #endif
 
 	/**
-	 * Calculate Cross product between this and another vector.
+	 * Calculate cross product between this and another vector.
 	 *
 	 * @param V The other vector.
-	 * @return The Cross product.
+	 * @return The cross product.
 	 */
 	FORCEINLINE FVector operator^( const FVector& V ) const;
 
 	/**
-	 * Calculate the Cross product of two vectors.
+	 * Calculate the cross product of two vectors.
 	 *
 	 * @param A The first vector.
 	 * @param B The second vector.
-	 * @return The Cross product.
+	 * @return The cross product.
 	 */
 	FORCEINLINE static FVector CrossProduct( const FVector& A, const FVector& B );
 
 	/**
-	 * Calculate the Dot Product between this and another vector.
+	 * Calculate the dot product between this and another vector.
 	 *
 	 * @param V The other vector.
-	 * @return The Dot Product.
+	 * @return The dot product.
 	 */
 	FORCEINLINE float operator|( const FVector& V ) const;
 
 	/**
-	 * Calculate the Dot product of two vectors.
+	 * Calculate the dot product of two vectors.
 	 *
 	 * @param A The first vector.
 	 * @param B The second vector.
-	 * @return The Dot product.
+	 * @return The dot product.
 	 */
 	FORCEINLINE static float DotProduct( const FVector& A, const FVector& B );
 
 	/**
-	 * Gets the result of adding a vector to this.
+	 * Gets the result of component-wise addition of this and another vector.
 	 *
-	 * @param V The vector to add.
+	 * @param V The vector to add to this.
 	 * @return The result of vector addition.
 	 */
 	FORCEINLINE FVector operator+( const FVector& V ) const;
 
 	/**
-	 * Gets the result of subtracting a vector from this.
+	 * Gets the result of component-wise subtraction of this by another vector.
 	 *
-	 * @param V The vector to subtract.
+	 * @param V The vector to subtract from this.
 	 * @return The result of vector subtraction.
 	 */
 	FORCEINLINE FVector operator-( const FVector& V ) const;
@@ -171,7 +173,7 @@ public:
 	/**
 	 * Gets the result of subtracting from each component of the vector.
 	 *
-	 * @param Bias What to subtract.
+	 * @param Bias How much to subtract from each component.
 	 * @return The result of subtraction.
 	 */
 	FORCEINLINE FVector operator-( float Bias ) const;
@@ -179,29 +181,29 @@ public:
 	/**
 	 * Gets the result of adding to each component of the vector.
 	 *
-	 * @param Bias What to add.
+	 * @param Bias How much to add to each component.
 	 * @return The result of addition.
 	 */
 	FORCEINLINE FVector operator+( float Bias ) const;
 
 	/**
-	 * Gets the result of multiplying each component of the vector.
+	 * Gets the result of scaling the vector (multiplying each component by a value).
 	 *
-	 * @param Scale What to multiply by.
+	 * @param Scale What to multiply each component by.
 	 * @return The result of multiplication.
 	 */
 	FORCEINLINE FVector operator*( float Scale ) const;
 
 	/**
-	 * Gets the result of dividing each component of the vector.
+	 * Gets the result of dividing each component of the vector by a value.
 	 *
-	 * @param Scale What to divide by.
+	 * @param Scale What to divide each component by.
 	 * @return The result of division.
 	 */
 	FVector operator/( float Scale ) const;
 
 	/**
-	 * Gets the result of multiplying vector with this.
+	 * Gets the result of component-wise multiplication of this vector by another.
 	 *
 	 * @param V The vector to multiply with.
 	 * @return The result of multiplication.
@@ -209,7 +211,7 @@ public:
 	FORCEINLINE FVector operator*( const FVector& V ) const;
 
 	/**
-	 * Gets the result of dividing this vector by another.
+	 * Gets the result of component-wise division of this vector by another.
 	 *
 	 * @param V The vector to divide by.
 	 * @return The result of division.
@@ -222,7 +224,7 @@ public:
 	 * Check against another vector for equality.
 	 *
 	 * @param V The vector to check against.
-	 * @return true if the vectors are equal, false otherwise..
+	 * @return true if the vectors are equal, false otherwise.
 	 */
 	bool operator==( const FVector& V ) const;
 
@@ -230,7 +232,7 @@ public:
 	 * Check against another vector for inequality.
 	 *
 	 * @param V The vector to check against.
-	 * @return true if the vectors are not equal, false otherwise..
+	 * @return true if the vectors are not equal, false otherwise.
 	 */
 	bool operator!=( const FVector& V ) const;
 
@@ -239,15 +241,15 @@ public:
 	 *
 	 * @param V The vector to check against.
 	 * @param Tolerance Error tolerance.
-	 * @return true if the vectors are equal within tolerance limits, false otherwise..
+	 * @return true if the vectors are equal within tolerance limits, false otherwise.
 	 */
 	bool Equals(const FVector& V, float Tolerance=KINDA_SMALL_NUMBER) const;
 
 	/**
-	 * Checks whether all components of the vector are the same, within a tolerance.
+	 * Checks whether all components of this vector are the same, within a tolerance.
 	 *
-	 * @param Tolerance Error Tolerance.
-	 * @return true if the vectors are equal within tolerance limits, false otherwise..
+	 * @param Tolerance Error tolerance.
+	 * @return true if the vectors are equal within tolerance limits, false otherwise.
 	 */
 	bool AllComponentsEqual(float Tolerance=KINDA_SMALL_NUMBER) const;
 
@@ -260,16 +262,18 @@ public:
 
 	/**
 	 * Adds another vector to this.
+	 * Uses component-wise addition.
 	 *
-	 * @param V Vector to add.
+	 * @param V Vector to add to this.
 	 * @return Copy of the vector after addition.
 	 */
 	FORCEINLINE FVector operator+=( const FVector& V );
 
 	/**
 	 * Subtracts another vector from this.
+	 * Uses component-wise subtraction.
 	 *
-	 * @param V Vector to subtract.
+	 * @param V Vector to subtract from this.
 	 * @return Copy of the vector after subtraction.
 	 */
 	FORCEINLINE FVector operator-=( const FVector& V );
@@ -277,7 +281,7 @@ public:
 	/**
 	 * Scales the vector.
 	 *
-	 * @param Scale What to scale vector by.
+	 * @param Scale Amount to scale this vector by.
 	 * @return Copy of the vector after scaling.
 	 */
 	FORCEINLINE FVector operator*=( float Scale );
@@ -285,21 +289,21 @@ public:
 	/**
 	 * Divides the vector by a number.
 	 *
-	 * @param V What to divide the vector by.
+	 * @param V What to divide this vector by.
 	 * @return Copy of the vector after division.
 	 */
 	FVector operator/=( float V );
 
 	/**
-	 * Multiplies the vector with another vector.
+	 * Multiplies the vector with another vector, using component-wise multiplication.
 	 *
-	 * @param V What to multiply vector with.
+	 * @param V What to multiply this vector with.
 	 * @return Copy of the vector after multiplication.
 	 */
 	FVector operator*=( const FVector& V );
 
 	/**
-	 * Divides the vector by another vector.
+	 * Divides the vector by another vector, using component-wise division.
 	 *
 	 * @param V What to divide vector by.
 	 * @return Copy of the vector after division.
@@ -322,6 +326,23 @@ public:
 	 */
 	float operator[]( int32 Index )const;
 
+	/**
+	* Gets a specific component of the vector.
+	*
+	* @param Index The index of the component required.
+	*
+	* @return Reference to the specified component.
+	*/
+	float& Component(int32 Index);
+
+	/**
+	* Gets a specific component of the vector.
+	*
+	* @param Index The index of the component required.
+	* @return Copy of the specified component.
+	*/
+	float Component(int32 Index) const;
+
 	// Simple functions.
 
 	/**
@@ -334,30 +355,30 @@ public:
 	void Set( float InX, float InY, float InZ );
 
 	/**
-	 * Get the maximum of the vectors coordinates.
+	 * Get the maximum value of the vector's components.
 	 *
-	 * @return The maximum of the vectors coordinates.
+	 * @return The maximum value of the vector's components.
 	 */
 	float GetMax() const;
 
 	/**
-	 * Get the absolute maximum of the vectors coordinates.
+	 * Get the maximum absolute value of the vector's components.
 	 *
-	 * @return The absolute maximum of the vectors coordinates.
+	 * @return The maximum absolute value of the vector's components.
 	 */
 	float GetAbsMax() const;
 
 	/**
-	 * Get the minimum of the vectors coordinates.
+	 * Get the minimum value of the vector's components.
 	 *
-	 * @return The minimum of the vectors coordinates.
+	 * @return The minimum value of the vector's components.
 	 */
 	float GetMin() const;
 
 	/**
-	 * Get the absolute minimum of the vectors coordinates.
+	 * Get the minimum absolute value of the vector's components.
 	 *
-	 * @return The absolute minimum of the vectors coordinates.
+	 * @return The minimum absolute value of the vector's components.
 	 */
 	float GetAbsMin() const;
 
@@ -368,14 +389,14 @@ public:
 	FVector ComponentMax(const FVector& Other) const;
 
 	/**
-	 * Get a copy of this vector with absolute components.
+	 * Get a copy of this vector with absolute value of each component.
 	 *
-	 * @return A copy of this vector with absolute components.
+	 * @return A copy of this vector with absolute value of each component.
 	 */
 	FVector GetAbs() const;
 
 	/**
-	 * Get the length of this vector.
+	 * Get the length (magnitude) of this vector.
 	 *
 	 * @return The length of this vector.
 	 */
@@ -406,22 +427,22 @@ public:
 	 * Checks whether vector is near to zero within a specified tolerance.
 	 *
 	 * @param Tolerance Error tolerance.
-	 * @return true if the vector is near to zero, false otherwise..
+	 * @return true if the vector is near to zero, false otherwise.
 	 */
 	bool IsNearlyZero(float Tolerance=KINDA_SMALL_NUMBER) const;
 
 	/**
-	 * Checks whether vector is exactly zero.
+	 * Checks whether all components of the vector are exactly zero.
 	 *
-	 * @return true if the vector is exactly zero, false otherwise..
+	 * @return true if the vector is exactly zero, false otherwise.
 	 */
 	bool IsZero() const;
 
 	/**
-	 * Normalize this vector if it is large enough.
+	 * Normalize this vector in-place if it is large enough, set it to (0,0,0) otherwise.
 	 *
 	 * @param Tolerance Minimum squared length of vector for normalization.
-	 * @return true if the vector was normalized correctly, false otherwise..
+	 * @return true if the vector was normalized correctly, false otherwise.
 	 */
 	bool Normalize(float Tolerance=SMALL_NUMBER);
 
@@ -433,7 +454,7 @@ public:
 	bool IsNormalized() const;
 
 	/**
-	 * Util to convert this vector into a unit direction vector, and its original length
+	 * Util to convert this vector into a unit direction vector and its original length.
 	 *
 	 * @param OutDir Reference passed in to store unit direction vector.
 	 * @param OutLength Reference passed in to store length of the vector.
@@ -442,22 +463,24 @@ public:
 
 	/**
 	 * Get a copy of the vector as sign only.
+	 * Each component is set to +1 or -1, with the sign of zero treated as +1.
 	 *
-	 * @param A copy of the vector with each component set to 1 or -1
+	 * @param A copy of the vector with each component set to +1 or -1
 	 */
 	FORCEINLINE FVector GetSignVector();
 
 	/**
 	 * Projects 2D components of vector based on Z.
 	 *
-	 * @return Projected version of vector.
+	 * @return Projected version of vector based on Z.
 	 */
 	FVector Projection() const;
 
 	/**
-	 * Calculates normalized version of vector without checking if it is non-zero.
+	 * Calculates normalized version of vector without checking for zero length.
 	 *
 	 * @return Normalized version of vector.
+	 * @see SafeNormal()
 	 */
 	FORCEINLINE FVector UnsafeNormal() const;
 
@@ -466,6 +489,7 @@ public:
 	 *
 	 * @param GridSz Grid dimension.
 	 * @return A copy of this vector snapped to a grid.
+	 * @see FMath::GridSnap()
 	 */
 	FVector GridSnap( const float& GridSz ) const;
 
@@ -483,10 +507,10 @@ public:
 	/** Create a copy of this vector, with the 2D magnitude clamped between Min and Max. Z is unchanged. */
 	FVector ClampSize2D(float Min, float Max) const;
 
-	/** Create a copy of this vector, with its magnitude clamped to MaxSize. */
+	/** Create a copy of this vector, with its maximum magnitude clamped to MaxSize. */
 	FVector ClampMaxSize(float MaxSize) const;
 
-	/** Create a copy of this vector, with the 2D magnitude clamped to MaxSize. Z is unchanged. */
+	/** Create a copy of this vector, with the maximum 2D magnitude clamped to MaxSize. Z is unchanged. */
 	FVector ClampMaxSize2D(float MaxSize) const;
 
 	/**
@@ -497,26 +521,9 @@ public:
 	 */
 	void AddBounded( const FVector& V, float Radius=MAX_int16 );
 
-	/**
-	 * Gets a specific component of the vector.
-	 *
-	 * @param Index The index of the component required.
-	 *
-	 * @return Reference to the specified component.
-	 */
-	float& Component( int32 Index );
-
-	/**
-	 * Gets a specific component of the vector.
-	 *
-	 * @param Index The index of the component required.
-	 * @return Copy of the specified component.
-	 */
-	float Component( int32 Index ) const;
-
 	/**
 	 * Gets the reciprocal of this vector, avoiding division by zero.
-	 * Zero components set to BIG_NUMBER.
+	 * Zero components are set to BIG_NUMBER.
 	 *
 	 * @return Reciprocal of this vector.
 	 */
@@ -556,15 +563,17 @@ public:
 	FVector RotateAngleAxis( const float AngleDeg, const FVector& Axis ) const;
 
 	/**
-	 * Gets a normalized copy of the vector, checking it is safe to do so.
+	 * Gets a normalized copy of the vector, checking it is safe to do so based on the length.
+	 * Returns zero vector if vector length is too small to safely normalize.
 	 *
 	 * @param Tolerance Minimum squared vector length.
-	 * @return Normalized copy if safe, otherwise returns zero vector.
+	 * @return A normalized copy if safe, (0,0,0) otherwise.
 	 */
 	FVector SafeNormal(float Tolerance=SMALL_NUMBER) const;
 
 	/**
-	 * Gets a normalized copy of the 2D components of the vector, checking it is safe to do so.
+	 * Gets a normalized copy of the 2D components of the vector, checking it is safe to do so. Z is set to zero. 
+	 * Returns zero vector if vector length is too small to normalize.
 	 *
 	 * @param Tolerance Minimum squared vector length.
 	 * @return Normalized copy if safe, otherwise returns zero vector.
@@ -572,7 +581,7 @@ public:
 	FVector SafeNormal2D(float Tolerance=SMALL_NUMBER) const;
 
 	/**
-	 * Returns the cosine of the angle between this vector and another projected onto the XY plane (no z)
+	 * Returns the cosine of the angle between this vector and another projected onto the XY plane (no Z).
 	 *
 	 * @param B the other vector to find the 2D cosine of the angle with.
 	 * @return The cosine.
@@ -580,9 +589,9 @@ public:
 	FORCEINLINE float CosineAngle2D(FVector B) const;
 
 	/**
-	 * Projects this vector onto the input vector.
+	 * Gets a copy of this vector projected onto the input vector.
 	 *
-	 * @param A Vector to project onto, does not assume it is unnormalized.
+	 * @param A Vector to project onto, does not assume it is normalized.
 	 * @return Projected vector.
 	 */
 	FORCEINLINE FVector ProjectOnTo( const FVector& A ) const ;
@@ -598,7 +607,7 @@ public:
 
 	/**
 	 * Find good arbitrary axis vectors to represent U and V axes of a plane,
-	 * given just the normal.
+	 * using this vector as the normal of the plane.
 	 *
 	 * @param Axis1 Reference to first axis.
 	 * @param Axis2 Reference to second axis.
@@ -611,7 +620,7 @@ public:
 	/**
 	 * Utility to check if there are any NaNs in this vector.
 	 *
-	 * @return true if there are any NaNs in this vector, false otherwise..
+	 * @return true if there are any NaNs in this vector, false otherwise.
 	 */
 	bool ContainsNaN() const;
 
@@ -653,7 +662,7 @@ public:
 	bool InitFromString( const FString & InSourceString );
 
 	/** 
-	 * Converts a cartesian unit vector into spherical coordinates on the unit sphere.
+	 * Converts a Cartesian unit vector into spherical coordinates on the unit sphere.
 	 * @return Output Theta will be in the range [0, PI], and output Phi will be in the range [-PI, PI]. 
 	 */
 	FVector2D UnitCartesianToSpherical() const;
@@ -682,7 +691,7 @@ public:
 	 *
 	 * @param P First vector.
 	 * @param Q Second vector.
-	 * @return 1=yes, 0=no.  Uses fast distance approximation.
+	 * @return Whether points are the same within a threshold. Uses fast distance approximation (linear per-component distance).
 	 */
 	static bool PointsAreSame( const FVector &P, const FVector &Q );
 	
@@ -692,45 +701,46 @@ public:
 	 * @param Point1 First vector.
 	 * @param Point2 Second vector.
 	 * @param Dist Specified distance.
-	 * @return 1=yes, 0=no.  Uses fast distance approximation.
+	 * @return Whether two points are within the specified distance. Uses fast distance approximation (linear per-component distance).
 	 */
 	static bool PointsAreNear( const FVector &Point1, const FVector &Point2, float Dist );
 
 	/**
-	 * Calculate the signed distance (in the direction of the normal) between
-	 * a point and a plane.
+	 * Calculate the signed distance (in the direction of the normal) between a point and a plane.
 	 *
 	 * @param Point The Point we are checking.
 	 * @param PlaneBase The Base Point in the plane.
 	 * @param PlaneNormal The Normal of the plane.
-	 * @return Signed distance  between point and plane.
+	 * @return Signed distance between point and plane.
 	 */
 	static float PointPlaneDist( const FVector &Point, const FVector &PlaneBase, const FVector &PlaneNormal );
 
 	/**
-	 * Calculate a the projection of a point on the given plane
+	 * Calculate the projection of a point on the given plane.
 	 *
-	 * @param Point - the point to project onto the plane
-	 * @param Plane - the plane
+	 * @param Point The point to project onto the plane
+	 * @param Plane The plane
 	 * @return Projection of Point onto Plane
 	 */
 	static FVector PointPlaneProject(const FVector& Point, const FPlane& Plane);
 
 	/**
-	 * Calculate a the projection of a point on the plane defined by CCW points A,B,C
+	 * Calculate the projection of a point on the plane defined by counter-clockwise (CCW) points A,B,C.
 	 *
-	 * @param Point - the point to project onto the plane
-	 * @param A,B,C - three points in CCW order defining the plane 
+	 * @param Point The point to project onto the plane
+	 * @param A 1st of three points in CCW order defining the plane 
+	 * @param B 2nd of three points in CCW order defining the plane
+	 * @param C 3rd of three points in CCW order defining the plane
 	 * @return Projection of Point onto plane ABC
 	 */
 	static FVector PointPlaneProject(const FVector& Point, const FVector& A, const FVector& B, const FVector& C);
 
 	/**
-	* Calculate a the projection of a point on the plane defined by PlaneBase, and PlaneNormal
+	* Calculate the projection of a point on the plane defined by PlaneBase and PlaneNormal.
 	*
-	* @param Point - the point to project onto the plane
-	* @param PlaneBase - point on the plane
-	* @param PlaneNorm - normal of the plane
+	* @param Point The point to project onto the plane
+	* @param PlaneBase Point on the plane
+	* @param PlaneNorm Normal of the plane
 	* @return Projection of Point onto plane ABC
 	*/
 	static FVector PointPlaneProject(const FVector& Point, const FVector& PlaneBase, const FVector& PlaneNorm);
@@ -767,7 +777,7 @@ public:
 	 *
 	 * @param Normal1 First normalized vector.
 	 * @param Normal1 Second normalized vector.
-	 * @return true if vectors are parallel, false otherwise..
+	 * @return true if vectors are nearly parallel, false otherwise.
 	 */
 	static bool Parallel( const FVector &Normal1, const FVector &Normal2 );
 
@@ -778,17 +788,17 @@ public:
 	 * @param Normal1 The normal of the first plane.
 	 * @param Base2 The base point in the second plane.
 	 * @param Normal2 The normal of the second plane.
-	 * @return true if the planes are coplanar, false otherwise..
+	 * @return true if the planes are coplanar, false otherwise.
 	 */
 	static bool Coplanar( const FVector &Base1, const FVector &Normal1, const FVector &Base2, const FVector &Normal2 );
 
 	/**
-	 * Triple product of three vectors.
+	 * Triple product of three vectors: X dot (Y cross Z).
 	 *
 	 * @param X The first vector.
 	 * @param Y The second vector.
 	 * @param Z The third vector.
-	 * @return The triple product.
+	 * @return The triple product: X dot (Y cross Z).
 	 */
 	static float Triple( const FVector& X, const FVector& Y, const FVector& Z );
 
@@ -826,7 +836,13 @@ public:
 		// See TArray::BulkSerialize for detailed description of implied limitations.
 		return Ar << V.X << V.Y << V.Z;
 	}
-
+	
+	/** 
+	 * Network serialization function.
+	 * FVectors NetSerialize without quantization (ie exact values are serialized).
+	 *
+	 * @see FVector_NetQuantize, FVector_NetQuantize10, FVector_NetQuantize100, FVector_NetQuantizeNormal
+	 */
 	CORE_API bool NetSerialize(FArchive& Ar, class UPackageMap* Map, bool& bOutSuccess);
 };
 
@@ -975,20 +991,20 @@ inline bool FVector::PointsAreSame( const FVector &P, const FVector &Q )
 			Temp=P.Z-Q.Z;
 			if( (Temp > -THRESH_POINTS_ARE_SAME) && (Temp < THRESH_POINTS_ARE_SAME) )
 			{
-				return 1;
+				return true;
 			}
 		}
 	}
-	return 0;
+	return false;
 }
 
 inline bool FVector::PointsAreNear( const FVector &Point1, const FVector &Point2, float Dist )
 {
 	float Temp;
-	Temp=(Point1.X - Point2.X); if (FMath::Abs(Temp)>=Dist) return 0;
-	Temp=(Point1.Y - Point2.Y); if (FMath::Abs(Temp)>=Dist) return 0;
-	Temp=(Point1.Z - Point2.Z); if (FMath::Abs(Temp)>=Dist) return 0;
-	return 1;
+	Temp=(Point1.X - Point2.X); if (FMath::Abs(Temp)>=Dist) return false;
+	Temp=(Point1.Y - Point2.Y); if (FMath::Abs(Temp)>=Dist) return false;
+	Temp=(Point1.Z - Point2.Z); if (FMath::Abs(Temp)>=Dist) return false;
+	return true;
 }
 
 inline float FVector::PointPlaneDist
@@ -1017,9 +1033,9 @@ inline bool FVector::Parallel( const FVector &Normal1, const FVector &Normal2 )
 
 inline bool FVector::Coplanar( const FVector &Base1, const FVector &Normal1, const FVector &Base2, const FVector &Normal2 )
 {
-	if      (!FVector::Parallel(Normal1,Normal2)) return 0;
-	else if (FVector::PointPlaneDist (Base2,Base1,Normal1) > THRESH_POINT_ON_PLANE) return 0;
-	else    return 1;
+	if      (!FVector::Parallel(Normal1,Normal2)) return false;
+	else if (FVector::PointPlaneDist (Base2,Base1,Normal1) > THRESH_POINT_ON_PLANE) return false;
+	else    return true;
 }
 
 inline float FVector::Triple( const FVector& X, const FVector& Y, const FVector& Z )
diff --git a/Engine/Source/Runtime/Core/Public/Math/Vector2D.h b/Engine/Source/Runtime/Core/Public/Math/Vector2D.h
--- a/Engine/Source/Runtime/Core/Public/Math/Vector2D.h
+++ b/Engine/Source/Runtime/Core/Public/Math/Vector2D.h
@@ -8,22 +8,22 @@
 
 
 /**
- * A 2x1 of FLOATs.
+ * A vector in 2-D space composed of components (X, Y) with floating point precision.
  */
 struct FVector2D 
 {
-	/** Holds the vector's X-component. */
+	/** Vector's X component. */
 	float X;
 
-	/** Holds the vector's Y-component. */
+	/** Vector's Y component. */
 	float Y;
 
 public:
 
-	/** Global 2D zero vector constant */
+	/** Global 2D zero vector constant (0,0) */
 	static CORE_API const FVector2D ZeroVector;
 
-	/** Global 2D unit vector constant */
+	/** Global 2D unit vector constant (1,1) */
 	static CORE_API const FVector2D UnitVector;
 
 public:
@@ -34,7 +34,7 @@ public:
 	FORCEINLINE FVector2D( ) { }
 
 	/**
-	 * Constructor
+	 * Constructor using initial values for each component.
 	 *
 	 * @param InX X coordinate.
 	 * @param InY Y coordinate.
@@ -42,21 +42,22 @@ public:
 	FORCEINLINE FVector2D( float InX,float InY );
 
 	/**
-	 * Constructor
+	 * Constructs a vector from an FIntPoint.
 	 *
 	 * @param InPos Integer point used to set this vector.
 	 */
 	FORCEINLINE FVector2D( FIntPoint InPos );
 
 	/**
-	 * Constructor
+	 * Constructor which initializes all components to zero.
 	 *
 	 * @param EForceInit Force init enum
 	 */
 	explicit FORCEINLINE FVector2D( EForceInit );
 
 	/**
-	 * Copy Constructor
+	 * Constructs a vector from an FVector.
+	 * Copies the X and Y components from the FVector.
 	 *
 	 * @param V Vector to copy from.
 	 */
@@ -75,23 +76,23 @@ public:
 	/**
 	 * Gets the result of subtracting a vector from this one.
 	 *
-	 * @param V The other vector to subtract.
+	 * @param V The other vector to subtract from this.
 	 * @return The result of the subtraction.
 	 */
 	FORCEINLINE FVector2D operator-( const FVector2D& V ) const;
 
 	/**
-	 * Gets the result of scaling this vector.
+	 * Gets the result of scaling the vector (multiplying each component by a value).
 	 *
-	 * @param Scale What to scale the vector by.
+	 * @param Scale How much to scale the vector by.
 	 * @return The result of scaling this vector.
 	 */
 	FORCEINLINE FVector2D operator*( float Scale ) const;
 
 	/**
-	 * Gets the result of dividing this vector.
+	 * Gets the result of dividing each component of the vector by a value.
 	 *
-	 * @param Scale What to divide the vector by.
+	 * @param Scale How much to divide the vector by.
 	 * @return The result of division on this vector.
 	 */
 	FVector2D operator/( float Scale ) const;
@@ -99,21 +100,21 @@ public:
 	/**
 	 * Gets the result of this vector + float A.
 	 *
-	 * @param A		Float to add.
+	 * @param A		Float to add to each component.
 	 * @return		The result of this vector + float A.
 	 */
 	FORCEINLINE FVector2D operator+( float A ) const;
 
 	/**
-	 * Gets the result of this vector - float A
+	 * Gets the result of subtracting from each component of the vector.
 	 *
-	 * @param A		Float to subtract
+	 * @param A		Float to subtract from each component
 	 * @return		The result of this vector - float A.
 	 */
 	FORCEINLINE FVector2D operator-( float A ) const;
 
 	/**
-	 * Gets the result of multiplying this vector by another.
+	 * Gets the result of component-wise multiplication of this vector by another.
 	 *
 	 * @param V The other vector to multiply this by.
 	 * @return The result of the multiplication.
@@ -121,7 +122,7 @@ public:
 	FORCEINLINE FVector2D operator*( const FVector2D& V ) const;
 
 	/**
-	 * Gets the result of dividing this vector by another.
+	 * Gets the result of component-wise division of this vector by another.
 	 *
 	 * @param V The other vector to divide this by.
 	 * @return The result of the division.
@@ -129,18 +130,18 @@ public:
 	FVector2D operator/( const FVector2D& V ) const;
 
 	/**
-	 * Calculates Dot product of this vector and another.
+	 * Calculates dot product of this vector and another.
 	 *
 	 * @param V The other vector.
-	 * @return The Dot product.
+	 * @return The dot product.
 	 */
 	FORCEINLINE float operator|( const FVector2D& V) const;
 
 	/**
-	 * Calculates Cross product of this vector and another.
+	 * Calculates cross product of this vector and another.
 	 *
 	 * @param V The other vector.
-	 * @return The Cross product.
+	 * @return The cross product.
 	 */
 	FORCEINLINE float operator^( const FVector2D& V) const;
 
@@ -163,7 +164,7 @@ public:
 	bool operator!=( const FVector2D& V ) const;
 
 	/**
-	 * Checks whether the components of this vector are less than another.
+	 * Checks whether both components of this vector are less than another.
 	 *
 	 * @param Other The vector to compare against.
 	 * @return true if this is the smaller vector, otherwise false.
@@ -171,7 +172,7 @@ public:
 	bool operator<( const FVector2D& Other ) const;
 
 	/**
-	 * Checks whether the components of this vector are greater than another.
+	 * Checks whether both components of this vector are greater than another.
 	 *
 	 * @param Other The vector to compare against.
 	 * @return true if this is the larger vector, otherwise false.
@@ -179,7 +180,7 @@ public:
 	bool operator>( const FVector2D& Other ) const;
 
 	/**
-	 * Checks whether the components of this vector are less than or equal to another.
+	 * Checks whether both components of this vector are less than or equal to another.
 	 *
 	 * @param Other The vector to compare against.
 	 * @return true if this vector is less than or equal to the other vector, otherwise false.
@@ -187,7 +188,7 @@ public:
 	bool operator<=( const FVector2D& Other ) const;
 
 	/**
-	 * Checks whether the components of this vector are greater than or equal to another.
+	 * Checks whether both components of this vector are greater than or equal to another.
 	 *
 	 * @param Other The vector to compare against.
 	 * @return true if this vector is greater than or equal to the other vector, otherwise false.
@@ -234,7 +235,7 @@ public:
 	FVector2D operator/=( float V );
 
 	/**
-	 * Multiplies this vector with another vector.
+	 * Multiplies this vector with another vector, using component-wise multiplication.
 	 *
 	 * @param V The vector to multiply with.
 	 * @return Copy of the vector after multiplication.
@@ -242,7 +243,7 @@ public:
 	FVector2D operator*=( const FVector2D& V );
 
 	/**
-	 * Divides this vector by another vector.
+	 * Divides this vector by another vector, using component-wise division.
 	 *
 	 * @param V The vector to divide by.
 	 * @return Copy of the vector after division.
@@ -265,14 +266,30 @@ public:
 	 */
 	float operator[]( int32 Index ) const;
 
+	/**
+	* Gets a specific component of the vector.
+	*
+	* @param Index The index of the component required.
+	* @return Reference to the specified component.
+	*/
+	float& Component(int32 Index);
+
+	/**
+	* Gets a specific component of the vector.
+	*
+	* @param Index The index of the component required.
+	* @return Copy of the specified component.
+	*/
+	float Component(int32 Index) const;
+
 public:
 
 	/**
-	 * Calculates the Dot product of two vectors.
+	 * Calculates the dot product of two vectors.
 	 *
 	 * @param A The first vector.
 	 * @param B The second vector.
-	 * @return The Dot product.
+	 * @return The dot product.
 	 */
 	FORCEINLINE static float DotProduct( const FVector2D& A, const FVector2D& B );
 
@@ -295,11 +312,11 @@ public:
 	FORCEINLINE static float Distance( const FVector2D& V1, const FVector2D& V2 );
 
 	/**
-	 * Calculate the Cross product of two vectors.
+	 * Calculate the cross product of two vectors.
 	 *
 	 * @param A The first vector.
 	 * @param B The second vector.
-	 * @return The Cross product.
+	 * @return The cross product.
 	 */
 	FORCEINLINE static float CrossProduct(const FVector2D& A, const FVector2D& B);
 
@@ -321,28 +338,28 @@ public:
 	void Set( float InX, float InY );
 
 	/**
-	 * Get the maximum of the vectors coordinates.
+	 * Get the maximum value of the vector's components.
 	 *
-	 * @return The maximum of the vectors coordinates.
+	 * @return The maximum value of the vector's components.
 	 */
 	float GetMax( ) const;
 
 	/**
-	 * Get the absolute maximum of the vectors coordinates.
+	 * Get the maximum absolute value of the vector's components.
 	 *
-	 * @return The absolute maximum of the vectors coordinates.
+	 * @return The maximum absolute value of the vector's components.
 	 */
 	float GetAbsMax( ) const;
 
 	/**
-	 * Get the minimum of the vectors coordinates.
+	 * Get the minimum value of the vector's components.
 	 *
-	 * @return The minimum of the vectors coordinates.
+	 * @return The minimum value of the vector's components.
 	 */
 	float GetMin( ) const;
 
 	/**
-	 * Get the length of this vector.
+	 * Get the length (magnitude) of this vector.
 	 *
 	 * @return The length of this vector.
 	 */
@@ -356,7 +373,8 @@ public:
 	float SizeSquared( ) const;
 
 	/**
-	 * Get a normalized copy of this vector if it is large enough.
+	 * Gets a normalized copy of the vector, checking it is safe to do so based on the length.
+	 * Returns zero vector if vector length is too small to safely normalize.
 	 *
 	 * @param Tolerance Minimum squared length of vector for normalization.
 	 * @return A normalized copy of the vector if safe, (0,0) otherwise.
@@ -364,9 +382,10 @@ public:
 	FVector2D SafeNormal( float Tolerance=SMALL_NUMBER ) const;
 
 	/**
-	 * Normalize this vector if it is large enough, set it to (0,0) otherwise.
+	 * Normalize this vector in-place if it is large enough, set it to (0,0) otherwise.
 	 *
 	 * @param Tolerance Minimum squared length of vector for normalization.
+	 * @see SafeNormal()
 	 */
 	void Normalize( float Tolerance=SMALL_NUMBER );
 
@@ -379,28 +398,12 @@ public:
 	bool IsNearlyZero( float Tolerance=KINDA_SMALL_NUMBER ) const;
 
 	/**
-	 * Checks whether vector is exactly zero.
+	 * Checks whether all components of the vector are exactly zero.
 	 *
 	 * @return true if vector is exactly zero, otherwise false.
 	 */
 	bool IsZero( ) const;
 
-	/**
-	 * Gets a specific component of the vector.
-	 *
-	 * @param Index The index of the component required.
-	 * @return Reference to the specified component.
-	 */
-	float& Component( int32 Index );
-
-	/**
-	 * Gets a specific component of the vector.
-	 *
-	 * @param Index The index of the component required.
-	 * @return Copy of the specified component.
-	 */
-	float Component( int32 Index ) const;
-
 	/**
 	 * Get this vector as an Int Point.
 	 *
@@ -464,9 +467,13 @@ public:
 			FMath::IsNaN(Y) || !FMath::IsFinite(Y));
 	}
 
+	/**
+	 * Network serialization function.
+	 * FVectors NetSerialize without quantization (ie exact values are serialized).
+	 */
 	CORE_API bool NetSerialize( FArchive& Ar, class UPackageMap* Map, bool& bOutSuccess );
 
-	/** Converts spherical coordinates on the unit sphere into a cartesian unit length vector. */
+	/** Converts spherical coordinates on the unit sphere into a Cartesian unit length vector. */
 	inline FVector SphericalToUnitCartesian( ) const;
 };
 
