Added a method to FCulture to allow canonicalization of a culture name. Useful for converting semi-friendly names like "en-US" into "en_US".

diff --git a/Engine/Source/Runtime/Core/Private/Internationalization/Culture.cpp b/Engine/Source/Runtime/Core/Private/Internationalization/Culture.cpp
--- a/Engine/Source/Runtime/Core/Private/Internationalization/Culture.cpp
+++ b/Engine/Source/Runtime/Core/Private/Internationalization/Culture.cpp
@@ -49,6 +49,11 @@ FString FCulture::GetName() const
 	return Implementation->GetName();
 }
 
+FString FCulture::GetCanonicalName(const FString& Name)
+{
+	return FImplementation::GetCanonicalName(Name);
+}
+
 FString FCulture::GetParentName(const FString& CultureName)
 {
 	return FImplementation::GetParentName(CultureName);
diff --git a/Engine/Source/Runtime/Core/Private/Internationalization/ICUCulture.cpp b/Engine/Source/Runtime/Core/Private/Internationalization/ICUCulture.cpp
--- a/Engine/Source/Runtime/Core/Private/Internationalization/ICUCulture.cpp
+++ b/Engine/Source/Runtime/Core/Private/Internationalization/ICUCulture.cpp
@@ -107,6 +107,18 @@ FString FCulture::FICUCultureImplementation::GetName() const
 	return ICULocale.getName();
 }
 
+FString FCulture::FICUCultureImplementation::GetCanonicalName(const FString& Name)
+{
+	static const int32 MaximumNameLength = 64;
+	const int32 NameLength = Name.Len();
+	check(NameLength < MaximumNameLength);
+	char CanonicalName[MaximumNameLength];
+
+	UErrorCode ICUStatus = U_ZERO_ERROR;
+	uloc_canonicalize(TCHAR_TO_ANSI( *Name ), CanonicalName, MaximumNameLength, &ICUStatus);
+	return CanonicalName;
+}
+
 FString FCulture::FICUCultureImplementation::GetParentName(const FString& CultureName)
 {
 	char Parent[32];
diff --git a/Engine/Source/Runtime/Core/Private/Internationalization/ICUCulture.h b/Engine/Source/Runtime/Core/Private/Internationalization/ICUCulture.h
--- a/Engine/Source/Runtime/Core/Private/Internationalization/ICUCulture.h
+++ b/Engine/Source/Runtime/Core/Private/Internationalization/ICUCulture.h
@@ -116,6 +116,8 @@ class FCulture::FICUCultureImplementation
 
 	FString GetName() const;
 
+	static FString GetCanonicalName(const FString& Name);
+
 	static FString GetParentName(const FString& CultureName);
 
 	FString GetNativeName() const;
diff --git a/Engine/Source/Runtime/Core/Private/Internationalization/Internationalization.cpp b/Engine/Source/Runtime/Core/Private/Internationalization/Internationalization.cpp
--- a/Engine/Source/Runtime/Core/Private/Internationalization/Internationalization.cpp
+++ b/Engine/Source/Runtime/Core/Private/Internationalization/Internationalization.cpp
@@ -23,27 +23,6 @@
 
 #define LOCTEXT_NAMESPACE "Internationalization"
 
-namespace
-{
-
-FString GetCanonicalName(const FString& Name)
-{
-#if UE_ENABLE_ICU
-	static const int32 MaximumNameLength = 64;
-	const int32 NameLength = Name.Len();
-	check(NameLength < MaximumNameLength);
-	char CanonicalName[MaximumNameLength];
-
-	UErrorCode ICUStatus = U_ZERO_ERROR;
-	uloc_canonicalize(TCHAR_TO_ANSI( *Name ), CanonicalName, MaximumNameLength, &ICUStatus);
-	return CanonicalName;
-#else
-	return Name;
-#endif
-}
-
-}
-
 FInternationalization* FInternationalization::Instance;
 
 FInternationalization& FInternationalization::Get()
@@ -120,7 +99,7 @@ FCulturePtr FInternationalization::GetCulture(const FString& Name) const
 
 int32 FInternationalization::GetCultureIndex(const FString& Name) const
 {
-	const FString CanonicalName = GetCanonicalName(Name);
+	const FString CanonicalName = FCulture::GetCanonicalName(Name);
 
 	const int32 CultureCount = AllCultures.Num();
 	int32 i;
@@ -589,7 +568,7 @@ void FInternationalization::GetCulturesWithAvailableLocalization(const TArray<FS
 				{
 					// UE localization resource folders use "en-US" style while ICU uses "en_US"
 					const FString LocalizationFolder = FPaths::GetCleanFilename(FilenameOrDirectory);
-					const FString CanonicalName = GetCanonicalName(LocalizationFolder);
+					const FString CanonicalName = FCulture::GetCanonicalName(LocalizationFolder);
 					LocalizationFolders.AddUnique(CanonicalName);
 				}
 
diff --git a/Engine/Source/Runtime/Core/Private/Internationalization/LegacyCulture.cpp b/Engine/Source/Runtime/Core/Private/Internationalization/LegacyCulture.cpp
--- a/Engine/Source/Runtime/Core/Private/Internationalization/LegacyCulture.cpp
+++ b/Engine/Source/Runtime/Core/Private/Internationalization/LegacyCulture.cpp
@@ -43,6 +43,11 @@ FString FCulture::FLegacyCultureImplementation::GetName() const
 	return Name;
 }
 
+FString FCulture::FLegacyCultureImplementation::GetCanonicalName(const FString& Name)
+{
+	return Name;
+}
+
 FString FCulture::FLegacyCultureImplementation::GetParentName(const FString& CultureName)
 {
 	FString Result;
diff --git a/Engine/Source/Runtime/Core/Private/Internationalization/LegacyCulture.h b/Engine/Source/Runtime/Core/Private/Internationalization/LegacyCulture.h
--- a/Engine/Source/Runtime/Core/Private/Internationalization/LegacyCulture.h
+++ b/Engine/Source/Runtime/Core/Private/Internationalization/LegacyCulture.h
@@ -28,6 +28,8 @@ class FCulture::FLegacyCultureImplementation
 
 	FString GetName() const;
 
+	static FString GetCanonicalName(const FString& Name);
+
 	static FString GetParentName(const FString& CultureName);
 
 	FString GetNativeName() const;
diff --git a/Engine/Source/Runtime/Core/Public/Internationalization/Culture.h b/Engine/Source/Runtime/Core/Public/Internationalization/Culture.h
--- a/Engine/Source/Runtime/Core/Public/Internationalization/Culture.h
+++ b/Engine/Source/Runtime/Core/Public/Internationalization/Culture.h
@@ -41,6 +41,8 @@ public:
 
 	FString GetName() const;
 
+	static FString GetCanonicalName(const FString& Name);
+
 	static FString GetParentName(const FString& CultureName);
 	
 	FString GetNativeName() const;
