Improved loading of text localization resources so that directories for culture-specific localization data can use names like "en-US" as well as "en_US". Specifically made to resolve issues with "zh-Hans-CN".

diff --git a/Engine/Source/Runtime/Core/Private/Internationalization/TextLocalizationManager.cpp b/Engine/Source/Runtime/Core/Private/Internationalization/TextLocalizationManager.cpp
--- a/Engine/Source/Runtime/Core/Private/Internationalization/TextLocalizationManager.cpp
+++ b/Engine/Source/Runtime/Core/Private/Internationalization/TextLocalizationManager.cpp
@@ -322,28 +322,78 @@ void FTextLocalizationManager::LoadResources(const bool ShouldLoadEditor, const
 	// Prioritized array of localization entry trackers.
 	TArray<FLocalizationEntryTracker> LocalizationEntryTrackers;
 
+	const auto MapCulturesToDirectories = [](const FString& LocalizationPath) -> TMap<FString, FString>
+	{
+		TMap<FString, FString> CultureToDirectoryMap;
+		IFileManager& FileManager = IFileManager::Get();
+
+		/* Visitor class used to enumerate directories of culture */
+		class FCultureDirectoryMapperVistor : public IPlatformFile::FDirectoryVisitor
+		{
+		public:
+			FCultureDirectoryMapperVistor( TMap<FString, FString>& OutCultureToDirectoryMap )
+				: CultureToDirectoryMap(OutCultureToDirectoryMap)
+			{
+			}
+
+			virtual bool Visit(const TCHAR* FilenameOrDirectory, bool bIsDirectory) override
+			{
+				if(bIsDirectory)
+				{
+					// UE localization resource folders use "en-US" style while ICU uses "en_US"
+					const FString LocalizationFolder = FPaths::GetCleanFilename(FilenameOrDirectory);
+					const FString CanonicalName = FCulture::GetCanonicalName(LocalizationFolder);
+					CultureToDirectoryMap.Add(CanonicalName, LocalizationFolder);
+				}
+
+				return true;
+			}
+
+			/** Array to fill with the names of the UE localization folders available at the given path */
+			TMap<FString, FString>& CultureToDirectoryMap;
+		};
+
+		FCultureDirectoryMapperVistor CultureEnumeratorVistor(CultureToDirectoryMap);
+		FileManager.IterateDirectory(*LocalizationPath, CultureEnumeratorVistor);
+
+		return CultureToDirectoryMap;
+	};
+
+	TMap< FString, TMap<FString, FString> > LocalizationPathToCultureDirectoryMap;
+	for (const FString& LocalizationPath : LocalizationPaths)
+	{
+		LocalizationPathToCultureDirectoryMap.Add(LocalizationPath, MapCulturesToDirectories(LocalizationPath));
+	}
+
 	// Read culture localization resources.
 	FLocalizationEntryTracker& CultureTracker = LocalizationEntryTrackers[LocalizationEntryTrackers.Add(FLocalizationEntryTracker())];
-	for(int32 PathIndex = 0; PathIndex < LocalizationPaths.Num(); ++PathIndex)
+	for (const FString& LocalizationPath : LocalizationPaths)
 	{
-		const FString& LocalizationPath = LocalizationPaths[PathIndex];
-		const FString CulturePath = LocalizationPath / CultureName;
+		const FString* const Entry = LocalizationPathToCultureDirectoryMap[LocalizationPath].Find(CultureName);
+		if (Entry)
+		{
+			const FString CulturePath = LocalizationPath / (*Entry);
 
-		CultureTracker.ReadFromDirectory(CulturePath);
+			CultureTracker.ReadFromDirectory(CulturePath);
+		}
 	}
 	CultureTracker.ReportCollisions();
 
 	// Read base language localization resources.
 	FLocalizationEntryTracker& BaseLanguageTracker = LocalizationEntryTrackers[LocalizationEntryTrackers.Add(FLocalizationEntryTracker())];
-	for(int32 PathIndex = 0; PathIndex < LocalizationPaths.Num(); ++PathIndex)
+	for (const FString& LocalizationPath : LocalizationPaths)
 	{
-		const FString& LocalizationPath = LocalizationPaths[PathIndex];
-		const FString BaseLanguagePath = LocalizationPath / BaseLanguageName;
-		const FString CulturePath = LocalizationPath / CultureName;
+		const FString* const BaseEntry = LocalizationPathToCultureDirectoryMap[LocalizationPath].Find(BaseLanguageName);
+		const FString* const Entry = LocalizationPathToCultureDirectoryMap[LocalizationPath].Find(CultureName);
 
-		if( BaseLanguagePath != CulturePath )
+		if (BaseEntry)
 		{
-			BaseLanguageTracker.ReadFromDirectory(BaseLanguagePath);
+			// If there is no entry or the paths are different.
+			if( !Entry || *BaseEntry != *Entry )
+			{
+				const FString BaseLanguagePath = LocalizationPath / (*BaseEntry);
+				BaseLanguageTracker.ReadFromDirectory(BaseLanguagePath);
+			}
 		}
 	}
 	BaseLanguageTracker.ReportCollisions();
