Fixing CIS (Mac)

diff --git a/Engine/Source/Runtime/Core/Private/Mac/MacApplication.cpp b/Engine/Source/Runtime/Core/Private/Mac/MacApplication.cpp
--- a/Engine/Source/Runtime/Core/Private/Mac/MacApplication.cpp
+++ b/Engine/Source/Runtime/Core/Private/Mac/MacApplication.cpp
@@ -11,7 +11,7 @@
 #include "HIDInputInterface.h"
 #include "AnalyticsEventAttribute.h"
 #include "IAnalyticsProvider.h"
-
+#include "CocoaThread.h"
 #include "ModuleManager.h"
 
 static NSString* NSWindowDraggingFinished = @"NSWindowDraggingFinished";
