Use a dispatch semaphore and always call CFRunLoopWakeUp when synchronously performing blocks on a different run loop to address a deadlock risk with nested dispatching.

diff --git a/Engine/Source/Runtime/Core/Private/Mac/CocoaThread.cpp b/Engine/Source/Runtime/Core/Private/Mac/CocoaThread.cpp
--- a/Engine/Source/Runtime/Core/Private/Mac/CocoaThread.cpp
+++ b/Engine/Source/Runtime/Core/Private/Mac/CocoaThread.cpp
@@ -49,17 +49,21 @@ static FCriticalSection GRunLoopModesLock;
 		
 		if(bWait)
 		{
-			__block bool bDone = false;
+			dispatch_semaphore_t Semaphore = dispatch_semaphore_create(0);
 			{
 				FScopeLock Lock(&GRunLoopModesLock);
 				TArray<NSString*>& Modes = GRunLoopModes.FindOrAdd([NSRunLoop currentRunLoop]);
 				Modes.Add([WaitMode retain]);
 			}
-			[RunLoop performBlock:^{ CopiedBlock(); bDone = true; } forMode:RunMode wake:true];
-			while(!bDone)
+			dispatch_block_t ExecuteBlock = Block_copy(^{ CopiedBlock(); dispatch_semaphore_signal(Semaphore); });
+			[RunLoop performBlock:ExecuteBlock forMode:RunMode wake:true];
+			while(dispatch_semaphore_wait(Semaphore, DISPATCH_TIME_NOW))
 			{
 				CFRunLoopRunInMode((CFStringRef)WaitMode, 0, true);
+				CFRunLoopWakeUp([RunLoop getCFRunLoop]);
 			}
+			dispatch_release(Semaphore);
+			Block_release(ExecuteBlock);
 			{
 				FScopeLock Lock(&GRunLoopModesLock);
 				TArray<NSString*>& Modes = GRunLoopModes.FindOrAdd([NSRunLoop currentRunLoop]);
