Added IsIdentity() for FTransform2D and FMatrix2x2.

diff --git a/Engine/Source/Runtime/Core/Public/Math/TransformCalculus2D.h b/Engine/Source/Runtime/Core/Public/Math/TransformCalculus2D.h
--- a/Engine/Source/Runtime/Core/Public/Math/TransformCalculus2D.h
+++ b/Engine/Source/Runtime/Core/Public/Math/TransformCalculus2D.h
@@ -446,6 +446,13 @@ public:
 		return FScale2D(FMath::Sqrt(ScaleSquared.GetVector().X), FMath::Sqrt(ScaleSquared.GetVector().Y));
 	}
 
+	/** Determines if the matrix is identity or not. Uses exact float comparison, so rounding error is not considered. */
+	bool IsIdentity() const
+	{
+		return M[0][0] == 1.0f && M[0][1] == 0.0f
+			&& M[1][0] == 0.0f && M[1][1] == 1.0f;
+	}
+
 private:
 	float M[2][2];
 };
@@ -606,6 +613,14 @@ public:
 	/** Access to the translation */
 	const FVector2D& GetTranslation() const { return Trans; }
 
+	/**
+	 * Specialized function to determine if a transform is precisely the identity transform. Uses exact float comparison, so rounding error is not considered.
+	 */
+	bool IsIdentity() const
+	{
+		return M.IsIdentity() && Trans == FVector2D::ZeroVector;
+	}
+
 private:
 	FMatrix2x2 M;
 	FVector2D Trans;
