added comments

diff --git a/Engine/Source/Runtime/Core/Public/Misc/ConfigCacheIni.h b/Engine/Source/Runtime/Core/Public/Misc/ConfigCacheIni.h
--- a/Engine/Source/Runtime/Core/Public/Misc/ConfigCacheIni.h
+++ b/Engine/Source/Runtime/Core/Public/Misc/ConfigCacheIni.h
@@ -505,8 +505,8 @@ public:
 	 * @param ConfigFile The output object to fill
 	 * @param IniName Either a Base ini name (Engine) or a full ini name (WrangleContent). NO PATH OR EXTENSION SHOULD BE USED!
 	 * @param bIsBaseIniName true if IniName is a Base name, which can be overridden on commandline, etc.
-	 * @param Platform The platform to use for Base ini names
-	 * @param GameName The game to use for Base ini names
+	 * @param Platform The platform to use for Base ini names, NULL means to use the current platform
+	 * @param GameName The game to use for Base ini names, NULL means to use the current
 	 */
 	static void LoadLocalIniFile(FConfigFile& ConfigFile, const TCHAR* IniName, bool bIsBaseIniName, const TCHAR* Platform=NULL, const TCHAR* GameName=NULL);
 
