editor console(OutputLog) is now saving it's history like the game console did

diff --git a/Engine/Source/Runtime/Core/Private/HAL/ConsoleManager.cpp b/Engine/Source/Runtime/Core/Private/HAL/ConsoleManager.cpp
--- a/Engine/Source/Runtime/Core/Private/HAL/ConsoleManager.cpp
+++ b/Engine/Source/Runtime/Core/Private/HAL/ConsoleManager.cpp
@@ -7,6 +7,7 @@ ConsoleManager.cpp: console command handling
 #include "Core.h"
 #include "ConsoleManager.h"
 #include "ModuleManager.h"
+#include "RemoteConfigIni.h"
 
 DEFINE_LOG_CATEGORY(LogConsoleResponse);
 DEFINE_LOG_CATEGORY_STATIC(LogConsoleManager, Log, All);
@@ -753,6 +754,60 @@ void FConsoleManager::UnregisterConsoleObject(const TCHAR* Name, bool bKeepState
 	}
 }
 
+void FConsoleManager::LoadHistoryIfNeeded()
+{
+	if(bHistoryWasLoaded)
+	{
+		return;
+	}
+
+	bHistoryWasLoaded = true;
+
+	HistoryEntries.Empty();
+
+	FConfigFile Ini;
+
+	FString ConfigPath = FPaths::GeneratedConfigDir() + TEXT("ConsoleHistory.ini");
+	ProcessIniContents(*ConfigPath, *ConfigPath, &Ini, false, false, false);
+
+	const FString History = TEXT("History");
+
+	FConfigSection* Section = Ini.Find(TEXT("ConsoleHistory"));
+
+	if(Section)
+	{
+		for (auto It : *Section)
+		{
+			const FString& Key = It.Key.ToString();
+
+			if(Key == History)
+			{
+				HistoryEntries.Add(It.Value);
+			}
+		}
+	}
+}
+
+void FConsoleManager::SaveHistory()
+{
+	const FName History = TEXT("History");
+
+	FConfigFile Ini;
+
+	FString ConfigPath = FPaths::GeneratedConfigDir() + TEXT("ConsoleHistory.ini");
+
+	FConfigSection& Section = Ini.Add(TEXT("ConsoleHistory"));
+
+	for(auto It : HistoryEntries)
+	{
+		Section.Add(History, It);
+	}
+
+	Ini.Dirty = true;
+	Ini.Write(ConfigPath);
+}
+
+
 void FConsoleManager::ForEachConsoleObject(const FConsoleObjectVisitor& Visitor, const TCHAR* ThatStartsWith) const
 {
 	check(Visitor.IsBound());
@@ -1058,11 +1113,23 @@ void IConsoleManager::SetupSingleton()
 
 void FConsoleManager::AddConsoleHistoryEntry(const TCHAR* Input)
 {
+	LoadHistoryIfNeeded();
+
+	// limit size to avoid a ever growing file
+	while(HistoryEntries.Num() > 64)
+	{
+		HistoryEntries.RemoveAt(0);
+	}
+
 	HistoryEntries.Add(FString(Input));
+
+	SaveHistory();
 }
 
-void FConsoleManager::GetConsoleHistory(TArray<FString>& Out) const
+void FConsoleManager::GetConsoleHistory(TArray<FString>& Out)
 {
+	LoadHistoryIfNeeded();
+
 	Out = HistoryEntries;
 }
 
diff --git a/Engine/Source/Runtime/Core/Public/HAL/ConsoleManager.h b/Engine/Source/Runtime/Core/Public/HAL/ConsoleManager.h
--- a/Engine/Source/Runtime/Core/Public/HAL/ConsoleManager.h
+++ b/Engine/Source/Runtime/Core/Public/HAL/ConsoleManager.h
@@ -12,7 +12,8 @@ class CORE_API FConsoleManager :public IConsoleManager
 public:
 	/** constructor */
 	FConsoleManager()
-		: ThreadPropagationCallback(0)
+		: bHistoryWasLoaded(false)
+		, ThreadPropagationCallback(0)
 		, ThreadPropagationThreadId(0)
 		, bCallAllConsoleVariableSinks(true)
 	{
@@ -65,7 +66,7 @@ public:
 	virtual void ForEachConsoleObject(const FConsoleObjectVisitor& Visitor, const TCHAR* ThatStartsWith) const override;
 	virtual bool ProcessUserConsoleInput(const TCHAR* InInput, FOutputDevice& Ar, UWorld* InWorld) override;
 	virtual void AddConsoleHistoryEntry(const TCHAR* Input);
-	virtual void GetConsoleHistory(TArray<FString>& Out) const; 
+	virtual void GetConsoleHistory(TArray<FString>& Out);
 	virtual bool IsNameRegistered(const TCHAR* Name) const override;	
 	virtual void RegisterThreadPropagation(uint32 ThreadId, IConsoleThreadPropagation* InCallback) override;
 	virtual void UnregisterConsoleObject( IConsoleObject* Object, bool bKeepState) override;
@@ -77,6 +78,7 @@ private: // ----------------------------------------------------
 	TMap<FString, IConsoleObject*> ConsoleObjects;
 
 	TArray<FString>	HistoryEntries;
+	bool bHistoryWasLoaded;
 	TArray<FConsoleCommandDelegate>	ConsoleVariableChangeSinks;
 
 	IConsoleThreadPropagation* ThreadPropagationCallback;
@@ -128,5 +130,11 @@ private: // ----------------------------------------------------
 	 * @param bool bKeepState if the current state is kept in memory until a cvar with the same name is registered
 	 */
 	void UnregisterConsoleObject(const TCHAR* Name, bool bKeepState);
+
+	// clears HistoryEntries and reads it from the .ini file
+	void LoadHistoryIfNeeded();
+
+	// clears HistoryEntries and writes it the .ini file
+	void SaveHistory();
 };
 
diff --git a/Engine/Source/Runtime/Core/Public/HAL/IConsoleManager.h b/Engine/Source/Runtime/Core/Public/HAL/IConsoleManager.h
--- a/Engine/Source/Runtime/Core/Public/HAL/IConsoleManager.h
+++ b/Engine/Source/Runtime/Core/Public/HAL/IConsoleManager.h
@@ -478,7 +478,7 @@ struct CORE_API IConsoleManager
 	
 	/**
 	 */
-	virtual void GetConsoleHistory(TArray<FString>& Out) const = 0; 
+	virtual void GetConsoleHistory(TArray<FString>& Out) = 0; 
 
 	/**
 	 * Check if a name (command or variable) has been registered with the console manager
