Disabled internationalization smoke tests. Temporary solution to work around the issue where smoke tests change cultures and cause those translations to stay resident, leading to unexpected mixtures of different translations from different cultures appearing.

diff --git a/Engine/Source/Runtime/Core/Private/Tests/Internationalization/DateTimeFormattingRulesTest.cpp b/Engine/Source/Runtime/Core/Private/Tests/Internationalization/DateTimeFormattingRulesTest.cpp
--- a/Engine/Source/Runtime/Core/Private/Tests/Internationalization/DateTimeFormattingRulesTest.cpp
+++ b/Engine/Source/Runtime/Core/Private/Tests/Internationalization/DateTimeFormattingRulesTest.cpp
@@ -6,7 +6,7 @@
 // Disable optimization for DateTimeFormattingRulesTest as it compiles very slowly in development builds
 PRAGMA_DISABLE_OPTIMIZATION
 
-IMPLEMENT_SIMPLE_AUTOMATION_TEST(FDateTimeFormattingRulesTest, "Core.Misc.DateTime Formatting Rules", EAutomationTestFlags::ATF_SmokeTest)
+IMPLEMENT_SIMPLE_AUTOMATION_TEST(FDateTimeFormattingRulesTest, "Core.Misc.DateTime Formatting Rules", EAutomationTestFlags::ATF_None)
 
 namespace
 {
diff --git a/Engine/Source/Runtime/Core/Private/Tests/Internationalization/NumberFormatingRulesTest.cpp b/Engine/Source/Runtime/Core/Private/Tests/Internationalization/NumberFormatingRulesTest.cpp
--- a/Engine/Source/Runtime/Core/Private/Tests/Internationalization/NumberFormatingRulesTest.cpp
+++ b/Engine/Source/Runtime/Core/Private/Tests/Internationalization/NumberFormatingRulesTest.cpp
@@ -10,7 +10,7 @@
 // Disable optimization for NumberFormatingRulesTest as it compiles very slowly in development builds
 PRAGMA_DISABLE_OPTIMIZATION
 
-IMPLEMENT_SIMPLE_AUTOMATION_TEST(FNumberFormattingRulesTest, "Core.Misc.Number Formatting Rules", EAutomationTestFlags::ATF_ApplicationMask)
+IMPLEMENT_SIMPLE_AUTOMATION_TEST(FNumberFormattingRulesTest, "Core.Misc.Number Formatting Rules", EAutomationTestFlags::ATF_None)
 
 namespace
 {
diff --git a/Engine/Source/Runtime/Core/Private/Tests/Internationalization/TextTest.cpp b/Engine/Source/Runtime/Core/Private/Tests/Internationalization/TextTest.cpp
--- a/Engine/Source/Runtime/Core/Private/Tests/Internationalization/TextTest.cpp
+++ b/Engine/Source/Runtime/Core/Private/Tests/Internationalization/TextTest.cpp
@@ -10,7 +10,7 @@ PRAGMA_DISABLE_OPTIMIZATION
 
 #define LOCTEXT_NAMESPACE "Core.Tests.TextFormatTest"
 
-IMPLEMENT_SIMPLE_AUTOMATION_TEST(FTextTest, "Core.Misc.Text", EAutomationTestFlags::ATF_ApplicationMask)
+IMPLEMENT_SIMPLE_AUTOMATION_TEST(FTextTest, "Core.Misc.Text", EAutomationTestFlags::ATF_None)
 
 namespace
 {
