Watcher - Checking in a suspected fix for the build.

diff --git a/Engine/Source/Runtime/Core/Private/IOS/IOSInputInterface.cpp b/Engine/Source/Runtime/Core/Private/IOS/IOSInputInterface.cpp
--- a/Engine/Source/Runtime/Core/Private/IOS/IOSInputInterface.cpp
+++ b/Engine/Source/Runtime/Core/Private/IOS/IOSInputInterface.cpp
@@ -2,6 +2,7 @@
 
 #include "Core.h"
 #include "IOSInputInterface.h"
+#include "IOSAppDelegate.h"
 
 DECLARE_LOG_CATEGORY_EXTERN(LogIOSInput, Log, All);
 
