Core - Fixed a typo, GetEngineMode() is now saved correctly

diff --git a/Engine/Source/Runtime/Core/Private/Windows/WindowsPlatformExceptionHandling.cpp b/Engine/Source/Runtime/Core/Private/Windows/WindowsPlatformExceptionHandling.cpp
--- a/Engine/Source/Runtime/Core/Private/Windows/WindowsPlatformExceptionHandling.cpp
+++ b/Engine/Source/Runtime/Core/Private/Windows/WindowsPlatformExceptionHandling.cpp
@@ -171,7 +171,7 @@ void SetReportParameters( HREPORT ReportHandle, EXCEPTION_POINTERS* ExceptionInf
 	StringCchPrintf( StringBuffer, MAX_SPRINTF, TEXT( "!%s!AssertLog=\"%s\"" ), FCommandLine::Get(), LocalBuffer );
 	Result = WerReportSetParameter( ReportHandle, WER_P8, TEXT( "Commandline" ), StringBuffer );
 
-	StringCchPrintf( StringBuffer, MAX_SPRINTF, TEXT( "%s!%s!%s!%d" ), TEXT( BRANCH_NAME ), FPlatformProcess::BaseDir(), FPlatformMisc::GetEngineMode, BUILT_FROM_CHANGELIST );
+	StringCchPrintf( StringBuffer, MAX_SPRINTF, TEXT( "%s!%s!%s!%d" ), TEXT( BRANCH_NAME ), FPlatformProcess::BaseDir(), FPlatformMisc::GetEngineMode(), BUILT_FROM_CHANGELIST );
 	Result = WerReportSetParameter( ReportHandle, WER_P9, TEXT( "BranchBaseDir" ), StringBuffer );
 }
 
