Major improvement to memory usage from ICU. When ICU requests a file to be opened and loaded into memory, we now cache that data and return it if further requests for the file are made.

diff --git a/Engine/Source/Runtime/Core/Private/Internationalization/Internationalization.cpp b/Engine/Source/Runtime/Core/Private/Internationalization/Internationalization.cpp
--- a/Engine/Source/Runtime/Core/Private/Internationalization/Internationalization.cpp
+++ b/Engine/Source/Runtime/Core/Private/Internationalization/Internationalization.cpp
@@ -19,12 +19,23 @@
 #include <unicode/timezone.h>
 #include <unicode/uclean.h>
 #include <unicode/udata.h>
+#include <unicode/umachine.h>
 #endif
 
 #define LOCTEXT_NAMESPACE "Internationalization"
 
 FInternationalization* FInternationalization::Instance;
 
+#if UE_ENABLE_ICU
+// Simple struct that operates as a namespace for ICU data file callbacks.
+// Struct used instead of namespace in order to avoid having to friend individual methods, which requires forward declaration and thus would create a "public" ICU dependency in the header.
+struct FICUDataCallbacks
+{
+	static UBool OpenDataFile(const void* context, void** fileContext, void** contents, const char* path);
+	static void CloseDataFile(const void* context, void* const fileContext, void* const contents);
+};
+#endif
+
 FInternationalization& FInternationalization::Get()
 {
 	if(!Instance)
@@ -163,57 +174,7 @@ namespace
 #if STATS
 	int64 FICUOverrides::BytesInUseCount = 0;
 	int64 FICUOverrides::CachedBytesInUseCount = 0;
-	static int64 DataFileBytesInUseCount = 0;
-	static int64 CachedDataFileBytesInUseCount = 0;
 #endif
-
-	UBool OpenDataFile(const void* context, void** fileContext, void** contents, const char* path)
-	{
-		FArchive* FileAr = IFileManager::Get().CreateFileReader(StringCast<TCHAR>(path).Get());
-		if(!FileAr)
-		{
-			return FALSE;
-		}
-
-		int64 FileSize = FileAr->TotalSize();
-
-		*fileContext = FileAr;
-		*contents = FICUOverrides::Malloc(nullptr, FileSize);
-
-#if STATS
-		DataFileBytesInUseCount += FMemory::GetAllocSize(*contents);
-		if(FThreadStats::IsThreadingReady() && CachedDataFileBytesInUseCount != DataFileBytesInUseCount)
-		{
-			SET_MEMORY_STAT(STAT_MemoryICUDataFileAllocationSize, DataFileBytesInUseCount);
-			CachedDataFileBytesInUseCount = DataFileBytesInUseCount;
-		}
-#endif
-
-		FileAr->Serialize(*contents, FileSize); 
-
-		delete FileAr;
-
-		return TRUE;
-	}
-
-	void CloseDataFile(const void* context, void* const fileContext, void* const contents)
-	{
-		if(fileContext)
-		{
-			if(contents)
-			{
-#if STATS
-				DataFileBytesInUseCount -= FMemory::GetAllocSize(contents);
-				if(FThreadStats::IsThreadingReady() && CachedDataFileBytesInUseCount != DataFileBytesInUseCount)
-				{
-					SET_MEMORY_STAT(STAT_MemoryICUDataFileAllocationSize, DataFileBytesInUseCount);
-					CachedDataFileBytesInUseCount = DataFileBytesInUseCount;
-				}
-#endif
-				FICUOverrides::Free(nullptr, contents);
-			}
-		}
-	}
 }
 #endif
 
@@ -276,7 +237,7 @@ void FInternationalization::Initialize()
 	};
 	checkf( HasFoundDataDirectory, TEXT("ICU data directory was not discovered:\n%s"), *(GetPrioritizedDataDirectoriesString()) );
 
-	u_setDataFileFunctions(nullptr, &OpenDataFile, &CloseDataFile, &(ICUStatus));
+	u_setDataFileFunctions(nullptr, &FICUDataCallbacks::OpenDataFile, &FICUDataCallbacks::CloseDataFile, &(ICUStatus));
 	u_init(&(ICUStatus));
 #endif /*UE_ENABLE_ICU*/
 
@@ -300,7 +261,6 @@ void FInternationalization::Terminate()
 	DefaultCulture.Reset();
 	InvariantCulture.Reset();
 	AllCultures.Empty();
-	bIsInitialized = false;
 
 #if UE_ENABLE_ICU
 	u_cleanup();
@@ -311,6 +271,7 @@ void FInternationalization::Terminate()
 
 	delete Instance;
 	Instance = nullptr;
+	bIsInitialized = false;
 }
 
 #if UE_ENABLE_ICU && (IS_PROGRAM || !IS_MONOLITHIC)
@@ -668,4 +629,131 @@ void FInternationalization::PopulateAllCultures(void)
 	InvariantCulture = GetCulture(TEXT(""));
 }
 
+#if UE_ENABLE_ICU
+#if STATS
+namespace
+{
+	int64 DataFileBytesInUseCount = 0;
+	int64 CachedDataFileBytesInUseCount = 0;
+}
+#endif
+
+UBool FICUDataCallbacks::OpenDataFile(const void* context, void** fileContext, void** contents, const char* path)
+{
+	auto& PathToCachedFileDataMap = FInternationalization::Get().PathToCachedFileDataMap;
+
+	// Try to find existing buffer
+	FInternationalization::FICUCachedFileData* CachedFileData = PathToCachedFileDataMap.Find(path);
+
+	// If there's no file context, we might have to load the file.
+	if (!CachedFileData)
+	{
+		// Attempt to load the file.
+		FArchive* FileAr = IFileManager::Get().CreateFileReader(StringCast<TCHAR>(path).Get());
+		if (FileAr)
+		{
+			const int64 FileSize = FileAr->TotalSize();
+
+			// Create file data.
+			CachedFileData = &(PathToCachedFileDataMap.Emplace(FString(path), FInternationalization::FICUCachedFileData(FileSize)));
+
+			// Load file into buffer.
+			FileAr->Serialize(CachedFileData->Buffer, FileSize); 
+			delete FileAr;
+
+			// Stat tracking.
+#if STATS
+			DataFileBytesInUseCount += FMemory::GetAllocSize(CachedFileData->Buffer);
+			if (FThreadStats::IsThreadingReady() && CachedDataFileBytesInUseCount != DataFileBytesInUseCount)
+			{
+				SET_MEMORY_STAT(STAT_MemoryICUDataFileAllocationSize, DataFileBytesInUseCount);
+				CachedDataFileBytesInUseCount = DataFileBytesInUseCount;
+			}
+#endif
+		}
+	}
+
+	// Add a reference, either the initial one or an additional one.
+	if (CachedFileData)
+	{
+		++(CachedFileData->ReferenceCount);
+	}
+
+	// Use the file path as the context, so we can look up the cached file data later and decrement its reference count.
+	*fileContext = CachedFileData ? new FString(path) : nullptr;
+
+	// Use the buffer from the cached file data.
+	*contents = CachedFileData ? CachedFileData->Buffer : nullptr;
+
+	// If we have cached file data, we must have loaded new data or found existing data, so we've successfully "opened" and "read" the file into "contents".
+	return CachedFileData != nullptr;
+}
+
+void FICUDataCallbacks::CloseDataFile(const void* context, void* const fileContext, void* const contents)
+{
+	// Early out on null context.
+	if (fileContext == nullptr)
+	{
+		return;
+	}
+
+	auto& PathToCachedFileDataMap = FInternationalization::Get().PathToCachedFileDataMap;
+
+	// The file context is the path to the file.
+	FString* const Path = reinterpret_cast<FString*>(fileContext);
+	check(Path);
+
+	// Look up the cached file data so we can maintain references.
+	FInternationalization::FICUCachedFileData* const CachedFileData = PathToCachedFileDataMap.Find(*Path);
+	check(CachedFileData);
+	check(CachedFileData->Buffer == contents);
+
+	// Remove a reference.
+	--(CachedFileData->ReferenceCount);
+
+	// If the last reference has been removed, the cached file data is not longer needed.
+	if (CachedFileData->ReferenceCount == 0)
+	{
+		// Stat tracking.
+#if STATS
+		DataFileBytesInUseCount -= FMemory::GetAllocSize(CachedFileData->Buffer);
+		if (FThreadStats::IsThreadingReady() && CachedDataFileBytesInUseCount != DataFileBytesInUseCount)
+		{
+			SET_MEMORY_STAT(STAT_MemoryICUDataFileAllocationSize, DataFileBytesInUseCount);
+			CachedDataFileBytesInUseCount = DataFileBytesInUseCount;
+		}
+#endif
+
+		// Delete the cached file data.
+		PathToCachedFileDataMap.Remove(*Path);
+	}
+
+	// The path string we allocated for tracking is no longer necessary.
+	delete Path;
+}
+
+FInternationalization::FICUCachedFileData::FICUCachedFileData(const int64 FileSize)
+	: ReferenceCount(0)
+	, Buffer( FICUOverrides::Malloc(nullptr, FileSize) )
+{
+}
+
+FInternationalization::FICUCachedFileData::FICUCachedFileData(FICUCachedFileData&& Source)
+	: ReferenceCount(Source.ReferenceCount)
+	, Buffer( Source.Buffer )
+{
+	// Make sure the moved source object doesn't retain the pointer and free the memory we now point to.
+	Source.Buffer = nullptr;
+}
+
+FInternationalization::FICUCachedFileData::~FICUCachedFileData()
+{
+	if (Buffer)
+	{
+		check(ReferenceCount == 0);
+		FICUOverrides::Free(nullptr, Buffer);
+	}
+}
+#endif
+
 #undef LOCTEXT_NAMESPACE
diff --git a/Engine/Source/Runtime/Core/Public/Internationalization/Internationalization.h b/Engine/Source/Runtime/Core/Public/Internationalization/Internationalization.h
--- a/Engine/Source/Runtime/Core/Public/Internationalization/Internationalization.h
+++ b/Engine/Source/Runtime/Core/Public/Internationalization/Internationalization.h
@@ -108,6 +108,24 @@ private:
 	TSharedPtr< FCulture, ESPMode::ThreadSafe > InvariantCulture;
 
 	TArray< void* > DLLHandles;
+
+#if UE_ENABLE_ICU
+	friend struct FICUDataCallbacks;
+
+	// Tidy class for storing the count of references for an ICU data file and the file's data itself.
+	struct FICUCachedFileData
+	{
+		FICUCachedFileData(const int64 FileSize);
+		FICUCachedFileData(FICUCachedFileData&& Source);
+		~FICUCachedFileData();
+
+		uint32 ReferenceCount;
+		void* Buffer;
+	};
+
+	// Map for associating ICU data file paths with cached file data, to prevent multiple copies of immutable ICU data files from residing in memory.
+	TMap<FString, FICUCachedFileData> PathToCachedFileDataMap;
+#endif
 };
 
 #undef LOC_DEFINE_REGION
