Resolved build issue related to FICUCachedFileData not being copy-constructible, a necessity for being a TMap value type when building with CLang.

diff --git a/Engine/Source/Runtime/Core/Private/Internationalization/Internationalization.cpp b/Engine/Source/Runtime/Core/Private/Internationalization/Internationalization.cpp
--- a/Engine/Source/Runtime/Core/Private/Internationalization/Internationalization.cpp
+++ b/Engine/Source/Runtime/Core/Private/Internationalization/Internationalization.cpp
@@ -738,6 +738,11 @@ FInternationalization::FICUCachedFileData::FICUCachedFileData(const int64 FileSi
 {
 }
 
+FInternationalization::FICUCachedFileData::FICUCachedFileData(const FICUCachedFileData& Source)
+{
+	checkf(false, TEXT("Cached file data for ICU may not be copy constructed. Something is trying to copy construct FICUCachedFileData."));
+}
+
 FInternationalization::FICUCachedFileData::FICUCachedFileData(FICUCachedFileData&& Source)
 	: ReferenceCount(Source.ReferenceCount)
 	, Buffer( Source.Buffer )
diff --git a/Engine/Source/Runtime/Core/Public/Internationalization/Internationalization.h b/Engine/Source/Runtime/Core/Public/Internationalization/Internationalization.h
--- a/Engine/Source/Runtime/Core/Public/Internationalization/Internationalization.h
+++ b/Engine/Source/Runtime/Core/Public/Internationalization/Internationalization.h
@@ -116,6 +116,7 @@ private:
 	struct FICUCachedFileData
 	{
 		FICUCachedFileData(const int64 FileSize);
+		FICUCachedFileData(const FICUCachedFileData& Source);
 		FICUCachedFileData(FICUCachedFileData&& Source);
 		~FICUCachedFileData();
 
