Specify explicit window levels for the different types of window in the editor on OS X, so that modal/tooltip/menu windows are on top, then parented or utility windows, then normal windows at the bottom. This should better match the behaviour under Windows, but should cause problems with tooltips/menus on modal windows or with hiding on deactivation.

diff --git a/Engine/Source/Runtime/Core/Private/Mac/MacWindow.cpp b/Engine/Source/Runtime/Core/Private/Mac/MacWindow.cpp
--- a/Engine/Source/Runtime/Core/Private/Mac/MacWindow.cpp
+++ b/Engine/Source/Runtime/Core/Private/Mac/MacWindow.cpp
@@ -101,13 +101,28 @@ void FMacWindow::Initialize( FMacApplication* const Application, const TSharedRe
 			[WindowHandle setDelegate: WindowHandle];
 
 			// @todo: We really need a window type in tab manager to know what window level to use and whether or not a window should hide on deactivate
-			if(Definition->IsRegularWindow && (!InParent.IsValid() || Definition->IsModalWindow || (!Definition->SupportsMaximize && !Definition->SupportsMinimize)))
+			if (Definition->IsModalWindow)
 			{
-				[WindowHandle setLevel: NSNormalWindowLevel];
+				[WindowHandle setLevel: NSModalPanelWindowLevel];
 			}
-			else
+			else if (Definition->IsRegularWindow)
+			{
+				if (InParent.IsValid())
+				{
+					[WindowHandle setLevel: NSFloatingWindowLevel];
+				}
+				else
+				{
+					[WindowHandle setLevel: NSNormalWindowLevel];
+				}
+			}
+			else if (!Definition->SupportsMaximize && !Definition->SupportsMinimize)
 			{
 				[WindowHandle setLevel: NSFloatingWindowLevel];
+			}
+			else
+			{
+				[WindowHandle setLevel: NSModalPanelWindowLevel];
 				[WindowHandle setHidesOnDeactivate: YES];
 			}
 			
