Put a scoped auto-release pool into every block that calls Objective-C code, including dispatched blocks to try and eliminate objects being autoreleased in the main run loop. Now when an over-release occurs it should happen within one of these blocks, making it easier to diagnose. Code that relies on autorelease has been annotated to make it clear that the caller must have properly configured the current auto-release pool.

diff --git a/Engine/Source/Runtime/Core/Private/Mac/CocoaMenu.cpp b/Engine/Source/Runtime/Core/Private/Mac/CocoaMenu.cpp
--- a/Engine/Source/Runtime/Core/Private/Mac/CocoaMenu.cpp
+++ b/Engine/Source/Runtime/Core/Private/Mac/CocoaMenu.cpp
@@ -18,6 +18,7 @@
 
 - (bool)isHighlightingKeyEquivalent
 {
+	SCOPED_AUTORELEASE_POOL;
 	FCocoaMenu* SuperMenu = [[self supermenu] isKindOfClass:[FCocoaMenu class]] ? (FCocoaMenu*)[self supermenu] : nil;
 	if ( SuperMenu )
 	{
@@ -31,6 +32,7 @@
 
 - (bool)highlightKeyEquivalent:(NSEvent *)Event
 {
+	SCOPED_AUTORELEASE_POOL;
 	bHighlightingKeyEquivalent = true;
 	bool bHighlighted = [super performKeyEquivalent:Event];
 	bHighlightingKeyEquivalent = false;
diff --git a/Engine/Source/Runtime/Core/Private/Mac/CocoaTextView.cpp b/Engine/Source/Runtime/Core/Private/Mac/CocoaTextView.cpp
--- a/Engine/Source/Runtime/Core/Private/Mac/CocoaTextView.cpp
+++ b/Engine/Source/Runtime/Core/Private/Mac/CocoaTextView.cpp
@@ -9,6 +9,7 @@
 
 - (id)initWithFrame:(NSRect)frame
 {
+	SCOPED_AUTORELEASE_POOL;
     self = [super initWithFrame:frame];
     if (self)
 	{
@@ -22,6 +23,7 @@
 
 - (void)drawRect:(NSRect)dirtyRect
 {
+	SCOPED_AUTORELEASE_POOL;
 	[super drawRect:dirtyRect];
 	
     // Drawing code here.
@@ -31,6 +33,7 @@
 {
 	if (IMMContext.IsValid())
 	{
+		SCOPED_AUTORELEASE_POOL;
 		reallyHandledEvent = true;
 		return [[self inputContext] handleEvent:theEvent] && reallyHandledEvent;
 	}
@@ -50,6 +53,7 @@
 
 - (void)mouseDown:(NSEvent *)theEvent
 {
+	SCOPED_AUTORELEASE_POOL;
 	if (IMMContext.IsValid())
 	{
 		[[self inputContext] handleEvent:theEvent];
@@ -66,12 +70,14 @@
 {
 	if (IMMContext.IsValid())
 	{
+		SCOPED_AUTORELEASE_POOL;
 		[[self inputContext] handleEvent:theEvent];
 	}
 }
 
 - (void)mouseUp:(NSEvent *)theEvent
 {
+	SCOPED_AUTORELEASE_POOL;
 	if (IMMContext.IsValid())
 	{
 		[[self inputContext] handleEvent:theEvent];
@@ -86,6 +92,7 @@
 
 - (void)rightMouseDown:(NSEvent*)Event
 {
+	SCOPED_AUTORELEASE_POOL;
 	FCocoaWindow* CocoaWindow = [[self window] isKindOfClass:[FCocoaWindow class]] ? (FCocoaWindow*)[self window] : nil;
 	if (CocoaWindow)
 	{
@@ -99,6 +106,7 @@
 
 - (void)otherMouseDown:(NSEvent*)Event
 {
+	SCOPED_AUTORELEASE_POOL;
 	FCocoaWindow* CocoaWindow = [[self window] isKindOfClass:[FCocoaWindow class]] ? (FCocoaWindow*)[self window] : nil;
 	if (CocoaWindow)
 	{
@@ -112,6 +120,7 @@
 
 - (void)rightMouseUp:(NSEvent*)Event
 {
+	SCOPED_AUTORELEASE_POOL;
 	FCocoaWindow* CocoaWindow = [[self window] isKindOfClass:[FCocoaWindow class]] ? (FCocoaWindow*)[self window] : nil;
 	if (CocoaWindow)
 	{
@@ -125,6 +134,7 @@
 
 - (void)otherMouseUp:(NSEvent*)Event
 {
+	SCOPED_AUTORELEASE_POOL;
 	FCocoaWindow* CocoaWindow = [[self window] isKindOfClass:[FCocoaWindow class]] ? (FCocoaWindow*)[self window] : nil;
 	if (CocoaWindow)
 	{
@@ -145,6 +155,7 @@
 
 - (void)activateInputMethod:(const TSharedRef<ITextInputMethodContext>&)InContext
 {
+	SCOPED_AUTORELEASE_POOL;
 	if (IMMContext.IsValid())
 	{
 		[self unmarkText];
@@ -157,6 +168,7 @@
 
 - (void)deactivateInputMethod
 {
+	SCOPED_AUTORELEASE_POOL;
 	[self unmarkText];
 	IMMContext = NULL;
 	[[self inputContext] deactivate];
@@ -174,6 +186,7 @@
  */
 - (void)insertText:(id)aString replacementRange:(NSRange)replacementRange
 {
+	SCOPED_AUTORELEASE_POOL;
 	if (IMMContext.IsValid() && [self hasMarkedText])
 	{
 		__block uint32 SelectionLocation;
@@ -210,6 +223,7 @@
 		}
 		
 		GameThreadCall(^{
+			SCOPED_AUTORELEASE_POOL;
 			FString TheFString(TheString);
 			IMMContext->SetTextInRange(SelectionLocation, SelectionLength, TheFString);
 			IMMContext->SetSelectionRange(SelectionLocation+TheFString.Len(), 0, ITextInputMethodContext::ECaretPosition::Ending);
@@ -237,6 +251,7 @@
 {
 	if (IMMContext.IsValid())
 	{
+		SCOPED_AUTORELEASE_POOL;
 		__block uint32 SelectionLocation;
 		__block uint32 SelectionLength;
 		if (replacementRange.location == NSNotFound)
@@ -310,6 +325,7 @@
 			}
 			
 			GameThreadCall(^{
+				SCOPED_AUTORELEASE_POOL;
 				FString TheFString(TheString);
 				IMMContext->SetTextInRange(SelectionLocation, SelectionLength, TheFString);
 				IMMContext->UpdateCompositionRange(CompositionRange.location, CompositionRange.length);
@@ -386,6 +402,7 @@
  */
 - (NSAttributedString *)attributedSubstringForProposedRange:(NSRange)aRange actualRange:(NSRangePointer)actualRange
 {
+	// Deliberately no autorelease pool here - relies on the OS setting one up beforehand, otherwise the returned object won't be properly ref-counted
 	NSAttributedString* AttributedString = nil;
 	if (IMMContext.IsValid())
 	{
@@ -414,7 +431,8 @@
     // We only allow these attributes to be set on our marked text (plus standard attributes)
     // NSMarkedClauseSegmentAttributeName is important for CJK input, among other uses
     // NSGlyphInfoAttributeName allows alternate forms of characters
-    return [NSArray arrayWithObjects:NSMarkedClauseSegmentAttributeName, NSGlyphInfoAttributeName, nil];
+	// Deliberately no autorelease pool here - relies on the OS setting one up beforehand, otherwise the returned object won't be properly ref-counted
+return [NSArray arrayWithObjects:NSMarkedClauseSegmentAttributeName, NSGlyphInfoAttributeName, nil];
 }
 
 /* Returns the first logical rectangular area for aRange. The return value is in the screen coordinate. The size value can be negative if the text flows to the left. If non-NULL, actuallRange contains the character range corresponding to the returned area.
@@ -423,6 +441,7 @@
 {
 	if (IMMContext.IsValid())
 	{
+		SCOPED_AUTORELEASE_POOL;
 		__block FVector2D Position;
 		__block FVector2D Size;
 		GameThreadCall(^{
@@ -471,6 +490,7 @@
  */
 - (NSInteger)windowLevel
 {
+	SCOPED_AUTORELEASE_POOL;
 	return [[self window] level];
 }
 
diff --git a/Engine/Source/Runtime/Core/Private/Mac/CocoaThread.cpp b/Engine/Source/Runtime/Core/Private/Mac/CocoaThread.cpp
--- a/Engine/Source/Runtime/Core/Private/Mac/CocoaThread.cpp
+++ b/Engine/Source/Runtime/Core/Private/Mac/CocoaThread.cpp
@@ -164,6 +164,7 @@ static FCriticalSection GRunLoopModesLock;
 
 NSString* InGameRunLoopMode(NSArray* AllowedModes)
 {
+	SCOPED_AUTORELEASE_POOL;
 	NSString* Mode = [[NSRunLoop gameRunLoop] intendedMode];
 	if(Mode == nil || ![AllowedModes containsObject:Mode])
 	{
@@ -189,6 +190,7 @@ void GameThreadCall(dispatch_block_t Block, NSString* SendMode, bool const bWait
 
 void RunGameThread(id Target, SEL Selector)
 {
+	SCOPED_AUTORELEASE_POOL;
 #if MAC_SEPARATE_GAME_THREAD
 	// Create a separate game thread and set it to the stack size to be the same as the main thread default of 8MB ( http://developer.apple.com/library/mac/#qa/qa1419/_index.html )
 	FCocoaGameThread* GameThread = [[FCocoaGameThread alloc] initWithTarget:Target selector:Selector object:nil];
@@ -201,12 +203,12 @@ void RunGameThread(id Target, SEL Selector)
 
 void ProcessGameThreadEvents(void)
 {
+	SCOPED_AUTORELEASE_POOL;
 #if MAC_SEPARATE_GAME_THREAD
 	// Make one pass through the loop, processing all events
 	CFRunLoopRef RunLoop = CFRunLoopGetCurrent();
 	CFRunLoopRunInMode(kCFRunLoopDefaultMode, 0, false);
 #else
-	SCOPED_AUTORELEASE_POOL;
 	while( NSEvent *Event = [NSApp nextEventMatchingMask: NSAnyEventMask untilDate: nil inMode: NSDefaultRunLoopMode dequeue: true] )
 	{
 		// Either the windowNumber is 0 or the window must be valid for the event to be processed.
diff --git a/Engine/Source/Runtime/Core/Private/Mac/CocoaWindow.cpp b/Engine/Source/Runtime/Core/Private/Mac/CocoaWindow.cpp
--- a/Engine/Source/Runtime/Core/Private/Mac/CocoaWindow.cpp
+++ b/Engine/Source/Runtime/Core/Private/Mac/CocoaWindow.cpp
@@ -51,6 +51,7 @@ NSString* NSPerformDragOperation = @"NSPerformDragOperation";
 
 - (NSRect)openGLFrame
 {
+	SCOPED_AUTORELEASE_POOL;
 	if(self.TargetWindowMode == EWindowMode::Fullscreen || WindowMode == EWindowMode::Fullscreen)
 	{
 		return self.PreFullScreenRect;
@@ -68,6 +69,7 @@ NSString* NSPerformDragOperation = @"NSPerformDragOperation";
 
 - (NSView*)openGLView
 {
+	SCOPED_AUTORELEASE_POOL;
 	if([self styleMask] & (NSTexturedBackgroundWindowMask))
 	{
 		NSView* SuperView = [[self contentView] superview];
@@ -95,6 +97,7 @@ NSString* NSPerformDragOperation = @"NSPerformDragOperation";
 	
 	if(bDeferOrderFront)
 	{
+		SCOPED_AUTORELEASE_POOL;
 		if(!(bDeferSetFrame || bDeferSetOrigin))
 		{
 			bDeferOrderFront = false;
@@ -111,7 +114,9 @@ NSString* NSPerformDragOperation = @"NSPerformDragOperation";
 {
 	if(bRenderInitialised && (bDeferSetFrame || bDeferSetOrigin))
 	{
+		SCOPED_AUTORELEASE_POOL;
 		dispatch_block_t Block = ^{
+			SCOPED_AUTORELEASE_POOL;
 			if(!bDeferSetFrame && bDeferSetOrigin)
 			{
 				DeferFrame.size = [self frame].size;
@@ -136,6 +141,7 @@ NSString* NSPerformDragOperation = @"NSPerformDragOperation";
 
 - (void)orderWindow:(NSWindowOrderingMode)OrderingMode relativeTo:(NSInteger)OtherWindowNumber
 {
+	SCOPED_AUTORELEASE_POOL;
 	bool bModal = FMacWindow::CurrentModalWindow() == nil || FMacWindow::CurrentModalWindow() == self || [self styleMask] == NSBorderlessWindowMask;
 	if(OrderingMode == NSWindowOut || bModal)
 	{
@@ -164,6 +170,7 @@ NSString* NSPerformDragOperation = @"NSPerformDragOperation";
 
 - (void)redrawContents
 {
+	SCOPED_AUTORELEASE_POOL;
 	if(bNeedsRedraw && bForwardEvents && ([self isVisible] && [super alphaValue] > 0.0f))
 	{
 		NSNotification* Notification = [NSNotification notificationWithName:NSWindowRedrawContents object:self];
@@ -189,6 +196,7 @@ NSString* NSPerformDragOperation = @"NSPerformDragOperation";
 
 - (void)orderFrontAndMakeMain:(bool)bMain andKey:(bool)bKey
 {
+	SCOPED_AUTORELEASE_POOL;
 	if ([NSApp isHidden] == NO)
 	{
 		bool bBringToFront = FMacWindow::CurrentModalWindow() == nil || FMacWindow::CurrentModalWindow() == self || [self styleMask] == NSBorderlessWindowMask;
@@ -211,17 +219,20 @@ NSString* NSPerformDragOperation = @"NSPerformDragOperation";
 // Following few methods overload NSWindow's methods from Cocoa API, so have to use Cocoa's BOOL (signed char), not bool (unsigned int)
 - (BOOL)canBecomeMainWindow
 {
+	SCOPED_AUTORELEASE_POOL;
 	bool bNoModalOrCurrent = FMacWindow::CurrentModalWindow() == nil || FMacWindow::CurrentModalWindow() == self;
 	return bAcceptsInput && ([self styleMask] != NSBorderlessWindowMask) && bNoModalOrCurrent;
 }
 
 - (BOOL)canBecomeKeyWindow
 {
+	SCOPED_AUTORELEASE_POOL;
 	return bAcceptsInput && ![self ignoresMouseEvents];
 }
 
 - (BOOL)validateMenuItem:(NSMenuItem *)MenuItem
 {
+	SCOPED_AUTORELEASE_POOL;
 	// Borderless windows we use do not automatically handle first responder's actions, so we force it here
 	return ([MenuItem action] == @selector(performClose:) || [MenuItem action] == @selector(performMiniaturize:) || [MenuItem action] == @selector(performZoom:)) ? YES : [super validateMenuItem:MenuItem];
 }
@@ -235,6 +246,7 @@ NSString* NSPerformDragOperation = @"NSPerformDragOperation";
 	}
 	else
 	{
+		SCOPED_AUTORELEASE_POOL;
 		if([self isVisible] && WindowAlpha > 0.0f)
 		{
 			[self performDeferredSetFrame];
@@ -245,6 +257,7 @@ NSString* NSPerformDragOperation = @"NSPerformDragOperation";
 
 - (void)orderOut:(id)Sender
 {
+	SCOPED_AUTORELEASE_POOL;
 	bDeferOrderFront = false;
 	
 	[super orderOut:Sender];
@@ -252,6 +265,7 @@ NSString* NSPerformDragOperation = @"NSPerformDragOperation";
 
 - (void)performClose:(id)Sender
 {
+	SCOPED_AUTORELEASE_POOL;
 	bDeferOrderFront = false;
 	
 	[self close];
@@ -259,17 +273,20 @@ NSString* NSPerformDragOperation = @"NSPerformDragOperation";
 
 - (void)performMiniaturize:(id)Sender
 {
+	SCOPED_AUTORELEASE_POOL;
 	[self miniaturize: self];
 }
 
 - (void)performZoom:(id)Sender
 {
+	SCOPED_AUTORELEASE_POOL;
 	bZoomed = !bZoomed;
 	[self zoom: self];
 }
 
 - (void)setFrame:(NSRect)FrameRect display:(BOOL)Flag
 {
+	SCOPED_AUTORELEASE_POOL;
 	NSSize Size = [self frame].size;
 	NSSize NewSize = FrameRect.size;
 	if(!bRenderInitialised || ([self isVisible] && [super alphaValue] > 0.0f && (Size.width > 1 || Size.height > 1 || NewSize.width > 1 || NewSize.height > 1)))
@@ -291,10 +308,12 @@ NSString* NSPerformDragOperation = @"NSPerformDragOperation";
 
 - (void)setFrameOrigin:(NSPoint)Point
 {
+	SCOPED_AUTORELEASE_POOL;
 	NSSize Size = [self frame].size;
 	if(!bRenderInitialised || ([self isVisible] && [super alphaValue] > 0.0f && (Size.width > 1 || Size.height > 1)))
 	{
 		MainThreadCall(^{
+			SCOPED_AUTORELEASE_POOL;
 			[super setFrameOrigin:Point];
 		});
 		bDeferSetOrigin = false;
@@ -349,6 +368,7 @@ NSString* NSPerformDragOperation = @"NSPerformDragOperation";
 
 - (void)windowDidBecomeKey:(NSNotification *)Notification
 {
+	SCOPED_AUTORELEASE_POOL;
 	if([NSApp isHidden] == NO)
 	{
 		if(FMacWindow::CurrentModalWindow() == nil || FMacWindow::CurrentModalWindow() == self || [self styleMask] == NSBorderlessWindowMask)
@@ -369,6 +389,7 @@ NSString* NSPerformDragOperation = @"NSPerformDragOperation";
 
 - (void)windowDidResignKey:(NSNotification *)Notification
 {
+	SCOPED_AUTORELEASE_POOL;
 	[self setMovable: YES];
 	[self setMovableByWindowBackground: NO];
 	
@@ -388,6 +409,7 @@ NSString* NSPerformDragOperation = @"NSPerformDragOperation";
 
 - (void)windowDidMove:(NSNotification *)Notification
 {
+	SCOPED_AUTORELEASE_POOL;
 	bZoomed = [self isZoomed];
 	
 	NSView* OpenGLView = [self openGLView];
@@ -407,6 +429,7 @@ NSString* NSPerformDragOperation = @"NSPerformDragOperation";
 	// It does however, work fine for handling display arrangement changes that cause a window to go offscreen.
 	if(bDisplayReconfiguring)
 	{
+		SCOPED_AUTORELEASE_POOL;
 		NSScreen* Screen = [self screen];
 		NSRect Frame = [self frame];
 		NSRect VisibleFrame = [Screen visibleFrame];
@@ -464,6 +487,7 @@ NSString* NSPerformDragOperation = @"NSPerformDragOperation";
 
 - (void)windowDidResize:(NSNotification *)Notification
 {
+	SCOPED_AUTORELEASE_POOL;
 	bZoomed = [self isZoomed];
 	if(self.bForwardEvents)
 	{
@@ -474,6 +498,7 @@ NSString* NSPerformDragOperation = @"NSPerformDragOperation";
 
 - (void)windowWillClose:(NSNotification *)notification
 {
+	SCOPED_AUTORELEASE_POOL;
 	if(self.bForwardEvents && MacApplication)
 	{
 		FMacEvent::SendToGameRunLoop(notification, self, EMacEventSendMethod::Async, InGameRunLoopMode(@[ UE4CloseEventMode ]));
@@ -494,6 +519,7 @@ NSString* NSPerformDragOperation = @"NSPerformDragOperation";
 {
 	// Really we shouldn't be doing this - on OS X only left-click changes focus,
 	// but for the moment it is easier than changing Slate.
+	SCOPED_AUTORELEASE_POOL;
 	if([self canBecomeKeyWindow])
 	{
 		[self makeKeyWindow];
@@ -546,6 +572,7 @@ NSString* NSPerformDragOperation = @"NSPerformDragOperation";
 {
 	if(self.bForwardEvents)
 	{
+		SCOPED_AUTORELEASE_POOL;
 		NSNotification* Notification = [NSNotification notificationWithName:NSDraggingExited object:Sender];
 		FMacEvent::SendToGameRunLoop(Notification, self, EMacEventSendMethod::Async);
 	}
@@ -555,6 +582,7 @@ NSString* NSPerformDragOperation = @"NSPerformDragOperation";
 {
 	if(self.bForwardEvents)
 	{
+		SCOPED_AUTORELEASE_POOL;
 		NSNotification* Notification = [NSNotification notificationWithName:NSDraggingUpdated object:Sender];
 		FMacEvent::SendToGameRunLoop(Notification, self, EMacEventSendMethod::Async);
 	}
@@ -565,6 +593,7 @@ NSString* NSPerformDragOperation = @"NSPerformDragOperation";
 {
 	if(self.bForwardEvents)
 	{
+		SCOPED_AUTORELEASE_POOL;
 		NSNotification* Notification = [NSNotification notificationWithName:NSPrepareForDragOperation object:Sender];
 		FMacEvent::SendToGameRunLoop(Notification, self, EMacEventSendMethod::Async);
 	}
@@ -575,6 +604,7 @@ NSString* NSPerformDragOperation = @"NSPerformDragOperation";
 {
 	if(self.bForwardEvents)
 	{
+		SCOPED_AUTORELEASE_POOL;
 		NSNotification* Notification = [NSNotification notificationWithName:NSPerformDragOperation object:Sender];
 		FMacEvent::SendToGameRunLoop(Notification, self, EMacEventSendMethod::Async);
 	}
@@ -583,6 +613,7 @@ NSString* NSPerformDragOperation = @"NSPerformDragOperation";
 
 - (BOOL)isMovable
 {
+	SCOPED_AUTORELEASE_POOL;
 	BOOL Movable = [super isMovable];
 	if(Movable && bRenderInitialised && MacApplication)
 	{
diff --git a/Engine/Source/Runtime/Core/Private/Mac/MacEvent.cpp b/Engine/Source/Runtime/Core/Private/Mac/MacEvent.cpp
--- a/Engine/Source/Runtime/Core/Private/Mac/MacEvent.cpp
+++ b/Engine/Source/Runtime/Core/Private/Mac/MacEvent.cpp
@@ -9,6 +9,7 @@
 FMacEvent::FMacEvent(NSEvent* const Event)
 : EventData(Event)
 {
+	SCOPED_AUTORELEASE_POOL;
 	check(Event);
 	[Event retain];
 	
@@ -27,6 +28,7 @@ FMacEvent::FMacEvent(NSEvent* const Event)
 FMacEvent::FMacEvent(NSNotification* const Notification, NSWindow* const Window)
 : EventData(Notification)
 {
+	SCOPED_AUTORELEASE_POOL;
 	check(Notification);
 	[Notification retain];
 	
@@ -44,6 +46,7 @@ FMacEvent::FMacEvent(NSNotification* const Notification, NSWindow* const Window)
 
 void FMacEvent::SendToGameRunLoop(FMacEvent const* const Event, EMacEventSendMethod SendMethod, NSString* SendMode)
 {
+	SCOPED_AUTORELEASE_POOL;
 	NSThread* GameThread = [NSThread gameThread];
 	check(GameThread);
 	
@@ -78,6 +81,7 @@ FMacEvent::~FMacEvent(void)
 {
 	if(EventData)
 	{
+		SCOPED_AUTORELEASE_POOL;
 		[EventData release];
 		EventData = nil;
 	}
@@ -85,6 +89,7 @@ FMacEvent::~FMacEvent(void)
 
 NSEvent* FMacEvent::GetEvent(void) const
 {
+	SCOPED_AUTORELEASE_POOL;
 	if(EventData && [EventData isKindOfClass:[NSEvent class]])
 	{
 		return (NSEvent*)EventData;
@@ -97,6 +102,7 @@ NSEvent* FMacEvent::GetEvent(void) const
 
 NSNotification* FMacEvent::GetNotification(void) const
 {
+	SCOPED_AUTORELEASE_POOL;
 	if(EventData && [EventData isKindOfClass:[NSNotification class]])
 	{
 		return (NSNotification*)EventData;
diff --git a/Engine/Source/Runtime/Core/Private/Mac/MacTextInputMethodSystem.cpp b/Engine/Source/Runtime/Core/Private/Mac/MacTextInputMethodSystem.cpp
--- a/Engine/Source/Runtime/Core/Private/Mac/MacTextInputMethodSystem.cpp
+++ b/Engine/Source/Runtime/Core/Private/Mac/MacTextInputMethodSystem.cpp
@@ -51,6 +51,7 @@ namespace
 		{
 			FCocoaWindow* CocoaWindow = (FCocoaWindow*)ContextWindow->GetOSWindowHandle();
 			MainThreadCall(^{
+				SCOPED_AUTORELEASE_POOL;
 				if(CocoaWindow && [CocoaWindow openGLView])
 				{
 					FCocoaTextView* TextView = (FCocoaTextView*)[CocoaWindow openGLView];
@@ -66,6 +67,7 @@ namespace
 		{
 			FCocoaWindow* CocoaWindow = (FCocoaWindow*)ContextWindow->GetOSWindowHandle();
 			MainThreadCall(^{
+				SCOPED_AUTORELEASE_POOL;
 				if(CocoaWindow && [CocoaWindow openGLView])
 				{
 					FCocoaTextView* TextView = (FCocoaTextView*)[CocoaWindow openGLView];
@@ -81,6 +83,7 @@ namespace
 		{
 			FCocoaWindow* CocoaWindow = (FCocoaWindow*)ContextWindow->GetOSWindowHandle();
 			MainThreadCall(^{
+				SCOPED_AUTORELEASE_POOL;
 				if(CocoaWindow && [CocoaWindow openGLView])
 				{
 					FCocoaTextView* TextView = (FCocoaTextView*)[CocoaWindow openGLView];
@@ -96,6 +99,7 @@ namespace
 		{
 			FCocoaWindow* CocoaWindow = (FCocoaWindow*)ContextWindow->GetOSWindowHandle();
 			MainThreadCall(^{
+				SCOPED_AUTORELEASE_POOL;
 				if(CocoaWindow && [CocoaWindow openGLView])
 				{
 					FCocoaTextView* TextView = (FCocoaTextView*)[CocoaWindow openGLView];
@@ -133,6 +137,7 @@ TSharedPtr<ITextInputMethodChangeNotifier> FMacTextInputMethodSystem::RegisterCo
 
 void FMacTextInputMethodSystem::UnregisterContext(const TSharedRef<ITextInputMethodContext>& Context)
 {
+	SCOPED_AUTORELEASE_POOL;
 	TWeakPtr<ITextInputMethodChangeNotifier>& NotifierRef = ContextMap[Context];
 	if(!NotifierRef.IsValid())
 	{
@@ -153,6 +158,7 @@ void FMacTextInputMethodSystem::UnregisterContext(const TSharedRef<ITextInputMet
 
 void FMacTextInputMethodSystem::ActivateContext(const TSharedRef<ITextInputMethodContext>& Context)
 {
+	SCOPED_AUTORELEASE_POOL;
 	TWeakPtr<ITextInputMethodChangeNotifier>& NotifierRef = ContextMap[Context];
 	if(!NotifierRef.IsValid())
 	{
@@ -191,6 +197,7 @@ void FMacTextInputMethodSystem::ActivateContext(const TSharedRef<ITextInputMetho
 
 void FMacTextInputMethodSystem::DeactivateContext(const TSharedRef<ITextInputMethodContext>& Context)
 {
+	SCOPED_AUTORELEASE_POOL;
 	TWeakPtr<ITextInputMethodChangeNotifier>& NotifierRef = ContextMap[Context];
 	if(!NotifierRef.IsValid())
 	{
@@ -228,6 +235,7 @@ void FMacTextInputMethodSystem::DeactivateContext(const TSharedRef<ITextInputMet
 
 bool FMacTextInputMethodSystem::IsActiveContext(const TSharedRef<ITextInputMethodContext>& Context) const
 {
+	SCOPED_AUTORELEASE_POOL;
 	const TWeakPtr<ITextInputMethodChangeNotifier>& NotifierRef = ContextMap[Context];
 	if(!NotifierRef.IsValid())
 	{
diff --git a/Engine/Source/Runtime/Core/Private/Mac/MacWindow.cpp b/Engine/Source/Runtime/Core/Private/Mac/MacWindow.cpp
--- a/Engine/Source/Runtime/Core/Private/Mac/MacWindow.cpp
+++ b/Engine/Source/Runtime/Core/Private/Mac/MacWindow.cpp
@@ -16,6 +16,7 @@ FMacWindow::~FMacWindow()
 	{
 		NSWindow* Window = WindowHandle;
 		dispatch_async(dispatch_get_main_queue(), ^{
+			SCOPED_AUTORELEASE_POOL;
 			[Window release];
 		});
 		WindowHandle = nil;
@@ -89,6 +90,7 @@ void FMacWindow::Initialize( FMacApplication* const Application, const TSharedRe
 	}
 
 	MainThreadCall(^{
+		SCOPED_AUTORELEASE_POOL;
 		WindowHandle = [[FCocoaWindow alloc] initWithContentRect: ViewRect styleMask: WindowStyle backing: NSBackingStoreBuffered defer: NO];
 		
 		if( WindowHandle != nullptr )
@@ -217,6 +219,7 @@ void FMacWindow::ReshapeWindow( int32 X, int32 Y, int32 Width, int32 Height )
 		if(WindowMode == EWindowMode::Windowed || WindowMode == EWindowMode::WindowedFullscreen)
 		{
 			MainThreadCall(^{
+				SCOPED_AUTORELEASE_POOL;
 				BOOL DisplayIfNeeded = (WindowMode == EWindowMode::Windowed);
 				
 				const int32 InvertedY = FPlatformMisc::ConvertSlateYPositionToCocoa(Y) - Height + 1;
@@ -280,6 +283,7 @@ void FMacWindow::Destroy()
 {
 	if( WindowHandle )
 	{
+		SCOPED_AUTORELEASE_POOL;
 		FCocoaWindow* Window = WindowHandle;
 		if(Definition->IsModalWindow)
 		{
@@ -301,6 +305,7 @@ void FMacWindow::Destroy()
 			
 			// Close the window
 			MainThreadCall(^{
+				SCOPED_AUTORELEASE_POOL;
 				[Window performClose:nil];
 			}, UE4CloseEventMode, true);
 		}
@@ -357,6 +362,7 @@ void FMacWindow::Show()
 		bool bMakeMainAndKey = ([WindowHandle canBecomeKeyWindow] && Definition->ActivateWhenFirstShown);
 		
 		MainThreadCall(^{
+			SCOPED_AUTORELEASE_POOL;
 			[WindowHandle orderFrontAndMakeMain:bMakeMainAndKey andKey:bMakeMainAndKey];
 		}, UE4ShowEventMode, true);
 		
@@ -390,6 +396,7 @@ void FMacWindow::Hide()
 			RemoveModalWindow(WindowHandle);
 		}
 		MainThreadCall(^{
+			SCOPED_AUTORELEASE_POOL;
 			[WindowHandle orderOut:nil];
 		}, UE4CloseEventMode, false);
 	}
@@ -423,6 +430,7 @@ void FMacWindow::SetWindowMode( EWindowMode::Type NewWindowMode )
 		WindowHandle.TargetWindowMode = NewWindowMode;
 		
 		MainThreadCall(^{
+			SCOPED_AUTORELEASE_POOL;
 			[WindowHandle setCollectionBehavior: Behaviour];
 			[WindowHandle toggleFullScreen:nil];
 		}, UE4FullscreenEventMode, true);
@@ -453,11 +461,13 @@ bool FMacWindow::IsMaximized() const
 
 bool FMacWindow::IsMinimized() const
 {
+	SCOPED_AUTORELEASE_POOL;
 	return [WindowHandle isMiniaturized];
 }
 
 bool FMacWindow::IsVisible() const
 {
+	SCOPED_AUTORELEASE_POOL;
 	return bIsVisible && [NSApp isHidden] == false;
 }
 
@@ -552,9 +562,9 @@ bool FMacWindow::IsForegroundWindow() const
 
 void FMacWindow::SetText(const TCHAR* const Text)
 {
-	SCOPED_AUTORELEASE_POOL;
 	CFStringRef CFName = FPlatformString::TCHARToCFString( Text );
 	MainThreadCall(^{
+		SCOPED_AUTORELEASE_POOL;
 		[WindowHandle setTitle: (NSString*)CFName];
 		if(IsRegularWindow())
 		{
@@ -579,6 +589,7 @@ void FMacWindow::OnDisplayReconfiguration(CGDirectDisplayID Display, CGDisplayCh
 	if(WindowHandle)
 	{
 		MainThreadCall(^{
+			SCOPED_AUTORELEASE_POOL;
 			if(Flags & kCGDisplayBeginConfigurationFlag)
 			{
 				[WindowHandle setMovable: YES];
@@ -599,6 +610,7 @@ bool FMacWindow::OnIMKKeyDown(NSEvent* Event)
 	if(WindowHandle && [WindowHandle openGLView])
 	{
 		return MainThreadReturn(^{
+			SCOPED_AUTORELEASE_POOL;
 			FCocoaTextView* View = (FCocoaTextView*)[WindowHandle openGLView];
 			return View && [View imkKeyDown:Event];
 		}, UE4IMEEventMode);
