Core - SetStoredValue/GetStoredValue uses platform specific implementation (fixed regression)

diff --git a/Engine/Source/Runtime/Core/Private/GenericPlatform/GenericPlatformMisc.cpp b/Engine/Source/Runtime/Core/Private/GenericPlatform/GenericPlatformMisc.cpp
--- a/Engine/Source/Runtime/Core/Private/GenericPlatform/GenericPlatformMisc.cpp
+++ b/Engine/Source/Runtime/Core/Private/GenericPlatform/GenericPlatformMisc.cpp
@@ -805,13 +805,13 @@ FGuid FGenericPlatformMisc::GetMachineId()
 	FString MachineIdStr;
 
 	// Check to see if we already have a valid machine ID to use
-	if( !GetStoredValue( TEXT( "Epic Games" ), TEXT( "Unreal Engine/Identifiers" ), TEXT( "MachineId" ), MachineIdStr ) || !FGuid::Parse( MachineIdStr, MachineId ) )
+	if( !FPlatformMisc::GetStoredValue( TEXT( "Epic Games" ), TEXT( "Unreal Engine/Identifiers" ), TEXT( "MachineId" ), MachineIdStr ) || !FGuid::Parse( MachineIdStr, MachineId ) )
 	{
 		// No valid machine ID, generate and save a new one
 		MachineId = FGuid::NewGuid();
 		MachineIdStr = MachineId.ToString( EGuidFormats::Digits );
 
-		if( !SetStoredValue( TEXT( "Epic Games" ), TEXT( "Unreal Engine/Identifiers" ), TEXT( "MachineId" ), MachineIdStr ) )
+		if( !FPlatformMisc::SetStoredValue( TEXT( "Epic Games" ), TEXT( "Unreal Engine/Identifiers" ), TEXT( "MachineId" ), MachineIdStr ) )
 		{
 			// Failed to persist the machine ID - reset it to zero to avoid returning a transient value
 			MachineId = FGuid();
@@ -824,13 +824,13 @@ FGuid FGenericPlatformMisc::GetMachineId()
 FString FGenericPlatformMisc::GetEpicAccountId()
 {
 	FString AccountId;
-	GetStoredValue( TEXT( "Epic Games" ), TEXT( "Unreal Engine/Identifiers" ), TEXT( "AccountId" ), AccountId );
+	FPlatformMisc::GetStoredValue( TEXT( "Epic Games" ), TEXT( "Unreal Engine/Identifiers" ), TEXT( "AccountId" ), AccountId );
 	return AccountId;
 }
 
 void FGenericPlatformMisc::SetEpicAccountId( const FString& AccountId )
 {
-	SetStoredValue( TEXT( "Epic Games" ), TEXT( "Unreal Engine/Identifiers" ), TEXT( "AccountId" ), AccountId );
+	FPlatformMisc::SetStoredValue( TEXT( "Epic Games" ), TEXT( "Unreal Engine/Identifiers" ), TEXT( "AccountId" ), AccountId );
 }
 
 const TCHAR* FGenericPlatformMisc::GetEngineMode()
@@ -840,4 +840,4 @@ const TCHAR* FGenericPlatformMisc::GetEngineMode()
 		GIsEditor ? TEXT( "Editor" ) :
 		IsRunningDedicatedServer() ? TEXT( "Server" ) :
 		TEXT( "Game" );
-}
+}
\ No newline at end of file
