Stats - Fixed a crash when the verbose group contains non-cycle stats, also fixes an issue with the missing data

diff --git a/Engine/Source/Runtime/Core/Private/Stats/Stats2.cpp b/Engine/Source/Runtime/Core/Private/Stats/Stats2.cpp
--- a/Engine/Source/Runtime/Core/Private/Stats/Stats2.cpp
+++ b/Engine/Source/Runtime/Core/Private/Stats/Stats2.cpp
@@ -144,12 +144,12 @@ void FStartupMessages::AddThreadMetadata( const FName InThreadName, uint32 InThr
 }
 
 
-void FStartupMessages::AddMetadata( FName InStatName, const TCHAR* InStatDesc, const char* InGroupName, const char* InGroupCategory, const TCHAR* InGroupDesc, bool bCanBeDisabled, EStatDataType::Type InStatType, bool bCycleStat, FPlatformMemory::EMemoryCounterRegion InMemoryRegion /*= FPlatformMemory::MCR_Invalid*/ )
+void FStartupMessages::AddMetadata( FName InStatName, const TCHAR* InStatDesc, const char* InGroupName, const char* InGroupCategory, const TCHAR* InGroupDesc, bool bShouldClearEveryFrame, EStatDataType::Type InStatType, bool bCycleStat, FPlatformMemory::EMemoryCounterRegion InMemoryRegion /*= FPlatformMemory::MCR_Invalid*/ )
 {
 	FScopeLock Lock( &CriticalSection );
 
 	new (DelayedMessages)FStatMessage( InGroupName, EStatDataType::ST_None, "Groups", InGroupCategory, InGroupDesc, false, false );
-	new (DelayedMessages)FStatMessage( InStatName, InStatType, InGroupName, InGroupCategory, InStatDesc, bCanBeDisabled, bCycleStat, InMemoryRegion );
+	new (DelayedMessages)FStatMessage( InStatName, InStatType, InGroupName, InGroupCategory, InStatDesc, bShouldClearEveryFrame, bCycleStat, InMemoryRegion );
 }
 
 
@@ -164,14 +164,14 @@ FStartupMessages& FStartupMessages::Get()
 	return *Messages;
 }
 
-void FThreadSafeStaticStatBase::DoSetup(const char* InStatName, const TCHAR* InStatDesc, const char* InGroupName, const char* InGroupCategory, const TCHAR* InGroupDesc, bool bDefaultEnable, bool bCanBeDisabled, EStatDataType::Type InStatType, bool bCycleStat, FPlatformMemory::EMemoryCounterRegion InMemoryRegion) const
+void FThreadSafeStaticStatBase::DoSetup(const char* InStatName, const TCHAR* InStatDesc, const char* InGroupName, const char* InGroupCategory, const TCHAR* InGroupDesc, bool bDefaultEnable, bool bShouldClearEveryFrame, EStatDataType::Type InStatType, bool bCycleStat, FPlatformMemory::EMemoryCounterRegion InMemoryRegion) const
 {
 	FName TempName(InStatName);
 
 	// send meta data, we don't use normal messages because the stats thread might not be running yet
-	FStartupMessages::Get().AddMetadata(TempName, InStatDesc, InGroupName, InGroupCategory, InGroupDesc, bCanBeDisabled, InStatType, bCycleStat, InMemoryRegion);
+	FStartupMessages::Get().AddMetadata(TempName, InStatDesc, InGroupName, InGroupCategory, InGroupDesc, bShouldClearEveryFrame, InStatType, bCycleStat, InMemoryRegion);
 
-	TStatIdData const* LocalHighPerformanceEnable(IStatGroupEnableManager::Get().GetHighPerformanceEnableForStat(FName(InStatName), InGroupName, InGroupCategory, bDefaultEnable, bCanBeDisabled, InStatType, InStatDesc, bCycleStat, InMemoryRegion).GetRawPointer());
+	TStatIdData const* LocalHighPerformanceEnable(IStatGroupEnableManager::Get().GetHighPerformanceEnableForStat(FName(InStatName), InGroupName, InGroupCategory, bDefaultEnable, bShouldClearEveryFrame, InStatType, InStatDesc, bCycleStat, InMemoryRegion).GetRawPointer());
 	TStatIdData const* OldHighPerformanceEnable = (TStatIdData const*)FPlatformAtomics::InterlockedCompareExchangePointer((void**)&HighPerformanceEnable, (void*)LocalHighPerformanceEnable, NULL);
 	check(!OldHighPerformanceEnable || HighPerformanceEnable == OldHighPerformanceEnable); // we are assigned two different groups?
 }
@@ -183,7 +183,6 @@ class FStatGroupEnableManager : public IStatGroupEnableManager
 	struct FGroupEnable
 	{
 		TMap<FName, TStatIdData *> NamesInThisGroup;
-		TMap<FName, TStatIdData *> AlwaysEnabledNamesInThisGroup;
 		bool DefaultEnable; 
 		bool CurrentEnable; 
 
@@ -335,11 +334,11 @@ public:
 		FThreadStats::MasterDisableChangeTagLockSubtract();
 	}
 
-	virtual TStatId GetHighPerformanceEnableForStat(FName StatShortName, const char* InGroup, const char* InCategory, bool bDefaultEnable, bool bCanBeDisabled, EStatDataType::Type InStatType, TCHAR const* InDescription, bool bCycleStat, FPlatformMemory::EMemoryCounterRegion MemoryRegion = FPlatformMemory::MCR_Invalid) override
+	virtual TStatId GetHighPerformanceEnableForStat(FName StatShortName, const char* InGroup, const char* InCategory, bool bDefaultEnable, bool bShouldClearEveryFrame, EStatDataType::Type InStatType, TCHAR const* InDescription, bool bCycleStat, FPlatformMemory::EMemoryCounterRegion MemoryRegion = FPlatformMemory::MCR_Invalid) override
 	{
 		FScopeLock ScopeLock(&SynchronizationObject);
 
-		FStatNameAndInfo LongName(StatShortName, InGroup, InCategory, InDescription, InStatType, bCanBeDisabled, bCycleStat, MemoryRegion);
+		FStatNameAndInfo LongName(StatShortName, InGroup, InCategory, InDescription, InStatType, bShouldClearEveryFrame, bCycleStat, MemoryRegion);
 
 		FName Stat = LongName.GetEncodedName();
 
@@ -352,23 +351,14 @@ public:
 				UE_LOG(LogStatGroupEnableManager, Fatal, TEXT("Stat group %s was was defined both on and off by default."), *Group.ToString());
 			}
 			TStatIdData** StatFound = Found->NamesInThisGroup.Find(Stat);
-			TStatIdData** StatFoundAlways = Found->AlwaysEnabledNamesInThisGroup.Find(Stat);
 			if (StatFound)
 			{
-				if (StatFoundAlways)
-				{
-					UE_LOG(LogStatGroupEnableManager, Fatal, TEXT("Stat %s is both always enabled and not always enabled, so it was used for two different things."), *Stat.ToString());
-				}
 				return TStatId(*StatFound);
 			}
-			else if (StatFoundAlways)
-			{
-				return TStatId(*StatFoundAlways);
-			}
 		}
 		else
 		{
-			Found = &HighPerformanceEnable.Add(Group, FGroupEnable(bDefaultEnable || !bCanBeDisabled));
+			Found = &HighPerformanceEnable.Add(Group, FGroupEnable(bDefaultEnable));
 
 			// this was set up before we saw the group, so set the enable now
 			if (EnableForNewGroup.Contains(Group))
@@ -411,13 +401,9 @@ public:
 		{
 			EnableStat(Stat, Result);
 		}
-		if (bCanBeDisabled)
-		{
-			Found->NamesInThisGroup.Add(Stat, Result);
-		}
 		else
 		{
-			Found->AlwaysEnabledNamesInThisGroup.Add(Stat, Result);
+			Found->NamesInThisGroup.Add(Stat, Result);
 		}
 		return TStatId(Result);
 	}
@@ -442,10 +428,6 @@ public:
 				{
 					UE_LOG(LogStatGroupEnableManager, Display, TEXT("      %d %s"), !ItInner.Value()->IsNone(), *ItInner.Key().ToString());
 				}
-				for (auto ItInner = It.Value().AlwaysEnabledNamesInThisGroup.CreateConstIterator(); ItInner; ++ItInner)
-				{
-					UE_LOG(LogStatGroupEnableManager, Display, TEXT("      (always enabled) %s"), *ItInner.Key().ToString());
-				}
 			}
 		}
 	}
diff --git a/Engine/Source/Runtime/Core/Public/Stats/Stats2.h b/Engine/Source/Runtime/Core/Public/Stats/Stats2.h
--- a/Engine/Source/Runtime/Core/Public/Stats/Stats2.h
+++ b/Engine/Source/Runtime/Core/Public/Stats/Stats2.h
@@ -1261,7 +1261,7 @@ public:
 	void AddThreadMetadata( const FName InThreadFName, uint32 InThreadID );
 
 	/** Adds a regular metadata. */
-	void AddMetadata( FName InStatName, const TCHAR* InStatDesc, const char* InGroupName, const char* InGroupCategory, const TCHAR* InGroupDesc, bool bCanBeDisabled, EStatDataType::Type InStatType, bool bCycleStat, FPlatformMemory::EMemoryCounterRegion InMemoryRegion = FPlatformMemory::MCR_Invalid );
+	void AddMetadata( FName InStatName, const TCHAR* InStatDesc, const char* InGroupName, const char* InGroupCategory, const TCHAR* InGroupDesc, bool bShouldClearEveryFrame, EStatDataType::Type InStatType, bool bCycleStat, FPlatformMemory::EMemoryCounterRegion InMemoryRegion = FPlatformMemory::MCR_Invalid );
 
 	/** Access the singleton. */
 	static FStartupMessages& Get();
@@ -1287,10 +1287,10 @@ public:
 	 * @param InGroup, group to look up
 	 * @param InCategory, the category the group belongs to
 	 * @param bDefaultEnable, If this is the first time this group has been set up, this sets the default enable value for this group.
-	 * @param bCanBeDisabled, If this is true, this is a memory counter or something and can never be disabled
+	 * @param bShouldClearEveryFrame, If this is true, this is a memory counter or an accumulator
 	 * @return a pointer to a FName (valid forever) that determines if this group is active
 	 */
-	virtual TStatId GetHighPerformanceEnableForStat(FName StatShortName, const char* InGroup, const char* InCategory, bool bDefaultEnable, bool bCanBeDisabled, EStatDataType::Type InStatType, TCHAR const* InDescription, bool bCycleStat, FPlatformMemory::EMemoryCounterRegion MemoryRegion = FPlatformMemory::MCR_Invalid) = 0;
+	virtual TStatId GetHighPerformanceEnableForStat(FName StatShortName, const char* InGroup, const char* InCategory, bool bDefaultEnable, bool bShouldClearEveryFrame, EStatDataType::Type InStatType, TCHAR const* InDescription, bool bCycleStat, FPlatformMemory::EMemoryCounterRegion MemoryRegion = FPlatformMemory::MCR_Invalid) = 0;
 
 	/**
 	 * Enables or disabled a particular group of stats
@@ -1330,7 +1330,7 @@ struct FThreadSafeStaticStatBase
 {
 protected:
 	mutable TStatIdData* HighPerformanceEnable; // must be uninitialized, because we need atomic initialization
-	CORE_API void DoSetup(const char* InStatName, const TCHAR* InStatDesc, const char* InGroupName, const char* InGroupCategory, const TCHAR* InGroupDesc, bool bDefaultEnable, bool bCanBeDisabled, EStatDataType::Type InStatType, bool bCycleStat, FPlatformMemory::EMemoryCounterRegion InMemoryRegion) const;
+	CORE_API void DoSetup(const char* InStatName, const TCHAR* InStatDesc, const char* InGroupName, const char* InGroupCategory, const TCHAR* InGroupDesc, bool bDefaultEnable, bool bShouldClearEveryFrame, EStatDataType::Type InStatType, bool bCycleStat, FPlatformMemory::EMemoryCounterRegion InMemoryRegion) const;
 };
 
 template<class TStatData, bool TCompiledIn>
