UMG - Fixing an issue with render transforms having the ability to generate infinity when numerical percision just happens to work out.

diff --git a/Engine/Source/Runtime/Core/Public/Math/TransformCalculus2D.h b/Engine/Source/Runtime/Core/Public/Math/TransformCalculus2D.h
--- a/Engine/Source/Runtime/Core/Public/Math/TransformCalculus2D.h
+++ b/Engine/Source/Runtime/Core/Public/Math/TransformCalculus2D.h
@@ -150,6 +150,22 @@ public:
 	/** Ctor. initialize from a 2D vector representing a set of shears parallel to the X and Y axis, respectively. */
 	explicit FShear2D(const FVector2D& InShear) :Shear(InShear) {}
 	
+	/**
+	 * Generates a shear structure based on angles instead of slope.
+	 * @param InShearAngles The angles of shear.
+	 * @return the sheare structure.
+	 */
+	static FShear2D FromShearAngles(const FVector2D& InShearAngles)
+	{
+		// Compute the M (Shear Slot) = CoTan(90 - SlopeAngle)
+
+		// 0 is a special case because Tan(90) == infinity
+		float ShearX = InShearAngles.X == 0 ? 0 : ( 1.0f / FMath::Tan(FMath::DegreesToRadians(90 - FMath::Clamp(InShearAngles.X, -89.0f, 89.0f))) );
+		float ShearY = InShearAngles.Y == 0 ? 0 : ( 1.0f / FMath::Tan(FMath::DegreesToRadians(90 - FMath::Clamp(InShearAngles.Y, -89.0f, 89.0f))) );
+
+		return FShear2D(ShearX, ShearY);
+	}
+
 	/** 
 	 * Transform 2D Point 
 	 * [X Y] * [1 YY] == [X+Y*XX Y+X*YY]
diff --git a/Engine/Source/Runtime/UMG/Private/Components/Widget.cpp b/Engine/Source/Runtime/UMG/Private/Components/Widget.cpp
--- a/Engine/Source/Runtime/UMG/Private/Components/Widget.cpp
+++ b/Engine/Source/Runtime/UMG/Private/Components/Widget.cpp
@@ -51,11 +51,7 @@ void UWidget::UpdateRenderTransform()
 	{
 		if (!RenderTransform.IsIdentity())
 		{
-			// Compute the M (Shear Slot) = CoTan(90 - SlopeAngle)
-			float ShearX = 1.0f / FMath::Tan(FMath::DegreesToRadians(90 - FMath::Clamp(RenderTransform.Shear.X, -89.0f, 89.0f)));
-			float ShearY = 1.0f / FMath::Tan(FMath::DegreesToRadians(90 - FMath::Clamp(RenderTransform.Shear.Y, -89.0f, 89.0f)));
-
-			FSlateRenderTransform Transform2D = ::Concatenate(FScale2D(RenderTransform.Scale), FShear2D(ShearX, ShearY), FQuat2D(FMath::DegreesToRadians(RenderTransform.Angle)), FVector2D(RenderTransform.Translation));
+			FSlateRenderTransform Transform2D = ::Concatenate(FScale2D(RenderTransform.Scale), FShear2D::FromShearAngles(RenderTransform.Shear), FQuat2D(FMath::DegreesToRadians(RenderTransform.Angle)), FVector2D(RenderTransform.Translation));
 			MyWidget.Pin()->SetRenderTransform(Transform2D);
 		}
 		else
