- New support for iPhone6+ sizing with native scale for 1920x1080

diff --git a/Engine/Source/Runtime/Core/Private/IOS/IOSView.cpp b/Engine/Source/Runtime/Core/Private/IOS/IOSView.cpp
--- a/Engine/Source/Runtime/Core/Private/IOS/IOSView.cpp
+++ b/Engine/Source/Runtime/Core/Private/IOS/IOSView.cpp
@@ -149,9 +149,10 @@ id<MTLDevice> GMetalDevice = nil;
 		if (RequestedContentScaleFactor == 0.0f)
 		{
 #ifdef __IPHONE_8_0
-			UE_LOG(LogIOS, Log, TEXT("Leaving view scale factor alone, with scale = %f, nativeScale = %f"), [[UIScreen mainScreen] scale], [[UIScreen mainScreen] nativeScale]);
+			self.contentScaleFactor = self.window.screen.nativeScale;
+			UE_LOG(LogIOS, Log, TEXT("Setting contentScaleFactor to nativeScale which is = %f"), self.contentScaleFactor);
 #else
-			UE_LOG(LogIOS, Log, TEXT("Leaving view scale factor alone, with scale = %f"), NativeScale);
+			UE_LOG(LogIOS, Log, TEXT("Leaving contentScaleFactor alone, with scale = %f"), NativeScale);
 #endif
 		}
 		else
@@ -162,8 +163,16 @@ id<MTLDevice> GMetalDevice = nil;
 		}
 
 
-		// metal has no further action needed
-		if (!bIsUsingMetal)
+		// handle Metal or GL sizing
+		if (bIsUsingMetal)
+		{
+			CAMetalLayer* MetalLayer = (CAMetalLayer*)self.layer;
+			CGSize DrawableSize = self.bounds.size;
+			DrawableSize.width *= self.contentScaleFactor;
+			DrawableSize.height *= self.contentScaleFactor;
+			MetalLayer.drawableSize = DrawableSize;
+		}
+		else
 		{
 			// make sure this is current
 			[self MakeCurrent];
