Fix offset in ReceiveHitBoxClick event (TTP #331994).

- Y position was still offset when receiving OnMovedWindow.

diff --git a/Engine/Source/Runtime/Core/Private/Linux/LinuxApplication.cpp b/Engine/Source/Runtime/Core/Private/Linux/LinuxApplication.cpp
--- a/Engine/Source/Runtime/Core/Private/Linux/LinuxApplication.cpp
+++ b/Engine/Source/Runtime/Core/Private/Linux/LinuxApplication.cpp
@@ -666,7 +666,7 @@ void FLinuxApplication::ProcessDeferredMessage( SDL_Event Event )
 						{
 							UE_LOG(LogLinuxWindow, Verbose, TEXT("Could not get Window border sizes!"));
 						}
-						MessageHandler->OnMovedWindow(CurrentEventWindow.ToSharedRef(), ClientScreenX, windowEvent.data2);
+						MessageHandler->OnMovedWindow(CurrentEventWindow.ToSharedRef(), ClientScreenX, ClientScreenY);
 					}
 					break;
 
