On Mac OS X when the underlying NSWindow has been closed, don't try and reopen it in response to mouse events or a ghost window will temporarily reopen.

diff --git a/Engine/Source/Runtime/Core/Private/Mac/MacWindow.cpp b/Engine/Source/Runtime/Core/Private/Mac/MacWindow.cpp
--- a/Engine/Source/Runtime/Core/Private/Mac/MacWindow.cpp
+++ b/Engine/Source/Runtime/Core/Private/Mac/MacWindow.cpp
@@ -198,6 +198,7 @@ FMacWindow::FMacWindow()
 	: WindowHandle(NULL)
 	, WindowMode(EWindowMode::Windowed)
 	, bIsVisible(false)
+	, bIsClosed(false)
 {
 	PreFullscreenWindowRect.origin.x = PreFullscreenWindowRect.origin.y = PreFullscreenWindowRect.size.width = PreFullscreenWindowRect.size.height = 0.0f;
 }
@@ -270,7 +271,7 @@ void FMacWindow::MoveWindowTo( int32 X, int32 Y )
 
 void FMacWindow::BringToFront( bool bForce )
 {
-	if (bIsVisible)
+	if (!bIsClosed && bIsVisible)
 	{
 		MainThreadCall(^{
 			SCOPED_AUTORELEASE_POOL;
@@ -290,6 +291,8 @@ void FMacWindow::Destroy()
 			RemoveModalWindow(Window);
 		}
 		
+		bIsClosed = true;
+		
 		if( MacApplication->OnWindowDestroyed( Window ) )
 		{
 			// This FMacWindow may have been destructed by now & so the WindowHandle will probably be invalid memory.
@@ -352,7 +355,7 @@ void FMacWindow::Restore()
 void FMacWindow::Show()
 {
 	SCOPED_AUTORELEASE_POOL;
-	if (!bIsVisible)
+	if (!bIsClosed && !bIsVisible)
 	{
 		if(Definition->IsModalWindow)
 		{
diff --git a/Engine/Source/Runtime/Core/Public/Mac/MacWindow.h b/Engine/Source/Runtime/Core/Public/Mac/MacWindow.h
--- a/Engine/Source/Runtime/Core/Public/Mac/MacWindow.h
+++ b/Engine/Source/Runtime/Core/Public/Mac/MacWindow.h
@@ -115,4 +115,5 @@ private:
 	NSRect PreFullscreenWindowRect;
 
 	bool bIsVisible : 1;
+	bool bIsClosed : 1;
 };
