Core - Implemented GetAllocationSize() for FMallocAnsi

diff --git a/Engine/Source/Runtime/Core/Public/HAL/MallocAnsi.h b/Engine/Source/Runtime/Core/Public/HAL/MallocAnsi.h
--- a/Engine/Source/Runtime/Core/Public/HAL/MallocAnsi.h
+++ b/Engine/Source/Runtime/Core/Public/HAL/MallocAnsi.h
@@ -82,7 +82,9 @@ public:
 		{
 			// Can't use realloc as it might screw with alignment.
 			Result = Malloc( NewSize, Alignment );
-			FMemory::Memcpy( Result, Ptr, FMath::Min(NewSize, *((SIZE_T*)( (uint8*)Ptr - sizeof(void*) - sizeof(SIZE_T))) ) );
+			SIZE_T PtrSize = 0;
+			GetAllocationSize(Ptr,PtrSize);
+			FMemory::Memcpy( Result, Ptr, FMath::Min(NewSize, PtrSize ) );
 			Free( Ptr );
 		}
 		else if( Ptr == nullptr )
@@ -116,40 +118,19 @@ public:
 	}
 
 	virtual bool GetAllocationSize(void *Original, SIZE_T &SizeOut) override
-	{ 
-		return false;
+	{
+		if (!Original)
+		{
+			return false;
+		}
+#if	USE_ALIGNED_MALLOC
+		SizeOut = _aligned_msize(Original,16,0); // Assumes alignment of 16
+#else
+		SizeOut = *((SIZE_T*)( (uint8*)Original - sizeof(void*) - sizeof(SIZE_T)));	
+#endif // USE_ALIGNED_MALLOC
+		return true;
 	}
 
-// #if PLATFORM_MAC
-// 	virtual void GetAllocationInfo( FMemoryAllocationStats_DEPRECATED& MemStats ) override
-// 	{
-// 		// @todo implement in platform memory
-// 		// Just get memory information for the process and report the working set instead
-// 		task_basic_info_64_data_t TaskInfo;
-// 		mach_msg_type_number_t TaskInfoCount = TASK_BASIC_INFO_COUNT;
-// 		task_info( mach_task_self(), TASK_BASIC_INFO, (task_info_t)&TaskInfo, &TaskInfoCount );
-// 		MemStats.TotalUsed___UsedPhysical = MemStats.TotalAllocated___UsedVirtual = TaskInfo.resident_size;
-// 	}
-// #elif PLATFORM_IOS
-// 	virtual void GetAllocationInfo( FMemoryAllocationStats_DEPRECATED& MemStats ) override
-// 	{
-// 		// @todo implement in platform memory
-// 		task_basic_info info;
-// 		mach_msg_type_number_t size = sizeof(info);
-// 		kern_return_t kerr = task_info(mach_task_self(), TASK_BASIC_INFO, (task_info_t)&info, &size);
-// 		MemStats.TotalUsed___UsedPhysical = MemStats.TotalAllocated___UsedVirtual = info.resident_size;
-// 
-// 		mach_port_t host_port = mach_host_self();
-// 		mach_msg_type_number_t host_size = sizeof(vm_statistics_data_t) / sizeof(integer_t);
-// 		vm_size_t pagesize;
-// 		vm_statistics_data_t vm_stat;
-// 
-// 		host_page_size(host_port, &pagesize);
-// 		(void) host_statistics(host_port, HOST_VM_INFO, (host_info_t)&vm_stat, &host_size);
-// 		//MemStats.OSReportedFree = vm_stat.free_count * pagesize;
-// 	}
-// #endif
-
 	/**
 	 * Returns if the allocator is guaranteed to be thread-safe and therefore
 	 * doesn't need a unnecessary thread-safety wrapper around it.
