Added _aligned_msize for Mac

diff --git a/Engine/Source/Runtime/Core/Public/Mac/MacPlatformIncludes.h b/Engine/Source/Runtime/Core/Public/Mac/MacPlatformIncludes.h
--- a/Engine/Source/Runtime/Core/Public/Mac/MacPlatformIncludes.h
+++ b/Engine/Source/Runtime/Core/Public/Mac/MacPlatformIncludes.h
@@ -53,6 +53,7 @@
 #include <execinfo.h>
 #include <signal.h>
 #include <sys/sysctl.h>
+#include <malloc/malloc.h>
 
 // SIMD intrinsics
 #include <xmmintrin.h>
@@ -74,6 +75,7 @@ struct RECT
 #define _aligned_malloc(Size,Align) malloc(Size)
 #define _aligned_realloc(Ptr,Size,Align) realloc(Ptr,Size)
 #define _aligned_free(Ptr) free(Ptr)
+#define _aligned_msize(Ptr,Align,Offset) malloc_size(Ptr)
 
 // include platform implementations
 #include "Mac/MacPlatformMemory.h"
