Rewrite separate game thread event dispatching on Mac OS X to use a custom run-loop source. This makes it easier to see where an event has gone once it has been dispatched, so problems can now be debugged. The manner of dispatch has also changed so that an event will be handled as soon as the target run-loop enters one of the acceptable run-loop modes, rather than waiting for the one selected a dispatch time to come round again.

diff --git a/Engine/Source/Runtime/Core/Private/Mac/CocoaTextView.cpp b/Engine/Source/Runtime/Core/Private/Mac/CocoaTextView.cpp
--- a/Engine/Source/Runtime/Core/Private/Mac/CocoaTextView.cpp
+++ b/Engine/Source/Runtime/Core/Private/Mac/CocoaTextView.cpp
@@ -205,7 +205,7 @@
 				GameThreadCall(^{
 					ITextInputMethodContext::ECaretPosition CaretPosition;
 					IMMContext->GetSelectionRange(SelectionLocation, SelectionLength, CaretPosition);
-				}, InGameRunLoopMode(@[ UE4IMEEventMode ]));
+				}, @[ NSDefaultRunLoopMode, UE4IMEEventMode ]);
 			}
 		}
 		else
@@ -229,7 +229,7 @@
 			FString TheFString(TheString);
 			IMMContext->SetTextInRange(SelectionLocation, SelectionLength, TheFString);
 			IMMContext->SetSelectionRange(SelectionLocation+TheFString.Len(), 0, ITextInputMethodContext::ECaretPosition::Ending);
-		}, InGameRunLoopMode(@[ UE4IMEEventMode ]));
+		}, @[ NSDefaultRunLoopMode, UE4IMEEventMode ]);
 		
 		[self unmarkText];
 		[[self inputContext] invalidateCharacterCoordinates]; // recentering
@@ -268,7 +268,7 @@
 				GameThreadCall(^{
 					ITextInputMethodContext::ECaretPosition CaretPosition;
 					IMMContext->GetSelectionRange(SelectionLocation, SelectionLength, CaretPosition);
-				}, InGameRunLoopMode(@[ UE4IMEEventMode ]));
+				}, @[ NSDefaultRunLoopMode, UE4IMEEventMode ]);
 			}
 		}
 		else
@@ -281,7 +281,7 @@
 		{
 			GameThreadCall(^{
 				IMMContext->SetTextInRange(SelectionLocation, SelectionLength, FString());
-			}, InGameRunLoopMode(@[ UE4IMEEventMode ]));
+			}, @[ NSDefaultRunLoopMode, UE4IMEEventMode ]);
 			[self unmarkText];
 		}
 		else
@@ -290,7 +290,7 @@
 			{
 				GameThreadCall(^{
 					IMMContext->BeginComposition();
-				}, InGameRunLoopMode(@[ UE4IMEEventMode ]));
+				}, @[ NSDefaultRunLoopMode, UE4IMEEventMode ]);
 			}
 			markedRange = NSMakeRange(SelectionLocation, [aString length]);
 			
@@ -332,7 +332,7 @@
 				IMMContext->SetTextInRange(SelectionLocation, SelectionLength, TheFString);
 				IMMContext->UpdateCompositionRange(CompositionRange.location, CompositionRange.length);
 				IMMContext->SetSelectionRange(markedRange.location + selectedRange.location, 0, ITextInputMethodContext::ECaretPosition::Ending);
-			}, InGameRunLoopMode(@[ UE4IMEEventMode ]));
+			}, @[ NSDefaultRunLoopMode, UE4IMEEventMode ]);
 		}
 		[[self inputContext] invalidateCharacterCoordinates]; // recentering
 	}
@@ -354,7 +354,7 @@
 			GameThreadCall(^{
 				IMMContext->UpdateCompositionRange(0, 0);
 				IMMContext->EndComposition();
-			}, InGameRunLoopMode(@[ UE4IMEEventMode ]));
+			}, @[ NSDefaultRunLoopMode, UE4IMEEventMode ]);
 		}
 	}
 }
@@ -370,7 +370,7 @@
 		__block ITextInputMethodContext::ECaretPosition CaretPosition;
 		GameThreadCall(^{
 			IMMContext->GetSelectionRange(SelectionLocation, SelectionLength, CaretPosition);
-		}, InGameRunLoopMode(@[ UE4IMEEventMode ]));
+		}, @[ NSDefaultRunLoopMode, UE4IMEEventMode ]);
 		return {SelectionLocation, SelectionLength};
 	}
 	else
@@ -411,7 +411,7 @@
 		__block FString String;
 		GameThreadCall(^{
 			IMMContext->GetTextInRange(aRange.location, aRange.length, String);
-		}, InGameRunLoopMode(@[ UE4IMEEventMode ]));
+		}, @[ NSDefaultRunLoopMode, UE4IMEEventMode ]);
 		CFStringRef CFString = FPlatformString::TCHARToCFString(*String);
 		if(CFString)
 		{
@@ -448,7 +448,7 @@ return [NSArray arrayWithObjects:NSMarkedClauseSegmentAttributeName, NSGlyphInfo
 		__block FVector2D Size;
 		GameThreadCall(^{
 			IMMContext->GetTextBounds(aRange.location, aRange.length, Position, Size);
-		}, InGameRunLoopMode(@[ UE4IMEEventMode ]));
+		}, @[ NSDefaultRunLoopMode, UE4IMEEventMode ]);
 		
 		if(actualRange)
 		{
@@ -476,7 +476,7 @@ return [NSArray arrayWithObjects:NSMarkedClauseSegmentAttributeName, NSGlyphInfo
 	if (IMMContext.IsValid())
 	{
 		FVector2D Pos(aPoint.x, aPoint.y);
-		int32 Index = GameThreadReturn(^{ return IMMContext->GetCharacterIndexFromPoint(Pos); }, InGameRunLoopMode(@[ UE4IMEEventMode ]));
+		int32 Index = GameThreadReturn(^{ return IMMContext->GetCharacterIndexFromPoint(Pos); }, @[ NSDefaultRunLoopMode, UE4IMEEventMode ]);
 		NSUInteger Idx = Index == INDEX_NONE ? NSNotFound : (NSUInteger)Index;
 		return Idx;
 	}
diff --git a/Engine/Source/Runtime/Core/Private/Mac/CocoaThread.cpp b/Engine/Source/Runtime/Core/Private/Mac/CocoaThread.cpp
--- a/Engine/Source/Runtime/Core/Private/Mac/CocoaThread.cpp
+++ b/Engine/Source/Runtime/Core/Private/Mac/CocoaThread.cpp
@@ -2,6 +2,7 @@
 
 #include "Core.h"
 #include "CocoaThread.h"
+#include "LockFreeList.h"
 
 #define MAC_SEPARATE_GAME_THREAD 1 // Separate the main & game threads so that we better handle the interaction between the Cocoa's event delegates and UE4's event polling.
 
@@ -13,114 +14,307 @@ NSString* UE4CloseEventMode = @"UE4CloseEventMode";
 NSString* UE4IMEEventMode = @"UE4IMEEventMode";
 
 static FCocoaGameThread* GCocoaGameThread = nil;
-static NSRunLoop* GCocoaGameRunLoop = nil;
-static TMap<NSRunLoop*, TArray<NSString*>> GRunLoopModes;
-static FCriticalSection GRunLoopModesLock;
 
-@implementation NSThread (FCocoaThread)
-
-+ (NSThread*) gameThread
+class FCocoaRunLoopSource;
+@interface FCocoaRunLoopSourceInfo : NSObject
 {
-	if(!GCocoaGameThread)
-	{
-		return [NSThread mainThread];
-	}
-	return GCocoaGameThread;
-}
-
-+ (bool) isGameThread
-{
-	bool const bIsGameThread = [[NSThread currentThread] isGameThread];
-	return bIsGameThread;
+@private
+	FCocoaRunLoopSource* Source;
+	CFRunLoopRef RunLoop;
+	CFStringRef Mode;
 }
+- (id)initWithSource:(FCocoaRunLoopSource*)InSource;
+- (void)dealloc;
+- (void)scheduleOn:(CFRunLoopRef)InRunLoop inMode:(CFStringRef)InMode;
+- (void)cancelFrom:(CFRunLoopRef)InRunLoop inMode:(CFStringRef)InMode;
+- (void)perform;
+@end
 
-- (bool) isGameThread
+class FCocoaRunLoopTask
 {
-	bool const bIsGameThread = (self == GCocoaGameThread);
-	return bIsGameThread;
-}
+	friend class FCocoaRunLoopSource;
+public:
+	FCocoaRunLoopTask(dispatch_block_t InBlock, NSArray* InModes)
+	: Block(Block_copy(InBlock))
+	, Modes([InModes retain])
+	{
+		
+	}
+	
+	~FCocoaRunLoopTask()
+	{
+		[Modes release];
+		Block_release(Block);
+	}
+private:
+	dispatch_block_t Block;
+	NSArray* Modes;
+};
 
-- (void) performBlock:(dispatch_block_t)Block onRunLoop:(NSRunLoop*)RunLoop forMode:(NSString*)RunMode wait:(bool)bWait inMode:(NSString*)WaitMode
+class FCocoaRunLoopSource : public FRefCountedObject
 {
-	check(RunLoop);
-	if(RunLoop != [NSRunLoop currentRunLoop])
+public:
+	static void RegisterMainRunLoop(CFRunLoopRef RunLoop)
+	{
+		check(!MainRunLoopSource);
+		MainRunLoopSource = new FCocoaRunLoopSource(RunLoop);
+	}
+	
+	static void RegisterGameRunLoop(CFRunLoopRef RunLoop)
+	{
+		check(!GameRunLoopSource);
+		GameRunLoopSource = new FCocoaRunLoopSource(RunLoop);
+	}
+	
+	static FCocoaRunLoopSource& GetMainRunLoopSource(void)
+	{
+		check(MainRunLoopSource);
+		return *MainRunLoopSource;
+	}
+	
+	static FCocoaRunLoopSource& GetGameRunLoopSource(void)
+	{
+		check(GameRunLoopSource);
+		return *GameRunLoopSource;
+	}
+	
+	void Schedule(dispatch_block_t InBlock, NSArray* InModes)
 	{
-		dispatch_block_t CopiedBlock = Block_copy(Block);
+		for (NSString* Mode in InModes)
+		{
+			Register((CFStringRef)Mode);
+		}
 		
-		if(bWait)
+		Tasks.Push(new FCocoaRunLoopTask(InBlock, InModes));
+		
+		CFDictionaryApplyFunction(SourceDictionary, &FCocoaRunLoopSource::SignalFunction, nullptr);
+	}
+	
+	void Wake(void)
+	{
+		CFRunLoopWakeUp(TargetRunLoop);
+	}
+	
+	void RunInMode(CFStringRef WaitMode)
+	{
+		CFRunLoopRunInMode(WaitMode, 0, true);
+	}
+	
+	void Process(CFStringRef Mode)
+	{
+		TArray<FCocoaRunLoopTask*> NewTasks;
+		Tasks.PopAll(NewTasks);
+		
+		for (uint32 i = NewTasks.Num(); i > 0; i--)
 		{
-			dispatch_semaphore_t Semaphore = dispatch_semaphore_create(0);
-			{
-				FScopeLock Lock(&GRunLoopModesLock);
-				TArray<NSString*>& Modes = GRunLoopModes.FindOrAdd([NSRunLoop currentRunLoop]);
-				Modes.Add([WaitMode retain]);
-			}
-			dispatch_block_t ExecuteBlock = Block_copy(^{ CopiedBlock(); dispatch_semaphore_signal(Semaphore); });
-			[RunLoop performBlock:ExecuteBlock forMode:RunMode wake:true];
-			while(dispatch_semaphore_wait(Semaphore, DISPATCH_TIME_NOW))
-			{
-				CFRunLoopRunInMode((CFStringRef)WaitMode, 0, true);
-				CFRunLoopWakeUp([RunLoop getCFRunLoop]);
-			}
-			dispatch_release(Semaphore);
-			Block_release(ExecuteBlock);
+			OutstandingTasks.Add(NewTasks[i - 1]);
+		}
+		
+		bool bDone = false;
+		while (!bDone)
+		{
+			bDone = true;
+			for (uint32 i = 0; i < OutstandingTasks.Num(); i++)
 			{
-				FScopeLock Lock(&GRunLoopModesLock);
-				TArray<NSString*>& Modes = GRunLoopModes.FindOrAdd([NSRunLoop currentRunLoop]);
-				check(Modes.Num() > 0);
-				[Modes.Last() release];
-				Modes.RemoveAt(Modes.Num() - 1);
+				FCocoaRunLoopTask* Task = OutstandingTasks[i];
+				if (Task && [Task->Modes containsObject:(NSString*)Mode])
+				{
+					OutstandingTasks.RemoveAt(i);
+					Task->Block();
+					delete Task;
+					bDone = false;
+					break;
+				}
 			}
 		}
-		else
+	}
+	
+private:
+	FCocoaRunLoopSource(CFRunLoopRef RunLoop)
+	: TargetRunLoop(RunLoop)
+	, SourceDictionary(CFDictionaryCreateMutable(nullptr, 0, &kCFTypeDictionaryKeyCallBacks, &kCFTypeDictionaryValueCallBacks))
+	{
+		check(TargetRunLoop);
+		
+		// Register for default modes
+		Register(kCFRunLoopDefaultMode);
+		Register((CFStringRef)UE4NilEventMode);
+		Register((CFStringRef)UE4ShowEventMode);
+		Register((CFStringRef)UE4ResizeEventMode);
+		Register((CFStringRef)UE4FullscreenEventMode);
+		Register((CFStringRef)UE4CloseEventMode);
+		Register((CFStringRef)UE4IMEEventMode);
+	}
+	
+	virtual ~FCocoaRunLoopSource()
+	{
+		if ( MainRunLoopSource == this )
+		{
+			MainRunLoopSource = nullptr;
+		}
+		else if ( GameRunLoopSource == this )
 		{
-			[RunLoop performBlock:CopiedBlock forMode:RunMode wake:true];
+			GameRunLoopSource = nullptr;
 		}
 		
-		Block_release(CopiedBlock);
+		CFDictionaryApplyFunction(SourceDictionary, &FCocoaRunLoopSource::ShutdownFunction, TargetRunLoop);
+		CFRelease(SourceDictionary);
 	}
-	else
+	
+	void Register(CFStringRef Mode)
 	{
-		Block();
+		if(CFDictionaryContainsKey(SourceDictionary, Mode) == false)
+		{
+			CFRunLoopSourceContext Context;
+			FCocoaRunLoopSourceInfo* Info = [[FCocoaRunLoopSourceInfo alloc] initWithSource:this];
+			FMemory::MemZero(Context);
+			Context.version = 0;
+			Context.info = Info;
+			Context.retain = CFRetain;
+			Context.release = CFRelease;
+			Context.copyDescription = CFCopyDescription;
+			Context.equal = CFEqual;
+			Context.hash = CFHash;
+			Context.schedule = &FCocoaRunLoopSource::Schedule;
+			Context.cancel = &FCocoaRunLoopSource::Cancel;
+			Context.perform = &FCocoaRunLoopSource::Perform;
+			
+			CFRunLoopSourceRef Source = CFRunLoopSourceCreate(nullptr, 0, &Context);
+			CFDictionaryAddValue(SourceDictionary, Mode, Source);
+			CFRunLoopAddSource(TargetRunLoop, Source, Mode);
+			CFRelease(Source);
+		}
 	}
-}
+	
+	static void SignalFunction(const void* Key, const void* Value, void* Context)
+	{
+		CFRunLoopSourceRef RunLoopSource = (CFRunLoopSourceRef)Value;
+		if(RunLoopSource)
+		{
+			CFRunLoopSourceSignal(RunLoopSource);
+		}
+	}
+	static void ShutdownFunction(const void* Key, const void* Value, void* Context)
+	{
+		CFRunLoopRef RunLoop = static_cast<CFRunLoopRef>(Context);
+		if(RunLoop)
+		{
+			CFRunLoopRemoveSource(RunLoop, (CFRunLoopSourceRef)Value, (CFStringRef)Key);
+		}
+	}
+	static void Schedule(void* Info, CFRunLoopRef RunLoop, CFStringRef Mode)
+	{
+		FCocoaRunLoopSourceInfo* Source = static_cast<FCocoaRunLoopSourceInfo*>(Info);
+		if (Source)
+		{
+			[Source scheduleOn:RunLoop inMode:Mode];
+		}
+	}
+    static void Cancel(void* Info, CFRunLoopRef RunLoop, CFStringRef Mode)
+	{
+		FCocoaRunLoopSourceInfo* Source = static_cast<FCocoaRunLoopSourceInfo*>(Info);
+		if (Source)
+		{
+			[Source cancelFrom:RunLoop inMode:Mode];
+		}
+	}
+    static void Perform(void* Info)
+	{
+		FCocoaRunLoopSourceInfo* Source = static_cast<FCocoaRunLoopSourceInfo*>(Info);
+		if (Source)
+		{
+			[Source perform];
+		}
+	}
+	
+private:
+	TLockFreePointerList<FCocoaRunLoopTask> Tasks;
+	TArray<FCocoaRunLoopTask*> OutstandingTasks;
+	CFRunLoopRef TargetRunLoop;
+	CFMutableDictionaryRef SourceDictionary;
+	
+private:
+	static FCocoaRunLoopSource* MainRunLoopSource;
+	static FCocoaRunLoopSource* GameRunLoopSource;
+};
 
-@end
+FCocoaRunLoopSource* FCocoaRunLoopSource::MainRunLoopSource = nullptr;
+FCocoaRunLoopSource* FCocoaRunLoopSource::GameRunLoopSource = nullptr;
 
-@implementation NSRunLoop (FCocoaRunLoop)
+@implementation FCocoaRunLoopSourceInfo
 
-+ (NSRunLoop*)gameRunLoop
+- (id)initWithSource:(FCocoaRunLoopSource*)InSource
 {
-	if(!GCocoaGameRunLoop)
+	id Self = [super init];
+	if(Self)
 	{
-		return [NSRunLoop mainRunLoop];
+		check(InSource);
+		Source = InSource;
+		Source->AddRef();
 	}
-	return GCocoaGameRunLoop;
+	return Self;
 }
 
-- (void)performBlock:(dispatch_block_t)Block forMode:(NSString *)Mode wake:(bool)bWakeRunLoop
+- (void)dealloc
 {
-	CFRunLoopRef RunLoop = [self getCFRunLoop];
-	check(RunLoop);
-	CFRunLoopPerformBlock(RunLoop, (CFStringRef)Mode, Block);
-	if(bWakeRunLoop)
+	check(Source);
+	Source->Release();
+	Source = nullptr;
+	[super dealloc];
+}
+
+- (void)scheduleOn:(CFRunLoopRef)InRunLoop inMode:(CFStringRef)InMode
+{
+	check(!RunLoop);
+	check(!Mode);
+	RunLoop = InRunLoop;
+	Mode = InMode;
+}
+
+- (void)cancelFrom:(CFRunLoopRef)InRunLoop inMode:(CFStringRef)InMode
+{
+	if(CFEqual(InRunLoop, RunLoop) && CFEqual(Mode, InMode))
 	{
-		CFRunLoopWakeUp(RunLoop);
+		RunLoop = nullptr;
+		Mode = nullptr;
 	}
 }
 
-- (NSString*)intendedMode
+- (void)perform
+{
+	check(Source);
+	check(RunLoop);
+	check(Mode);
+	check(CFEqual(RunLoop, CFRunLoopGetCurrent()));
+	CFStringRef CurrentMode = CFRunLoopCopyCurrentMode(CFRunLoopGetCurrent());
+	check(CFEqual(CurrentMode, Mode));
+	Source->Process(CurrentMode);
+	CFRelease(CurrentMode);
+}
+
+@end
+
+@implementation NSThread (FCocoaThread)
+
++ (NSThread*) gameThread
 {
-	NSString* Mode = nil;
+	if(!GCocoaGameThread)
 	{
-		FScopeLock Lock(&GRunLoopModesLock);
-		TArray<NSString*>& Modes = GRunLoopModes.FindOrAdd(self);
-		if(Modes.Num() > 0)
-		{
-			Mode = Modes.Last();
-		}
+		return [NSThread mainThread];
 	}
-	return Mode;
+	return GCocoaGameThread;
+}
+
++ (bool) isGameThread
+{
+	bool const bIsGameThread = [[NSThread currentThread] isGameThread];
+	return bIsGameThread;
+}
+
+- (bool) isGameThread
+{
+	bool const bIsGameThread = (self == GCocoaGameThread);
+	return bIsGameThread;
 }
 
 @end
@@ -150,48 +344,84 @@ static FCriticalSection GRunLoopModesLock;
 - (void)dealloc
 {
 	GCocoaGameThread = nil;
-	GCocoaGameRunLoop = nil;
 	[super dealloc];
 }
 
 - (void)main
 {
-	GCocoaGameRunLoop = [NSRunLoop currentRunLoop];
+	NSRunLoop* GameRunLoop = [NSRunLoop currentRunLoop];
+	
+	// Register game run loop source
+	FCocoaRunLoopSource::RegisterGameRunLoop([GameRunLoop getCFRunLoop]);
+	
 	[super main];
 }
 
 @end
 
-NSString* InGameRunLoopMode(NSArray* AllowedModes)
+static void PerformBlockOnThread(FCocoaRunLoopSource& ThreadSource, dispatch_block_t Block, NSArray* SendModes, NSString* WaitMode, bool const bWait)
 {
-	SCOPED_AUTORELEASE_POOL;
-	NSString* Mode = [[NSRunLoop gameRunLoop] intendedMode];
-	if(Mode == nil || ![AllowedModes containsObject:Mode])
+	dispatch_block_t CopiedBlock = Block_copy(Block);
+	
+	if ( bWait )
+	{
+		dispatch_semaphore_t Semaphore = dispatch_semaphore_create(0);
+		dispatch_block_t ExecuteBlock = Block_copy(^{ CopiedBlock(); dispatch_semaphore_signal(Semaphore); });
+		
+		ThreadSource.Schedule(ExecuteBlock, SendModes);
+		ThreadSource.Wake();
+		
+		while (dispatch_semaphore_wait(Semaphore, DISPATCH_TIME_NOW))
+		{
+			ThreadSource.RunInMode( (CFStringRef)WaitMode );
+			ThreadSource.Wake();
+		}
+		
+		Block_release(ExecuteBlock);
+		dispatch_release(Semaphore);
+	}
+	else
 	{
-		Mode = NSDefaultRunLoopMode;
+		ThreadSource.Schedule(Block, SendModes);
+		ThreadSource.Wake();
 	}
-	return Mode;
-}
-
-void PerformBlockOnRunLoop(dispatch_block_t Block, NSRunLoop* const RunLoop, NSString* SendMode, NSString* WaitMode, bool const bWait)
-{
-	[[NSThread currentThread] performBlock:Block onRunLoop:RunLoop forMode:SendMode wait:bWait inMode:WaitMode];
+	
+	Block_release(CopiedBlock);
 }
 
 void MainThreadCall(dispatch_block_t Block, NSString* WaitMode, bool const bWait)
 {
-	PerformBlockOnRunLoop(Block, [NSRunLoop mainRunLoop], NSDefaultRunLoopMode, WaitMode, bWait);
+	if ( [NSThread mainThread] != [NSThread currentThread] )
+	{
+		FCocoaRunLoopSource& MainRunLoopSource = FCocoaRunLoopSource::GetMainRunLoopSource();
+		PerformBlockOnThread(MainRunLoopSource, Block, @[ NSDefaultRunLoopMode ], WaitMode, bWait);
+	}
+	else
+	{
+		Block();
+	}
 }
 
-void GameThreadCall(dispatch_block_t Block, NSString* SendMode, bool const bWait)
+void GameThreadCall(dispatch_block_t Block, NSArray* SendModes, bool const bWait)
 {
-	PerformBlockOnRunLoop(Block, [NSRunLoop gameRunLoop], SendMode, NSDefaultRunLoopMode, bWait);
+	if ( [NSThread gameThread] != [NSThread currentThread] )
+	{
+		FCocoaRunLoopSource& GameRunLoopSource = FCocoaRunLoopSource::GetGameRunLoopSource();
+		PerformBlockOnThread(GameRunLoopSource, Block, SendModes, NSDefaultRunLoopMode, bWait);
+	}
+	else
+	{
+		Block();
+	}
 }
 
 void RunGameThread(id Target, SEL Selector)
 {
 	SCOPED_AUTORELEASE_POOL;
 #if MAC_SEPARATE_GAME_THREAD
+	// Register main run loop source
+	FCocoaRunLoopSource::RegisterMainRunLoop(CFRunLoopGetCurrent());
+	
 	// Create a separate game thread and set it to the stack size to be the same as the main thread default of 8MB ( http://developer.apple.com/library/mac/#qa/qa1419/_index.html )
 	FCocoaGameThread* GameThread = [[FCocoaGameThread alloc] initWithTarget:Target selector:Selector object:nil];
 	[GameThread setStackSize:8*1024*1024];
diff --git a/Engine/Source/Runtime/Core/Private/Mac/CocoaWindow.cpp b/Engine/Source/Runtime/Core/Private/Mac/CocoaWindow.cpp
--- a/Engine/Source/Runtime/Core/Private/Mac/CocoaWindow.cpp
+++ b/Engine/Source/Runtime/Core/Private/Mac/CocoaWindow.cpp
@@ -49,6 +49,11 @@ NSString* NSPerformDragOperation = @"NSPerformDragOperation";
 	return NewSelf;
 }
 
+- (void)dealloc
+{
+	[super dealloc];
+}
+
 - (NSRect)openGLFrame
 {
 	SCOPED_AUTORELEASE_POOL;
@@ -174,7 +179,7 @@ NSString* NSPerformDragOperation = @"NSPerformDragOperation";
 	if(bNeedsRedraw && bForwardEvents && ([self isVisible] && [super alphaValue] > 0.0f))
 	{
 		NSNotification* Notification = [NSNotification notificationWithName:NSWindowRedrawContents object:self];
-		FMacEvent::SendToGameRunLoop(Notification, self, EMacEventSendMethod::Sync, InGameRunLoopMode(@[ UE4NilEventMode, UE4ShowEventMode, UE4ResizeEventMode, UE4FullscreenEventMode, UE4CloseEventMode, UE4IMEEventMode ]));
+		FMacEvent::SendToGameRunLoop(Notification, self, EMacEventSendMethod::Sync, @[ NSDefaultRunLoopMode, UE4NilEventMode, UE4ShowEventMode, UE4ResizeEventMode, UE4FullscreenEventMode, UE4CloseEventMode, UE4IMEEventMode ]);
 	}
 	bNeedsRedraw = false;
 }
@@ -301,7 +306,7 @@ NSString* NSPerformDragOperation = @"NSPerformDragOperation";
 		if(self.bForwardEvents)
 		{
 			NSNotification* Notification = [NSNotification notificationWithName:NSWindowDidResizeNotification object:self];
-			FMacEvent::SendToGameRunLoop(Notification, self, EMacEventSendMethod::Async, InGameRunLoopMode(@[ UE4ResizeEventMode, UE4ShowEventMode ]));
+			FMacEvent::SendToGameRunLoop(Notification, self, EMacEventSendMethod::Async, @[ NSDefaultRunLoopMode, UE4ResizeEventMode, UE4ShowEventMode ]);
 		}
 	}
 }
@@ -323,7 +328,7 @@ NSString* NSPerformDragOperation = @"NSPerformDragOperation";
 		bDeferSetOrigin = true;
 		DeferFrame.origin = Point;
 		NSNotification* Notification = [NSNotification notificationWithName:NSWindowDidMoveNotification object:self];
-		FMacEvent::SendToGameRunLoop(Notification, self, EMacEventSendMethod::Async, InGameRunLoopMode(@[ UE4ResizeEventMode, UE4ShowEventMode ]));
+		FMacEvent::SendToGameRunLoop(Notification, self, EMacEventSendMethod::Async, @[ NSDefaultRunLoopMode, UE4ResizeEventMode, UE4ShowEventMode ]);
 	}
 }
 
@@ -352,7 +357,7 @@ NSString* NSPerformDragOperation = @"NSPerformDragOperation";
 	WindowMode = self.TargetWindowMode;
 	if(self.bForwardEvents)
 	{
-		FMacEvent::SendToGameRunLoop(notification, self, EMacEventSendMethod::Async, InGameRunLoopMode(@[ UE4FullscreenEventMode ]));
+		FMacEvent::SendToGameRunLoop(notification, self, EMacEventSendMethod::Async, @[ NSDefaultRunLoopMode, UE4FullscreenEventMode ]);
 	}
 }
 
@@ -362,7 +367,7 @@ NSString* NSPerformDragOperation = @"NSPerformDragOperation";
 	self.TargetWindowMode = EWindowMode::Windowed;
 	if(self.bForwardEvents)
 	{
-		FMacEvent::SendToGameRunLoop(notification, self, EMacEventSendMethod::Async, InGameRunLoopMode(@[ UE4FullscreenEventMode ]));
+		FMacEvent::SendToGameRunLoop(notification, self, EMacEventSendMethod::Async, @[ NSDefaultRunLoopMode, UE4FullscreenEventMode ]);
 	}
 }
 
@@ -383,7 +388,7 @@ NSString* NSPerformDragOperation = @"NSPerformDragOperation";
 	
 	if(self.bForwardEvents)
 	{
-		FMacEvent::SendToGameRunLoop(Notification, self, EMacEventSendMethod::Async, InGameRunLoopMode(@[ UE4ShowEventMode, UE4CloseEventMode, UE4FullscreenEventMode ]));
+		FMacEvent::SendToGameRunLoop(Notification, self, EMacEventSendMethod::Async, @[ NSDefaultRunLoopMode, UE4ShowEventMode, UE4CloseEventMode, UE4FullscreenEventMode ]);
 	}
 }
 
@@ -395,7 +400,7 @@ NSString* NSPerformDragOperation = @"NSPerformDragOperation";
 	
 	if(self.bForwardEvents)
 	{
-		FMacEvent::SendToGameRunLoop(Notification, self, EMacEventSendMethod::Async, InGameRunLoopMode(@[ UE4ShowEventMode, UE4CloseEventMode, UE4FullscreenEventMode ]));
+		FMacEvent::SendToGameRunLoop(Notification, self, EMacEventSendMethod::Async, @[ NSDefaultRunLoopMode, UE4ShowEventMode, UE4CloseEventMode, UE4FullscreenEventMode ]);
 	}
 }
 
@@ -403,7 +408,7 @@ NSString* NSPerformDragOperation = @"NSPerformDragOperation";
 {
 	if(self.bForwardEvents)
 	{
-		FMacEvent::SendToGameRunLoop(Notification, self, EMacEventSendMethod::Async, InGameRunLoopMode(@[ UE4ResizeEventMode, UE4ShowEventMode, UE4FullscreenEventMode ]));
+		FMacEvent::SendToGameRunLoop(Notification, self, EMacEventSendMethod::Async, @[ NSDefaultRunLoopMode, UE4ResizeEventMode, UE4ShowEventMode, UE4FullscreenEventMode ]);
 	}
 }
 
@@ -417,7 +422,7 @@ NSString* NSPerformDragOperation = @"NSPerformDragOperation";
 	
 	if(self.bForwardEvents)
 	{
-		FMacEvent::SendToGameRunLoop(Notification, self, EMacEventSendMethod::Async, InGameRunLoopMode(@[ UE4ResizeEventMode, UE4ShowEventMode, UE4FullscreenEventMode ]));
+		FMacEvent::SendToGameRunLoop(Notification, self, EMacEventSendMethod::Async, @[ NSDefaultRunLoopMode, UE4ResizeEventMode, UE4ShowEventMode, UE4FullscreenEventMode ]);
 	}
 }
 
@@ -491,7 +496,7 @@ NSString* NSPerformDragOperation = @"NSPerformDragOperation";
 	bZoomed = [self isZoomed];
 	if(self.bForwardEvents)
 	{
-		FMacEvent::SendToGameRunLoop(Notification, self, EMacEventSendMethod::Async, InGameRunLoopMode(@[ UE4ResizeEventMode, UE4ShowEventMode, UE4FullscreenEventMode ]));
+		FMacEvent::SendToGameRunLoop(Notification, self, EMacEventSendMethod::Async, @[ NSDefaultRunLoopMode, UE4ResizeEventMode, UE4ShowEventMode, UE4FullscreenEventMode ]);
 	}
 	bNeedsRedraw = true;
 }
@@ -501,7 +506,7 @@ NSString* NSPerformDragOperation = @"NSPerformDragOperation";
 	SCOPED_AUTORELEASE_POOL;
 	if(self.bForwardEvents && MacApplication)
 	{
-		FMacEvent::SendToGameRunLoop(notification, self, EMacEventSendMethod::Async, InGameRunLoopMode(@[ UE4CloseEventMode ]));
+		FMacEvent::SendToGameRunLoop(notification, self, EMacEventSendMethod::Async, @[ NSDefaultRunLoopMode, UE4CloseEventMode ]);
 	}
 	self.bForwardEvents = false;
 	[self setDelegate:nil];
@@ -617,7 +622,7 @@ NSString* NSPerformDragOperation = @"NSPerformDragOperation";
 	BOOL Movable = [super isMovable];
 	if(Movable && bRenderInitialised && MacApplication)
 	{
-		Movable &= (BOOL)(GameThreadReturn(^{ return MacApplication->IsWindowMovable(self, NULL); }, InGameRunLoopMode(@[ UE4NilEventMode, UE4ShowEventMode, UE4ResizeEventMode, UE4FullscreenEventMode, UE4CloseEventMode, UE4IMEEventMode ])));
+		Movable &= (BOOL)(GameThreadReturn(^{ return MacApplication->IsWindowMovable(self, NULL); }, @[ NSDefaultRunLoopMode, UE4NilEventMode, UE4ShowEventMode, UE4ResizeEventMode, UE4FullscreenEventMode, UE4CloseEventMode, UE4IMEEventMode ]));
 	}
 	return Movable;
 }
diff --git a/Engine/Source/Runtime/Core/Private/Mac/MacApplication.cpp b/Engine/Source/Runtime/Core/Private/Mac/MacApplication.cpp
--- a/Engine/Source/Runtime/Core/Private/Mac/MacApplication.cpp
+++ b/Engine/Source/Runtime/Core/Private/Mac/MacApplication.cpp
@@ -44,34 +44,39 @@ FMacApplication* FMacApplication::CreateMacApplication()
 
 NSEvent* FMacApplication::HandleNSEvent(NSEvent* Event)
 {
-	NSEvent* ReturnEvent = Event;
-	const bool bIsMouseClickOrKeyEvent = [Event type] == NSLeftMouseDown || [Event type] == NSLeftMouseUp
-		|| [Event type] == NSRightMouseDown || [Event type] == NSRightMouseUp
-		|| [Event type] == NSOtherMouseDown || [Event type] == NSOtherMouseUp;
-	const bool bIsResentEvent = [Event type] == NSApplicationDefined && [Event subtype] == FMacApplication::ResentEvent;
+	NSEvent* ReturnEvent = nil;
 	
-	if (MacApplication)
+	if ([Event windowNumber] == 0 || [Event window] != nil)
 	{
-		if ( bIsResentEvent )
-		{
-			ReturnEvent = (NSEvent*)[Event data1];
-		}
+		ReturnEvent = Event;
+		const bool bIsMouseClickOrKeyEvent = [Event type] == NSLeftMouseDown || [Event type] == NSLeftMouseUp
+			|| [Event type] == NSRightMouseDown || [Event type] == NSRightMouseUp
+			|| [Event type] == NSOtherMouseDown || [Event type] == NSOtherMouseUp;
+		const bool bIsResentEvent = [Event type] == NSApplicationDefined && [Event subtype] == FMacApplication::ResentEvent;
 		
-		if ( !bIsResentEvent && ( !bIsMouseClickOrKeyEvent || [Event window] == NULL ) )
+		if (MacApplication)
 		{
-			FMacEvent::SendToGameRunLoop(Event, EMacEventSendMethod::Async);
+			if ( bIsResentEvent )
+			{
+				ReturnEvent = (NSEvent*)[Event data1];
+			}
+			
+			if ( !bIsResentEvent && ( !bIsMouseClickOrKeyEvent || [Event window] == NULL ) )
+			{
+				FMacEvent::SendToGameRunLoop(Event, EMacEventSendMethod::Async);
+				
+				if ( [Event type] == NSKeyDown || [Event type] == NSKeyUp )
+				{
+					ReturnEvent = nil;
+				}
+			}
 			
-			if ( [Event type] == NSKeyDown || [Event type] == NSKeyUp )
+			if ([Event type] == NSLeftMouseUp)
 			{
-				ReturnEvent = nil;
+				NSNotification* Notification = [NSNotification notificationWithName:NSWindowDraggingFinished object:[Event window]];
+				FMacEvent::SendToGameRunLoop(Notification, [Event window], EMacEventSendMethod::Async);
 			}
 		}
-		
-		if ([Event type] == NSLeftMouseUp)
-		{
-			NSNotification* Notification = [NSNotification notificationWithName:NSWindowDraggingFinished object:[Event window]];
-			FMacEvent::SendToGameRunLoop(Notification, [Event window], EMacEventSendMethod::Async);
-		}
 	}
 	
 	return ReturnEvent;
diff --git a/Engine/Source/Runtime/Core/Private/Mac/MacEvent.cpp b/Engine/Source/Runtime/Core/Private/Mac/MacEvent.cpp
--- a/Engine/Source/Runtime/Core/Private/Mac/MacEvent.cpp
+++ b/Engine/Source/Runtime/Core/Private/Mac/MacEvent.cpp
@@ -44,36 +44,30 @@ FMacEvent::FMacEvent(NSNotification* const Notification, NSWindow* const Window)
 	CFRelease(Event);
 }
 
-void FMacEvent::SendToGameRunLoop(FMacEvent const* const Event, EMacEventSendMethod SendMethod, NSString* SendMode)
+void FMacEvent::SendToGameRunLoop(FMacEvent const* const Event, EMacEventSendMethod SendMethod, NSArray* SendModes)
 {
 	SCOPED_AUTORELEASE_POOL;
-	NSThread* GameThread = [NSThread gameThread];
-	check(GameThread);
-	
-	NSRunLoop* GameRunLoop = [NSRunLoop gameRunLoop];
-	check(GameRunLoop);
-	
 	dispatch_block_t Block = ^{ FMacApplication::ProcessEvent(Event); };
 	const bool bWait = (SendMethod == EMacEventSendMethod::Sync);
 	
-	[GameThread performBlock: Block onRunLoop: GameRunLoop forMode: SendMode wait: bWait inMode: NSDefaultRunLoopMode];
+	GameThreadCall(Block, SendModes, bWait);
 }
 	
-void FMacEvent::SendToGameRunLoop(NSEvent* const Event, EMacEventSendMethod SendMethod, NSString* SendMode)
+void FMacEvent::SendToGameRunLoop(NSEvent* const Event, EMacEventSendMethod SendMethod, NSArray* SendModes)
 {
 	if(MacApplication)
 	{
 		FMacEvent* MacEvent = new FMacEvent(Event);
-		FMacEvent::SendToGameRunLoop(MacEvent, SendMethod, SendMode);
+		FMacEvent::SendToGameRunLoop(MacEvent, SendMethod, SendModes);
 	}
 }
 
-void FMacEvent::SendToGameRunLoop(NSNotification* const Notification, NSWindow* const Window, EMacEventSendMethod SendMethod, NSString* SendMode)
+void FMacEvent::SendToGameRunLoop(NSNotification* const Notification, NSWindow* const Window, EMacEventSendMethod SendMethod, NSArray* SendModes)
 {
 	if(MacApplication)
 	{
 		FMacEvent* MacEvent = new FMacEvent(Notification, Window);
-		FMacEvent::SendToGameRunLoop(MacEvent, SendMethod, SendMode);
+		FMacEvent::SendToGameRunLoop(MacEvent, SendMethod, SendModes);
 	}
 }
 
diff --git a/Engine/Source/Runtime/Core/Public/Mac/CocoaThread.h b/Engine/Source/Runtime/Core/Public/Mac/CocoaThread.h
--- a/Engine/Source/Runtime/Core/Public/Mac/CocoaThread.h
+++ b/Engine/Source/Runtime/Core/Public/Mac/CocoaThread.h
@@ -3,8 +3,6 @@
 
 #include "Core.h"
 
-#ifdef __OBJC__
-
 /* Custom run-loop modes for UE4 that process only certain kinds of events to simulate Windows event ordering. */
 extern NSString* UE4NilEventMode; /* Process only mandatory events */
 extern NSString* UE4ShowEventMode; /* Process only show window events */
@@ -17,13 +15,6 @@ extern NSString* UE4IMEEventMode; /* Process only input method events */
 + (NSThread*) gameThread; // Returns the main game thread, or nil if has yet to be constructed.
 + (bool) isGameThread; // True if the current thread is the main game thread, else false.
 - (bool) isGameThread; // True if this thread object is the main game thread, else false.
-- (void) performBlock:(dispatch_block_t)Block onRunLoop:(NSRunLoop*)RunLoop forMode:(NSString*)RunMode wait:(bool)bWait inMode:(NSString*)WaitMode; // Perform Block on RunLoop in RunMode, which must not be the current thread's runloop and wait in WaitMode if bWait is true or return immediately after queueing to the RunLoop.
-@end
-
-@interface NSRunLoop (FCocoaRunLoop)
-+ (NSRunLoop*)gameRunLoop; // The run loop for the main game thread, or nil if the main game thread has not been started.
-- (void)performBlock:(dispatch_block_t)Block forMode:(NSString *)Mode wake:(bool)bWakeRunLoop; // Add Block to the run loop for processing in the given Mode, only waking the run loop if bWakeRunLoop is set otherwise processing occurs the next time the run-loop is processed.
-- (NSString*)intendedMode;
 @end
 
 @interface FCocoaGameThread : NSThread
@@ -33,42 +24,26 @@ extern NSString* UE4IMEEventMode; /* Process only input method events */
 - (void)main; // Override that sets the variable backing +[NSRunLoop gameRunLoop], do not override in a subclass.
 @end
 
-NSString* InGameRunLoopMode(NSArray* AllowedModes);
+void MainThreadCall(dispatch_block_t Block, NSString* WaitMode = NSDefaultRunLoopMode, bool const bWait = true);
 
 template<typename ReturnType>
-ReturnType PerformBlockOnRunLoop(ReturnType (^Block)(void), NSRunLoop* const RunLoop, NSString* SendMode, NSString* WaitMode)
+ReturnType MainThreadReturn(ReturnType (^Block)(void), NSString* WaitMode = NSDefaultRunLoopMode)
 {
 	__block ReturnType ReturnValue;
-	[[NSThread currentThread] performBlock:^{ ReturnValue = Block(); } onRunLoop:RunLoop forMode:SendMode wait:true inMode:WaitMode];
+	MainThreadCall(^{ ReturnValue = Block(); }, WaitMode, true);
 	return ReturnValue;
 }
 
-void PerformBlockOnRunLoop(dispatch_block_t Block, NSRunLoop* const RunLoop, NSString* SendMode, NSString* WaitMode, bool const bWait);
-
-template<typename ReturnType>
-ReturnType MainThreadReturn(ReturnType (^Block)(void), NSString* WaitMode = NSDefaultRunLoopMode)
-{
-	return PerformBlockOnRunLoop(Block, [NSRunLoop mainRunLoop], NSDefaultRunLoopMode, WaitMode);
-}
-
-void MainThreadCall(dispatch_block_t Block, NSString* WaitMode = NSDefaultRunLoopMode, bool const bWait = true);
+void GameThreadCall(dispatch_block_t Block, NSArray* SendModes = @[ NSDefaultRunLoopMode ], bool const bWait = true);
 
 template<typename ReturnType>
-ReturnType GameThreadReturn(ReturnType (^Block)(void), NSString* SendMode = NSDefaultRunLoopMode)
+ReturnType GameThreadReturn(ReturnType (^Block)(void), NSArray* SendModes = @[ NSDefaultRunLoopMode ])
 {
-	return PerformBlockOnRunLoop(Block, [NSRunLoop gameRunLoop], SendMode, NSDefaultRunLoopMode);
+	__block ReturnType ReturnValue;
+	GameThreadCall(^{ ReturnValue = Block(); }, SendModes, true);
+	return ReturnValue;
 }
 
-void GameThreadCall(dispatch_block_t Block, NSString* SendMode = NSDefaultRunLoopMode, bool const bWait = true);
-
 void RunGameThread(id Target, SEL Selector);
 
 void ProcessGameThreadEvents(void);
-
-#else
-
-class NSThread;
-class NSRunLoop;
-class FCocoaGameThread;
-
-#endif
diff --git a/Engine/Source/Runtime/Core/Public/Mac/MacEvent.h b/Engine/Source/Runtime/Core/Public/Mac/MacEvent.h
--- a/Engine/Source/Runtime/Core/Public/Mac/MacEvent.h
+++ b/Engine/Source/Runtime/Core/Public/Mac/MacEvent.h
@@ -27,14 +27,14 @@ class FMacEvent
 	FMacEvent(NSNotification* const Notification, NSWindow* const Window);
 	
 	// Send the event to the game run loop for processing using the specified SendMethod.
-	static void SendToGameRunLoop(FMacEvent const* const Event, EMacEventSendMethod SendMethod, NSString* SendMode = NSDefaultRunLoopMode);
+	static void SendToGameRunLoop(FMacEvent const* const Event, EMacEventSendMethod SendMethod, NSArray* SendModes = @[ NSDefaultRunLoopMode ]);
 	
 public:
 	// Send an NSEvent to the Game run loop as an FMacEvent using the specified SendMethod
-	static void SendToGameRunLoop(NSEvent* const Event, EMacEventSendMethod SendMethod, NSString* SendMode = NSDefaultRunLoopMode);
+	static void SendToGameRunLoop(NSEvent* const Event, EMacEventSendMethod SendMethod, NSArray* SendModes = @[ NSDefaultRunLoopMode ]);
 	
 	// Send an NSNotification to the Game run loop as an FMacEvent using the specified SendMethod
-	static void SendToGameRunLoop(NSNotification* const Notification, NSWindow* const Window, EMacEventSendMethod SendMethod, NSString* SendMode = NSDefaultRunLoopMode);
+	static void SendToGameRunLoop(NSNotification* const Notification, NSWindow* const Window, EMacEventSendMethod SendMethod, NSArray* SendModes = @[ NSDefaultRunLoopMode ]);
 	
 	// Destructor
 	~FMacEvent(void);
diff --git a/Engine/Source/Runtime/Launch/Private/Mac/LaunchMac.cpp b/Engine/Source/Runtime/Launch/Private/Mac/LaunchMac.cpp
--- a/Engine/Source/Runtime/Launch/Private/Mac/LaunchMac.cpp
+++ b/Engine/Source/Runtime/Launch/Private/Mac/LaunchMac.cpp
@@ -120,7 +120,7 @@ void EngineCrashHandler(const FGenericCrashContext & GenericContext)
 				GEngine->DeferredCommands.Add(TEXT("EXIT"));
 			}
 		}
-	}, NSDefaultRunLoopMode, false);
+	}, @[ NSDefaultRunLoopMode ], false);
 }
 
 - (IBAction)showAboutWindow:(id)Sender
@@ -131,7 +131,7 @@ void EngineCrashHandler(const FGenericCrashContext & GenericContext)
 		{
 			FModuleManager::GetModuleChecked<IMainFrameModule>(TEXT("MainFrame")).ShowAboutWindow();
 		}
-	}, NSDefaultRunLoopMode, false);
+	}, @[ NSDefaultRunLoopMode ], false);
 #else
 	[NSApp orderFrontStandardAboutPanel:Sender];
 #endif
@@ -145,7 +145,7 @@ void EngineCrashHandler(const FGenericCrashContext & GenericContext)
 		{
 			FModuleManager::LoadModuleChecked<ISettingsModule>("Settings").ShowViewer(FName("Editor"), FName("General"), FName("Appearance"));
 		}
-	}, NSDefaultRunLoopMode, false);
+	}, @[ NSDefaultRunLoopMode ], false);
 }
 #endif
 
diff --git a/Engine/Source/Runtime/Slate/Private/Framework/MultiBox/Mac/MacMenu.cpp b/Engine/Source/Runtime/Slate/Private/Framework/MultiBox/Mac/MacMenu.cpp
--- a/Engine/Source/Runtime/Slate/Private/Framework/MultiBox/Mac/MacMenu.cpp
+++ b/Engine/Source/Runtime/Slate/Private/Framework/MultiBox/Mac/MacMenu.cpp
@@ -325,7 +325,7 @@ void FSlateMacMenu::ExecuteMenuItemAction(const TSharedRef< const class FMenuEnt
 			(*MenuBlock)->GetDirectActions().Execute();
 		}
         delete MenuBlock;
-	}, NSDefaultRunLoopMode, false);
+	}, @[ NSDefaultRunLoopMode ], false);
 }
 
 static const TSharedRef<SWidget> FindTextBlockWidget(TSharedRef<SWidget> Content)
