Small fixes to OS X crash reporting so that we don't inject erroneus line endings into the XML report file.

diff --git a/Engine/Source/Runtime/Core/Private/Mac/MacPlatformMisc.cpp b/Engine/Source/Runtime/Core/Private/Mac/MacPlatformMisc.cpp
--- a/Engine/Source/Runtime/Core/Private/Mac/MacPlatformMisc.cpp
+++ b/Engine/Source/Runtime/Core/Private/Mac/MacPlatformMisc.cpp
@@ -1254,11 +1254,16 @@ void FMacCrashContext::GenerateWindowsErrorReport(char const* WERPath) const
 			
 			// Crash Module name
 			WriteUTF16String(ReportFile, TEXT("\t\t<Parameter3>"));
-			for(uint32 i = 0; i < PATH_MAX && (i < FCStringAnsi::Strlen(Info.dli_fname) + 1); i++)
+			if (FCStringAnsi::Strlen(Info.dli_fname))
 			{
-				Line[i] = Info.dli_fname[i];
+				FMemory::Memzero(Line, PATH_MAX * sizeof(TCHAR));
+				FUTF8ToTCHAR_Convert::Convert(Line, PATH_MAX, Info.dli_fname, FCStringAnsi::Strlen(Info.dli_fname));
+				WriteUTF16String(ReportFile, Line);
+			}
+			else
+			{
+				WriteUTF16String(ReportFile, TEXT("Unknown"));
 			}
-			WriteLine(ReportFile, Line);
 			WriteLine(ReportFile, TEXT("</Parameter3>"));
 			
 			// Check header
@@ -1304,11 +1309,9 @@ void FMacCrashContext::GenerateWindowsErrorReport(char const* WERPath) const
 		
 		// Fault type -> Signal
 		WriteUTF16String(ReportFile, TEXT("\t\t<Parameter8>"));
-		for(uint32 i = 0; i < PATH_MAX && (i < FCStringAnsi::Strlen(SignalDescription) + 1); i++)
-		{
-			Line[i] = SignalDescription[i];
-		}
-		WriteLine(ReportFile, Line);
+		FMemory::Memzero(Line, PATH_MAX * sizeof(TCHAR));
+		FUTF8ToTCHAR_Convert::Convert(Line, PATH_MAX, SignalDescription, FCStringAnsi::Strlen(SignalDescription));
+		WriteUTF16String(ReportFile, Line);
 		WriteLine(ReportFile, TEXT("</Parameter8>"));
 		
 		WriteUTF16String(ReportFile, TEXT("\t\t<Parameter9>"));
