Fixed crash caused by ICU break iterators persisting beyond ICU shutdown

TTP# 347830 - CRITICAL: Regression: PROJECT BROWSER: CRASH: Occurs on exit

The lifespan of the ICU break iterators is now managed by FICUBreakIteratorManager, which allows us to safely destroy any break iterators that still exist when ICU is shutdown. Our UE iterators now hold onto the internal ICU iterator via a weak pointer so that we know if it's been destroyed.

diff --git a/Engine/Source/Runtime/Core/Private/Internationalization/ICUBreakIterator.cpp b/Engine/Source/Runtime/Core/Private/Internationalization/ICUBreakIterator.cpp
--- a/Engine/Source/Runtime/Core/Private/Internationalization/ICUBreakIterator.cpp
+++ b/Engine/Source/Runtime/Core/Private/Internationalization/ICUBreakIterator.cpp
@@ -6,68 +6,134 @@
 #if UE_ENABLE_ICU
 #include "ICUTextCharacterIterator.h"
 
-FICUBreakIterator::FICUBreakIterator(TSharedRef<icu::BreakIterator>&& InICUBreakIterator)
-	: ICUBreakIterator(InICUBreakIterator)
+FICUBreakIteratorManager* FICUBreakIteratorManager::Singleton = nullptr;
+
+void FICUBreakIteratorManager::Create()
+{
+	check(!Singleton);
+	Singleton = new FICUBreakIteratorManager();
+}
+
+void FICUBreakIteratorManager::Destroy()
+{
+	check(Singleton);
+	delete Singleton;
+	Singleton = nullptr;
+}
+
+FICUBreakIteratorManager& FICUBreakIteratorManager::Get()
+{
+	check(Singleton);
+	return *Singleton;
+}
+
+TWeakPtr<icu::BreakIterator> FICUBreakIteratorManager::CreateCharacterBoundaryIterator()
+{
+	UErrorCode ICUStatus = U_ZERO_ERROR;
+	TSharedRef<icu::BreakIterator> Iterator = MakeShareable(icu::BreakIterator::createCharacterInstance(icu::Locale::getDefault(), ICUStatus));
+	AllocatedIterators.Add(Iterator);
+	return Iterator;
+}
+
+TWeakPtr<icu::BreakIterator> FICUBreakIteratorManager::CreateWordBreakIterator()
+{
+	UErrorCode ICUStatus = U_ZERO_ERROR;
+	TSharedRef<icu::BreakIterator> Iterator = MakeShareable(icu::BreakIterator::createWordInstance(icu::Locale::getDefault(), ICUStatus));
+	AllocatedIterators.Add(Iterator);
+	return Iterator;
+}
+
+TWeakPtr<icu::BreakIterator> FICUBreakIteratorManager::CreateLineBreakIterator()
+{
+	UErrorCode ICUStatus = U_ZERO_ERROR;
+	TSharedRef<icu::BreakIterator> Iterator = MakeShareable(icu::BreakIterator::createLineInstance(icu::Locale::getDefault(), ICUStatus));
+	AllocatedIterators.Add(Iterator);
+	return Iterator;
+}
+
+void FICUBreakIteratorManager::DestroyIterator(TWeakPtr<icu::BreakIterator>& InIterator)
+{
+	TSharedPtr<icu::BreakIterator> Iterator = InIterator.Pin();
+	if(Iterator.IsValid())
+	{
+		AllocatedIterators.Remove(Iterator);
+	}
+	InIterator.Reset();
+}
+
+FICUBreakIterator::FICUBreakIterator(TWeakPtr<icu::BreakIterator>&& InICUBreakIteratorHandle)
+	: ICUBreakIteratorHandle(InICUBreakIteratorHandle)
+{
+}
+
+FICUBreakIterator::~FICUBreakIterator()
 {
+	// This assumes that FICUBreakIterator owns the iterator, and that nothing ever copies an FICUBreakIterator instance
+	FICUBreakIteratorManager::Get().DestroyIterator(ICUBreakIteratorHandle);
 }
 
 void FICUBreakIterator::SetString(const FText& InText)
 {
-	ICUBreakIterator->adoptText(new FICUTextCharacterIterator(InText)); // ICUBreakIterator takes ownership of this instance
+	GetInternalBreakIterator()->adoptText(new FICUTextCharacterIterator(InText)); // ICUBreakIterator takes ownership of this instance
 	ResetToBeginning();
 }
 
 void FICUBreakIterator::SetString(const FString& InString)
 {
-	ICUBreakIterator->adoptText(new FICUTextCharacterIterator(InString)); // ICUBreakIterator takes ownership of this instance
+	GetInternalBreakIterator()->adoptText(new FICUTextCharacterIterator(InString)); // ICUBreakIterator takes ownership of this instance
 	ResetToBeginning();
 }
 
 void FICUBreakIterator::SetString(const TCHAR* const InString, const int32 InStringLength) 
 {
-	ICUBreakIterator->adoptText(new FICUTextCharacterIterator(InString, InStringLength)); // ICUBreakIterator takes ownership of this instance
+	GetInternalBreakIterator()->adoptText(new FICUTextCharacterIterator(InString, InStringLength)); // ICUBreakIterator takes ownership of this instance
 	ResetToBeginning();
 }
 
 void FICUBreakIterator::ClearString()
 {
-	ICUBreakIterator->setText(icu::UnicodeString());
+	GetInternalBreakIterator()->setText(icu::UnicodeString());
 	ResetToBeginning();
 }
 
 int32 FICUBreakIterator::GetCurrentPosition() const
 {
-	return ICUBreakIterator->current();
+	return GetInternalBreakIterator()->current();
 }
 
 int32 FICUBreakIterator::ResetToBeginning()
 {
-	return ICUBreakIterator->first();
+	return GetInternalBreakIterator()->first();
 }
 
 int32 FICUBreakIterator::ResetToEnd()
 {
-	return ICUBreakIterator->last();
+	return GetInternalBreakIterator()->last();
 }
 
 int32 FICUBreakIterator::MoveToPrevious()
 {
-	return ICUBreakIterator->previous();
+	return GetInternalBreakIterator()->previous();
 }
 
 int32 FICUBreakIterator::MoveToNext()
 {
-	return ICUBreakIterator->next();
+	return GetInternalBreakIterator()->next();
 }
 
 int32 FICUBreakIterator::MoveToCandidateBefore(const int32 InIndex)
 {
-	return ICUBreakIterator->preceding(InIndex);
+	return GetInternalBreakIterator()->preceding(InIndex);
 }
 
 int32 FICUBreakIterator::MoveToCandidateAfter(const int32 InIndex)
 {
-	return ICUBreakIterator->following(InIndex);
+	return GetInternalBreakIterator()->following(InIndex);
+}
+
+TSharedRef<icu::BreakIterator> FICUBreakIterator::GetInternalBreakIterator() const
+{
+	return ICUBreakIteratorHandle.Pin().ToSharedRef();
 }
 
 #endif
diff --git a/Engine/Source/Runtime/Core/Private/Internationalization/ICUBreakIterator.h b/Engine/Source/Runtime/Core/Private/Internationalization/ICUBreakIterator.h
--- a/Engine/Source/Runtime/Core/Private/Internationalization/ICUBreakIterator.h
+++ b/Engine/Source/Runtime/Core/Private/Internationalization/ICUBreakIterator.h
@@ -8,13 +8,34 @@
 
 #include <unicode/brkiter.h>
 
+/**
+ * Manages the lifespan of ICU break iterators
+ */
+class FICUBreakIteratorManager
+{
+public:
+	static void Create();
+	static void Destroy();
+	static FICUBreakIteratorManager& Get();
+
+	TWeakPtr<icu::BreakIterator> CreateCharacterBoundaryIterator();
+	TWeakPtr<icu::BreakIterator> CreateWordBreakIterator();
+	TWeakPtr<icu::BreakIterator> CreateLineBreakIterator();
+	void DestroyIterator(TWeakPtr<icu::BreakIterator>& InIterator);
+
+private:
+	static FICUBreakIteratorManager* Singleton;
+	TSet<TSharedPtr<icu::BreakIterator>> AllocatedIterators;
+};
+
 /**
  * Wraps an ICU break iterator instance inside our own break iterator API
  */
 class FICUBreakIterator : public IBreakIterator
 {
 public:
-	FICUBreakIterator(TSharedRef<icu::BreakIterator>&& InICUBreakIterator);
+	FICUBreakIterator(TWeakPtr<icu::BreakIterator>&& InICUBreakIteratorHandle);
+	virtual ~FICUBreakIterator();
 
 	virtual void SetString(const FText& InText) override;
 	virtual void SetString(const FString& InString) override;
@@ -32,7 +53,10 @@ public:
 	virtual int32 MoveToCandidateAfter(const int32 InIndex) override;
 
 protected:
-	TSharedRef<icu::BreakIterator> ICUBreakIterator;
+	TSharedRef<icu::BreakIterator> GetInternalBreakIterator() const;
+
+private:
+	TWeakPtr<icu::BreakIterator> ICUBreakIteratorHandle;
 };
 
 #endif
diff --git a/Engine/Source/Runtime/Core/Private/Internationalization/ICUCharacterBoundaryIterator.cpp b/Engine/Source/Runtime/Core/Private/Internationalization/ICUCharacterBoundaryIterator.cpp
--- a/Engine/Source/Runtime/Core/Private/Internationalization/ICUCharacterBoundaryIterator.cpp
+++ b/Engine/Source/Runtime/Core/Private/Internationalization/ICUCharacterBoundaryIterator.cpp
@@ -10,20 +10,11 @@ class FICUCharacterBoundaryIterator : public FICUBreakIterator
 {
 public:
 	FICUCharacterBoundaryIterator();
-
-private:
-	static TSharedRef<icu::BreakIterator> CreateInternalGraphemeBreakIterator();
 };
 
 FICUCharacterBoundaryIterator::FICUCharacterBoundaryIterator()
-	: FICUBreakIterator(CreateInternalGraphemeBreakIterator())
-{
-}
-
-TSharedRef<icu::BreakIterator> FICUCharacterBoundaryIterator::CreateInternalGraphemeBreakIterator()
+	: FICUBreakIterator(FICUBreakIteratorManager::Get().CreateCharacterBoundaryIterator())
 {
-	UErrorCode ICUStatus = U_ZERO_ERROR;
-	return MakeShareable(icu::BreakIterator::createCharacterInstance(icu::Locale::getDefault(), ICUStatus));
 }
 
 TSharedRef<IBreakIterator> FBreakIterator::CreateCharacterBoundaryIterator()
diff --git a/Engine/Source/Runtime/Core/Private/Internationalization/ICULineBreakIterator.cpp b/Engine/Source/Runtime/Core/Private/Internationalization/ICULineBreakIterator.cpp
--- a/Engine/Source/Runtime/Core/Private/Internationalization/ICULineBreakIterator.cpp
+++ b/Engine/Source/Runtime/Core/Private/Internationalization/ICULineBreakIterator.cpp
@@ -10,20 +10,11 @@ class FICULineBreakIterator : public FICUBreakIterator
 {
 public:
 	FICULineBreakIterator();
-
-private:
-	static TSharedRef<icu::BreakIterator> CreateInternalLineBreakIterator();
 };
 
 FICULineBreakIterator::FICULineBreakIterator()
-	: FICUBreakIterator(CreateInternalLineBreakIterator())
-{
-}
-
-TSharedRef<icu::BreakIterator> FICULineBreakIterator::CreateInternalLineBreakIterator()
+	: FICUBreakIterator(FICUBreakIteratorManager::Get().CreateLineBreakIterator())
 {
-	UErrorCode ICUStatus = U_ZERO_ERROR;
-	return MakeShareable( icu::BreakIterator::createLineInstance(icu::Locale::getDefault(), ICUStatus));
 }
 
 TSharedRef<IBreakIterator> FBreakIterator::CreateLineBreakIterator()
diff --git a/Engine/Source/Runtime/Core/Private/Internationalization/ICUWordBreakIterator.cpp b/Engine/Source/Runtime/Core/Private/Internationalization/ICUWordBreakIterator.cpp
--- a/Engine/Source/Runtime/Core/Private/Internationalization/ICUWordBreakIterator.cpp
+++ b/Engine/Source/Runtime/Core/Private/Internationalization/ICUWordBreakIterator.cpp
@@ -10,20 +10,11 @@ class FICUWordBreakIterator : public FICUBreakIterator
 {
 public:
 	FICUWordBreakIterator();
-
-private:
-	static TSharedRef<icu::BreakIterator> CreateInternalWordBreakIterator();
 };
 
 FICUWordBreakIterator::FICUWordBreakIterator()
-	: FICUBreakIterator(CreateInternalWordBreakIterator())
-{
-}
-
-TSharedRef<icu::BreakIterator> FICUWordBreakIterator::CreateInternalWordBreakIterator()
+	: FICUBreakIterator(FICUBreakIteratorManager::Get().CreateWordBreakIterator())
 {
-	UErrorCode ICUStatus = U_ZERO_ERROR;
-	return MakeShareable(icu::BreakIterator::createWordInstance(icu::Locale::getDefault(), ICUStatus));
 }
 
 TSharedRef<IBreakIterator> FBreakIterator::CreateWordBreakIterator()
diff --git a/Engine/Source/Runtime/Core/Private/Internationalization/Internationalization.cpp b/Engine/Source/Runtime/Core/Private/Internationalization/Internationalization.cpp
--- a/Engine/Source/Runtime/Core/Private/Internationalization/Internationalization.cpp
+++ b/Engine/Source/Runtime/Core/Private/Internationalization/Internationalization.cpp
@@ -15,6 +15,7 @@
 
 #if UE_ENABLE_ICU
 #include "ICUUtilities.h"
+#include "ICUBreakIterator.h"
 #include <unicode/locid.h>
 #include <unicode/timezone.h>
 #include <unicode/uclean.h>
@@ -239,6 +240,8 @@ void FInternationalization::Initialize()
 
 	u_setDataFileFunctions(nullptr, &FICUDataCallbacks::OpenDataFile, &FICUDataCallbacks::CloseDataFile, &(ICUStatus));
 	u_init(&(ICUStatus));
+
+	FICUBreakIteratorManager::Create();
 #endif /*UE_ENABLE_ICU*/
 
 	PopulateAllCultures();
@@ -263,6 +266,7 @@ void FInternationalization::Terminate()
 	AllCultures.Empty();
 
 #if UE_ENABLE_ICU
+	FICUBreakIteratorManager::Destroy();
 	u_cleanup();
 #if IS_PROGRAM || !IS_MONOLITHIC
 	UnloadDLLs();
