TArray documentation and a coupld of methods deprecation.

diff --git a/Engine/Source/Runtime/Core/Public/Containers/Array.h b/Engine/Source/Runtime/Core/Public/Containers/Array.h
--- a/Engine/Source/Runtime/Core/Public/Containers/Array.h
+++ b/Engine/Source/Runtime/Core/Public/Containers/Array.h
@@ -358,15 +358,18 @@ class TArray
 public:
 	typedef InElementType ElementType;
 
+	/**
+	 * Constructor, initializes element number counters.
+	 */
 	TArray()
-	:   ArrayNum( 0 )
-	,	ArrayMax( 0 )
+		: ArrayNum(0)
+		, ArrayMax(0)
 	{}
 
 	/**
-	 * Copy constructor. Use the common routine to perform the copy
+	 * Copy constructor with changed allocator. Use the common routine to perform the copy.
 	 *
-	 * @param Other the source array to copy
+	 * @param Other The source array to copy.
 	 */
 	template<typename OtherAllocator>
 	explicit TArray(const TArray<ElementType,OtherAllocator>& Other)
@@ -374,29 +377,55 @@ public:
 		CopyToEmpty(Other);
 	}
 
+	/**
+	 * Copy constructor. Use the common routine to perform the copy.
+	 *
+	 * @param Other The source array to copy.
+	 */
 	TArray(const TArray& Other)
 	{
 		CopyToEmpty(Other);
 	}
 
+	/**
+	 * Copy constructor. Use the common routine to perform the copy.
+	 *
+	 * @param Other The source array to copy.
+	 * @param ExtraSlack Tells how much extra memory should be preallocated
+	 *                   at the end of the array in the number of elements.
+	 */
 	TArray(const TArray& Other, int32 ExtraSlack)
 	{
 		CopyToEmpty(Other, ExtraSlack);
 	}
 
+	/**
+	 * Assignment operator. First deletes all currently contained elements
+	 * and then copies from other array.
+	 *
+	 * Allocator changing version.
+	 *
+	 * @param Other The source array to assign from.
+	 */
 	template<typename OtherAllocator>
-	TArray& operator=( const TArray<ElementType,OtherAllocator>& Other )
+	TArray& operator=(const TArray<ElementType, OtherAllocator>& Other)
 	{
-		DestructItems(GetTypedData(), ArrayNum);
+		DestructItems(GetData(), ArrayNum);
 		CopyToEmpty(Other);
 		return *this;
 	}
 
-	TArray& operator=( const TArray<ElementType,Allocator>& Other )
+	/**
+	 * Assignment operator. First deletes all currently contained elements
+	 * and then copies from other array.
+	 *
+	 * @param Other The source array to assign from.
+	 */
+	TArray& operator=(const TArray<ElementType, Allocator>& Other)
 	{
 		if (this != &Other)
 		{
-			DestructItems(GetTypedData(), ArrayNum);
+			DestructItems(GetData(), ArrayNum);
 			CopyToEmpty(Other);
 		}
 		return *this;
@@ -405,6 +434,14 @@ public:
 #if PLATFORM_COMPILER_HAS_RVALUE_REFERENCES
 
 private:
+	/**
+	 * Moves or copies array. Depends on the array type traits.
+	 *
+	 * This override moves.
+	 *
+	 * @param ToArray Array to move into.
+	 * @param FromArray Array to move from.
+	 */
 	template <typename ArrayType>
 	static FORCEINLINE typename TEnableIf<TContainerTraits<ArrayType>::MoveWillEmptyContainer>::Type MoveOrCopy(ArrayType& ToArray, ArrayType& FromArray)
 	{
@@ -416,6 +453,16 @@ private:
 		FromArray.ArrayMax = 0;
 	}
 
+	/**
+	 * Moves or copies array. Depends on the array type traits.
+	 *
+	 * This override moves.
+	 *
+	 * @param ToArray Array to move into.
+	 * @param FromArray Array to move from.
+	 * @param ExtraSlack Tells how much extra memory should be preallocated
+	 *                   at the end of the array in the number of elements.
+	 */
 	template <typename ArrayType>
 	static FORCEINLINE typename TEnableIf<TContainerTraits<ArrayType>::MoveWillEmptyContainer>::Type MoveOrCopy(ArrayType& ToArray, ArrayType& FromArray, int32 ExtraSlack)
 	{
@@ -424,6 +471,16 @@ private:
 		ToArray.Reserve(ToArray.ArrayNum + ExtraSlack);
 	}
 
+	/**
+	 * Moves or copies array. Depends on the array type traits.
+	 *
+	 * This override copies.
+	 *
+	 * @param ToArray Array to move into.
+	 * @param FromArray Array to move from.
+	 * @param ExtraSlack Tells how much extra memory should be preallocated
+	 *                   at the end of the array in the number of elements.
+	 */
 	template <typename ArrayType>
 	static FORCEINLINE typename TEnableIf<!TContainerTraits<ArrayType>::MoveWillEmptyContainer>::Type MoveOrCopy(ArrayType& ToArray, ArrayType& FromArray, int32 ExtraSlack = 0)
 	{
@@ -435,21 +492,38 @@ public:
 	// to tell if they're compatible with the current one.  Probably going to be a pretty
 	// rare requirement anyway.
 
-	TArray(TArray<ElementType,Allocator>&& Other)
+	/**
+	 * Move constructor.
+	 *
+	 * @param Other Array to move from.
+	 */
+	TArray(TArray<ElementType, Allocator>&& Other)
 	{
 		MoveOrCopy(*this, Other);
 	}
 
-	TArray(TArray<ElementType,Allocator>&& Other, int32 ExtraSlack)
+	/**
+	 * Move constructor.
+	 *
+	 * @param Other Array to move from.
+	 * @param ExtraSlack Tells how much extra memory should be preallocated
+	 *                   at the end of the array in the number of elements.
+	 */
+	TArray(TArray<ElementType, Allocator>&& Other, int32 ExtraSlack)
 	{
 		MoveOrCopy(*this, Other, ExtraSlack);
 	}
 
-	TArray& operator=(TArray<ElementType,Allocator>&& Other)
+	/**
+	 * Move assignment operator.
+	 *
+	 * @param Other Array to assign and move from.
+	 */
+	TArray& operator=(TArray<ElementType, Allocator>&& Other)
 	{
 		if (this != &Other)
 		{
-			DestructItems(GetTypedData(), ArrayNum);
+			DestructItems(GetData(), ArrayNum);
 			MoveOrCopy(*this, Other);
 		}
 		return *this;
@@ -457,9 +531,12 @@ public:
 
 #endif
 
+	/**
+	 * Destructor.
+	 */
 	~TArray()
 	{
-		DestructItems(GetTypedData(), ArrayNum);
+		DestructItems(GetData(), ArrayNum);
 
 		#if defined(_MSC_VER) && !defined(__clang__)	// Relies on MSVC-specific lazy template instantiation to support arrays of incomplete types
 			// ensure that DebugGet gets instantiated.
@@ -471,33 +548,49 @@ public:
 	/**
 	 * Helper function for returning a typed pointer to the first array entry.
 	 *
-	 * @return pointer to first array entry or NULL if ArrayMax==0
+	 * @returns Pointer to first array entry or NULL if ArrayMax == 0.
 	 */
+	DEPRECATED(4.6, "This function is deprecated as it does the same as GetData(). Please use GetData() instead.")
 	FORCEINLINE ElementType* GetTypedData()
 	{
 		return (ElementType*)AllocatorInstance.GetAllocation();
 	}
+
+	/**
+	 * Helper function for returning a typed pointer to the first array entry.
+	 *
+	 * @returns Pointer to first array entry or NULL if ArrayMax == 0.
+	 */
 	FORCEINLINE ElementType* GetData()
 	{
 		return (ElementType*)AllocatorInstance.GetAllocation();
 	}
+
 	/**
 	 * Helper function for returning a typed pointer to the first array entry.
 	 *
-	 * @return pointer to first array entry or NULL if ArrayMax==0
+	 * @returns Pointer to first array entry or NULL if ArrayMax == 0.
 	 */
+	DEPRECATED(4.6, "This function is deprecated as it does the same as GetData(). Please use GetData() instead.")
 	FORCEINLINE const ElementType* GetTypedData() const
 	{
 		return (const ElementType*)AllocatorInstance.GetAllocation();
 	}
+
+	/**
+	 * Helper function for returning a typed pointer to the first array entry.
+	 *
+	 * @returns Pointer to first array entry or NULL if ArrayMax == 0.
+	 */
 	FORCEINLINE const ElementType* GetData() const
 	{
 		return (const ElementType*)AllocatorInstance.GetAllocation();
 	}
+
 	/** 
-	 * Helper function returning the size of the inner type
+	 * Helper function returning the size of the inner type.
 	 *
-	 * @return size in bytes of array type
+	 * @returns Size in bytes of array type.
 	 */
 	FORCEINLINE uint32 GetTypeSize() const
 	{
@@ -505,11 +598,12 @@ public:
 	}
 
 	/** 
-	 * Helper function to return the amount of memory allocated by this container 
+	 * Helper function to return the amount of memory allocated by this
+	 * container.
 	 *
-	 * @return number of bytes allocated by this container
+	 * @returns Number of bytes allocated by this container.
 	 */
-	FORCEINLINE uint32 GetAllocatedSize( void ) const
+	FORCEINLINE uint32 GetAllocatedSize(void) const
 	{
 		return AllocatorInstance.GetAllocatedSize(ArrayMax, sizeof(ElementType));
 	}
@@ -522,10 +616,20 @@ public:
 		return ArrayMax - ArrayNum;
 	}
 
+	/**
+	 * Checks array invariants: if array size is greater than zero and less
+	 * than maximum.
+	 */
 	FORCEINLINE void CheckInvariants() const
 	{
 		checkSlow((ArrayNum >= 0) & (ArrayMax >= ArrayNum)); // & for one branch
 	}
+
+	/**
+	 * Checks if index is in array range.
+	 *
+	 * @param Index Index to check.
+	 */
 	FORCEINLINE void RangeCheck(int32 Index) const
 	{
 		CheckInvariants();
@@ -537,81 +641,192 @@ public:
 		}
 	}
 
-	FORCEINLINE bool IsValidIndex( int32 i ) const
+	/**
+	 * Tests if index is valid, i.e. greater than zero and less than number of
+	 * elements in array.
+	 *
+	 * @param Index Index to test.
+	 *
+	 * @returns True if index is valid. False otherwise.
+	 */
+	FORCEINLINE bool IsValidIndex(int32 Index) const
 	{
-		return i>=0 && i<ArrayNum;
+		return Index >= 0 && Index < ArrayNum;
 	}
+
+	/**
+	 * Returns number of elements in array.
+	 *
+	 * @returns Number of elements in array.
+	 */
 	FORCEINLINE int32 Num() const
 	{
 		return ArrayNum;
 	}
+
+	/**
+	 * Returns maximum number of elements in array.
+	 *
+	 * @returns Maximum number of elements in array.
+	 */
 	FORCEINLINE int32 Max() const
 	{
 		return ArrayMax;
 	}
 
-	FORCEINLINE ElementType& operator[]( int32 i )
+	/**
+	 * Array bracket operator. Returns reference to element at give index.
+	 *
+	 * @returns Reference to indexed element.
+	 */
+	FORCEINLINE ElementType& operator[](int32 Index)
 	{
-		RangeCheck(i);
-		return GetTypedData()[i];
+		RangeCheck(Index);
+		return GetData()[Index];
 	}
-	FORCEINLINE const ElementType& operator[]( int32 i ) const
+
+	/**
+	 * Array bracket operator. Returns reference to element at give index.
+	 *
+	 * Const version of the above.
+	 *
+	 * @returns Reference to indexed element.
+	 */
+	FORCEINLINE const ElementType& operator[](int32 Index) const
 	{
-		RangeCheck(i);
-		return GetTypedData()[i];
+		RangeCheck(Index);
+		return GetData()[Index];
 	}
 
-	ElementType Pop(bool bAllowShrinking=true)
+	/**
+	 * Pops element from the array.
+	 *
+	 * @param bAllowShrinking If this call allows shrinking of the array during
+	 *                        element remove.
+	 *
+	 * @returns Popped element.
+	 */
+	ElementType Pop(bool bAllowShrinking = true)
 	{
 		RangeCheck(0);
-		ElementType Result = MoveTemp(GetTypedData()[ArrayNum-1]);
-		RemoveAt( ArrayNum-1, 1, bAllowShrinking );
+		ElementType Result = MoveTemp(GetData()[ArrayNum - 1]);
+		RemoveAt(ArrayNum - 1, 1, bAllowShrinking);
 		return Result;
 	}
-	void Push( ElementType&& Item )
+
+	/**
+	 * Pushes element into the array.
+	 *
+	 * @param Item Item to push.
+	 */
+	void Push(ElementType&& Item)
 	{
 		Add(MoveTemp(Item));
 	}
-	void Push( const ElementType& Item )
+
+	/**
+	 * Pushes element into the array.
+	 *
+	 * Const ref version of the above.
+	 *
+	 * @param Item Item to push.
+	 */
+	void Push(const ElementType& Item)
 	{
 		Add(Item);
 	}
+
+	/**
+	 * Returns the top element, i.e. the last one.
+	 *
+	 * @returns Reference to the top element.
+	 */
 	ElementType& Top()
 	{
 		return Last();
 	}
+
+	/**
+	 * Returns the top element, i.e. the last one.
+	 *
+	 * Const version of the above.
+	 *
+	 * @returns Reference to the top element.
+	 */
 	const ElementType& Top() const
 	{
 		return Last();
 	}
-	ElementType& Last( int32 c=0 )
+
+	/**
+	 * Returns n-th last element from the array.
+	 *
+	 * @param IndexFromTheEnd (Optional) Index from the end of array.
+	 *                        Default is 0.
+	 *
+	 * @returns Reference to n-th last element from the array.
+	 */
+	ElementType& Last(int32 IndexFromTheEnd = 0)
 	{
-		RangeCheck(ArrayNum-c-1);
-		return GetTypedData()[ArrayNum-c-1];
+		RangeCheck(ArrayNum - IndexFromTheEnd - 1);
+		return GetData()[ArrayNum - IndexFromTheEnd - 1];
 	}
-	const ElementType& Last( int32 c=0 ) const
+
+	/**
+	 * Returns n-th last element from the array.
+	 *
+	 * Const version of the above.
+	 *
+	 * @param IndexFromTheEnd (Optional) Index from the end of array.
+	 *                        Default is 0.
+	 *
+	 * @returns Reference to n-th last element from the array.
+	 */
+	const ElementType& Last(int32 IndexFromTheEnd = 0) const
 	{
-		RangeCheck(ArrayNum-c-1);
-		return GetTypedData()[ArrayNum-c-1];
+		RangeCheck(ArrayNum - IndexFromTheEnd - 1);
+		return GetData()[ArrayNum - IndexFromTheEnd - 1];
 	}
+
+	/**
+	 * Shrinks the array's used memory to smallest possible to store elements
+	 * currently in it.
+	 */
 	void Shrink()
 	{
 		CheckInvariants();
-		if( ArrayMax != ArrayNum )
+		if (ArrayMax != ArrayNum)
 		{
 			ArrayMax = ArrayNum;
-			AllocatorInstance.ResizeAllocation(ArrayNum,ArrayMax,sizeof(ElementType));
+			AllocatorInstance.ResizeAllocation(ArrayNum, ArrayMax, sizeof(ElementType));
 		}
 	}
-	FORCEINLINE bool Find( const ElementType& Item, int32& Index ) const
+
+	/**
+	 * Finds element within the array.
+	 *
+	 * @param Item Item to look for.
+	 * @param Index Output parameter. Found index.
+	 *
+	 * @returns True if found. False otherwise.
+	 */
+	FORCEINLINE bool Find(const ElementType& Item, int32& Index) const
 	{
 		Index = this->Find(Item);
 		return Index != INDEX_NONE;
 	}
-	int32 Find( const ElementType& Item ) const
+
+	/**
+	 * Finds element within the array.
+	 *
+	 * @param Item Item to look for.
+	 *
+	 * @returns Index of the found element. INDEX_NONE otherwise.
+	 */
+	int32 Find(const ElementType& Item) const
 	{
-		const ElementType* RESTRICT Start = GetTypedData();
-		for (const ElementType* RESTRICT Data = Start, * RESTRICT DataEnd = Data + ArrayNum; Data != DataEnd; ++Data)
+		const ElementType* RESTRICT Start = GetData();
+		for (const ElementType* RESTRICT Data = Start, *RESTRICT DataEnd = Data + ArrayNum; Data != DataEnd; ++Data)
 		{
 			if (*Data == Item)
 			{
@@ -620,15 +835,32 @@ public:
 		}
 		return INDEX_NONE;
 	}
-	FORCEINLINE bool FindLast( const ElementType& Item, int32& Index ) const
+
+	/**
+	 * Finds element within the array starting from the end.
+	 *
+	 * @param Item Item to look for.
+	 * @param Index Output parameter. Found index.
+	 *
+	 * @returns True if found. False otherwise.
+	 */
+	FORCEINLINE bool FindLast(const ElementType& Item, int32& Index) const
 	{
 		Index = this->FindLast(Item);
 		return Index != INDEX_NONE;
 	}
-	int32 FindLast( const ElementType& Item ) const
+
+	/**
+	 * Finds element within the array starting from the end.
+	 *
+	 * @param Item Item to look for.
+	 *
+	 * @returns Index of the found element. INDEX_NONE otherwise.
+	 */
+	int32 FindLast(const ElementType& Item) const
 	{
-		const ElementType* RESTRICT End = GetTypedData() + ArrayNum;
-		for (const ElementType* RESTRICT Data = End, * RESTRICT DataStart = Data - ArrayNum; Data != DataStart; )
+		const ElementType* RESTRICT End = GetData() + ArrayNum;
+		for (const ElementType* RESTRICT Data = End, *RESTRICT DataStart = Data - ArrayNum; Data != DataStart;)
 		{
 			--Data;
 			if (*Data == Item)
@@ -638,32 +870,48 @@ public:
 		}
 		return INDEX_NONE;
 	}
+
+	/**
+	 * Finds element within the array that fulfills given predicate.
+	 *
+	 * @param MatchFunctorType A functor object with implemented
+	 *                         bool Matches(const ElementType& Element) method
+	 *                         that returns true if given element is the one we
+	 *                         look for.
+	 *
+	 * @returns Index of the found element. INDEX_NONE otherwise.
+	 */
 	template<typename MatchFunctorType>
-	int32 FindMatch( const MatchFunctorType& Matcher ) const
+	int32 FindMatch(const MatchFunctorType& Matcher) const
 	{
-		const ElementType* const RESTRICT DataEnd = GetTypedData() + ArrayNum;
-		for(const ElementType* RESTRICT Data = GetTypedData();
+		const ElementType* const RESTRICT DataEnd = GetData() + ArrayNum;
+		for (const ElementType* RESTRICT Data = GetData();
 			Data < DataEnd;
 			Data++
 			)
 		{
-			if( Matcher.Matches(*Data) )
+			if (Matcher.Matches(*Data))
 			{
-				return static_cast<int32>(Data - GetTypedData());
+				return static_cast<int32>(Data - GetData());
 			}
 		}
 		return INDEX_NONE;
 	}
+
 	/**
-	 * Finds an item by key (assuming the ElementType overloads operator== for the comparison).
-	 * @param Key	The key to search by
-	 * @return		Index to the first matching element, or INDEX_NONE if none is found
+	 * Finds an item by key (assuming the ElementType overloads operator== for
+	 * the comparison).
+	 *
+	 * @param Key The key to search by.
+	 *
+	 * @returns Index to the first matching element, or INDEX_NONE if none is
+	 *          found.
 	 */
 	template <typename KeyType>
-	int32 IndexOfByKey( const KeyType& Key ) const
+	int32 IndexOfByKey(const KeyType& Key) const
 	{
-		const ElementType* RESTRICT Start = GetTypedData();
-		for (const ElementType* RESTRICT Data = Start, * RESTRICT DataEnd = Start + ArrayNum; Data != DataEnd; ++Data)
+		const ElementType* RESTRICT Start = GetData();
+		for (const ElementType* RESTRICT Data = Start, *RESTRICT DataEnd = Start + ArrayNum; Data != DataEnd; ++Data)
 		{
 			if (*Data == Key)
 			{
@@ -672,16 +920,20 @@ public:
 		}
 		return INDEX_NONE;
 	}
+
 	/**
 	 * Finds an item by predicate.
-	 * @param Pred	The predicate to match
-	 * @return		Index to the first matching element, or INDEX_NONE if none is found
+	 *
+	 * @param Pred The predicate to match.
+	 *
+	 * @returns Index to the first matching element, or INDEX_NONE if none is
+	 *          found.
 	 */
 	template <typename Predicate>
 	int32 IndexOfByPredicate(Predicate Pred) const
 	{
-		const ElementType* RESTRICT Start = GetTypedData();
-		for (const ElementType* RESTRICT Data = Start, * RESTRICT DataEnd = Start + ArrayNum; Data != DataEnd; ++Data)
+		const ElementType* RESTRICT Start = GetData();
+		for (const ElementType* RESTRICT Data = Start, *RESTRICT DataEnd = Start + ArrayNum; Data != DataEnd; ++Data)
 		{
 			if (Pred(*Data))
 			{
@@ -690,25 +942,33 @@ public:
 		}
 		return INDEX_NONE;
 	}
+
 	/**
-	 * Finds an item by key (assuming the ElementType overloads operator== for the comparison).
-	 * @param Key	The key to search by
-	 * @return		Pointer to the first matching element, or NULL if none is found
+	 * Finds an item by key (assuming the ElementType overloads operator== for
+	 * the comparison).
+	 *
+	 * @param Key The key to search by.
+	 *
+	 * @returns Pointer to the first matching element, or NULL if none is found.
 	 */
 	template <typename KeyType>
-	FORCEINLINE const ElementType* FindByKey( const KeyType& Key ) const
+	FORCEINLINE const ElementType* FindByKey(const KeyType& Key) const
 	{
 		return const_cast<TArray*>(this)->FindByKey(Key);
 	}
+
 	/**
-	 * Finds an item by key (assuming the ElementType overloads operator== for the comparison).
-	 * @param Key	The key to search by
-	 * @return		Pointer to the first matching element, or NULL if none is found
+	 * Finds an item by key (assuming the ElementType overloads operator== for
+	 * the comparison).
+	 *
+	 * @param Key The key to search by.
+	 *
+	 * @returns Pointer to the first matching element, or NULL if none is found.
 	 */
 	template <typename KeyType>
-	ElementType* FindByKey( const KeyType& Key )
+	ElementType* FindByKey(const KeyType& Key)
 	{
-		for (ElementType* RESTRICT Data = GetTypedData(), * RESTRICT DataEnd = Data + ArrayNum; Data != DataEnd; ++Data)
+		for (ElementType* RESTRICT Data = GetData(), *RESTRICT DataEnd = Data + ArrayNum; Data != DataEnd; ++Data)
 		{
 			if (*Data == Key)
 			{
@@ -718,25 +978,33 @@ public:
 
 		return NULL;
 	}
+
 	/**
 	 * Finds an element which matches a predicate functor.
-	 * @param  Pred The functor to apply to each element.
-	 * @return      Pointer to the first element for which the predicate returns true, or NULL if none is found.
+	 *
+	 * @param Pred The functor to apply to each element.
+	 *
+	 * @returns Pointer to the first element for which the predicate returns
+	 *          true, or NULL if none is found.
 	 */
 	template <typename Predicate>
 	FORCEINLINE const ElementType* FindByPredicate(Predicate Pred) const
 	{
 		return const_cast<TArray*>(this)->FindByPredicate(Pred);
 	}
+
 	/**
 	 * Finds an element which matches a predicate functor.
-	 * @param  Pred The functor to apply to each element.
-	 * @return      Pointer to the first element for which the predicate returns true, or NULL if none is found.
+	 *
+	 * @param Pred The functor to apply to each element.
+	 *
+	 * @return Pointer to the first element for which the predicate returns
+	 *         true, or NULL if none is found.
 	 */
 	template <typename Predicate>
 	ElementType* FindByPredicate(Predicate Pred)
 	{
-		for (ElementType* RESTRICT Data = GetTypedData(), * RESTRICT DataEnd = Data + ArrayNum; Data != DataEnd; ++Data)
+		for (ElementType* RESTRICT Data = GetData(), *RESTRICT DataEnd = Data + ArrayNum; Data != DataEnd; ++Data)
 		{
 			if (Pred(*Data))
 			{
@@ -746,16 +1014,20 @@ public:
 
 		return NULL;
 	}
+
 	/**
-	 * Filters the elements in the array based on a predicate functor. 
-	 * @param  Pred The functor to apply to each element.
-	 * @return      TArray with the same type as this object which contains the subset of elements for which the functor returns true.
+	 * Filters the elements in the array based on a predicate functor.
+	 *
+	 * @param Pred The functor to apply to each element.
+	 *
+	 * @returns TArray with the same type as this object which contains
+	 *          the subset of elements for which the functor returns true.
 	 */
 	template <typename Predicate>
 	TArray<ElementType> FilterByPredicate(Predicate Pred) const
 	{
 		TArray<ElementType> FilterResults;
-		for (const ElementType* RESTRICT Data = GetTypedData(), *RESTRICT DataEnd = Data + ArrayNum; Data != DataEnd; ++Data)
+		for (const ElementType* RESTRICT Data = GetData(), *RESTRICT DataEnd = Data + ArrayNum; Data != DataEnd; ++Data)
 		{
 			if (Pred(*Data))
 			{
@@ -764,12 +1036,15 @@ public:
 		}
 		return FilterResults;
 	}
+
 	/**
-	 * @return		true if found
+	 * Checks if this array contains the element.
+	 *
+	 * @returns	True if found. False otherwise.
 	 */
-	bool Contains( const ElementType& Item ) const
+	bool Contains(const ElementType& Item) const
 	{
-		for (const ElementType* RESTRICT Data = GetTypedData(), * RESTRICT DataEnd = Data + ArrayNum; Data != DataEnd; ++Data)
+		for (const ElementType* RESTRICT Data = GetData(), *RESTRICT DataEnd = Data + ArrayNum; Data != DataEnd; ++Data)
 		{
 			if (*Data == Item)
 			{
@@ -778,46 +1053,76 @@ public:
 		}
 		return false;
 	}
+
 	/**
-	 * @return		true if found
+	 * Checks if this array contains element for which the predicate is true.
+	 *
+	 * @param Predicate to use
+	 *
+	 * @returns	True if found. False otherwise.
 	 */
 	template <typename Predicate>
-	FORCEINLINE bool ContainsByPredicate( Predicate Pred ) const
+	FORCEINLINE bool ContainsByPredicate(Predicate Pred) const
 	{
 		return FindByPredicate(Pred) != NULL;
 	}
+
+	/**
+	 * Equality operator.
+	 *
+	 * @param OtherArray Array to compare.
+	 *
+	 * @returns True if this array is the same as OtherArray. False otherwise.
+	 */
 	bool operator==(const TArray& OtherArray) const
 	{
 		int32 Count = Num();
 
-		return Count == OtherArray.Num() && CompareItems(GetTypedData(), OtherArray.GetTypedData(), Count);
+		return Count == OtherArray.Num() && CompareItems(GetData(), OtherArray.GetData(), Count);
 	}
+
+	/**
+	 * Inequality operator.
+	 *
+	 * @param OtherArray Array to compare.
+	 *
+	 * @returns True if this array is NOT the same as OtherArray. False otherwise.
+	 */
 	bool operator!=(const TArray& OtherArray) const
 	{
 		return !(*this == OtherArray);
 	}
-	friend FArchive& operator<<( FArchive& Ar, TArray& A )
+
+	/**
+	 * Serialization operator.
+	 *
+	 * @param Ar Archive to serialize the array with.
+	 * @param A Array to serialize.
+	 *
+	 * @returns Passing the given archive.
+	 */
+	friend FArchive& operator<<(FArchive& Ar, TArray& A)
 	{
-		A.CountBytes( Ar );
-		if( sizeof(ElementType)==1 )
+		A.CountBytes(Ar);
+		if (sizeof(ElementType) == 1)
 		{
 			// Serialize simple bytes which require no construction or destruction.
 			Ar << A.ArrayNum;
-			check( A.ArrayNum >= 0 );
-			if( Ar.IsLoading() )
+			check(A.ArrayNum >= 0);
+			if (Ar.IsLoading())
 			{
 				A.ArrayMax = A.ArrayNum;
-				A.AllocatorInstance.ResizeAllocation(0,A.ArrayMax,sizeof(ElementType));
+				A.AllocatorInstance.ResizeAllocation(0, A.ArrayMax, sizeof(ElementType));
 			}
-			Ar.Serialize( A.GetData(), A.Num() );
+			Ar.Serialize(A.GetData(), A.Num());
 		}
-		else if( Ar.IsLoading() )
+		else if (Ar.IsLoading())
 		{
 			// Load array.
 			int32 NewNum;
 			Ar << NewNum;
-			A.Empty( NewNum );
-			for( int32 i=0; i<NewNum; i++ )
+			A.Empty(NewNum);
+			for (int32 i = 0; i < NewNum; i++)
 			{
 				Ar << *::new(A)ElementType;
 			}
@@ -826,13 +1131,14 @@ public:
 		{
 			// Save array.
 			Ar << A.ArrayNum;
-			for( int32 i=0; i<A.ArrayNum; i++ )
+			for (int32 i = 0; i < A.ArrayNum; i++)
 			{
-				Ar << A[ i ];
+				Ar << A[i];
 			}
 		}
 		return Ar;
 	}
+
 	/**
 	 * Bulk serialize array as a single memory blob when loading. Uses regular serialization code for saving
 	 * and doesn't serialize at all otherwise (e.g. transient, garbage collection, ...).
@@ -853,7 +1159,7 @@ public:
 	 *   - it is safe to call BulkSerialize on TTransArrays
 	 *
 	 * IMPORTANT:
-	 *   - This is Overridden in XeD3dResourceArray.h  Please make certain changes are propogated accordingly
+	 *   - This is Overridden in XeD3dResourceArray.h Please make certain changes are propogated accordingly
 	 *
 	 * @param Ar	FArchive to bulk serialize this TArray to/from
 	 */
@@ -866,87 +1172,133 @@ public:
 
 		if (bForcePerElementSerialization
 			|| (Ar.IsSaving()			// if we are saving, we always do the ordinary serialize as a way to make sure it matches up with bulk serialization
-				&& !Ar.IsCooking())		// but cooking is performance critical, so we skip that
+			&& !Ar.IsCooking())		// but cooking is performance critical, so we skip that
 			|| Ar.IsByteSwapping()		// if we are byteswapping, we need to do that per-element
 			)
 		{
 			Ar << *this;
 		}
-		else 
+		else
 		{
 			CountBytes(Ar);
 			if (Ar.IsLoading())
 			{
 				// Basic sanity checking to ensure that sizes match.
-				checkf(SerializedElementSize==0 || SerializedElementSize==ElementSize,TEXT("Expected %i, Got: %i"),ElementSize,SerializedElementSize);
+				checkf(SerializedElementSize == 0 || SerializedElementSize == ElementSize, TEXT("Expected %i, Got: %i"), ElementSize, SerializedElementSize);
 				// Serialize the number of elements, block allocate the right amount of memory and deserialize
 				// the data as a giant memory blob in a single call to Serialize. Please see the function header
 				// for detailed documentation on limitations and implications.
 				int32 NewArrayNum;
 				Ar << NewArrayNum;
-				Empty( NewArrayNum );
-				AddUninitialized( NewArrayNum );
-				Ar.Serialize( GetData(), NewArrayNum * SerializedElementSize );
+				Empty(NewArrayNum);
+				AddUninitialized(NewArrayNum);
+				Ar.Serialize(GetData(), NewArrayNum * SerializedElementSize);
 			}
 			else if (Ar.IsSaving())
 			{
 				int32 ArrayCount = Num();
 				Ar << ArrayCount;
-				Ar.Serialize( GetData(), ArrayCount * SerializedElementSize );
+				Ar.Serialize(GetData(), ArrayCount * SerializedElementSize);
 			}
 		}
 	}
-	void CountBytes( FArchive& Ar )
+
+	/**
+	 * Count bytes needed to serialize this array.
+	 *
+	 * @param Ar Archive to count for.
+	 */
+	void CountBytes(FArchive& Ar)
 	{
-		Ar.CountBytes( ArrayNum*sizeof(ElementType), ArrayMax*sizeof(ElementType) );
+		Ar.CountBytes(ArrayNum*sizeof(ElementType), ArrayMax*sizeof(ElementType));
 	}
 
-	// Add, Insert, Remove, Empty interface.
-	/** Caution, Add() will create elements without calling the constructor and this is not appropriate for element types that require a constructor to function properly. */
-	int32 AddUninitialized( int32 Count=1 )
+	/**
+	 * Adds a given number of uninitialized elements into the array.
+	 *
+	 * Caution, AddUninitialized() will create elements without calling
+	 * the constructor and this is not appropriate for element types that
+	 * require a constructor to function properly.
+	 *
+	 * @param Count Number of elements to add.
+	 *
+	 * @returns Number of elements in array before addition.
+	 */
+	int32 AddUninitialized(int32 Count = 1)
 	{
 		CheckInvariants();
-		checkSlow(Count>=0);
+		checkSlow(Count >= 0);
 
 		const int32 OldNum = ArrayNum;
-		if( (ArrayNum+=Count)>ArrayMax )
+		if ((ArrayNum += Count) > ArrayMax)
 		{
-			ArrayMax = AllocatorInstance.CalculateSlack( ArrayNum, ArrayMax, sizeof(ElementType) );
-			AllocatorInstance.ResizeAllocation(OldNum,ArrayMax, sizeof(ElementType));
+			ArrayMax = AllocatorInstance.CalculateSlack(ArrayNum, ArrayMax, sizeof(ElementType));
+			AllocatorInstance.ResizeAllocation(OldNum, ArrayMax, sizeof(ElementType));
 		}
 
 		return OldNum;
 	}
-	/** Caution, Insert() will create elements without calling the constructor and this is not appropriate for element types that require a constructor to function properly. */
-	void InsertUninitialized( int32 Index, int32 Count=1 )
+
+	/**
+	 * Inserts a given number of uninitialized elements into the array at given
+	 * location.
+	 *
+	 * Caution, InsertUninitialized() will create elements without calling the
+	 * constructor and this is not appropriate for element types that require
+	 * a constructor to function properly.
+	 *
+	 * @param Index Tells where to insert the new elements.
+	 * @param Count Number of elements to add.
+	 */
+	void InsertUninitialized(int32 Index, int32 Count = 1)
 	{
 		CheckInvariants();
-		checkSlow((Count>=0) & (Index>=0) & (Index<=ArrayNum));
+		checkSlow((Count >= 0) & (Index >= 0) & (Index <= ArrayNum));
 
 		const int32 OldNum = ArrayNum;
-		if( (ArrayNum+=Count)>ArrayMax )
+		if ((ArrayNum += Count) > ArrayMax)
 		{
-			ArrayMax = AllocatorInstance.CalculateSlack( ArrayNum, ArrayMax, sizeof(ElementType) );
-			AllocatorInstance.ResizeAllocation(OldNum,ArrayMax,sizeof(ElementType));
+			ArrayMax = AllocatorInstance.CalculateSlack(ArrayNum, ArrayMax, sizeof(ElementType));
+			AllocatorInstance.ResizeAllocation(OldNum, ArrayMax, sizeof(ElementType));
 		}
-		ElementType* Data = GetTypedData() + Index;
+		ElementType* Data = GetData() + Index;
 		RelocateItems(Data + Count, Data, OldNum - Index);
 	}
-	/** Caution, InsertZeroed() will create elements without calling the constructor and this is not appropriate for element types that require a constructor to function properly. */
-	void InsertZeroed( int32 Index, int32 Count=1 )
+
+	/**
+	 * Inserts a given number of zeroed elements into the array at given
+	 * location.
+	 *
+	 * Caution, InsertZeroed() will create elements without calling the
+	 * constructor and this is not appropriate for element types that require
+	 * a constructor to function properly.
+	 *
+	 * @param Index Tells where to insert the new elements.
+	 * @param Count Number of elements to add.
+	 */
+	void InsertZeroed(int32 Index, int32 Count = 1)
 	{
-		InsertUninitialized( Index, Count );
-		FMemory::Memzero( (uint8*)AllocatorInstance.GetAllocation()+Index*sizeof(ElementType), Count*sizeof(ElementType) );
+		InsertUninitialized(Index, Count);
+		FMemory::Memzero((uint8*)AllocatorInstance.GetAllocation() + Index*sizeof(ElementType), Count*sizeof(ElementType));
 	}
-	int32 Insert( const TArray<ElementType>& Items, const int32 InIndex )
+
+	/**
+	 * Inserts given elements into the array at given location.
+	 *
+	 * @param Items Array of elements to insert.
+	 * @param InIndex Tells where to insert the new elements.
+	 *
+	 * @returns Location at which the insert was done.
+	 */
+	int32 Insert(const TArray<ElementType>& Items, const int32 InIndex)
 	{
-		check(this!=&Items);
-		InsertUninitialized( InIndex, Items.Num() );
-		int32 Index=InIndex;
-		for( auto It=Items.CreateConstIterator(); It; ++It )
+		check(this != &Items);
+		InsertUninitialized(InIndex, Items.Num());
+		int32 Index = InIndex;
+		for (auto It = Items.CreateConstIterator(); It; ++It)
 		{
 			RangeCheck(Index);
-			new(GetTypedData() + Index++) ElementType(MoveTemp(*It));
+			new(GetData() + Index++) ElementType(MoveTemp(*It));
 		}
 		return InIndex;
 	}
@@ -965,138 +1317,190 @@ public:
 		check(Ptr != NULL);
 
 		InsertUninitialized(Index, Count);
-		CopyConstructItems(GetTypedData() + Index, Ptr, Count);
+		CopyConstructItems(GetData() + Index, Ptr, Count);
 
 		return Index;
 	}
 
 	/**
-	 * Checks that the specified address is not part of an element within the container.  Used for implementations
-	 * to check that reference arguments aren't going to be invalidated by possible reallocation.
+	 * Checks that the specified address is not part of an element within the
+	 * container. Used for implementations to check that reference arguments
+	 * aren't going to be invalidated by possible reallocation.
 	 *
 	 * @param Addr The address to check.
 	 */
 	FORCEINLINE void CheckAddress(const ElementType* Addr) const
 	{
-		checkf(Addr < GetTypedData() || Addr >= (GetTypedData() + ArrayMax), TEXT("Attempting to add a container element (0x%08x) which already comes from the container (0x%08x, ArrayMax: %d)!"), Addr, GetTypedData(), ArrayMax);
+		checkf(Addr < GetData() || Addr >= (GetData() + ArrayMax), TEXT("Attempting to add a container element (0x%08x) which already comes from the container (0x%08x, ArrayMax: %d)!"), Addr, GetData(), ArrayMax);
 	}
 
-	int32 Insert( ElementType&& Item, int32 Index )
+	/**
+	 * Inserts a given element into the array at given location. Move semantics
+	 * version.
+	 *
+	 * @param Item The element to insert.
+	 * @param Index Tells where to insert the new elements.
+	 *
+	 * @returns Location at which the insert was done.
+	 */
+	int32 Insert(ElementType&& Item, int32 Index)
 	{
 		CheckAddress(&Item);
 
 		// construct a copy in place at Index (this new operator will insert at 
 		// Index, then construct that memory with Item)
-		InsertUninitialized(Index,1);
-		new(GetTypedData() + Index) ElementType(MoveTemp(Item));
+		InsertUninitialized(Index, 1);
+		new(GetData() + Index) ElementType(MoveTemp(Item));
 		return Index;
 	}
-	int32 Insert( const ElementType& Item, int32 Index )
+
+	/**
+	 * Inserts a given element into the array at given location.
+	 *
+	 * @param Item The element to insert.
+	 * @param Index Tells where to insert the new elements.
+	 *
+	 * @returns Location at which the insert was done.
+	 */
+	int32 Insert(const ElementType& Item, int32 Index)
 	{
 		CheckAddress(&Item);
 
 		// construct a copy in place at Index (this new operator will insert at 
 		// Index, then construct that memory with Item)
-		InsertUninitialized(Index,1);
-		new(GetTypedData() + Index) ElementType(Item);
+		InsertUninitialized(Index, 1);
+		new(GetData() + Index) ElementType(Item);
 		return Index;
 	}
-	void RemoveAt( int32 Index, int32 Count=1, bool bAllowShrinking=true )
+
+	/**
+	 * Removes an element (or elements) at given location optionally shrinking
+	 * the array.
+	 *
+	 * @param Index Location in array of the element to remove.
+	 * @param Count (Optional) Number of elements to remove. Default is 1.
+	 * @param bAllowShrinking (Optional) Tells if this call can shrink array if
+	 *                        suitable after remove. Default is true.
+	 */
+	void RemoveAt(int32 Index, int32 Count = 1, bool bAllowShrinking = true)
 	{
 		CheckInvariants();
-		checkSlow((Count >= 0) & (Index >= 0) & (Index+Count <= ArrayNum));
+		checkSlow((Count >= 0) & (Index >= 0) & (Index + Count <= ArrayNum));
 
-		DestructItems(GetTypedData() + Index, Count);
+		DestructItems(GetData() + Index, Count);
 
 		// Skip memmove in the common case that there is nothing to move.
 		int32 NumToMove = ArrayNum - Index - Count;
-		if( NumToMove )
+		if (NumToMove)
 		{
 			FMemory::Memmove
-			(
-				(uint8*)AllocatorInstance.GetAllocation() + (Index      ) * sizeof(ElementType),
-				(uint8*)AllocatorInstance.GetAllocation() + (Index+Count) * sizeof(ElementType),
+				(
+				(uint8*)AllocatorInstance.GetAllocation() + (Index)* sizeof(ElementType),
+				(uint8*)AllocatorInstance.GetAllocation() + (Index + Count) * sizeof(ElementType),
 				NumToMove * sizeof(ElementType)
-			);
+				);
 		}
 		ArrayNum -= Count;
-		
-		const int32 NewArrayMax = AllocatorInstance.CalculateSlack(ArrayNum,ArrayMax,sizeof(ElementType));
-		if(NewArrayMax != ArrayMax && bAllowShrinking)
+
+		const int32 NewArrayMax = AllocatorInstance.CalculateSlack(ArrayNum, ArrayMax, sizeof(ElementType));
+		if (NewArrayMax != ArrayMax && bAllowShrinking)
 		{
 			ArrayMax = NewArrayMax;
-			AllocatorInstance.ResizeAllocation(ArrayNum,ArrayMax,sizeof(ElementType));
+			AllocatorInstance.ResizeAllocation(ArrayNum, ArrayMax, sizeof(ElementType));
 		}
 	}
-	// RemoveSwap, this version is much more efficient O(Count) instead of O(ArrayNum), but does not preserve the order
-	void RemoveAtSwap( int32 Index, int32 Count=1, bool bAllowShrinking=true)
+
+	/**
+	 * Removes an element (or elements) at given location optionally shrinking
+	 * the array.
+	 *
+	 * This version is much more efficient than RemoveAt (O(Count) instead of
+	 * O(ArrayNum)), but does not preserve the order.
+	 *
+	 * @param Index Location in array of the element to remove.
+	 * @param Count (Optional) Number of elements to remove. Default is 1.
+	 * @param bAllowShrinking (Optional) Tells if this call can shrink array if
+	 *                        suitable after remove. Default is true.
+	 */
+	void RemoveAtSwap(int32 Index, int32 Count = 1, bool bAllowShrinking = true)
 	{
 		CheckInvariants();
-		checkSlow((Count >= 0) & (Index >= 0) & (Index+Count <= ArrayNum));
+		checkSlow((Count >= 0) & (Index >= 0) & (Index + Count <= ArrayNum));
+
+		DestructItems(GetData() + Index, Count);
 
-		DestructItems(GetTypedData() + Index, Count);
-		
 		// Replace the elements in the hole created by the removal with elements from the end of the array, so the range of indices used by the array is contiguous.
 		const int32 NumElementsInHole = Count;
 		const int32 NumElementsAfterHole = ArrayNum - (Index + Count);
-		const int32 NumElementsToMoveIntoHole = FMath::Min(NumElementsInHole,NumElementsAfterHole);
-		if(NumElementsToMoveIntoHole)
+		const int32 NumElementsToMoveIntoHole = FMath::Min(NumElementsInHole, NumElementsAfterHole);
+		if (NumElementsToMoveIntoHole)
 		{
 			FMemory::Memcpy(
-				(uint8*)AllocatorInstance.GetAllocation() + (Index                             ) * sizeof(ElementType),
-				(uint8*)AllocatorInstance.GetAllocation() + (ArrayNum-NumElementsToMoveIntoHole) * sizeof(ElementType),
+				(uint8*)AllocatorInstance.GetAllocation() + (Index)* sizeof(ElementType),
+				(uint8*)AllocatorInstance.GetAllocation() + (ArrayNum - NumElementsToMoveIntoHole) * sizeof(ElementType),
 				NumElementsToMoveIntoHole * sizeof(ElementType)
 				);
 		}
 		ArrayNum -= Count;
 
-		const int32 NewArrayMax = AllocatorInstance.CalculateSlack(ArrayNum,ArrayMax,sizeof(ElementType));
-		if(NewArrayMax != ArrayMax && bAllowShrinking)
+		const int32 NewArrayMax = AllocatorInstance.CalculateSlack(ArrayNum, ArrayMax, sizeof(ElementType));
+		if (NewArrayMax != ArrayMax && bAllowShrinking)
 		{
 			ArrayMax = NewArrayMax;
-			AllocatorInstance.ResizeAllocation(ArrayNum,ArrayMax,sizeof(ElementType));
+			AllocatorInstance.ResizeAllocation(ArrayNum, ArrayMax, sizeof(ElementType));
 		}
 	}
 
 	/**
-	* Same as empty, but doesn't change memory allocations, unless the new size is larger than
-	* the current array. It calls the destructors on held items if needed and then zeros the ArrayNum.
-	*
-	* @param NewSize the expected usage size
-	*/
-	void Reset( int32 NewSize = 0 )
+	 * Same as empty, but doesn't change memory allocations, unless the new size is larger than
+	 * the current array. It calls the destructors on held items if needed and then zeros the ArrayNum.
+	 *
+	 * @param NewSize The expected usage size after calling this function.
+	 */
+	void Reset(int32 NewSize = 0)
 	{
 		// If we have space to hold the excepted size, then don't reallocate
-		if( NewSize <= ArrayMax )
+		if (NewSize <= ArrayMax)
 		{
-			DestructItems( GetTypedData(), ArrayNum );
+			DestructItems(GetData(), ArrayNum);
 			ArrayNum = 0;
 		}
 		else
 		{
-			Empty( NewSize );
+			Empty(NewSize);
 		}
 	}
 
-	void Empty( int32 Slack=0 )
+	/**
+	 * Empties the array. It calls the destructors on held items if needed.
+	 *
+	 * @param Slack (Optional) The expected usage size after empty operation. Default is 0.
+	 */
+	void Empty(int32 Slack = 0)
 	{
-		DestructItems(GetTypedData(), ArrayNum);
+		DestructItems(GetData(), ArrayNum);
 
-		checkSlow(Slack>=0);
+		checkSlow(Slack >= 0);
 		ArrayNum = 0;
 		// only reallocate if we need to, I don't trust realloc to the same size to work
 		if (ArrayMax != Slack)
 		{
 			ArrayMax = Slack;
-			AllocatorInstance.ResizeAllocation(0,ArrayMax,sizeof(ElementType));
+			AllocatorInstance.ResizeAllocation(0, ArrayMax, sizeof(ElementType));
 		}
 	}
 
-	void SetNum( int32 NewNum, bool bAllowShrinking=true )
+	/**
+	 * Resizes array to given number of elements.
+	 *
+	 * @param NewNum New size of the array.
+	 * @param bAllowShrinking Tell if this function can shrink the memory in-use if suitable.
+	 */
+	void SetNum(int32 NewNum, bool bAllowShrinking = true)
 	{
 		if (NewNum > Num())
 		{
-			const int32 Diff  = NewNum - ArrayNum;
+			const int32 Diff = NewNum - ArrayNum;
 			const int32 Index = AddUninitialized(Diff);
 			DefaultConstructItems<ElementType>((uint8*)AllocatorInstance.GetAllocation() + Index * sizeof(ElementType), Diff);
 		}
@@ -1106,11 +1510,16 @@ public:
 		}
 	}
 
-	void SetNumZeroed( int32 NewNum )
+	/**
+	 * Resizes array to given number of elements. New elements will be zeroed.
+	 *
+	 * @param NewNum New size of the array.
+	 */
+	void SetNumZeroed(int32 NewNum)
 	{
 		if (NewNum > Num())
 		{
-			AddZeroed(NewNum-Num());
+			AddZeroed(NewNum - Num());
 		}
 		else if (NewNum < Num())
 		{
@@ -1118,11 +1527,16 @@ public:
 		}
 	}
 
-	void SetNumUninitialized( int32 NewNum )
+	/**
+	 * Resizes array to given number of elements. New elements will be uninitialized.
+	 *
+	 * @param NewNum New size of the array.
+	 */
+	void SetNumUninitialized(int32 NewNum)
 	{
 		if (NewNum > Num())
 		{
-			AddUninitialized(NewNum-Num());
+			AddUninitialized(NewNum - Num());
 		}
 		else if (NewNum < Num())
 		{
@@ -1133,7 +1547,9 @@ public:
 	/**
 	 * Appends the specified array to this array.
 	 *
-	 * @param Source: The array to append.
+	 * Allocator changing version.
+	 *
+	 * @param Source The array to append.
 	 */
 	template <typename OtherAllocator>
 	FORCEINLINE void Append(const TArray<InElementType, OtherAllocator>& Source)
@@ -1151,11 +1567,16 @@ public:
 		// Allocate memory for the new elements.
 		Reserve(ArrayNum + SourceCount);
 
-		CopyConstructItems(GetTypedData() + ArrayNum, Source.GetTypedData(), SourceCount);
+		CopyConstructItems(GetData() + ArrayNum, Source.GetData(), SourceCount);
 
 		ArrayNum += SourceCount;
 	}
 
+	/**
+	 * Appends the specified array to this array.
+	 *
+	 * @param Source The array to append.
+	 */
 	FORCEINLINE void Append(TArray&& Source)
 	{
 		check((void*)this != (void*)&Source);
@@ -1171,7 +1592,7 @@ public:
 		// Allocate memory for the new elements.
 		Reserve(ArrayNum + SourceCount);
 
-		RelocateItems(GetTypedData() + ArrayNum, Source.GetTypedData(), SourceCount);
+		RelocateItems(GetData() + ArrayNum, Source.GetData(), SourceCount);
 		Source.ArrayNum = 0;
 
 		ArrayNum += SourceCount;
@@ -1188,15 +1609,26 @@ public:
 		check(Ptr != NULL);
 
 		int32 Pos = AddUninitialized(Count);
-		CopyConstructItems(GetTypedData() + Pos, Ptr, Count);
+		CopyConstructItems(GetData() + Pos, Ptr, Count);
 	}
 
 	/**
 	 * Appends the specified array to this array.
 	 * Cannot append to self.
+	 *
+	 * Move semantics version.
+	 *
+	 * @param Other The array to append.
 	 */
-	FORCEINLINE TArray& operator+=(      TArray&& Other) { Append(MoveTemp(Other)); return *this; }
-	FORCEINLINE TArray& operator+=(const TArray&  Other) { Append(         Other ); return *this; }
+	FORCEINLINE TArray& operator+=(TArray&& Other) { Append(MoveTemp(Other)); return *this; }
+
+	/**
+	 * Appends the specified array to this array.
+	 * Cannot append to self.
+	 *
+	 * @param Other The array to append.
+	 */
+	FORCEINLINE TArray& operator+=(const TArray& Other) { Append(Other); return *this; }
 
 	/**
 	 * Adds a new item to the end of the array, possibly reallocating the whole array to fit.
@@ -1210,7 +1642,7 @@ public:
 		int32 Emplace(ArgsType&&... Args)
 		{
 			const int32 Index = AddUninitialized(1);
-			new(GetTypedData() + Index) ElementType(Forward<ArgsType>(Args)...);
+			new(GetData() + Index) ElementType(Forward<ArgsType>(Args)...);
 			return Index;
 		}
 
@@ -1219,7 +1651,7 @@ public:
 		int32 Emplace()
 		{
 			const int32 Index = AddUninitialized(1);
-			new(GetTypedData() + Index) ElementType();
+			new(GetData() + Index) ElementType();
 			return Index;
 		}
 
@@ -1227,7 +1659,7 @@ public:
 		int32 Emplace(Arg0Type&& Arg0)
 		{
 			const int32 Index = AddUninitialized(1);
-			new(GetTypedData() + Index) ElementType(Forward<Arg0Type>(Arg0));
+			new(GetData() + Index) ElementType(Forward<Arg0Type>(Arg0));
 			return Index;
 		}
 
@@ -1235,7 +1667,7 @@ public:
 		int32 Emplace(Arg0Type&& Arg0, Arg1Type&& Arg1)
 		{
 			const int32 Index = AddUninitialized(1);
-			new(GetTypedData() + Index) ElementType(Forward<Arg0Type>(Arg0), Forward<Arg1Type>(Arg1));
+			new(GetData() + Index) ElementType(Forward<Arg0Type>(Arg0), Forward<Arg1Type>(Arg1));
 			return Index;
 		}
 
@@ -1243,7 +1675,7 @@ public:
 		int32 Emplace(Arg0Type&& Arg0, Arg1Type&& Arg1, Arg2Type&& Arg2)
 		{
 			const int32 Index = AddUninitialized(1);
-			new(GetTypedData() + Index) ElementType(Forward<Arg0Type>(Arg0), Forward<Arg1Type>(Arg1), Forward<Arg2Type>(Arg2));
+			new(GetData() + Index) ElementType(Forward<Arg0Type>(Arg0), Forward<Arg1Type>(Arg1), Forward<Arg2Type>(Arg2));
 			return Index;
 		}
 
@@ -1251,81 +1683,143 @@ public:
 		int32 Emplace(Arg0Type&& Arg0, Arg1Type&& Arg1, Arg2Type&& Arg2, Arg3Type&& Arg3)
 		{
 			const int32 Index = AddUninitialized(1);
-			new(GetTypedData() + Index) ElementType(Forward<Arg0Type>(Arg0), Forward<Arg1Type>(Arg1), Forward<Arg2Type>(Arg2), Forward<Arg3Type>(Arg3));
+			new(GetData() + Index) ElementType(Forward<Arg0Type>(Arg0), Forward<Arg1Type>(Arg1), Forward<Arg2Type>(Arg2), Forward<Arg3Type>(Arg3));
 			return Index;
 		}
 
 	#endif
 
+	/**
+	 * Adds a new item to the end of the array, possibly reallocating the whole array to fit.
+	 *
+	 * Move semantics version.
+	 *
+	 * @param Item	The item to add
+	 * @return		Index to the new item
+	 */
+	FORCEINLINE int32 Add(ElementType&& Item) { CheckAddress(&Item); return Emplace(MoveTemp(Item)); }
+
 	/**
 	 * Adds a new item to the end of the array, possibly reallocating the whole array to fit.
 	 *
 	 * @param Item	The item to add
 	 * @return		Index to the new item
 	 */
-	FORCEINLINE int32 Add(       ElementType&& Item ) { CheckAddress(&Item); return Emplace(MoveTemp(Item)); }
-	FORCEINLINE int32 Add( const ElementType&  Item ) { CheckAddress(&Item); return Emplace(         Item ); }
+	FORCEINLINE int32 Add(const ElementType& Item) { CheckAddress(&Item); return Emplace(Item); }
 
-	/** Caution, AddZeroed() will create elements without calling the constructor and this is not appropriate for element types that require a constructor to function properly. */
-	int32 AddZeroed( int32 Count=1 )
+	/**
+	 * Adds a new item to the end of the array, possibly reallocating the whole
+	 * array to fit. The new item will be zeroed.
+	 *
+	 * Caution, AddZeroed() will create elements without calling the
+	 * constructor and this is not appropriate for element types that require
+	 * a constructor to function properly.
+	 *
+	 * @param Item The item to add.
+	 *
+	 * @return Index to the new item.
+	 */
+	int32 AddZeroed(int32 Count = 1)
 	{
-		const int32 Index = AddUninitialized( Count );
-		FMemory::Memzero( (uint8*)AllocatorInstance.GetAllocation()+Index*sizeof(ElementType), Count*sizeof(ElementType) );
+		const int32 Index = AddUninitialized(Count);
+		FMemory::Memzero((uint8*)AllocatorInstance.GetAllocation() + Index*sizeof(ElementType), Count*sizeof(ElementType));
 		return Index;
 	}
 
 private:
+	/**
+	 * Adds unique element to array if it doesn't exist.
+	 *
+	 * @param Args Item to add.
+	 *
+	 * @returns Index of the element in the array.
+	 */
 	template <typename ArgsType>
-	int32 AddUniqueImpl( ArgsType&& Args )
+	int32 AddUniqueImpl(ArgsType&& Args)
 	{
 		int32 Index;
 		if (Find(Args, Index))
+		{
 			return Index;
+		}
 
 		return Add(Forward<ArgsType>(Args));
 	}
 
 public:
-	FORCEINLINE int32 AddUnique(       ElementType&& Item ) { return AddUniqueImpl(MoveTemp(Item)); }
-	FORCEINLINE int32 AddUnique( const ElementType&  Item ) { return AddUniqueImpl(         Item ); }
+	/**
+	 * Adds unique element to array if it doesn't exist.
+	 *
+	 * Move semantics version.
+	 *
+	 * @param Args Item to add.
+	 *
+	 * @returns Index of the element in the array.
+	 */
+	FORCEINLINE int32 AddUnique(ElementType&& Item) { return AddUniqueImpl(MoveTemp(Item)); }
+
+	/**
+	 * Adds unique element to array if it doesn't exist.
+	 *
+	 * @param Args Item to add.
+	 *
+	 * @returns Index of the element in the array.
+	 */
+	FORCEINLINE int32 AddUnique(const ElementType& Item) { return AddUniqueImpl(Item); }
 
 	/**
 	 * Reserves memory such that the array can contain at least Number elements.
+	 *
+	 * @param Number The number of elements that the array should be able to
+	 *               contain after allocation.
 	 */
 	void Reserve(int32 Number)
 	{
 		if (Number > ArrayMax)
 		{
 			ArrayMax = Number;
-			AllocatorInstance.ResizeAllocation(ArrayNum,ArrayMax,sizeof(ElementType));
+			AllocatorInstance.ResizeAllocation(ArrayNum, ArrayMax, sizeof(ElementType));
 		}
 	}
 	
-	/** Sets the size of the array. */
+	/**
+	 * Sets the size of the array. The added elements will be uninitialized.
+	 *
+	 * @param Number The number of elements that the array should be able to
+	 *               contain after allocation.
+	 */
 	void Init(int32 Number)
 	{
 		Empty(Number);
 		AddUninitialized(Number);
 	}
 
-	/** Sets the size of the array, filling it with the given element. */
-	void Init(const ElementType& Element,int32 Number)
+	/**
+	 * Sets the size of the array, filling it with the given element.
+	 *
+	 * @param Element The element to fill array with.
+	 * @param Number The number of elements that the array should be able to
+	 *               contain after allocation.
+	 */
+	void Init(const ElementType& Element, int32 Number)
 	{
 		Empty(Number);
-		for(int32 Index = 0;Index < Number;++Index)
+		for (int32 Index = 0; Index < Number; ++Index)
 		{
 			new(*this) ElementType(Element);
 		}
 	}
 
 	/**
-	 * Removes the first occurrence of the specified item in the array, maintaining order but not indices.
+	 * Removes the first occurrence of the specified item in the array,
+	 * maintaining order but not indices.
 	 *
-	 * @param	Item	The item to remove
+	 * @param Item The item to remove.
 	 *
-	 * @return	The number of items removed.  For RemoveSingleItem, this is always either 0 or 1.
+	 * @returns The number of items removed. For RemoveSingleItem, this is
+	 *          always either 0 or 1.
 	 */
-	int32 RemoveSingle( const ElementType& Item )
+	int32 RemoveSingle(const ElementType& Item)
 	{
 		int32 Index = Find(Item);
 		if (Index == INDEX_NONE)
@@ -1333,7 +1827,7 @@ public:
 			return 0;
 		}
 
-		auto* RemovePtr = GetTypedData() + Index;
+		auto* RemovePtr = GetData() + Index;
 
 		// Destruct items that match the specified Item.
 		DestructItems(RemovePtr, 1);
@@ -1347,7 +1841,14 @@ public:
 		return 1;
 	}
 
-	/** Removes as many instances of Item as there are in the array, maintaining order but not indices. */
+	/**
+	 * Removes as many instances of Item as there are in the array, maintaining
+	 * order but not indices.
+	 *
+	 * @param Item Item to remove from array.
+	 *
+	 * @returns Number of removed elements.
+	 */
 	int32 Remove(const ElementType& Item)
 	{
 		CheckAddress(&Item);
@@ -1362,6 +1863,8 @@ public:
 	 * Optimized to work with runs of matches/non-matches
 	 *
 	 * @param Predicate Predicate class instance
+	 *
+	 * @returns Number of removed elements.
 	 */
 	template <class PREDICATE_CLASS>
 	int32 RemoveAll(const PREDICATE_CLASS& Predicate)
@@ -1374,11 +1877,11 @@ public:
 
 		int32 WriteIndex = 0;
 		int32 ReadIndex = 0;
-		bool NotMatch = !Predicate(GetTypedData()[ReadIndex]); // use a ! to guarantee it can't be anything other than zero or one
+		bool NotMatch = !Predicate(GetData()[ReadIndex]); // use a ! to guarantee it can't be anything other than zero or one
 		do
 		{
 			int32 RunStartIndex = ReadIndex++;
-			while (ReadIndex < OriginalNum && NotMatch == !Predicate(GetTypedData()[ReadIndex]))
+			while (ReadIndex < OriginalNum && NotMatch == !Predicate(GetData()[ReadIndex]))
 			{
 				ReadIndex++;
 			}
@@ -1389,14 +1892,14 @@ public:
 				// this was a non-matching run, we need to move it
 				if (WriteIndex != RunStartIndex)
 				{
-					FMemory::Memmove(&GetTypedData()[WriteIndex], &GetTypedData()[RunStartIndex], sizeof(ElementType) * RunLength);
+					FMemory::Memmove(&GetData()[WriteIndex], &GetData()[RunStartIndex], sizeof(ElementType)* RunLength);
 				}
 				WriteIndex += RunLength;
 			}
 			else
 			{
 				// this was a matching run, delete it
-				DestructItems(GetTypedData() + RunStartIndex, RunLength);
+				DestructItems(GetData() + RunStartIndex, RunLength);
 			}
 			NotMatch = !NotMatch;
 		} while (ReadIndex < OriginalNum);
@@ -1411,11 +1914,11 @@ public:
 	 * @param Predicate Predicate class instance
 	 */
 	template <class PREDICATE_CLASS>
-	void RemoveAllSwap( const PREDICATE_CLASS& Predicate )
+	void RemoveAllSwap(const PREDICATE_CLASS& Predicate)
 	{
-		for ( int32 ItemIndex=0; ItemIndex < Num(); )
+		for (int32 ItemIndex = 0; ItemIndex < Num();)
 		{
-			if ( Predicate( (*this)[ItemIndex] ) )
+			if (Predicate((*this)[ItemIndex]))
 			{
 				RemoveAtSwap(ItemIndex);
 			}
@@ -1427,14 +1930,14 @@ public:
 	}
 
 	/**
-	 * Removes the first occurrence of the specified item in the array.  This version is much more efficient
+	 * Removes the first occurrence of the specified item in the array. This version is much more efficient
 	 * O(Count) instead of O(ArrayNum), but does not preserve the order
 	 *
-	 * @param	Item	The item to remove
+	 * @param Item The item to remove
 	 *
-	 * @return	The number of items removed.  For RemoveSingleItem, this is always either 0 or 1.
+	 * @returns The number of items removed. For RemoveSingleItem, this is always either 0 or 1.
 	 */
-	int32 RemoveSingleSwap( const ElementType& Item )
+	int32 RemoveSingleSwap(const ElementType& Item)
 	{
 		int32 Index = Find(Item);
 		if (Index == INDEX_NONE)
@@ -1448,47 +1951,80 @@ public:
 		return 1;
 	}
 
-	/** RemoveItemSwap, this version is much more efficient O(Count) instead of O(ArrayNum), but does not preserve the order */
-	int32 RemoveSwap( const ElementType& Item )
+	/**
+	 * Removes item from the array.
+	 *
+	 * This version is much more efficient, because it uses RemoveAtSwap
+	 * internally which is O(Count) instead of RemoveAt which is O(ArrayNum),
+	 * but does not preserve the order.
+	 *
+	 * @returns Number of elements removed.
+	 */
+	int32 RemoveSwap(const ElementType& Item)
 	{
 		CheckAddress(&Item);
 
-		const int32 OriginalNum=ArrayNum;
-		for( int32 Index=0; Index<ArrayNum; Index++ )
+		const int32 OriginalNum = ArrayNum;
+		for (int32 Index = 0; Index < ArrayNum; Index++)
 		{
-			if( (*this)[Index]==Item )
+			if ((*this)[Index] == Item)
 			{
-				RemoveAtSwap( Index-- );
+				RemoveAtSwap(Index--);
 			}
 		}
 		return OriginalNum - ArrayNum;
 	}
 
-	void SwapMemory(int32 A, int32 B)
+	/**
+	 * Element-wise array memory swap.
+	 *
+	 * @param FirstIndexToSwap Position of the first element to swap.
+	 * @param SecondIndexToSwap Position of the second element to swap.
+	 */
+	void SwapMemory(int32 FirstIndexToSwap, int32 SecondIndexToSwap)
 	{
 		FMemory::Memswap(
-			(uint8*)AllocatorInstance.GetAllocation()+(sizeof(ElementType)*A),
-			(uint8*)AllocatorInstance.GetAllocation()+(sizeof(ElementType)*B),
+			(uint8*)AllocatorInstance.GetAllocation() + (sizeof(ElementType)*FirstIndexToSwap),
+			(uint8*)AllocatorInstance.GetAllocation() + (sizeof(ElementType)*SecondIndexToSwap),
 			sizeof(ElementType)
 			);
 	}
 
-	void Swap(int32 A, int32 B)
+	/**
+	 * Element-wise array element swap.
+	 *
+	 * This version is doing more sanity checks than SwapMemory.
+	 *
+	 * @param FirstIndexToSwap Position of the first element to swap.
+	 * @param SecondIndexToSwap Position of the second element to swap.
+	 */
+	void Swap(int32 FirstIndexToSwap, int32 SecondIndexToSwap)
 	{
-		check((A >= 0) && (B >= 0));
-		check((ArrayNum > A) && (ArrayNum > B));
-		if (A != B)
+		check((FirstIndexToSwap >= 0) && (SecondIndexToSwap >= 0));
+		check((ArrayNum > FirstIndexToSwap) && (ArrayNum > SecondIndexToSwap));
+		if (FirstIndexToSwap != SecondIndexToSwap)
 		{
-			SwapMemory(A,B);
+			SwapMemory(FirstIndexToSwap, SecondIndexToSwap);
 		}
 	}
 
 	/**
-	 * Searches for the first entry of the specified type, will only work
-	 * with TArray<UObject*>.  Optionally return the item's index, and can
-	 * specify the start index.
+	 * Searches for the first entry of the specified type, will only work with
+	 * TArray<UObject*>. Optionally return the item's index, and can specify
+	 * the start index.
+	 *
+	 * @param Item (Optional output) If it's not null, then it will point to
+	 *             the found element. Untouched if element hasn't been found.
+	 * @param ItemIndex (Optional output) If it's not null, then it will be set
+	 *             to the position of found element in the array. Untouched if
+	 *             element hasn't been found.
+	 * @param StartIndex (Optional) Index in array at which the function should
+	 *             start to look for element.
+	 *
+	 * @returns True if element was found. False otherwise.
 	 */
-	template<typename SearchType> bool FindItemByClass(SearchType **Item = NULL, int32 *ItemIndex = NULL, int32 StartIndex = 0) const
+	template<typename SearchType>
+	bool FindItemByClass(SearchType **Item = NULL, int32 *ItemIndex = NULL, int32 StartIndex = 0) const
 	{
 		UClass* SearchClass = SearchType::StaticClass();
 		for (int32 Idx = StartIndex; Idx < ArrayNum; Idx++)
@@ -1513,13 +2049,21 @@ public:
 	typedef TIndexedContainerIterator<      TArray,       ElementType, int32> TIterator;
 	typedef TIndexedContainerIterator<const TArray, const ElementType, int32> TConstIterator;
 
-	/** Creates an iterator for the contents of this array */
+	/**
+	 * Creates an iterator for the contents of this array
+	 *
+	 * @returns The iterator.
+	 */
 	TIterator CreateIterator()
 	{
 		return TIterator(*this);
 	}
 
-	/** Creates a const iterator for the contents of this array */
+	/**
+	 * Creates a const iterator for the contents of this array
+	 *
+	 * @returns The const iterator.
+	 */
 	TConstIterator CreateConstIterator() const
 	{
 		return TConstIterator(*this);
@@ -1536,55 +2080,59 @@ private:
 	FORCEINLINE friend TConstIterator end  (const TArray& Array) { return TConstIterator(Array, Array.Num()); }
 
 public:
-	/** Sorts the array assuming < operator is defined for the item type */
+	/**
+	 * Sorts the array assuming < operator is defined for the item type.
+	 */
 	void Sort()
 	{
-		::Sort( GetTypedData(), Num() );
+		::Sort(GetData(), Num());
 	}
 
 	/**
 	 * Sorts the array using user define predicate class.
 	 *
-	 * @param Predicate Predicate class instance
+	 * @param Predicate Predicate class instance.
 	 */
 	template <class PREDICATE_CLASS>
-	void Sort( const PREDICATE_CLASS& Predicate )
+	void Sort(const PREDICATE_CLASS& Predicate)
 	{
-		::Sort( GetTypedData(), Num(), Predicate );
+		::Sort(GetData(), Num(), Predicate);
 	}
 
 	/**
 	 * Stable sorts the array assuming < operator is defined for the item type.
+	 *
 	 * Stable sort is slower than non-stable algorithm.
 	 */
 	void StableSort()
 	{
-		::StableSort(GetTypedData(), Num());
+		::StableSort(GetData(), Num());
 	}
 
 	/**
-	 * Stable sorts the array using user defined predicate class. Stable sort
-	 * is slower than non-stable algorithm.
+	 * Stable sorts the array using user defined predicate class.
+	 *
+	 * Stable sort is slower than non-stable algorithm.
 	 *
 	 * @param Predicate Predicate class instance
 	 */
 	template <class PREDICATE_CLASS>
 	void StableSort(const PREDICATE_CLASS& Predicate)
 	{
-		::StableSort(GetTypedData(), Num(), Predicate);
+		::StableSort(GetData(), Num(), Predicate);
 	}
 
 #if defined(_MSC_VER) && !defined(__clang__)	// Relies on MSVC-specific lazy template instantiation to support arrays of incomplete types
 private:
 	/**
-	* Helper function that can be used inside the debuggers watch window to debug TArrays. E.g. "*Class->Defaults.DebugGet(5)". 
-	*
-	* @param	i	Index
-	* @return		pointer to type T at Index i
-	*/
-	FORCENOINLINE const ElementType& DebugGet( int32 i ) const
+	 * Helper function that can be used inside the debuggers watch window to debug TArrays. E.g. "*Class->Defaults.DebugGet(5)". 
+	 *
+	 * @param Index Position to get.
+	 * @returns Reference to the element at given position.
+	 */
+	FORCENOINLINE const ElementType& DebugGet(int32 Index) const
 	{
-		return GetTypedData()[i];
+		return GetData()[Index];
 	}
 #endif
 
@@ -1594,17 +2142,20 @@ private:
 	 * Copies data from one array into this array. Uses the fast path if the
 	 * data in question does not need a constructor.
 	 *
-	 * @param Source the source array to copy
+	 * @param Source The source array to copy
+	 * @param ExtraSlack (Optional) Additional amount of memory to allocate at
+	 *                   the end of the buffer. Counted in elements. Zero by
+	 *                   default.
 	 */
 	template<typename OtherAllocator>
-	void CopyToEmpty(const TArray<ElementType,OtherAllocator>& Source, int32 ExtraSlack = 0)
+	void CopyToEmpty(const TArray<ElementType, OtherAllocator>& Source, int32 ExtraSlack = 0)
 	{
 		check(ExtraSlack >= 0);
 
 		int32 SourceCount = Source.Num();
 		AllocatorInstance.ResizeAllocation(0, SourceCount + ExtraSlack, sizeof(ElementType));
 
-		CopyConstructItems(GetTypedData(), Source.GetTypedData(), SourceCount);
+		CopyConstructItems(GetData(), Source.GetData(), SourceCount);
 
 		ArrayNum = SourceCount;
 		ArrayMax = SourceCount + ExtraSlack;
@@ -1629,15 +2180,15 @@ protected:
 public:
 
 	/** 
-	 * Builds an implicit heap from the array
+	 * Builds an implicit heap from the array.
 	 *
-	 * @pram Predicate Predicate class instance
+	 * @param Predicate Predicate class instance.
 	 */
 	template <class PREDICATE_CLASS>
 	void Heapify(const PREDICATE_CLASS& Predicate)
 	{
-		TDereferenceWrapper< ElementType, PREDICATE_CLASS> PredicateWrapper( Predicate );
-		for(int32 Index=HeapGetParentIndex(Num()-1); Index>=0; Index--)
+		TDereferenceWrapper< ElementType, PREDICATE_CLASS> PredicateWrapper(Predicate);
+		for (int32 Index = HeapGetParentIndex(Num() - 1); Index >= 0; Index--)
 		{
 			SiftDown(Index, Num(), PredicateWrapper);
 		}
@@ -1647,46 +2198,56 @@ public:
 #endif
 	}
 
+	/**
+	 * Builds an implicit heap from the array. Assumes < operator is defined
+	 * for the template type.
+	 */
 	void Heapify()
 	{
-		Heapify( TLess<ElementType>() );
+		Heapify(TLess<ElementType>());
 	}
 
 	/** 
-	 * Adds a new element to the heap
+	 * Adds a new element to the heap.
 	 *
-	 * @pram InIntem Item to be added
-	 * @pram Predicate Predicate class instance
+	 * @param InIntem Item to be added.
+	 * @param Predicate Predicate class instance.
 	 */
 	template <class PREDICATE_CLASS>
-	void HeapPush( const ElementType& InItem, const PREDICATE_CLASS& Predicate )
+	void HeapPush(const ElementType& InItem, const PREDICATE_CLASS& Predicate)
 	{
 		// Add at the end, then sift up
 		Add(InItem);
-		TDereferenceWrapper< ElementType, PREDICATE_CLASS> PredicateWrapper(Predicate);
-		SiftUp(0, Num()-1, PredicateWrapper);
+		TDereferenceWrapper<ElementType, PREDICATE_CLASS> PredicateWrapper(Predicate);
+		SiftUp(0, Num() - 1, PredicateWrapper);
 
 #if DEBUG_HEAP
 		VerifyHeap(PredicateWrapper);
 #endif
 	}
 
-	void HeapPush( const ElementType& InItem )
+	/** 
+	 * Adds a new element to the heap. Assumes < operator is defined for the
+	 * template type.
+	 *
+	 * @param InIntem Item to be added.
+	 */
+	void HeapPush(const ElementType& InItem)
 	{
-		HeapPush( InItem, TLess<ElementType>() );
+		HeapPush(InItem, TLess<ElementType>());
 	}
 
 	/** 
 	 * Removes the top element from the heap.
 	 *
-	 * @pram OutItem The removed item
-	 * @pram Predicate Predicate class instance
+	 * @param OutItem The removed item.
+	 * @param Predicate Predicate class instance.
 	 */
 	template <class PREDICATE_CLASS>
-	void HeapPop( ElementType& OutItem, const PREDICATE_CLASS& Predicate )
+	void HeapPop(ElementType& OutItem, const PREDICATE_CLASS& Predicate)
 	{
-		OutItem = (*this)[ 0 ];
-		RemoveAtSwap( 0 );
+		OutItem = (*this)[0];
+		RemoveAtSwap(0);
 
 		TDereferenceWrapper< ElementType, PREDICATE_CLASS> PredicateWrapper(Predicate);
 		SiftDown(0, Num(), PredicateWrapper);
@@ -1696,22 +2257,33 @@ public:
 #endif
 	}
 
-	void HeapPop( ElementType& OutItem )
+	/** 
+	 * Removes the top element from the heap. Assumes < operator is defined for
+	 * the template type.
+	 *
+	 * @param OutItem The removed item.
+	 */
+	void HeapPop(ElementType& OutItem)
 	{
-		HeapPop( OutItem, TLess<ElementType>() );
+		HeapPop(OutItem, TLess<ElementType>());
 	}
 
+	/**
+	 * Verifies the heap.
+	 *
+	 * @param Predicate Predicate class instance.
+	 */
 	template <class PREDICATE_CLASS>
-	void VerifyHeap(const PREDICATE_CLASS& Predicate )
+	void VerifyHeap(const PREDICATE_CLASS& Predicate)
 	{
 		// Verify Predicate
-		ElementType* Heap = GetTypedData();
-		for( int32 Index=1; Index<Num(); Index++ )
+		ElementType* Heap = GetData();
+		for (int32 Index = 1; Index < Num(); Index++)
 		{
 			int32 ParentIndex = HeapGetParentIndex(Index);
-			if( Predicate(Heap[Index], Heap[ParentIndex]) )
+			if (Predicate(Heap[Index], Heap[ParentIndex]))
 			{
-				check( false );
+				check(false);
 			}
 		}
 	}
@@ -1719,12 +2291,12 @@ public:
 	/** 
 	 * Removes the top element from the heap.
 	 *
-	 * @pram Predicate Predicate class instance
+	 * @param Predicate Predicate class instance.
 	 */
 	template <class PREDICATE_CLASS>
-	void HeapPopDiscard(const PREDICATE_CLASS& Predicate )
+	void HeapPopDiscard(const PREDICATE_CLASS& Predicate)
 	{
-		RemoveAtSwap( 0 );
+		RemoveAtSwap(0);
 		TDereferenceWrapper< ElementType, PREDICATE_CLASS> PredicateWrapper(Predicate);
 		SiftDown(0, Num(), PredicateWrapper);
 
@@ -1733,35 +2305,45 @@ public:
 #endif
 	}
 
+	/** 
+	 * Removes the top element from the heap. Assumes < operator is defined for
+	 * the template type.
+	 */
 	void HeapPopDiscard()
 	{
-		HeapPopDiscard( TLess<ElementType>() );
+		HeapPopDiscard(TLess<ElementType>());
 	}
 
 	/** 
 	 * Returns the top element from the heap (does not remove the element).
+	 *
+	 * Const version.
+	 *
+	 * @returns The reference to the top element from the heap.
 	 */
 	const ElementType& HeapTop() const
 	{
-		return (*this)[ 0 ];
+		return (*this)[0];
 	}
 
 	/** 
 	 * Returns the top element from the heap (does not remove the element).
+	 *
+	 * @returns The reference to the top element from the heap.
 	 */
 	ElementType& HeapTop()
 	{
-		return (*this)[ 0 ];
+		return (*this)[0];
 	}
 
 	/**
 	 * Removes an element from the heap.
 	 *
-	 * @param Index at which to remove item.
-	 * @param Predicate Predicate class.
+	 * @param Index Position at which to remove item.
+	 * @param Predicate Predicate class instance.
 	 */
 	template <class PREDICATE_CLASS>
-	void HeapRemoveAt( int32 Index, const PREDICATE_CLASS& Predicate )
+	void HeapRemoveAt(int32 Index, const PREDICATE_CLASS& Predicate)
 	{
 		RemoveAtSwap(Index);
 
@@ -1774,9 +2356,15 @@ public:
 #endif
 	}
 
-	void HeapRemoveAt( int32 Index )
+	/**
+	 * Removes an element from the heap. Assumes < operator is defined for the
+	 * template type.
+	 *
+	 * @param Index Position at which to remove item.
+	 */
+	void HeapRemoveAt(int32 Index)
 	{
-		HeapRemoveAt( Index, TLess< ElementType >() );
+		HeapRemoveAt(Index, TLess< ElementType >());
 	}
 
 	/**
@@ -1785,12 +2373,12 @@ public:
 	 * @param Predicate Predicate class instance.
 	 */
 	template <class PREDICATE_CLASS>
-	void HeapSort( const PREDICATE_CLASS& Predicate )
+	void HeapSort(const PREDICATE_CLASS& Predicate)
 	{
- 		TReversePredicateWrapper< ElementType, PREDICATE_CLASS > ReversePredicateWrapper(Predicate);
+		TReversePredicateWrapper<ElementType, PREDICATE_CLASS> ReversePredicateWrapper(Predicate);
 		Heapify(ReversePredicateWrapper);
 
-		ElementType* Heap = GetTypedData();
+		ElementType* Heap = GetData();
 		for(int32 Index=Num()-1; Index>0; Index--)
 		{
 			Exchange(Heap[0], Heap[Index]);
@@ -1798,7 +2386,7 @@ public:
 		}
 
 #if DEBUG_HEAP
-		TDereferenceWrapper< ElementType, PREDICATE_CLASS> PredicateWrapper(Predicate);
+		TDereferenceWrapper<ElementType, PREDICATE_CLASS> PredicateWrapper(Predicate);
 
 		// Verify Heap Property
 		VerifyHeap(PredicateWrapper);
@@ -1806,17 +2394,21 @@ public:
 		// Also verify Array is properly sorted
 		for(int32 Index=1; Index<Num(); Index++)
 		{
-			if( PredicateWrapper(Heap[Index], Heap[Index-1]) )
+			if (PredicateWrapper(Heap[Index], Heap[Index - 1]))
 			{
-				check( false );
+				check(false);
 			}
 		}
 #endif
 	}
 
+	/**
+	 * Performs heap sort on the array. Assumes < operator is defined for the
+	 * template type.
+	 */
 	void HeapSort()
 	{
-		HeapSort( TLess<ElementType>() );
+		HeapSort(TLess<ElementType>());
 	}
 
 private:
@@ -1825,58 +2417,61 @@ private:
 	 * Gets the index of the left child of node at Index.
 	 *
 	 * @param Index Node for which the left child index is to be returned.
-	 * @return Index of the left child.
+	 *
+	 * @returns Index of the left child.
 	 */
-	FORCEINLINE int32 HeapGetLeftChildIndex( int32 Index ) const
+	FORCEINLINE int32 HeapGetLeftChildIndex(int32 Index) const
 	{
 		return Index * 2 + 1;
 	}
 
 	/** 
-	 * Checks if node located at Index is a leaf or not .
+	 * Checks if node located at Index is a leaf or not.
+	 *
+	 * @param Index Node index.
 	 *
-	 * @param Index - Node index.
-	 * @return true if node is a leaf.
+	 * @returns True if node is a leaf. False otherwise.
 	 */
-	FORCEINLINE bool HeapIsLeaf( int32 Index, int32 Count ) const
+	FORCEINLINE bool HeapIsLeaf(int32 Index, int32 Count) const
 	{
-		return HeapGetLeftChildIndex( Index ) >= Count;
+		return HeapGetLeftChildIndex(Index) >= Count;
 	}
 
 	/**
 	 * Gets the parent index for node at Index.
 	 *
 	 * @param Index node index.
-	 * @return Parent index.
+	 *
+	 * @returns Parent index.
 	 */
-	FORCEINLINE int32 HeapGetParentIndex( int32 Index ) const
+	FORCEINLINE int32 HeapGetParentIndex(int32 Index) const
 	{
-		return ( Index - 1 ) / 2;
+		return (Index - 1) / 2;
 	}
 
 	/**
 	 * Fixes a possible violation of order property between node at Index and a child.
 	 *
-	 * @param Index node index.
+	 * @param Index Node index.
 	 * @param Count Size of the heap (to avoid using Num()).
 	 * @param Predicate Predicate class instance.
 	 */
 	template <class PREDICATE_CLASS>
 	FORCEINLINE void SiftDown(int32 Index, const int32 Count, const PREDICATE_CLASS& Predicate)
 	{
-		ElementType* Heap = GetTypedData();
-		while( !HeapIsLeaf(Index, Count) )
+		ElementType* Heap = GetData();
+		while (!HeapIsLeaf(Index, Count))
 		{
 			const int32 LeftChildIndex = HeapGetLeftChildIndex(Index);
 			const int32 RightChildIndex = LeftChildIndex + 1;
 
 			int32 MinChildIndex = LeftChildIndex;
-			if( RightChildIndex < Count )
+			if (RightChildIndex < Count)
 			{
 				MinChildIndex = Predicate(Heap[LeftChildIndex], Heap[RightChildIndex]) ? LeftChildIndex : RightChildIndex;
 			}
 
-			if( Predicate(Heap[MinChildIndex], Heap[Index]) )
+			if (Predicate(Heap[MinChildIndex], Heap[Index]))
 			{
 				Exchange(Heap[Index], Heap[MinChildIndex]);
 				Index = MinChildIndex;
@@ -1891,22 +2486,22 @@ private:
 	/**
 	 * Fixes a possible violation of order property between node at NodeIndex and a parent.
 	 *
-	 * @param RootIndex, how far to go up
-	 * @param NodeIndex node index.
+	 * @param RootIndex How far to go up?
+	 * @param NodeIndex Node index.
 	 * @param Predicate Predicate class instance.
 	 */
 	template <class PREDICATE_CLASS>
 	FORCEINLINE void SiftUp(int32 RootIndex, const int32 NodeIndex, const PREDICATE_CLASS& Predicate)
 	{
-		ElementType* Heap = GetTypedData();
+		ElementType* Heap = GetData();
 		int32 Index = FMath::Min(NodeIndex, Num() - 1);
-		while(Index > RootIndex)
+		while (Index > RootIndex)
 		{
 			int32 ParentIndex = HeapGetParentIndex(Index);
-			if( Predicate(Heap[Index], Heap[ParentIndex]) )
+			if (Predicate(Heap[Index], Heap[ParentIndex]))
 			{
 				Exchange(Heap[Index], Heap[ParentIndex]);
-				Index = ParentIndex; 
+				Index = ParentIndex;
 			}
 			else
 			{
@@ -1930,7 +2525,6 @@ struct TContainerTraits<TArray<InElementType, Allocator> > : public TContainerTr
 		TAllocatorTraits<Allocator>::SupportsMove };
 };
 
-
 /**
  * Traits class which determines whether or not a type is a TArray.
  */
@@ -1941,7 +2535,6 @@ template <typename InElementType, typename Allocator> struct TIsTArray<const
 template <typename InElementType, typename Allocator> struct TIsTArray<      volatile TArray<InElementType, Allocator>> { enum { Value = true }; };
 template <typename InElementType, typename Allocator> struct TIsTArray<const volatile TArray<InElementType, Allocator>> { enum { Value = true }; };
 
-
 //
 // Array operator news.
 //
@@ -1958,17 +2551,32 @@ template <typename T,typename Allocator> void* operator new( size_t Size, TArray
 	return &Array[Index];
 }
 
-/** A specialization of the exchange macro that avoids reallocating when exchanging two arrays. */
-template <typename T> inline void Exchange( TArray<T>& A, TArray<T>& B )
+/**
+ * A specialization of the exchange macro that avoids reallocating when
+ * exchanging two arrays.
+ *
+ * @param FirstArrayToExchange First array to exchange.
+ * @param SecondArrayToExchange Second array to exchange.
+ */
+template <typename T>
+inline void Exchange(TArray<T>& FirstArrayToExchange, TArray<T>& SecondArrayToExchange)
 {
-	FMemory::Memswap( &A, &B, sizeof(TArray<T>) );
+	FMemory::Memswap(&FirstArrayToExchange, &SecondArrayToExchange, sizeof(TArray<T>));
 }
 
-/** A specialization of the exchange macro that avoids reallocating when exchanging two arrays. */
-template<typename ElementType,typename Allocator>
-inline void Exchange( TArray<ElementType,Allocator>& A, TArray<ElementType,Allocator>& B )
+/**
+ * A specialization of the exchange macro that avoids reallocating when
+ * exchanging two arrays.
+ *
+ * Non-default allocator version.
+ *
+ * @param FirstArrayToExchange First array to exchange.
+ * @param SecondArrayToExchange Second array to exchange.
+ */
+template<typename ElementType, typename Allocator>
+inline void Exchange(TArray<ElementType, Allocator>& FirstArrayToExchange, TArray<ElementType, Allocator>& SecondArrayToExchange)
 {
-	FMemory::Memswap( &A, &B, sizeof(TArray<ElementType,Allocator>) );
+	FMemory::Memswap(&FirstArrayToExchange, &SecondArrayToExchange, sizeof(TArray<ElementType, Allocator>));
 }
 
 /*-----------------------------------------------------------------------------
@@ -1980,88 +2588,134 @@ inline void Exchange( TArray<ElementType,Allocator>& A, TArray<ElementType,Alloc
  * - Has an upper limit of the number of items it will store.
  * - Any item that is added to the array is moved to the top.
  */
-template<typename T,typename Allocator = FDefaultAllocator>
-class TMRUArray : public TArray<T,Allocator>
+template<typename T, typename Allocator = FDefaultAllocator>
+class TMRUArray : public TArray<T, Allocator>
 {
 public:
-	typedef TArray<T,Allocator> Super;
+	typedef TArray<T, Allocator> Super;
 
 	/** The maximum number of items we can store in this array. */
 	int32 MaxItems;
 
+	/**
+	 * Constructor.
+	 */
 	TMRUArray()
-	:	Super()
-	{
-		MaxItems = 0;
-	}
-	TMRUArray( int32 InNum )
-	:	Super( InNum )
+		: Super()
 	{
 		MaxItems = 0;
 	}
 
 #if PLATFORM_COMPILER_HAS_DEFAULTED_FUNCTIONS
 
-	TMRUArray( const TMRUArray& ) = default;
-	TMRUArray& operator=( const TMRUArray& ) = default;
+	TMRUArray(const TMRUArray&) = default;
+	TMRUArray& operator=(const TMRUArray&) = default;
 
-	#if PLATFORM_COMPILER_HAS_RVALUE_REFERENCES
+#if PLATFORM_COMPILER_HAS_RVALUE_REFERENCES
 
-		TMRUArray( TMRUArray&& ) = default;
-		TMRUArray& operator=( TMRUArray&& ) = default;
+	TMRUArray(TMRUArray&&) = default;
+	TMRUArray& operator=(TMRUArray&&) = default;
 
-	#endif
+#endif
 
 #else
 
-	FORCEINLINE TMRUArray( const TMRUArray& Other )
+	/**
+	 * Copy constructor.
+	 *
+	 * @param Other Other array to copy.
+	 */
+	FORCEINLINE TMRUArray(const TMRUArray& Other)
 		: Super((const Super&)Other)
 	{
 	}
 
-	FORCEINLINE TMRUArray& operator=( const TMRUArray& Other )
+	/**
+	 * Assign operator.
+	 *
+	 * @param Other Other array to assign with.
+	 *
+	 * @returns Reference to this object.
+	 */
+	FORCEINLINE TMRUArray& operator=(const TMRUArray& Other)
 	{
 		(Super&)*this = (const Super&)Other;
 		return *this;
 	}
 
-	#if PLATFORM_COMPILER_HAS_RVALUE_REFERENCES
-
-		FORCEINLINE TMRUArray( TMRUArray&& Other )
-			: Super((TMRUArray&&)Other)
-		{
-		}
-
-		FORCEINLINE TMRUArray& operator=( TMRUArray&& Other )
-		{
-			(Super&)*this = (Super&&)Other;
-			return *this;
-		}
+#if PLATFORM_COMPILER_HAS_RVALUE_REFERENCES
+	/**
+	 * Move constructor.
+	 *
+	 * @param Other Other array to move.
+	 */
+	FORCEINLINE TMRUArray(TMRUArray&& Other)
+		: Super((TMRUArray&&)Other)
+	{
+	}
 
-	#endif
+	/**
+	 * Move assign operator.
+	 *
+	 * @param Other Other array to assign with.
+	 *
+	 * @returns Reference to this object.
+	 */
+	FORCEINLINE TMRUArray& operator=(TMRUArray&& Other)
+	{
+		(Super&)*this = (Super&&)Other;
+		return *this;
+	}
 
 #endif
 
-	int32 Add( const T& Item )
+#endif
+	/**
+	 * Adds item to the array. Makes sure that we don't add more than the
+	 * limit.
+	 *
+	 * @param Item Item to add.
+	 *
+	 * @returns Always 0.
+	 */
+	int32 Add(const T& Item)
 	{
-		const int32 idx = Super::Add( Item );
-		this->Swap( idx, 0 );
+		const int32 idx = Super::Add(Item);
+		this->Swap(idx, 0);
 		CullArray();
 		return 0;
 	}
-	int32 AddZeroed( int32 Count=1 )
+
+	/**
+	 * Adds a number of zeroed elements to the array. Makes sure that we don't
+	 * add more than the limit.
+	 *
+	 * @param Count (Optional) A number of elements to add. Default is 0.
+	 *
+	 * @returns Always 0.
+	 */
+	int32 AddZeroed(int32 Count = 1)
 	{
-		const int32 idx = Super::AddZeroed( Count );
-		this->Swap( idx, 0 );
+		const int32 idx = Super::AddZeroed(Count);
+		this->Swap(idx, 0);
 		CullArray();
 		return 0;
 	}
-	int32 AddUnique( const T& Item )
+
+	/**
+	 * Adds unique item to the array. Makes sure that we don't add more than
+	 * the limit. If the item existed it will be removed before addition.
+	 *
+	 * @param Item Element to add.
+	 *
+	 * @returns Always 0.
+	 */
+	int32 AddUnique(const T& Item)
 	{
 		// Remove any existing copies of the item.
 		this->Remove(Item);
 
-		this->InsertUninitialized( 0 );
+		this->InsertUninitialized(0);
 		(*this)[0] = Item;
 
 		CullArray();
@@ -2072,18 +2726,17 @@ public:
 	/**
 	 * Makes sure that the array never gets beyond MaxItems in size.
 	 */
-
 	void CullArray()
 	{
 		// 0 = no limit
-		if( !MaxItems )
+		if (!MaxItems)
 		{
 			return;
 		}
 
-		while( this->Num() > MaxItems )
+		while (this->Num() > MaxItems)
 		{
-			this->RemoveAt( this->Num()-1, 1 );
+			this->RemoveAt(this->Num() - 1, 1);
 		}
 	}
 };
@@ -2101,7 +2754,6 @@ struct TContainerTraits<TMRUArray<T, Allocator> > : public TContainerTraitsBase<
 	Same as a TArray above, but stores pointers to the elements, to allow
 	resizing the array index without relocating the actual elements.
 -----------------------------------------------------------------------------*/
-
 template<typename T,typename Allocator = FDefaultAllocator>
 class TIndirectArray
 {
@@ -2109,15 +2761,14 @@ public:
 	typedef T                        ElementType;
 	typedef TArray<void*, Allocator> InternalArrayType;
 
-	TIndirectArray( int32 InNum )
-	:	Array( InNum )
-	{}
-
+	/**
+	 * Default constructors.
+	 */
 #if PLATFORM_COMPILER_HAS_DEFAULTED_FUNCTIONS
 
 	TIndirectArray() = default;
-	TIndirectArray( TIndirectArray&& ) = default;
-	TIndirectArray& operator=( TIndirectArray&& ) = default;
+	TIndirectArray(TIndirectArray&&) = default;
+	TIndirectArray& operator=(TIndirectArray&&) = default;
 
 #else
 
@@ -2127,33 +2778,51 @@ public:
 
 #endif
 
-	TIndirectArray( const TIndirectArray& Other )
+	/**
+	 * Copy constructor.
+	 *
+	 * @param Other Other array to copy from.
+	 */
+	TIndirectArray(const TIndirectArray& Other)
 	{
 		for (auto& Item : Other)
 		{
-			AddRawItem(new T(Item));
+			Add(new T(Item));
 		}
 	}
 
-	TIndirectArray& operator=( const TIndirectArray& Other )
+	/**
+	 * Assignment operator.
+	 *
+	 * @param Other Other array to assign with.
+	 */
+	TIndirectArray& operator=(const TIndirectArray& Other)
 	{
 		if (&Other != this)
 		{
 			Empty(Other.Num());
 			for (auto& Item : Other)
 			{
-				AddRawItem(new T(Item));
+				Add(new T(Item));
 			}
 		}
 
 		return *this;
 	}
 
+	/**
+	 * Destructor.
+	 */
 	~TIndirectArray()
 	{
 		Empty();
 	}
 
+	/**
+	 * Gets number of elements in array.
+	 *
+	 * @returns Number of elements in array.
+	 */
 	FORCEINLINE int32 Num() const
 	{
 		return Array.Num();
@@ -2162,50 +2831,125 @@ public:
 	/**
 	 * Helper function for returning a typed pointer to the first array entry.
 	 *
-	 * @return pointer to first array entry or NULL if this->ArrayMax==0
+	 * @returns Pointer to first array entry or NULL if this->ArrayMax == 0.
 	 */
+	DEPRECATED(4.6, "This function is deprecated as it does the same as GetData(). Please use GetData() instead.")
 	FORCEINLINE T** GetTypedData()
 	{
-		return (T**)Array.GetTypedData();
+		return (T**)Array.GetData();
+	}
+
+	/**
+	 * Helper function for returning a typed pointer to the first array entry.
+	 *
+	 * @returns Pointer to first array entry or NULL if this->ArrayMax == 0.
+	 */
+	FORCEINLINE T** GetData()
+	{
+		return (T**)Array.GetData();
 	}
+
 	/**
 	 * Helper function for returning a typed pointer to the first array entry.
 	 *
-	 * @return pointer to first array entry or NULL if this->ArrayMax==0
+	 * @returns Pointer to first array entry or NULL if this->ArrayMax == 0.
 	 */
+	DEPRECATED(4.6, "This function is deprecated as it does the same as GetData(). Please use GetData() instead.")
 	FORCEINLINE const T** GetTypedData() const
 	{
-		return (const T**)Array.GetTypedData();
+		return (const T**)Array.GetData();
+	}
+
+	/**
+	 * Helper function for returning a typed pointer to the first array entry.
+	 *
+	 * @returns Pointer to first array entry or NULL if this->ArrayMax == 0.
+	 */
+	FORCEINLINE const T** GetData() const
+	{
+		return (const T**)Array.GetData();
 	}
+
 	/** 
-	 * Helper function returning the size of the inner type
+	 * Helper function returning the size of the inner type.
 	 *
-	 * @return size in bytes of array type
+	 * @returns Size in bytes of array type.
 	 */
 	uint32 GetTypeSize() const
 	{
 		return sizeof(T*);
 	}
-	T& operator[]( int32 i )
+
+	/**
+	 * Bracket array access operator.
+	 *
+	 * @param Index Position of element to return.
+	 *
+	 * @returns Reference to element in array at given position.
+	 */
+	FORCEINLINE T& operator[](int32 Index)
 	{
-		return *(T*)Array[i];
+		return *(T*)Array[Index];
 	}
-	FORCEINLINE const T& operator[]( int32 i ) const
+
+	/**
+	 * Bracket array access operator.
+	 *
+	 * Const version.
+	 *
+	 * @param Index Position of element to return.
+	 *
+	 * @returns Reference to element in array at given position.
+	 */
+	FORCEINLINE const T& operator[](int32 Index) const
 	{
-		return *(T*)Array[i];
+		return *(T*)Array[Index];
 	}
-	FORCEINLINE ElementType& Last( int32 c=0 )
+
+	/**
+	 * Returns n-th last element from the array.
+	 *
+	 * @param IndexFromTheEnd (Optional) Index from the end of array.
+	 *                        Default is 0.
+	 *
+	 * @returns Reference to n-th last element from the array.
+	 */
+	FORCEINLINE ElementType& Last(int32 IndexFromTheEnd = 0)
 	{
-		return *(T*)Array.Last(c);
+		return *(T*)Array.Last(IndexFromTheEnd);
 	}
-	FORCEINLINE const ElementType& Last( int32 c=0 ) const
+
+	/**
+	 * Returns n-th last element from the array.
+	 *
+	 * Const version.
+	 *
+	 * @param IndexFromTheEnd (Optional) Index from the end of array.
+	 *                        Default is 0.
+	 *
+	 * @returns Reference to n-th last element from the array.
+	 */
+	FORCEINLINE const ElementType& Last(int32 IndexFromTheEnd = 0) const
 	{
-		return *(T*)Array.Last(c);
+		return *(T*)Array.Last(IndexFromTheEnd);
 	}
+
+	/**
+	 * Shrinks the array's used memory to smallest possible to store elements
+	 * currently in it.
+	 */
 	void Shrink()
 	{
-		Array.Shrink( );
+		Array.Shrink();
 	}
+
+	/**
+	 * Resets the array to the new given size. It calls the destructors on held
+	 * items.
+	 *
+	 * @param NewSize (Optional) The expected usage size after calling this
+	 *                function. Default is 0.
+	 */
 	void Reset(int32 NewSize = 0)
 	{
 		DestructAndFreeItems();
@@ -2213,28 +2957,28 @@ public:
 	}
 
 	/**
-	 * Special serialize function passing the owning UObject along as required by FUnytpedBulkData
-	 * serialization.
+	 * Special serialize function passing the owning UObject along as required
+	 * by FUnytpedBulkData serialization.
 	 *
-	 * @param	Ar		Archive to serialize with
-	 * @param	Owner	UObject this structure is serialized within
+	 * @param	Ar		Archive to serialize with.
+	 * @param	Owner	UObject this structure is serialized within.
 	 */
-	void Serialize( FArchive& Ar, UObject* Owner )
+	void Serialize(FArchive& Ar, UObject* Owner)
 	{
-		CountBytes( Ar );
-		if( Ar.IsLoading() )
+		CountBytes(Ar);
+		if (Ar.IsLoading())
 		{
 			// Load array.
 			int32 NewNum;
 			Ar << NewNum;
-			Empty( NewNum );
-			for( int32 i=0; i<NewNum; i++ )
+			Empty(NewNum);
+			for (int32 Index = 0; Index < NewNum; Index++)
 			{
 				new(*this) T;
 			}
-			for( int32 i=0; i<NewNum; i++ )
+			for (int32 Index = 0; Index < NewNum; Index++)
 			{
-				(*this)[i].Serialize( Ar, Owner, i );
+				(*this)[Index].Serialize(Ar, Owner, Index);
 			}
 		}
 		else
@@ -2242,49 +2986,74 @@ public:
 			// Save array.
 			int32 Num = Array.Num();
 			Ar << Num;
-			for( int32 i=0; i<Num; i++ )
+			for (int32 Index = 0; Index < Num; Index++)
 			{
-				(*this)[i].Serialize( Ar, Owner, i );
+				(*this)[Index].Serialize(Ar, Owner, Index);
 			}
 		}
 	}
-	friend FArchive& operator<<( FArchive& Ar, TIndirectArray& A )
+
+	/**
+	 * Serialization operator for TIndirectArray.
+	 *
+	 * @param	Ar		Archive to serialize with.
+	 * @param	Array	Array to serialize.
+	 *
+	 * @returns Passing down serializing archive.
+	 */
+	friend FArchive& operator<<(FArchive& Ar, TIndirectArray& Array)
 	{
-		A.CountBytes( Ar );
-		if( Ar.IsLoading() )
+		Array.CountBytes(Ar);
+		if (Ar.IsLoading())
 		{
 			// Load array.
 			int32 NewNum;
 			Ar << NewNum;
-			A.Empty( NewNum );
-			for( int32 i=0; i<NewNum; i++ )
+			Array.Empty(NewNum);
+			for (int32 Index = 0; Index < NewNum; Index++)
 			{
-				Ar << *new(A)T;
+				Ar << *new(Array)T;
 			}
 		}
 		else
 		{
 			// Save array.
-			int32 Num = A.Num();
+			int32 Num = Array.Num();
 			Ar << Num;
-			for( int32 i=0; i<Num; i++ )
+			for (int32 Index = 0; Index < Num; Index++)
 			{
-				Ar << A[ i ];
+				Ar << Array[Index];
 			}
 		}
 		return Ar;
 	}
-	void CountBytes( FArchive& Ar )
+
+	/**
+	 * Count bytes needed to serialize this array.
+	 *
+	 * @param Ar Archive to count for.
+	 */
+	void CountBytes(FArchive& Ar)
 	{
-		Array.CountBytes( Ar );
+		Array.CountBytes(Ar);
 	}
-	void RemoveAt( int32 Index, int32 Count=1, bool AllowShrinking=true )
+
+	/**
+	 * Removes an element (or elements) at given location optionally shrinking
+	 * the array.
+	 *
+	 * @param Index Location in array of the element to remove.
+	 * @param Count (Optional) Number of elements to remove. Default is 1.
+	 * @param bAllowShrinking (Optional) Tells if this call can shrink array if
+	 *                        suitable after remove. Default is true.
+	 */
+	void RemoveAt(int32 Index, int32 Count = 1, bool bAllowShrinking = true)
 	{
 		check(Index >= 0);
 		check(Index <= Array.Num());
 		check(Index + Count <= Array.Num());
-		T** Element = GetTypedData() + Index;
-		for (int32 i = Count; i; --i)
+		T** Element = GetData() + Index;
+		for (int32 ElementId = Count; ElementId; --ElementId)
 		{
 			// We need a typedef here because VC won't compile the destructor call below if T itself has a member called T
 			typedef T IndirectArrayDestructElementType;
@@ -2293,15 +3062,28 @@ public:
 			FMemory::Free(*Element);
 			++Element;
 		}
-		Array.RemoveAt( Index, Count, AllowShrinking );
+		Array.RemoveAt(Index, Count, bAllowShrinking);
 	}
-	void RemoveAtSwap( int32 Index, int32 Count=1, bool AllowShrinking=true )
+
+	/**
+	 * Removes an element (or elements) at given location optionally shrinking
+	 * the array.
+	 *
+	 * This version is much more efficient than RemoveAt (O(Count) instead of
+	 * O(ArrayNum)), but does not preserve the order.
+	 *
+	 * @param Index Location in array of the element to remove.
+	 * @param Count (Optional) Number of elements to remove. Default is 1.
+	 * @param bAllowShrinking (Optional) Tells if this call can shrink array if
+	 *                        suitable after remove. Default is true.
+	 */
+	void RemoveAtSwap(int32 Index, int32 Count = 1, bool bAllowShrinking = true)
 	{
 		check(Index >= 0);
 		check(Index <= Array.Num());
 		check(Index + Count <= Array.Num());
-		T** Element = GetTypedData() + Index;
-		for (int32 i = Count; i; --i)
+		T** Element = GetData() + Index;
+		for (int32 ElementId = Count; ElementId; --ElementId)
 		{
 			// We need a typedef here because VC won't compile the destructor call below if T itself has a member called T
 			typedef T IndirectArrayDestructElementType;
@@ -2310,78 +3092,156 @@ public:
 			FMemory::Free(*Element);
 			++Element;
 		}
-		Array.RemoveAtSwap( Index, Count, AllowShrinking );
+		Array.RemoveAtSwap(Index, Count, bAllowShrinking);
 	}
-	void Swap( int32 A, int32 B )
+
+	/**
+	 * Element-wise array element swap.
+	 *
+	 * This version is doing more sanity checks than SwapMemory.
+	 *
+	 * @param FirstIndexToSwap Position of the first element to swap.
+	 * @param SecondIndexToSwap Position of the second element to swap.
+	 */
+	void Swap(int32 FirstIndexToSwap, int32 SecondIndexToSwap)
 	{
-		Array.Swap(A,B);
+		Array.Swap(FirstIndexToSwap, SecondIndexToSwap);
 	}
-	void Empty( int32 Slack=0 )
+
+	/**
+	 * Empties the array. It calls the destructors on held items.
+	 *
+	 * @param Slack (Optional) The expected usage size after empty operation.
+	 *              Default is 0.
+	 */
+	void Empty(int32 Slack = 0)
 	{
 		DestructAndFreeItems();
-		Array.Empty( Slack );
+		Array.Empty(Slack);
 	}
+
+	/**
+	 * Adds a new item to the end of the array, possibly reallocating the
+	 * whole array to fit.
+	 *
+	 * @param Item The item to add.
+	 *
+	 * @returns Index to the new item.
+	 */
 	FORCEINLINE int32 Add(T* Item)
 	{
 		return Array.Add(Item);
 	}
+
+	/**
+	 * Adds a new item to the end of the array, possibly reallocating the
+	 * whole array to fit.
+	 *
+	 * @param Item The item to add.
+	 *
+	 * @returns Index to the new item.
+	 */
+	DEPRECATED(4.6, "This function is deprecated as it does the same as Add(). Please use Add() instead.")
 	FORCEINLINE int32 AddRawItem(T* Item)
 	{
 		return Array.Add(Item);
 	}
+
+	/**
+	 * Inserts a given element into the array at given location.
+	 *
+	 * @param Item The element to insert.
+	 * @param Index Tells where to insert the new elements.
+	 */
 	FORCEINLINE void Insert(T* Item, int32 Index)
 	{
 		Array.Insert(Item, Index);
 	}
+
+	/**
+	 * Inserts a given element into the array at given location.
+	 *
+	 * @param Item The element to insert.
+	 * @param Index Tells where to insert the new elements.
+	 */
+	DEPRECATED(4.6, "This function is deprecated as it does the same as Insert(). Please use Insert() instead.")
 	FORCEINLINE void InsertRawItem(T* Item, int32 Index)
 	{
 		Array.Insert(Item, Index);
 	}
+
+	/**
+	 * Reserves memory such that the array can contain at least Number elements.
+	 *
+	 * @param Number The number of elements that the array should be able to
+	 *               contain after allocation.
+	 */
 	FORCEINLINE void Reserve(int32 Number)
 	{
 		Array.Reserve(Number);
 	}
-	FORCEINLINE bool IsValidIndex(int32 i) const
+
+	/**
+	 * Tests if index is valid, i.e. greater than zero and less than number of
+	 * elements in array.
+	 *
+	 * @param Index Index to test.
+	 *
+	 * @returns True if index is valid. False otherwise.
+	 */
+	FORCEINLINE bool IsValidIndex(int32 Index) const
 	{
-		return Array.IsValidIndex(i);
+		return Array.IsValidIndex(Index);
 	}
+
 	/** 
-	* Helper function to return the amount of memory allocated by this container 
-	*
-	* @return number of bytes allocated by this container
-	*/
-	uint32 GetAllocatedSize( void ) const
+	 * Helper function to return the amount of memory allocated by this container 
+	 *
+	 * @returns Number of bytes allocated by this container.
+	 */
+	uint32 GetAllocatedSize() const
 	{
-		return Array.Max() * (sizeof(T) + sizeof(T*));
+		return Array.Max() * (sizeof(T)+sizeof(T*));
 	}
 
 	// Iterators
 	typedef TIndexedContainerIterator<      TIndirectArray,       ElementType, int32> TIterator;
 	typedef TIndexedContainerIterator<const TIndirectArray, const ElementType, int32> TConstIterator;
 
-	/** Creates an iterator for the contents of this array */
+	/**
+	 * Creates an iterator for the contents of this array.
+	 *
+	 * @returns The iterator.
+	 */
 	TIterator CreateIterator()
 	{
 		return TIterator(*this);
 	}
 
-	/** Creates a const iterator for the contents of this array */
+	/**
+	 * Creates a const iterator for the contents of this array.
+	 *
+	 * @returns The const iterator.
+	 */
 	TConstIterator CreateConstIterator() const
 	{
 		return TConstIterator(*this);
 	}
 
 private:
+	/**
+	 * Calls destructor and frees memory on every element in the array.
+	 */
 	void DestructAndFreeItems()
 	{
-		T** Element = GetTypedData();
-		for( int32 i = Array.Num(); i; --i )
+		T** Element = GetData();
+		for (int32 Index = Array.Num(); Index; --Index)
 		{
 			// We need a typedef here because VC won't compile the destructor call below if T itself has a member called T
 			typedef T IndirectArrayDestructElementType;
 
 			(*Element)->IndirectArrayDestructElementType::~IndirectArrayDestructElementType();
-			FMemory::Free( *Element );
+			FMemory::Free(*Element);
 			++Element;
 		}
 	}
@@ -2409,28 +3269,43 @@ struct TContainerTraits<TIndirectArray<T, Allocator> > : public TContainerTraits
 template <typename T,typename Allocator> void* operator new( size_t Size, TIndirectArray<T,Allocator>& Array )
 {
 	check(Size == sizeof(T));
-	const int32 Index = Array.AddRawItem((T*)FMemory::Malloc(Size));
+	const int32 Index = Array.Add((T*)FMemory::Malloc(Size));
 	return &Array[Index];
 }
 
 template <typename T,typename Allocator> void* operator new( size_t Size, TIndirectArray<T,Allocator>& Array, int32 Index )
 {
 	check(Size == sizeof(T));
-	Array.InsertRawItem((T*)FMemory::Malloc(Size), Index);
+	Array.Insert((T*)FMemory::Malloc(Size), Index);
 	return &Array[Index];
 }
 
-/** A specialization of the exchange macro that avoids reallocating when exchanging two arrays. */
-template <typename T> inline void Exchange( TIndirectArray<T>& A, TIndirectArray<T>& B )
+/**
+ * A specialization of the exchange macro that avoids reallocating when
+ * exchanging two arrays.
+ *
+ * @param FirstArrayToExchange First array to exchange.
+ * @param SecondArrayToExchange Second array to exchange.
+ */
+template <typename T>
+inline void Exchange(TIndirectArray<T>& FirstArrayToExchange, TIndirectArray<T>& SecondArrayToExchange)
 {
-	FMemory::Memswap( &A, &B, sizeof(TIndirectArray<T>) );
+	FMemory::Memswap(&FirstArrayToExchange, &SecondArrayToExchange, sizeof(TIndirectArray<T>));
 }
 
-/** A specialization of the exchange macro that avoids reallocating when exchanging two arrays. */
+/**
+ * A specialization of the exchange macro that avoids reallocating when
+ * exchanging two arrays.
+ *
+ * Non-default allocator version.
+ *
+ * @param FirstArrayToExchange First array to exchange.
+ * @param SecondArrayToExchange Second array to exchange.
+ */
 template<typename ElementType,typename Allocator>
-inline void Exchange( TIndirectArray<ElementType,Allocator>& A, TIndirectArray<ElementType,Allocator>& B )
+inline void Exchange(TIndirectArray<ElementType, Allocator>& FirstArrayToExchange, TIndirectArray<ElementType, Allocator>& SecondArrayToExchange)
 {
-	FMemory::Memswap( &A, &B, sizeof(TIndirectArray<ElementType,Allocator>) );
+	FMemory::Memswap(&FirstArrayToExchange, &SecondArrayToExchange, sizeof(TIndirectArray<ElementType, Allocator>));
 }
 
 /*-----------------------------------------------------------------------------
@@ -2553,7 +3428,7 @@ public:
 	int32 AddZeroed( int32 n=1 )
 	{
 		const int32 Index = AddUninitialized(n);
-		FMemory::Memzero( this->GetTypedData() + Index, n*sizeof(T) );
+		FMemory::Memzero(this->GetData() + Index, n*sizeof(T));
 		return Index;
 	}
 	int32 AddUnique( const T& Item )
diff --git a/Engine/Source/Runtime/Core/Public/Containers/StringConv.h b/Engine/Source/Runtime/Core/Public/Containers/StringConv.h
--- a/Engine/Source/Runtime/Core/Public/Containers/StringConv.h
+++ b/Engine/Source/Runtime/Core/Public/Containers/StringConv.h
@@ -789,7 +789,7 @@ FORCEINLINE TArray<ToType> StringToArray(const FromType* Src, int32 SrcLen)
 
 	TArray<ToType> Result;
 	Result.AddUninitialized(DestLen);
-	FPlatformString::Convert(Result.GetTypedData(), DestLen, Src, SrcLen);
+	FPlatformString::Convert(Result.GetData(), DestLen, Src, SrcLen);
 
 	return Result;
 }
