Preallocating TMaps to contain indicated number of elements #UE4

diff --git a/Engine/Source/Runtime/Core/Public/Containers/Map.h b/Engine/Source/Runtime/Core/Public/Containers/Map.h
--- a/Engine/Source/Runtime/Core/Public/Containers/Map.h
+++ b/Engine/Source/Runtime/Core/Public/Containers/Map.h
@@ -232,6 +232,12 @@ public:
 		Pairs.Compact();
 	}
 
+	/** Preallocates enough memory to contain Number elements */
+	FORCEINLINE void Reserve(int32 Number)
+	{
+		Pairs.Reserve(Number);
+	}
+
 	/** @return The number of elements in the map. */
 	FORCEINLINE int32 Num() const
 	{
diff --git a/Engine/Source/Runtime/Core/Public/Containers/Set.h b/Engine/Source/Runtime/Core/Public/Containers/Set.h
--- a/Engine/Source/Runtime/Core/Public/Containers/Set.h
+++ b/Engine/Source/Runtime/Core/Public/Containers/Set.h
@@ -314,6 +314,17 @@ public:
 		Rehash();
 	}
 
+	/** Preallocates enough memory to contain Number elements */
+	FORCEINLINE void Reserve(int32 Number)
+	{
+		// makes sense only when Number > Elements.Num() since TSparseArray::Reserve 
+		// does any work only if that's the case
+		if (Number > Elements.Num())
+		{
+			Elements.Reserve(Number);
+		}
+	}
+
 	/** Relaxes the set's hash to a size strictly bounded by the number of elements in the set. */
 	FORCEINLINE void Relax()
 	{
