Variadic template implementation of delegates.

* Variadic version is compiler-switched on variadic template support.
* Variadic version doesn't repeatedly include lots of the same headers to generate code.
* Tidying up of some original delegate stuff.
* Removed return type from lambda in K2Node_TemporaryVariable which was causing a Clang ICE.
* TTuple, TIntegerSequence and associated functionality - generic, but currently only implementation details of delegates and not yet for use in general code because there is no non-variadic version of them.
* MessageRouter and MessageTracer updated to use an actual delegate instead of implementation-specific identifiers.
* TTypeWrapper and TUnwrapType added to aid 'recursive' template definitions.
* C4373 disabled on VC - was only hitting this warning during development, but it's good to disable it anyway, as it's a 'this code is now standard' warning.

diff --git a/Engine/Plugins/Media/WmfMedia/Source/WmfMedia/Private/Wmf/WmfMediaSampler.h b/Engine/Plugins/Media/WmfMedia/Source/WmfMedia/Private/Wmf/WmfMediaSampler.h
--- a/Engine/Plugins/Media/WmfMedia/Source/WmfMedia/Private/Wmf/WmfMediaSampler.h
+++ b/Engine/Plugins/Media/WmfMedia/Source/WmfMedia/Private/Wmf/WmfMediaSampler.h
@@ -14,6 +14,7 @@ class FWmfMediaSampler
 	: public IMFSampleGrabberSinkCallback
 {
 public:
+	DECLARE_DELEGATE(CommandDelegate)
 
 	/** Default constructor. */
 	FWmfMediaSampler( )
@@ -86,7 +87,7 @@ public:
 	STDMETHODIMP OnProcessSample( REFGUID MajorMediaType, DWORD SampleFlags, LONGLONG SampleTime, LONGLONG SampleDuration, const BYTE* SampleBuffer, DWORD SampleSize)
 	{
 		// process pending commands
-		TBaseDelegate_NoParams<void> Command;
+		CommandDelegate Command;
 
 		while (Commands.Dequeue(Command))
 		{
@@ -156,7 +157,7 @@ private:
 private:
 
 	// Holds the router command queue.
-	TQueue<TBaseDelegate_NoParams<void>, EQueueMode::Mpsc> Commands;
+	TQueue<CommandDelegate, EQueueMode::Mpsc> Commands;
 
 	// The collection of registered media sinks.
 	TArray<IMediaSinkWeakPtr> Sinks;
diff --git a/Engine/Source/Editor/BlueprintGraph/Private/K2Node_TemporaryVariable.cpp b/Engine/Source/Editor/BlueprintGraph/Private/K2Node_TemporaryVariable.cpp
--- a/Engine/Source/Editor/BlueprintGraph/Private/K2Node_TemporaryVariable.cpp
+++ b/Engine/Source/Editor/BlueprintGraph/Private/K2Node_TemporaryVariable.cpp
@@ -143,7 +143,7 @@ void UK2Node_TemporaryVariable::GetMenuActions(FBlueprintActionDatabaseRegistrar
 		return;
 	}
 
-	auto MakeTempVarNodeSpawner = [](FEdGraphPinType const& VarType, bool bIsPersistent) -> UBlueprintNodeSpawner*
+	auto MakeTempVarNodeSpawner = [](FEdGraphPinType const& VarType, bool bIsPersistent)
 	{
 		UBlueprintNodeSpawner* NodeSpawner = UBlueprintNodeSpawner::Create(UK2Node_TemporaryVariable::StaticClass());
 		check(NodeSpawner != nullptr);
diff --git a/Engine/Source/Runtime/Core/Public/Delegates/Delegate.h b/Engine/Source/Runtime/Core/Public/Delegates/Delegate.h
--- a/Engine/Source/Runtime/Core/Public/Delegates/Delegate.h
+++ b/Engine/Source/Runtime/Core/Public/Delegates/Delegate.h
@@ -168,22 +168,33 @@
 /** Helper macro that enables passing comma-separated arguments as a single macro parameter */
 #define FUNC_CONCAT( ... ) __VA_ARGS__
 
+#if PLATFORM_COMPILER_HAS_VARIADIC_TEMPLATES
+	#define FUNC_DECLARE_DELEGATE_BASE(Prefix, Suffix, ...) Prefix<__VA_ARGS__>
+	#define FUNC_DECLARE_DYNAMIC_DELEGATE_BASE(TWeakPtr, Prefix, Suffix, ...) FUNC_DECLARE_DELEGATE_BASE(Prefix, Suffix, TWeakPtr, __VA_ARGS__)
+#else
+	#define FUNC_DECLARE_DELEGATE_BASE(Prefix, Suffix, ...) Prefix##_##Suffix<__VA_ARGS__>
+	#define FUNC_DECLARE_DYNAMIC_DELEGATE_BASE(TWeakPtr, Prefix, Suffix, ...) FUNC_DECLARE_DELEGATE_BASE(Prefix, Suffix, __VA_ARGS__, TWeakPtr)
+#endif
+
 /** Declare the user's delegate object */
 // NOTE: The last parameter is variadic and is used as the 'template args' for this delegate's classes (__VA_ARGS__)
 #define FUNC_DECLARE_DELEGATE( Suffix, DelegateName, ... ) \
-	typedef TBaseDelegate_##Suffix< __VA_ARGS__ > DelegateName;
+	typedef FUNC_DECLARE_DELEGATE_BASE(TBaseDelegate, Suffix, __VA_ARGS__) DelegateName;
 
 /** Declare the user's multicast delegate object */
 // NOTE: The last parameter is variadic and is used as the 'template args' for this delegate's classes (__VA_ARGS__)
 #define FUNC_DECLARE_MULTICAST_DELEGATE( Suffix, MulticastDelegateName, ... ) \
-	typedef TMulticastDelegate_##Suffix< __VA_ARGS__ > MulticastDelegateName;
+	typedef FUNC_DECLARE_DELEGATE_BASE(TMulticastDelegate, Suffix, __VA_ARGS__) MulticastDelegateName;
 
 #define FUNC_DECLARE_EVENT( OwningType, EventName, Suffix, ... ) \
-	class EventName : public TBaseMulticastDelegate_##Suffix< __VA_ARGS__ >  { friend class OwningType; };
+	class EventName : public FUNC_DECLARE_DELEGATE_BASE(TBaseMulticastDelegate, Suffix, __VA_ARGS__) \
+	{ \
+		friend class OwningType; \
+	};
 
 /** Declare user's dynamic delegate, with wrapper proxy method for executing the delegate */
-#define FUNC_DECLARE_DYNAMIC_DELEGATE( Suffix, DynamicDelegateName, ExecFunction, FuncParamList, FuncParamPassThru, ... ) \
-	class DynamicDelegateName : public TBaseDynamicDelegate_##Suffix< __VA_ARGS__ > \
+#define FUNC_DECLARE_DYNAMIC_DELEGATE( TWeakPtr, Suffix, DynamicDelegateName, ExecFunction, FuncParamList, FuncParamPassThru, ... ) \
+	class DynamicDelegateName : public FUNC_DECLARE_DYNAMIC_DELEGATE_BASE(TWeakPtr, TBaseDynamicDelegate, Suffix, __VA_ARGS__) \
 	{ \
 	public: \
 		/** Default constructor */ \
@@ -193,7 +204,7 @@
 		\
 		/** Construction from an FScriptDelegate must be explicit.  This is really only used by UObject system internals. */ \
 		explicit DynamicDelegateName( const TScriptDelegate<>& InScriptDelegate ) \
-			: TBaseDynamicDelegate_##Suffix( InScriptDelegate ) \
+			: FUNC_DECLARE_DYNAMIC_DELEGATE_BASE(TWeakPtr, TBaseDynamicDelegate, Suffix, __VA_ARGS__)( InScriptDelegate ) \
 		{ \
 		} \
 		\
@@ -216,22 +227,8 @@
 		} \
 	};
 
-// Helper struct for converting uint32 return value to bool
-template <typename RetVal>
-struct TRetValCast
-{
-	inline RetVal operator()( RetVal Val ) const
-	{
-		return Val;
-	}
-	inline bool operator()( uint32 Val ) const
-	{
-		return !!Val;
-	}
-};
-
-#define FUNC_DECLARE_DYNAMIC_DELEGATE_RETVAL( Suffix, DynamicDelegateName, ExecFunction, RetValType, FuncParamList, FuncParamPassThru, ... ) \
-	class DynamicDelegateName : public TBaseDynamicDelegate_##Suffix< __VA_ARGS__ > \
+#define FUNC_DECLARE_DYNAMIC_DELEGATE_RETVAL(TWeakPtr, Suffix, DynamicDelegateName, ExecFunction, RetValType, FuncParamList, FuncParamPassThru, ...) \
+	class DynamicDelegateName : public FUNC_DECLARE_DYNAMIC_DELEGATE_BASE(TWeakPtr, TBaseDynamicDelegate, Suffix, __VA_ARGS__) \
 	{ \
 	public: \
 		/** Default constructor */ \
@@ -241,7 +238,7 @@ struct TRetValCast
 		\
 		/** Construction from an FScriptDelegate must be explicit.  This is really only used by UObject system internals. */ \
 		explicit DynamicDelegateName( const TScriptDelegate<>& InScriptDelegate ) \
-			: TBaseDynamicDelegate_##Suffix( InScriptDelegate ) \
+			: FUNC_DECLARE_DYNAMIC_DELEGATE_BASE(TWeakPtr, TBaseDynamicDelegate, Suffix, __VA_ARGS__)( InScriptDelegate ) \
 		{ \
 		} \
 		\
@@ -250,15 +247,14 @@ struct TRetValCast
 		{ \
 			/* Verify that the user object is still valid.  We only have a weak reference to it. */ \
 			checkSlow( IsBound() ); \
-			TRetValCast<RetValType> CastRetVal; \
-			return CastRetVal( ExecFunction( FuncParamPassThru ) ); \
+			return ExecFunction( FuncParamPassThru ); \
 		} \
 	};
 
 
 /** Declare user's dynamic multi-cast delegate, with wrapper proxy method for executing the delegate */
-#define FUNC_DECLARE_DYNAMIC_MULTICAST_DELEGATE( Suffix, DynamicMulticastDelegateName, ExecFunction, FuncParamList, FuncParamPassThru, ... ) \
-class DynamicMulticastDelegateName : public TBaseDynamicMulticastDelegate_##Suffix< __VA_ARGS__ > \
+#define FUNC_DECLARE_DYNAMIC_MULTICAST_DELEGATE(TWeakPtr, Suffix, DynamicMulticastDelegateName, ExecFunction, FuncParamList, FuncParamPassThru, ...) \
+class DynamicMulticastDelegateName : public FUNC_DECLARE_DYNAMIC_DELEGATE_BASE(TWeakPtr, TBaseDynamicMulticastDelegate, Suffix, __VA_ARGS__) \
 	{ \
 	public: \
 		/** Default constructor */ \
@@ -268,7 +264,7 @@ class DynamicMulticastDelegateName : public TBaseDynamicMulticastDelegate_##Suff
 		\
 		/** Construction from an FMulticastScriptDelegate must be explicit.  This is really only used by UObject system internals. */ \
 		explicit DynamicMulticastDelegateName( const TMulticastScriptDelegate<>& InMulticastScriptDelegate ) \
-			: TBaseDynamicMulticastDelegate_##Suffix( InMulticastScriptDelegate ) \
+			: FUNC_DECLARE_DYNAMIC_DELEGATE_BASE(TWeakPtr, TBaseDynamicMulticastDelegate, Suffix, __VA_ARGS__)( InMulticastScriptDelegate ) \
 		{ \
 		} \
 		\
@@ -307,9 +303,15 @@ class DynamicMulticastDelegateName : public TBaseDynamicMulticastDelegate_##Suff
 // We define this as a guard to prevent DelegateSignatureImpl.inl being included outside of this file
 #define FUNC_INCLUDING_INLINE_IMPL
 
-
 #ifndef UE_BUILD_DOCS
-#include "DelegateCombinations.h"
+	#if PLATFORM_COMPILER_HAS_VARIADIC_TEMPLATES
+		#include "DelegateInstanceInterface_Variadics.h"
+		#include "DelegateInstancesImpl_Variadics.inl"
+		#include "DelegateSignatureImpl_Variadics.inl"
+		#include "DelegateCombinations_Variadics.h"
+	#else
+		#include "DelegateCombinations.h"
+	#endif
 #endif
 
 
diff --git a/Engine/Source/Runtime/Core/Public/Delegates/DelegateBase.h b/Engine/Source/Runtime/Core/Public/Delegates/DelegateBase.h
--- a/Engine/Source/Runtime/Core/Public/Delegates/DelegateBase.h
+++ b/Engine/Source/Runtime/Core/Public/Delegates/DelegateBase.h
@@ -20,8 +20,6 @@ public:
 		: DelegateInstance(InDelegateInstance)
 	{ }
 
-public:
-
 	/**
 	 * If this is a UFunction or UObject delegate, return the UObject.
 	 *
@@ -73,8 +71,6 @@ public:
 		}
 	}
 
-protected:
-
 	/**
 	 * Gets the delegate instance.
 	 *
diff --git a/Engine/Source/Runtime/Core/Public/Delegates/DelegateCombinations.h b/Engine/Source/Runtime/Core/Public/Delegates/DelegateCombinations.h
--- a/Engine/Source/Runtime/Core/Public/Delegates/DelegateCombinations.h
+++ b/Engine/Source/Runtime/Core/Public/Delegates/DelegateCombinations.h
@@ -198,8 +198,8 @@
 #define DECLARE_DELEGATE( DelegateName ) FUNC_DECLARE_DELEGATE( NoParams, DelegateName, void )
 #define DECLARE_MULTICAST_DELEGATE( DelegateName ) FUNC_DECLARE_MULTICAST_DELEGATE( NoParams, DelegateName, void )
 #define DECLARE_EVENT( OwningType, EventName ) FUNC_DECLARE_EVENT( OwningType, EventName, NoParams, void )
-#define DECLARE_DYNAMIC_DELEGATE( DelegateName ) FUNC_DECLARE_DYNAMIC_DELEGATE( NoParams, DelegateName, DelegateName##_DelegateWrapper, , FUNC_CONCAT( *this ), void )
-#define DECLARE_DYNAMIC_MULTICAST_DELEGATE( DelegateName ) FUNC_DECLARE_DYNAMIC_MULTICAST_DELEGATE( NoParams, DelegateName, DelegateName##_DelegateWrapper, , FUNC_CONCAT( *this ), void )
+#define DECLARE_DYNAMIC_DELEGATE( DelegateName ) FUNC_DECLARE_DYNAMIC_DELEGATE( FWeakObjectPtr, NoParams, DelegateName, DelegateName##_DelegateWrapper, , FUNC_CONCAT( *this ), void )
+#define DECLARE_DYNAMIC_MULTICAST_DELEGATE( DelegateName ) FUNC_DECLARE_DYNAMIC_MULTICAST_DELEGATE( FWeakObjectPtr, NoParams, DelegateName, DelegateName##_DelegateWrapper, , FUNC_CONCAT( *this ), void )
 #undef FUNC_SUFFIX
 #undef FUNC_TEMPLATE_DECL
 #undef FUNC_TEMPLATE_DECL_TYPENAME
@@ -400,7 +400,7 @@
 #include "DelegateSignatureImpl.inl"
 
 #define DECLARE_DELEGATE_RetVal( RetValType, DelegateName ) FUNC_DECLARE_DELEGATE( RetVal_NoParams, DelegateName, RetValType )
-#define DECLARE_DYNAMIC_DELEGATE_RetVal( RetValType, DelegateName ) FUNC_DECLARE_DYNAMIC_DELEGATE_RETVAL( RetVal_NoParams, DelegateName, DelegateName##_DelegateWrapper, RetValType, , FUNC_CONCAT( *this ), RetValType )
+#define DECLARE_DYNAMIC_DELEGATE_RetVal( RetValType, DelegateName ) FUNC_DECLARE_DYNAMIC_DELEGATE_RETVAL( FWeakObjectPtr, RetVal_NoParams, DelegateName, DelegateName##_DelegateWrapper, RetValType, , FUNC_CONCAT( *this ), RetValType )
 #undef FUNC_SUFFIX
 #undef FUNC_TEMPLATE_DECL
 #undef FUNC_TEMPLATE_DECL_TYPENAME
@@ -603,8 +603,8 @@
 #define DECLARE_DELEGATE_OneParam( DelegateName, Param1Type ) FUNC_DECLARE_DELEGATE( OneParam, DelegateName, void, Param1Type )
 #define DECLARE_MULTICAST_DELEGATE_OneParam( DelegateName, Param1Type ) FUNC_DECLARE_MULTICAST_DELEGATE( OneParam, DelegateName, void, Param1Type )
 #define DECLARE_EVENT_OneParam( OwningType, EventName, Param1Type ) FUNC_DECLARE_EVENT( OwningType, EventName, OneParam, void, Param1Type )
-#define DECLARE_DYNAMIC_DELEGATE_OneParam( DelegateName, Param1Type, Param1Name ) FUNC_DECLARE_DYNAMIC_DELEGATE( OneParam, DelegateName, DelegateName##_DelegateWrapper, FUNC_CONCAT( Param1Type InParam1 ), FUNC_CONCAT( *this, InParam1 ), void, Param1Type )
-#define DECLARE_DYNAMIC_MULTICAST_DELEGATE_OneParam( DelegateName, Param1Type, Param1Name ) FUNC_DECLARE_DYNAMIC_MULTICAST_DELEGATE( OneParam, DelegateName, DelegateName##_DelegateWrapper, FUNC_CONCAT( Param1Type InParam1 ), FUNC_CONCAT( *this, InParam1 ), void, Param1Type )
+#define DECLARE_DYNAMIC_DELEGATE_OneParam( DelegateName, Param1Type, Param1Name ) FUNC_DECLARE_DYNAMIC_DELEGATE( FWeakObjectPtr, OneParam, DelegateName, DelegateName##_DelegateWrapper, FUNC_CONCAT( Param1Type InParam1 ), FUNC_CONCAT( *this, InParam1 ), void, Param1Type )
+#define DECLARE_DYNAMIC_MULTICAST_DELEGATE_OneParam( DelegateName, Param1Type, Param1Name ) FUNC_DECLARE_DYNAMIC_MULTICAST_DELEGATE( FWeakObjectPtr, OneParam, DelegateName, DelegateName##_DelegateWrapper, FUNC_CONCAT( Param1Type InParam1 ), FUNC_CONCAT( *this, InParam1 ), void, Param1Type )
 #undef FUNC_SUFFIX
 #undef FUNC_TEMPLATE_DECL
 #undef FUNC_TEMPLATE_DECL_TYPENAME
@@ -805,7 +805,7 @@
 #include "DelegateSignatureImpl.inl"
 
 #define DECLARE_DELEGATE_RetVal_OneParam( RetValType, DelegateName, Param1Type ) FUNC_DECLARE_DELEGATE( RetVal_OneParam, DelegateName, RetValType, Param1Type )
-#define DECLARE_DYNAMIC_DELEGATE_RetVal_OneParam( RetValType, DelegateName, Param1Type, Param1Name ) FUNC_DECLARE_DYNAMIC_DELEGATE_RETVAL( RetVal_OneParam, DelegateName, DelegateName##_DelegateWrapper, RetValType, FUNC_CONCAT( Param1Type InParam1 ), FUNC_CONCAT( *this, InParam1 ), RetValType, Param1Type )
+#define DECLARE_DYNAMIC_DELEGATE_RetVal_OneParam( RetValType, DelegateName, Param1Type, Param1Name ) FUNC_DECLARE_DYNAMIC_DELEGATE_RETVAL( FWeakObjectPtr, RetVal_OneParam, DelegateName, DelegateName##_DelegateWrapper, RetValType, FUNC_CONCAT( Param1Type InParam1 ), FUNC_CONCAT( *this, InParam1 ), RetValType, Param1Type )
 #undef FUNC_SUFFIX
 #undef FUNC_TEMPLATE_DECL
 #undef FUNC_TEMPLATE_DECL_TYPENAME
@@ -1008,8 +1008,8 @@
 #define DECLARE_DELEGATE_TwoParams( DelegateName, Param1Type, Param2Type ) FUNC_DECLARE_DELEGATE( TwoParams, DelegateName, void, Param1Type, Param2Type )
 #define DECLARE_MULTICAST_DELEGATE_TwoParams( DelegateName, Param1Type, Param2Type ) FUNC_DECLARE_MULTICAST_DELEGATE( TwoParams, DelegateName, void, Param1Type, Param2Type )
 #define DECLARE_EVENT_TwoParams( OwningType, EventName, Param1Type, Param2Type ) FUNC_DECLARE_EVENT( OwningType, EventName, TwoParams, void, Param1Type, Param2Type )
-#define DECLARE_DYNAMIC_DELEGATE_TwoParams( DelegateName, Param1Type, Param1Name, Param2Type, Param2Name ) FUNC_DECLARE_DYNAMIC_DELEGATE( TwoParams, DelegateName, DelegateName##_DelegateWrapper, FUNC_CONCAT( Param1Type InParam1, Param2Type InParam2 ), FUNC_CONCAT( *this, InParam1, InParam2 ), void, Param1Type, Param2Type )
-#define DECLARE_DYNAMIC_MULTICAST_DELEGATE_TwoParams( DelegateName, Param1Type, Param1Name, Param2Type, Param2Name ) FUNC_DECLARE_DYNAMIC_MULTICAST_DELEGATE( TwoParams, DelegateName, DelegateName##_DelegateWrapper, FUNC_CONCAT( Param1Type InParam1, Param2Type InParam2 ), FUNC_CONCAT( *this, InParam1, InParam2 ), void, Param1Type, Param2Type )
+#define DECLARE_DYNAMIC_DELEGATE_TwoParams( DelegateName, Param1Type, Param1Name, Param2Type, Param2Name ) FUNC_DECLARE_DYNAMIC_DELEGATE( FWeakObjectPtr, TwoParams, DelegateName, DelegateName##_DelegateWrapper, FUNC_CONCAT( Param1Type InParam1, Param2Type InParam2 ), FUNC_CONCAT( *this, InParam1, InParam2 ), void, Param1Type, Param2Type )
+#define DECLARE_DYNAMIC_MULTICAST_DELEGATE_TwoParams( DelegateName, Param1Type, Param1Name, Param2Type, Param2Name ) FUNC_DECLARE_DYNAMIC_MULTICAST_DELEGATE( FWeakObjectPtr, TwoParams, DelegateName, DelegateName##_DelegateWrapper, FUNC_CONCAT( Param1Type InParam1, Param2Type InParam2 ), FUNC_CONCAT( *this, InParam1, InParam2 ), void, Param1Type, Param2Type )
 #undef FUNC_SUFFIX
 #undef FUNC_TEMPLATE_DECL
 #undef FUNC_TEMPLATE_DECL_TYPENAME
@@ -1210,7 +1210,7 @@
 #include "DelegateSignatureImpl.inl"
 
 #define DECLARE_DELEGATE_RetVal_TwoParams( RetValType, DelegateName, Param1Type, Param2Type ) FUNC_DECLARE_DELEGATE( RetVal_TwoParams, DelegateName, RetValType, Param1Type, Param2Type )
-#define DECLARE_DYNAMIC_DELEGATE_RetVal_TwoParams( RetValType, DelegateName, Param1Type, Param1Name, Param2Type, Param2Name ) FUNC_DECLARE_DYNAMIC_DELEGATE_RETVAL( RetVal_TwoParams, DelegateName, DelegateName##_DelegateWrapper, RetValType, FUNC_CONCAT( Param1Type InParam1, Param2Type InParam2 ), FUNC_CONCAT( *this, InParam1, InParam2 ), RetValType, Param1Type, Param2Type )
+#define DECLARE_DYNAMIC_DELEGATE_RetVal_TwoParams( RetValType, DelegateName, Param1Type, Param1Name, Param2Type, Param2Name ) FUNC_DECLARE_DYNAMIC_DELEGATE_RETVAL( FWeakObjectPtr, RetVal_TwoParams, DelegateName, DelegateName##_DelegateWrapper, RetValType, FUNC_CONCAT( Param1Type InParam1, Param2Type InParam2 ), FUNC_CONCAT( *this, InParam1, InParam2 ), RetValType, Param1Type, Param2Type )
 #undef FUNC_SUFFIX
 #undef FUNC_TEMPLATE_DECL
 #undef FUNC_TEMPLATE_DECL_TYPENAME
@@ -1413,8 +1413,8 @@
 #define DECLARE_DELEGATE_ThreeParams( DelegateName, Param1Type, Param2Type, Param3Type ) FUNC_DECLARE_DELEGATE( ThreeParams, DelegateName, void, Param1Type, Param2Type, Param3Type )
 #define DECLARE_MULTICAST_DELEGATE_ThreeParams( DelegateName, Param1Type, Param2Type, Param3Type ) FUNC_DECLARE_MULTICAST_DELEGATE( ThreeParams, DelegateName, void, Param1Type, Param2Type, Param3Type )
 #define DECLARE_EVENT_ThreeParams( OwningType, EventName, Param1Type, Param2Type, Param3Type ) FUNC_DECLARE_EVENT( OwningType, EventName, ThreeParams, void, Param1Type, Param2Type, Param3Type )
-#define DECLARE_DYNAMIC_DELEGATE_ThreeParams( DelegateName, Param1Type, Param1Name, Param2Type, Param2Name, Param3Type, Param3Name ) FUNC_DECLARE_DYNAMIC_DELEGATE( ThreeParams, DelegateName, DelegateName##_DelegateWrapper, FUNC_CONCAT( Param1Type InParam1, Param2Type InParam2, Param3Type InParam3 ), FUNC_CONCAT( *this, InParam1, InParam2, InParam3 ), void, Param1Type, Param2Type, Param3Type )
-#define DECLARE_DYNAMIC_MULTICAST_DELEGATE_ThreeParams( DelegateName, Param1Type, Param1Name, Param2Type, Param2Name, Param3Type, Param3Name ) FUNC_DECLARE_DYNAMIC_MULTICAST_DELEGATE( ThreeParams, DelegateName, DelegateName##_DelegateWrapper, FUNC_CONCAT( Param1Type InParam1, Param2Type InParam2, Param3Type InParam3 ), FUNC_CONCAT( *this, InParam1, InParam2, InParam3 ), void, Param1Type, Param2Type, Param3Type )
+#define DECLARE_DYNAMIC_DELEGATE_ThreeParams( DelegateName, Param1Type, Param1Name, Param2Type, Param2Name, Param3Type, Param3Name ) FUNC_DECLARE_DYNAMIC_DELEGATE( FWeakObjectPtr, ThreeParams, DelegateName, DelegateName##_DelegateWrapper, FUNC_CONCAT( Param1Type InParam1, Param2Type InParam2, Param3Type InParam3 ), FUNC_CONCAT( *this, InParam1, InParam2, InParam3 ), void, Param1Type, Param2Type, Param3Type )
+#define DECLARE_DYNAMIC_MULTICAST_DELEGATE_ThreeParams( DelegateName, Param1Type, Param1Name, Param2Type, Param2Name, Param3Type, Param3Name ) FUNC_DECLARE_DYNAMIC_MULTICAST_DELEGATE( FWeakObjectPtr, ThreeParams, DelegateName, DelegateName##_DelegateWrapper, FUNC_CONCAT( Param1Type InParam1, Param2Type InParam2, Param3Type InParam3 ), FUNC_CONCAT( *this, InParam1, InParam2, InParam3 ), void, Param1Type, Param2Type, Param3Type )
 #undef FUNC_SUFFIX
 #undef FUNC_TEMPLATE_DECL
 #undef FUNC_TEMPLATE_DECL_TYPENAME
@@ -1615,7 +1615,7 @@
 #include "DelegateSignatureImpl.inl"
 
 #define DECLARE_DELEGATE_RetVal_ThreeParams( RetValType, DelegateName, Param1Type, Param2Type, Param3Type ) FUNC_DECLARE_DELEGATE( RetVal_ThreeParams, DelegateName, RetValType, Param1Type, Param2Type, Param3Type )
-#define DECLARE_DYNAMIC_DELEGATE_RetVal_ThreeParams( RetValType, DelegateName, Param1Type, Param1Name, Param2Type, Param2Name, Param3Type, Param3Name ) FUNC_DECLARE_DYNAMIC_DELEGATE_RETVAL( RetVal_ThreeParams, DelegateName, DelegateName##_DelegateWrapper, RetValType, FUNC_CONCAT( Param1Type InParam1, Param2Type InParam2, Param3Type InParam3 ), FUNC_CONCAT( *this, InParam1, InParam2, InParam3 ), RetValType, Param1Type, Param2Type, Param3Type )
+#define DECLARE_DYNAMIC_DELEGATE_RetVal_ThreeParams( RetValType, DelegateName, Param1Type, Param1Name, Param2Type, Param2Name, Param3Type, Param3Name ) FUNC_DECLARE_DYNAMIC_DELEGATE_RETVAL( FWeakObjectPtr, RetVal_ThreeParams, DelegateName, DelegateName##_DelegateWrapper, RetValType, FUNC_CONCAT( Param1Type InParam1, Param2Type InParam2, Param3Type InParam3 ), FUNC_CONCAT( *this, InParam1, InParam2, InParam3 ), RetValType, Param1Type, Param2Type, Param3Type )
 #undef FUNC_SUFFIX
 #undef FUNC_TEMPLATE_DECL
 #undef FUNC_TEMPLATE_DECL_TYPENAME
@@ -1818,8 +1818,8 @@
 #define DECLARE_DELEGATE_FourParams( DelegateName, Param1Type, Param2Type, Param3Type, Param4Type ) FUNC_DECLARE_DELEGATE( FourParams, DelegateName, void, Param1Type, Param2Type, Param3Type, Param4Type )
 #define DECLARE_MULTICAST_DELEGATE_FourParams( DelegateName, Param1Type, Param2Type, Param3Type, Param4Type ) FUNC_DECLARE_MULTICAST_DELEGATE( FourParams, DelegateName, void, Param1Type, Param2Type, Param3Type, Param4Type )
 #define DECLARE_EVENT_FourParams( OwningType, EventName, Param1Type, Param2Type, Param3Type, Param4Type ) FUNC_DECLARE_EVENT( OwningType, EventName, FourParams, void, Param1Type, Param2Type, Param3Type, Param4Type )
-#define DECLARE_DYNAMIC_DELEGATE_FourParams( DelegateName, Param1Type, Param1Name, Param2Type, Param2Name, Param3Type, Param3Name, Param4Type, Param4Name ) FUNC_DECLARE_DYNAMIC_DELEGATE( FourParams, DelegateName, DelegateName##_DelegateWrapper, FUNC_CONCAT( Param1Type InParam1, Param2Type InParam2, Param3Type InParam3, Param4Type InParam4 ), FUNC_CONCAT( *this, InParam1, InParam2, InParam3, InParam4 ), void, Param1Type, Param2Type, Param3Type, Param4Type )
-#define DECLARE_DYNAMIC_MULTICAST_DELEGATE_FourParams( DelegateName, Param1Type, Param1Name, Param2Type, Param2Name, Param3Type, Param3Name, Param4Type, Param4Name ) FUNC_DECLARE_DYNAMIC_MULTICAST_DELEGATE( FourParams, DelegateName, DelegateName##_DelegateWrapper, FUNC_CONCAT( Param1Type InParam1, Param2Type InParam2, Param3Type InParam3, Param4Type InParam4 ), FUNC_CONCAT( *this, InParam1, InParam2, InParam3, InParam4 ), void, Param1Type, Param2Type, Param3Type, Param4Type )
+#define DECLARE_DYNAMIC_DELEGATE_FourParams( DelegateName, Param1Type, Param1Name, Param2Type, Param2Name, Param3Type, Param3Name, Param4Type, Param4Name ) FUNC_DECLARE_DYNAMIC_DELEGATE( FWeakObjectPtr, FourParams, DelegateName, DelegateName##_DelegateWrapper, FUNC_CONCAT( Param1Type InParam1, Param2Type InParam2, Param3Type InParam3, Param4Type InParam4 ), FUNC_CONCAT( *this, InParam1, InParam2, InParam3, InParam4 ), void, Param1Type, Param2Type, Param3Type, Param4Type )
+#define DECLARE_DYNAMIC_MULTICAST_DELEGATE_FourParams( DelegateName, Param1Type, Param1Name, Param2Type, Param2Name, Param3Type, Param3Name, Param4Type, Param4Name ) FUNC_DECLARE_DYNAMIC_MULTICAST_DELEGATE( FWeakObjectPtr, FourParams, DelegateName, DelegateName##_DelegateWrapper, FUNC_CONCAT( Param1Type InParam1, Param2Type InParam2, Param3Type InParam3, Param4Type InParam4 ), FUNC_CONCAT( *this, InParam1, InParam2, InParam3, InParam4 ), void, Param1Type, Param2Type, Param3Type, Param4Type )
 #undef FUNC_SUFFIX
 #undef FUNC_TEMPLATE_DECL
 #undef FUNC_TEMPLATE_DECL_TYPENAME
@@ -2020,7 +2020,7 @@
 #include "DelegateSignatureImpl.inl"
 
 #define DECLARE_DELEGATE_RetVal_FourParams( RetValType, DelegateName, Param1Type, Param2Type, Param3Type, Param4Type ) FUNC_DECLARE_DELEGATE( RetVal_FourParams, DelegateName, RetValType, Param1Type, Param2Type, Param3Type, Param4Type )
-#define DECLARE_DYNAMIC_DELEGATE_RetVal_FourParams( RetValType, DelegateName, Param1Type, Param1Name, Param2Type, Param2Name, Param3Type, Param3Name, Param4Type, Param4Name ) FUNC_DECLARE_DYNAMIC_DELEGATE_RETVAL( RetVal_FourParams, DelegateName, DelegateName##_DelegateWrapper, RetValType, FUNC_CONCAT( Param1Type InParam1, Param2Type InParam2, Param3Type InParam3, Param4Type InParam4 ), FUNC_CONCAT( *this, InParam1, InParam2, InParam3, InParam4 ), RetValType, Param1Type, Param2Type, Param3Type, Param4Type )
+#define DECLARE_DYNAMIC_DELEGATE_RetVal_FourParams( RetValType, DelegateName, Param1Type, Param1Name, Param2Type, Param2Name, Param3Type, Param3Name, Param4Type, Param4Name ) FUNC_DECLARE_DYNAMIC_DELEGATE_RETVAL( FWeakObjectPtr, RetVal_FourParams, DelegateName, DelegateName##_DelegateWrapper, RetValType, FUNC_CONCAT( Param1Type InParam1, Param2Type InParam2, Param3Type InParam3, Param4Type InParam4 ), FUNC_CONCAT( *this, InParam1, InParam2, InParam3, InParam4 ), RetValType, Param1Type, Param2Type, Param3Type, Param4Type )
 #undef FUNC_SUFFIX
 #undef FUNC_TEMPLATE_DECL
 #undef FUNC_TEMPLATE_DECL_TYPENAME
@@ -2223,8 +2223,8 @@
 #define DECLARE_DELEGATE_FiveParams( DelegateName, Param1Type, Param2Type, Param3Type, Param4Type, Param5Type ) FUNC_DECLARE_DELEGATE( FiveParams, DelegateName, void, Param1Type, Param2Type, Param3Type, Param4Type, Param5Type )
 #define DECLARE_MULTICAST_DELEGATE_FiveParams( DelegateName, Param1Type, Param2Type, Param3Type, Param4Type, Param5Type ) FUNC_DECLARE_MULTICAST_DELEGATE( FiveParams, DelegateName, void, Param1Type, Param2Type, Param3Type, Param4Type, Param5Type )
 #define DECLARE_EVENT_FiveParams( OwningType, EventName, Param1Type, Param2Type, Param3Type, Param4Type, Param5Type ) FUNC_DECLARE_EVENT( OwningType, EventName, FiveParams, void, Param1Type, Param2Type, Param3Type, Param4Type, Param5Type )
-#define DECLARE_DYNAMIC_DELEGATE_FiveParams( DelegateName, Param1Type, Param1Name, Param2Type, Param2Name, Param3Type, Param3Name, Param4Type, Param4Name, Param5Type, Param5Name ) FUNC_DECLARE_DYNAMIC_DELEGATE( FiveParams, DelegateName, DelegateName##_DelegateWrapper, FUNC_CONCAT( Param1Type InParam1, Param2Type InParam2, Param3Type InParam3, Param4Type InParam4, Param5Type InParam5 ), FUNC_CONCAT( *this, InParam1, InParam2, InParam3, InParam4, InParam5 ), void, Param1Type, Param2Type, Param3Type, Param4Type, Param5Type )
-#define DECLARE_DYNAMIC_MULTICAST_DELEGATE_FiveParams( DelegateName, Param1Type, Param1Name, Param2Type, Param2Name, Param3Type, Param3Name, Param4Type, Param4Name, Param5Type, Param5Name ) FUNC_DECLARE_DYNAMIC_MULTICAST_DELEGATE( FiveParams, DelegateName, DelegateName##_DelegateWrapper, FUNC_CONCAT( Param1Type InParam1, Param2Type InParam2, Param3Type InParam3, Param4Type InParam4, Param5Type InParam5 ), FUNC_CONCAT( *this, InParam1, InParam2, InParam3, InParam4, InParam5 ), void, Param1Type, Param2Type, Param3Type, Param4Type, Param5Type )
+#define DECLARE_DYNAMIC_DELEGATE_FiveParams( DelegateName, Param1Type, Param1Name, Param2Type, Param2Name, Param3Type, Param3Name, Param4Type, Param4Name, Param5Type, Param5Name ) FUNC_DECLARE_DYNAMIC_DELEGATE( FWeakObjectPtr, FiveParams, DelegateName, DelegateName##_DelegateWrapper, FUNC_CONCAT( Param1Type InParam1, Param2Type InParam2, Param3Type InParam3, Param4Type InParam4, Param5Type InParam5 ), FUNC_CONCAT( *this, InParam1, InParam2, InParam3, InParam4, InParam5 ), void, Param1Type, Param2Type, Param3Type, Param4Type, Param5Type )
+#define DECLARE_DYNAMIC_MULTICAST_DELEGATE_FiveParams( DelegateName, Param1Type, Param1Name, Param2Type, Param2Name, Param3Type, Param3Name, Param4Type, Param4Name, Param5Type, Param5Name ) FUNC_DECLARE_DYNAMIC_MULTICAST_DELEGATE( FWeakObjectPtr, FiveParams, DelegateName, DelegateName##_DelegateWrapper, FUNC_CONCAT( Param1Type InParam1, Param2Type InParam2, Param3Type InParam3, Param4Type InParam4, Param5Type InParam5 ), FUNC_CONCAT( *this, InParam1, InParam2, InParam3, InParam4, InParam5 ), void, Param1Type, Param2Type, Param3Type, Param4Type, Param5Type )
 #undef FUNC_SUFFIX
 #undef FUNC_TEMPLATE_DECL
 #undef FUNC_TEMPLATE_DECL_TYPENAME
@@ -2425,7 +2425,7 @@
 #include "DelegateSignatureImpl.inl"
 
 #define DECLARE_DELEGATE_RetVal_FiveParams( RetValType, DelegateName, Param1Type, Param2Type, Param3Type, Param4Type, Param5Type ) FUNC_DECLARE_DELEGATE( RetVal_FiveParams, DelegateName, RetValType, Param1Type, Param2Type, Param3Type, Param4Type, Param5Type )
-#define DECLARE_DYNAMIC_DELEGATE_RetVal_FiveParams( RetValType, DelegateName, Param1Type, Param1Name, Param2Type, Param2Name, Param3Type, Param3Name, Param4Type, Param4Name, Param5Type, Param5Name ) FUNC_DECLARE_DYNAMIC_DELEGATE_RETVAL( RetVal_FiveParams, DelegateName, DelegateName##_DelegateWrapper, RetValType, FUNC_CONCAT( Param1Type InParam1, Param2Type InParam2, Param3Type InParam3, Param4Type InParam4, Param5Type InParam5 ), FUNC_CONCAT( *this, InParam1, InParam2, InParam3, InParam4, InParam5 ), RetValType, Param1Type, Param2Type, Param3Type, Param4Type, Param5Type )
+#define DECLARE_DYNAMIC_DELEGATE_RetVal_FiveParams( RetValType, DelegateName, Param1Type, Param1Name, Param2Type, Param2Name, Param3Type, Param3Name, Param4Type, Param4Name, Param5Type, Param5Name ) FUNC_DECLARE_DYNAMIC_DELEGATE_RETVAL( FWeakObjectPtr, RetVal_FiveParams, DelegateName, DelegateName##_DelegateWrapper, RetValType, FUNC_CONCAT( Param1Type InParam1, Param2Type InParam2, Param3Type InParam3, Param4Type InParam4, Param5Type InParam5 ), FUNC_CONCAT( *this, InParam1, InParam2, InParam3, InParam4, InParam5 ), RetValType, Param1Type, Param2Type, Param3Type, Param4Type, Param5Type )
 #undef FUNC_SUFFIX
 #undef FUNC_TEMPLATE_DECL
 #undef FUNC_TEMPLATE_DECL_TYPENAME
@@ -2628,8 +2628,8 @@
 #define DECLARE_DELEGATE_SixParams( DelegateName, Param1Type, Param2Type, Param3Type, Param4Type, Param5Type, Param6Type ) FUNC_DECLARE_DELEGATE( SixParams, DelegateName, void, Param1Type, Param2Type, Param3Type, Param4Type, Param5Type, Param6Type )
 #define DECLARE_MULTICAST_DELEGATE_SixParams( DelegateName, Param1Type, Param2Type, Param3Type, Param4Type, Param5Type, Param6Type ) FUNC_DECLARE_MULTICAST_DELEGATE( SixParams, DelegateName, void, Param1Type, Param2Type, Param3Type, Param4Type, Param5Type, Param6Type )
 #define DECLARE_EVENT_SixParams( OwningType, EventName, Param1Type, Param2Type, Param3Type, Param4Type, Param5Type, Param6Type ) FUNC_DECLARE_EVENT( OwningType, EventName, SixParams, void, Param1Type, Param2Type, Param3Type, Param4Type, Param5Type, Param6Type )
-#define DECLARE_DYNAMIC_DELEGATE_SixParams( DelegateName, Param1Type, Param1Name, Param2Type, Param2Name, Param3Type, Param3Name, Param4Type, Param4Name, Param5Type, Param5Name, Param6Type, Param6Name ) FUNC_DECLARE_DYNAMIC_DELEGATE( SixParams, DelegateName, DelegateName##_DelegateWrapper, FUNC_CONCAT( Param1Type InParam1, Param2Type InParam2, Param3Type InParam3, Param4Type InParam4, Param5Type InParam5, Param6Type InParam6 ), FUNC_CONCAT( *this, InParam1, InParam2, InParam3, InParam4, InParam5, InParam6 ), void, Param1Type, Param2Type, Param3Type, Param4Type, Param5Type, Param6Type )
-#define DECLARE_DYNAMIC_MULTICAST_DELEGATE_SixParams( DelegateName, Param1Type, Param1Name, Param2Type, Param2Name, Param3Type, Param3Name, Param4Type, Param4Name, Param5Type, Param5Name, Param6Type, Param6Name ) FUNC_DECLARE_DYNAMIC_MULTICAST_DELEGATE( SixParams, DelegateName, DelegateName##_DelegateWrapper, FUNC_CONCAT( Param1Type InParam1, Param2Type InParam2, Param3Type InParam3, Param4Type InParam4, Param5Type InParam5, Param6Type InParam6 ), FUNC_CONCAT( *this, InParam1, InParam2, InParam3, InParam4, InParam5, InParam6 ), void, Param1Type, Param2Type, Param3Type, Param4Type, Param5Type, Param6Type )
+#define DECLARE_DYNAMIC_DELEGATE_SixParams( DelegateName, Param1Type, Param1Name, Param2Type, Param2Name, Param3Type, Param3Name, Param4Type, Param4Name, Param5Type, Param5Name, Param6Type, Param6Name ) FUNC_DECLARE_DYNAMIC_DELEGATE( FWeakObjectPtr, SixParams, DelegateName, DelegateName##_DelegateWrapper, FUNC_CONCAT( Param1Type InParam1, Param2Type InParam2, Param3Type InParam3, Param4Type InParam4, Param5Type InParam5, Param6Type InParam6 ), FUNC_CONCAT( *this, InParam1, InParam2, InParam3, InParam4, InParam5, InParam6 ), void, Param1Type, Param2Type, Param3Type, Param4Type, Param5Type, Param6Type )
+#define DECLARE_DYNAMIC_MULTICAST_DELEGATE_SixParams( DelegateName, Param1Type, Param1Name, Param2Type, Param2Name, Param3Type, Param3Name, Param4Type, Param4Name, Param5Type, Param5Name, Param6Type, Param6Name ) FUNC_DECLARE_DYNAMIC_MULTICAST_DELEGATE( FWeakObjectPtr, SixParams, DelegateName, DelegateName##_DelegateWrapper, FUNC_CONCAT( Param1Type InParam1, Param2Type InParam2, Param3Type InParam3, Param4Type InParam4, Param5Type InParam5, Param6Type InParam6 ), FUNC_CONCAT( *this, InParam1, InParam2, InParam3, InParam4, InParam5, InParam6 ), void, Param1Type, Param2Type, Param3Type, Param4Type, Param5Type, Param6Type )
 #undef FUNC_SUFFIX
 #undef FUNC_TEMPLATE_DECL
 #undef FUNC_TEMPLATE_DECL_TYPENAME
@@ -2830,7 +2830,7 @@
 #include "DelegateSignatureImpl.inl"
 
 #define DECLARE_DELEGATE_RetVal_SixParams( RetValType, DelegateName, Param1Type, Param2Type, Param3Type, Param4Type, Param5Type, Param6Type ) FUNC_DECLARE_DELEGATE( RetVal_SixParams, DelegateName, RetValType, Param1Type, Param2Type, Param3Type, Param4Type, Param5Type, Param6Type )
-#define DECLARE_DYNAMIC_DELEGATE_RetVal_SixParams( RetValType, DelegateName, Param1Type, Param1Name, Param2Type, Param2Name, Param3Type, Param3Name, Param4Type, Param4Name, Param5Type, Param5Name, Param6Type, Param6Name ) FUNC_DECLARE_DYNAMIC_DELEGATE_RETVAL( RetVal_SixParams, DelegateName, DelegateName##_DelegateWrapper, RetValType, FUNC_CONCAT( Param1Type InParam1, Param2Type InParam2, Param3Type InParam3, Param4Type InParam4, Param5Type InParam5, Param6Type InParam6 ), FUNC_CONCAT( *this, InParam1, InParam2, InParam3, InParam4, InParam5, InParam6 ), RetValType, Param1Type, Param2Type, Param3Type, Param4Type, Param5Type, Param6Type )
+#define DECLARE_DYNAMIC_DELEGATE_RetVal_SixParams( RetValType, DelegateName, Param1Type, Param1Name, Param2Type, Param2Name, Param3Type, Param3Name, Param4Type, Param4Name, Param5Type, Param5Name, Param6Type, Param6Name ) FUNC_DECLARE_DYNAMIC_DELEGATE_RETVAL( FWeakObjectPtr, RetVal_SixParams, DelegateName, DelegateName##_DelegateWrapper, RetValType, FUNC_CONCAT( Param1Type InParam1, Param2Type InParam2, Param3Type InParam3, Param4Type InParam4, Param5Type InParam5, Param6Type InParam6 ), FUNC_CONCAT( *this, InParam1, InParam2, InParam3, InParam4, InParam5, InParam6 ), RetValType, Param1Type, Param2Type, Param3Type, Param4Type, Param5Type, Param6Type )
 #undef FUNC_SUFFIX
 #undef FUNC_TEMPLATE_DECL
 #undef FUNC_TEMPLATE_DECL_TYPENAME
@@ -3033,8 +3033,8 @@
 #define DECLARE_DELEGATE_SevenParams( DelegateName, Param1Type, Param2Type, Param3Type, Param4Type, Param5Type, Param6Type, Param7Type ) FUNC_DECLARE_DELEGATE( SevenParams, DelegateName, void, Param1Type, Param2Type, Param3Type, Param4Type, Param5Type, Param6Type, Param7Type )
 #define DECLARE_MULTICAST_DELEGATE_SevenParams( DelegateName, Param1Type, Param2Type, Param3Type, Param4Type, Param5Type, Param6Type, Param7Type ) FUNC_DECLARE_MULTICAST_DELEGATE( SevenParams, DelegateName, void, Param1Type, Param2Type, Param3Type, Param4Type, Param5Type, Param6Type, Param7Type )
 #define DECLARE_EVENT_SevenParams( OwningType, EventName, Param1Type, Param2Type, Param3Type, Param4Type, Param5Type, Param6Type, Param7Type ) FUNC_DECLARE_EVENT( OwningType, EventName, SevenParams, void, Param1Type, Param2Type, Param3Type, Param4Type, Param5Type, Param6Type, Param7Type )
-#define DECLARE_DYNAMIC_DELEGATE_SevenParams( DelegateName, Param1Type, Param1Name, Param2Type, Param2Name, Param3Type, Param3Name, Param4Type, Param4Name, Param5Type, Param5Name, Param6Type, Param6Name, Param7Type, Param7Name ) FUNC_DECLARE_DYNAMIC_DELEGATE( SevenParams, DelegateName, DelegateName##_DelegateWrapper, FUNC_CONCAT( Param1Type InParam1, Param2Type InParam2, Param3Type InParam3, Param4Type InParam4, Param5Type InParam5, Param6Type InParam6, Param7Type InParam7 ), FUNC_CONCAT( *this, InParam1, InParam2, InParam3, InParam4, InParam5, InParam6, InParam7 ), void, Param1Type, Param2Type, Param3Type, Param4Type, Param5Type, Param6Type, Param7Type )
-#define DECLARE_DYNAMIC_MULTICAST_DELEGATE_SevenParams( DelegateName, Param1Type, Param1Name, Param2Type, Param2Name, Param3Type, Param3Name, Param4Type, Param4Name, Param5Type, Param5Name, Param6Type, Param6Name, Param7Type, Param7Name ) FUNC_DECLARE_DYNAMIC_MULTICAST_DELEGATE( SevenParams, DelegateName, DelegateName##_DelegateWrapper, FUNC_CONCAT( Param1Type InParam1, Param2Type InParam2, Param3Type InParam3, Param4Type InParam4, Param5Type InParam5, Param6Type InParam6, Param7Type InParam7 ), FUNC_CONCAT( *this, InParam1, InParam2, InParam3, InParam4, InParam5, InParam6, InParam7 ), void, Param1Type, Param2Type, Param3Type, Param4Type, Param5Type, Param6Type, Param7Type )
+#define DECLARE_DYNAMIC_DELEGATE_SevenParams( DelegateName, Param1Type, Param1Name, Param2Type, Param2Name, Param3Type, Param3Name, Param4Type, Param4Name, Param5Type, Param5Name, Param6Type, Param6Name, Param7Type, Param7Name ) FUNC_DECLARE_DYNAMIC_DELEGATE( FWeakObjectPtr, SevenParams, DelegateName, DelegateName##_DelegateWrapper, FUNC_CONCAT( Param1Type InParam1, Param2Type InParam2, Param3Type InParam3, Param4Type InParam4, Param5Type InParam5, Param6Type InParam6, Param7Type InParam7 ), FUNC_CONCAT( *this, InParam1, InParam2, InParam3, InParam4, InParam5, InParam6, InParam7 ), void, Param1Type, Param2Type, Param3Type, Param4Type, Param5Type, Param6Type, Param7Type )
+#define DECLARE_DYNAMIC_MULTICAST_DELEGATE_SevenParams( DelegateName, Param1Type, Param1Name, Param2Type, Param2Name, Param3Type, Param3Name, Param4Type, Param4Name, Param5Type, Param5Name, Param6Type, Param6Name, Param7Type, Param7Name ) FUNC_DECLARE_DYNAMIC_MULTICAST_DELEGATE( FWeakObjectPtr, SevenParams, DelegateName, DelegateName##_DelegateWrapper, FUNC_CONCAT( Param1Type InParam1, Param2Type InParam2, Param3Type InParam3, Param4Type InParam4, Param5Type InParam5, Param6Type InParam6, Param7Type InParam7 ), FUNC_CONCAT( *this, InParam1, InParam2, InParam3, InParam4, InParam5, InParam6, InParam7 ), void, Param1Type, Param2Type, Param3Type, Param4Type, Param5Type, Param6Type, Param7Type )
 #undef FUNC_SUFFIX
 #undef FUNC_TEMPLATE_DECL
 #undef FUNC_TEMPLATE_DECL_TYPENAME
@@ -3235,7 +3235,7 @@
 #include "DelegateSignatureImpl.inl"
 
 #define DECLARE_DELEGATE_RetVal_SevenParams( RetValType, DelegateName, Param1Type, Param2Type, Param3Type, Param4Type, Param5Type, Param6Type, Param7Type ) FUNC_DECLARE_DELEGATE( RetVal_SevenParams, DelegateName, RetValType, Param1Type, Param2Type, Param3Type, Param4Type, Param5Type, Param6Type, Param7Type )
-#define DECLARE_DYNAMIC_DELEGATE_RetVal_SevenParams( RetValType, DelegateName, Param1Type, Param1Name, Param2Type, Param2Name, Param3Type, Param3Name, Param4Type, Param4Name, Param5Type, Param5Name, Param6Type, Param6Name, Param7Type, Param7Name ) FUNC_DECLARE_DYNAMIC_DELEGATE_RETVAL( RetVal_SevenParams, DelegateName, DelegateName##_DelegateWrapper, RetValType, FUNC_CONCAT( Param1Type InParam1, Param2Type InParam2, Param3Type InParam3, Param4Type InParam4, Param5Type InParam5, Param6Type InParam6, Param7Type InParam7 ), FUNC_CONCAT( *this, InParam1, InParam2, InParam3, InParam4, InParam5, InParam6, InParam7 ), RetValType, Param1Type, Param2Type, Param3Type, Param4Type, Param5Type, Param6Type, Param7Type )
+#define DECLARE_DYNAMIC_DELEGATE_RetVal_SevenParams( RetValType, DelegateName, Param1Type, Param1Name, Param2Type, Param2Name, Param3Type, Param3Name, Param4Type, Param4Name, Param5Type, Param5Name, Param6Type, Param6Name, Param7Type, Param7Name ) FUNC_DECLARE_DYNAMIC_DELEGATE_RETVAL( FWeakObjectPtr, RetVal_SevenParams, DelegateName, DelegateName##_DelegateWrapper, RetValType, FUNC_CONCAT( Param1Type InParam1, Param2Type InParam2, Param3Type InParam3, Param4Type InParam4, Param5Type InParam5, Param6Type InParam6, Param7Type InParam7 ), FUNC_CONCAT( *this, InParam1, InParam2, InParam3, InParam4, InParam5, InParam6, InParam7 ), RetValType, Param1Type, Param2Type, Param3Type, Param4Type, Param5Type, Param6Type, Param7Type )
 #undef FUNC_SUFFIX
 #undef FUNC_TEMPLATE_DECL
 #undef FUNC_TEMPLATE_DECL_TYPENAME
@@ -3438,8 +3438,8 @@
 #define DECLARE_DELEGATE_EightParams( DelegateName, Param1Type, Param2Type, Param3Type, Param4Type, Param5Type, Param6Type, Param7Type, Param8Type ) FUNC_DECLARE_DELEGATE( EightParams, DelegateName, void, Param1Type, Param2Type, Param3Type, Param4Type, Param5Type, Param6Type, Param7Type, Param8Type )
 #define DECLARE_MULTICAST_DELEGATE_EightParams( DelegateName, Param1Type, Param2Type, Param3Type, Param4Type, Param5Type, Param6Type, Param7Type, Param8Type ) FUNC_DECLARE_MULTICAST_DELEGATE( EightParams, DelegateName, void, Param1Type, Param2Type, Param3Type, Param4Type, Param5Type, Param6Type, Param7Type, Param8Type )
 #define DECLARE_EVENT_EightParams( OwningType, EventName, Param1Type, Param2Type, Param3Type, Param4Type, Param5Type, Param6Type, Param7Type, Param8Type ) FUNC_DECLARE_EVENT( OwningType, EventName, EightParams, void, Param1Type, Param2Type, Param3Type, Param4Type, Param5Type, Param6Type, Param7Type, Param8Type )
-#define DECLARE_DYNAMIC_DELEGATE_EightParams( DelegateName, Param1Type, Param1Name, Param2Type, Param2Name, Param3Type, Param3Name, Param4Type, Param4Name, Param5Type, Param5Name, Param6Type, Param6Name, Param7Type, Param7Name, Param8Type, Param8Name ) FUNC_DECLARE_DYNAMIC_DELEGATE( EightParams, DelegateName, DelegateName##_DelegateWrapper, FUNC_CONCAT( Param1Type InParam1, Param2Type InParam2, Param3Type InParam3, Param4Type InParam4, Param5Type InParam5, Param6Type InParam6, Param7Type InParam7, Param8Type InParam8 ), FUNC_CONCAT( *this, InParam1, InParam2, InParam3, InParam4, InParam5, InParam6, InParam7, InParam8 ), void, Param1Type, Param2Type, Param3Type, Param4Type, Param5Type, Param6Type, Param7Type, Param8Type )
-#define DECLARE_DYNAMIC_MULTICAST_DELEGATE_EightParams( DelegateName, Param1Type, Param1Name, Param2Type, Param2Name, Param3Type, Param3Name, Param4Type, Param4Name, Param5Type, Param5Name, Param6Type, Param6Name, Param7Type, Param7Name, Param8Type, Param8Name ) FUNC_DECLARE_DYNAMIC_MULTICAST_DELEGATE( EightParams, DelegateName, DelegateName##_DelegateWrapper, FUNC_CONCAT( Param1Type InParam1, Param2Type InParam2, Param3Type InParam3, Param4Type InParam4, Param5Type InParam5, Param6Type InParam6, Param7Type InParam7, Param8Type InParam8 ), FUNC_CONCAT( *this, InParam1, InParam2, InParam3, InParam4, InParam5, InParam6, InParam7, InParam8 ), void, Param1Type, Param2Type, Param3Type, Param4Type, Param5Type, Param6Type, Param7Type, Param8Type )
+#define DECLARE_DYNAMIC_DELEGATE_EightParams( DelegateName, Param1Type, Param1Name, Param2Type, Param2Name, Param3Type, Param3Name, Param4Type, Param4Name, Param5Type, Param5Name, Param6Type, Param6Name, Param7Type, Param7Name, Param8Type, Param8Name ) FUNC_DECLARE_DYNAMIC_DELEGATE( FWeakObjectPtr, EightParams, DelegateName, DelegateName##_DelegateWrapper, FUNC_CONCAT( Param1Type InParam1, Param2Type InParam2, Param3Type InParam3, Param4Type InParam4, Param5Type InParam5, Param6Type InParam6, Param7Type InParam7, Param8Type InParam8 ), FUNC_CONCAT( *this, InParam1, InParam2, InParam3, InParam4, InParam5, InParam6, InParam7, InParam8 ), void, Param1Type, Param2Type, Param3Type, Param4Type, Param5Type, Param6Type, Param7Type, Param8Type )
+#define DECLARE_DYNAMIC_MULTICAST_DELEGATE_EightParams( DelegateName, Param1Type, Param1Name, Param2Type, Param2Name, Param3Type, Param3Name, Param4Type, Param4Name, Param5Type, Param5Name, Param6Type, Param6Name, Param7Type, Param7Name, Param8Type, Param8Name ) FUNC_DECLARE_DYNAMIC_MULTICAST_DELEGATE( FWeakObjectPtr, EightParams, DelegateName, DelegateName##_DelegateWrapper, FUNC_CONCAT( Param1Type InParam1, Param2Type InParam2, Param3Type InParam3, Param4Type InParam4, Param5Type InParam5, Param6Type InParam6, Param7Type InParam7, Param8Type InParam8 ), FUNC_CONCAT( *this, InParam1, InParam2, InParam3, InParam4, InParam5, InParam6, InParam7, InParam8 ), void, Param1Type, Param2Type, Param3Type, Param4Type, Param5Type, Param6Type, Param7Type, Param8Type )
 #undef FUNC_SUFFIX
 #undef FUNC_TEMPLATE_DECL
 #undef FUNC_TEMPLATE_DECL_TYPENAME
@@ -3640,7 +3640,7 @@
 #include "DelegateSignatureImpl.inl"
 
 #define DECLARE_DELEGATE_RetVal_EightParams( RetValType, DelegateName, Param1Type, Param2Type, Param3Type, Param4Type, Param5Type, Param6Type, Param7Type, Param8Type ) FUNC_DECLARE_DELEGATE( RetVal_EightParams, DelegateName, RetValType, Param1Type, Param2Type, Param3Type, Param4Type, Param5Type, Param6Type, Param7Type, Param8Type )
-#define DECLARE_DYNAMIC_DELEGATE_RetVal_EightParams( RetValType, DelegateName, Param1Type, Param1Name, Param2Type, Param2Name, Param3Type, Param3Name, Param4Type, Param4Name, Param5Type, Param5Name, Param6Type, Param6Name, Param7Type, Param7Name, Param8Type, Param8Name ) FUNC_DECLARE_DYNAMIC_DELEGATE_RETVAL( RetVal_EightParams, DelegateName, DelegateName##_DelegateWrapper, RetValType, FUNC_CONCAT( Param1Type InParam1, Param2Type InParam2, Param3Type InParam3, Param4Type InParam4, Param5Type InParam5, Param6Type InParam6, Param7Type InParam7, Param8Type InParam8 ), FUNC_CONCAT( *this, InParam1, InParam2, InParam3, InParam4, InParam5, InParam6, InParam7, InParam8 ), RetValType, Param1Type, Param2Type, Param3Type, Param4Type, Param5Type, Param6Type, Param7Type, Param8Type )
+#define DECLARE_DYNAMIC_DELEGATE_RetVal_EightParams( RetValType, DelegateName, Param1Type, Param1Name, Param2Type, Param2Name, Param3Type, Param3Name, Param4Type, Param4Name, Param5Type, Param5Name, Param6Type, Param6Name, Param7Type, Param7Name, Param8Type, Param8Name ) FUNC_DECLARE_DYNAMIC_DELEGATE_RETVAL( FWeakObjectPtr, RetVal_EightParams, DelegateName, DelegateName##_DelegateWrapper, RetValType, FUNC_CONCAT( Param1Type InParam1, Param2Type InParam2, Param3Type InParam3, Param4Type InParam4, Param5Type InParam5, Param6Type InParam6, Param7Type InParam7, Param8Type InParam8 ), FUNC_CONCAT( *this, InParam1, InParam2, InParam3, InParam4, InParam5, InParam6, InParam7, InParam8 ), RetValType, Param1Type, Param2Type, Param3Type, Param4Type, Param5Type, Param6Type, Param7Type, Param8Type )
 #undef FUNC_SUFFIX
 #undef FUNC_TEMPLATE_DECL
 #undef FUNC_TEMPLATE_DECL_TYPENAME
diff --git a/Engine/Source/Runtime/Core/Public/Delegates/DelegateInstanceInterfaceImpl.inl b/Engine/Source/Runtime/Core/Public/Delegates/DelegateInstanceInterfaceImpl.inl
--- a/Engine/Source/Runtime/Core/Public/Delegates/DelegateInstanceInterfaceImpl.inl
+++ b/Engine/Source/Runtime/Core/Public/Delegates/DelegateInstanceInterfaceImpl.inl
@@ -12,13 +12,6 @@
 
 #define DELEGATE_INSTANCE_INTERFACE_CLASS FUNC_COMBINE( IBaseDelegateInstance_, FUNC_SUFFIX )
 
-#if FUNC_HAS_PARAMS
-	#define DELEGATE_INSTANCE_COLON_PARAM_INITIALIZER_LIST : FUNC_PARAM_INITIALIZER_LIST
-#else
-	#define DELEGATE_INSTANCE_COLON_PARAM_INITIALIZER_LIST
-#endif
-
-
 /**
  * Delegate instance base interface.  For internal use only.  This abstract class is used to interact with
  * a delegate instance for *any* type of object (e.g. raw, shared pointer, UObject-based, etc.)
@@ -35,10 +28,6 @@ class DELEGATE_INSTANCE_INTERFACE_CLASS
 #if !FUNC_IS_VOID
 		RetValType Result;
 #endif
-
-		FParms( FUNC_PARAM_LIST )
-			DELEGATE_INSTANCE_COLON_PARAM_INITIALIZER_LIST
-		{ }
 	};
 
 public:
@@ -48,7 +37,7 @@ public:
 	 *
 	 * @return	The newly created copy
 	 */
-	virtual DELEGATE_INSTANCE_INTERFACE_CLASS< FUNC_TEMPLATE_ARGS >* CreateCopy() = 0;
+	virtual DELEGATE_INSTANCE_INTERFACE_CLASS* CreateCopy() = 0;
 
 	/**
 	 * Returns true if this delegate points to exactly the same object and method as the specified delegate,
@@ -57,7 +46,7 @@ public:
 	 * @param  InOtherDelegate
 	 * @return  True if delegates match
 	 */
-	virtual bool IsSameFunction( const DELEGATE_INSTANCE_INTERFACE_CLASS< FUNC_TEMPLATE_ARGS >& InOtherDelegate ) const = 0;
+	virtual bool IsSameFunction( const DELEGATE_INSTANCE_INTERFACE_CLASS& InOtherDelegate ) const = 0;
 
 	/**
 	 * Execute the delegate.  If the function pointer is not valid, an error will occur.
@@ -94,5 +83,4 @@ public:
 };
 
 
-#undef DELEGATE_INSTANCE_COLON_PARAM_INITIALIZER_LIST
 #undef DELEGATE_INSTANCE_INTERFACE_CLASS
diff --git a/Engine/Source/Runtime/Core/Public/Delegates/DelegateInstanceInterface_Variadics.h b/Engine/Source/Runtime/Core/Public/Delegates/DelegateInstanceInterface_Variadics.h
new file mode 100644
--- /dev/null
+++ b/Engine/Source/Runtime/Core/Public/Delegates/DelegateInstanceInterface_Variadics.h
@@ -0,0 +1,163 @@
+// Copyright 1998-2014 Epic Games, Inc. All Rights Reserved.
+
+#pragma once
+
+#include "Tuple.h"
+
+template <typename FuncType>
+struct IBaseDelegateInstance;
+
+template <typename FuncType>
+struct IBaseDelegateInstanceCommon;
+
+template <typename RetType, typename... Args>
+struct IBaseDelegateInstanceCommon<RetType(Args...)> : public IDelegateInstance
+{
+	/**
+	 * Creates a copy of the delegate instance
+	 *
+	 * @return	The newly created copy
+	 */
+	virtual IBaseDelegateInstance<RetType(Args...)>* CreateCopy() = 0;
+
+	/**
+	 * Returns true if this delegate points to exactly the same object and method as the specified delegate,
+	 * even if the delegate objects themselves are different.  Also, the delegate types *must* be compatible.
+	 *
+	 * @param  InOtherDelegate
+	 * @return  True if delegates match
+	 */
+	virtual bool IsSameFunction( const IBaseDelegateInstance<RetType(Args...)>& InOtherDelegate ) const = 0;
+
+	/**
+	 * Execute the delegate.  If the function pointer is not valid, an error will occur.
+	 */
+	virtual RetType Execute(Args...) const = 0;
+};
+
+template <typename FuncType>
+struct IBaseDelegateInstance : public IBaseDelegateInstanceCommon<FuncType>
+{
+};
+
+template <typename... Args>
+struct IBaseDelegateInstance<void(Args...)> : public IBaseDelegateInstanceCommon<void(Args...)>
+{
+	/**
+	 * Execute the delegate, but only if the function pointer is still valid
+	 *
+	 * @return  Returns true if the function was executed
+	 */
+	// NOTE: Currently only delegates with no return value support ExecuteIfSafe()
+	virtual bool ExecuteIfSafe(Args...) const = 0;
+};
+
+template <bool Const, typename Class, typename FuncType>
+struct TMemFunPtrType;
+
+template <typename Class, typename RetType, typename... Args>
+struct TMemFunPtrType<false, Class, RetType(Args...)>
+{
+	typedef RetType (Class::* Type)(Args...);
+};
+
+template <typename Class, typename RetType, typename... Args>
+struct TMemFunPtrType<true, Class, RetType(Args...)>
+{
+	typedef RetType (Class::* Type)(Args...) const;
+};
+
+template <typename FuncType>
+struct TPayload;
+
+template <typename Ret, typename... Types>
+struct TPayload<Ret(Types...)>
+{
+	TTuple<Types..., Ret> Values;
+
+	template <typename... ArgTypes>
+	explicit TPayload(ArgTypes&&... Args)
+		: Values(Forward<ArgTypes>(Args)..., Ret())
+	{
+	}
+
+	Ret& GetResult()
+	{
+		return Values.template Get<sizeof...(Types)>();
+	}
+};
+
+template <typename... Types>
+struct TPayload<void(Types...)>
+{
+	TTuple<Types...> Values;
+
+	template <typename... ArgTypes>
+	explicit TPayload(ArgTypes&&... Args)
+		: Values(Forward<ArgTypes>(Args)...)
+	{
+	}
+
+	void GetResult()
+	{
+	}
+};
+
+template <typename T>
+struct TPlacementNewer
+{
+	TPlacementNewer()
+		: bConstructed(false)
+	{
+	}
+
+	~TPlacementNewer()
+	{
+		if (bConstructed)
+		{
+			// We need a typedef here because VC won't compile the destructor call below if T itself has a member called T
+			typedef T TPlacementNewerTTypeTypedef;
+
+			((TPlacementNewerTTypeTypedef*)&Bytes)->TPlacementNewerTTypeTypedef::~TPlacementNewerTTypeTypedef();
+		}
+	}
+
+	template <typename... ArgTypes>
+	T* operator()(ArgTypes&&... Args)
+	{
+		check(!bConstructed);
+		T* Result = new (&Bytes) T(Forward<ArgTypes>(Args)...);
+		bConstructed = true;
+		return Result;
+	}
+
+	T* operator->()
+	{
+		check(bConstructed);
+		return (T*)&Bytes;
+	}
+
+private:
+	TTypeCompatibleBytes<T> Bytes;
+	bool                    bConstructed;
+};
+
+template <typename T, typename MemFunPtrType>
+class TMemberFunctionCaller
+{
+	T*            Obj;
+	MemFunPtrType MemFunPtr;
+
+public:
+	TMemberFunctionCaller(T* InObj, MemFunPtrType InMemFunPtr)
+		: Obj      (InObj)
+		, MemFunPtr(InMemFunPtr)
+	{
+	}
+
+	template <typename... ArgTypes>
+	auto operator()(ArgTypes&&... Args) -> decltype((Obj->*MemFunPtr)(Forward<ArgTypes>(Args)...))
+	{
+		return (Obj->*MemFunPtr)(Forward<ArgTypes>(Args)...);
+	}
+};
diff --git a/Engine/Source/Runtime/Core/Public/Delegates/DelegateInstancesImpl.inl b/Engine/Source/Runtime/Core/Public/Delegates/DelegateInstancesImpl.inl
--- a/Engine/Source/Runtime/Core/Public/Delegates/DelegateInstancesImpl.inl
+++ b/Engine/Source/Runtime/Core/Public/Delegates/DelegateInstancesImpl.inl
@@ -203,7 +203,7 @@ public:
 	 */
 	FORCEINLINE static DELEGATE_INSTANCE_INTERFACE_CLASS<FUNC_TEMPLATE_ARGS>* Create( const TSharedPtr<UserClass, SPMode>& InUserObjectRef, FMethodPtr InFunc DELEGATE_COMMA_PAYLOAD_LIST )
 	{
-		return new SP_METHOD_DELEGATE_INSTANCE_CLASS<UserClass, FUNC_PAYLOAD_TEMPLATE_ARGS, SPMode>(InUserObjectRef, InFunc DELEGATE_COMMA_PAYLOAD_PASSTHRU);
+		return new SP_METHOD_DELEGATE_INSTANCE_CLASS(InUserObjectRef, InFunc DELEGATE_COMMA_PAYLOAD_PASSTHRU);
 	}
 
 	/**
@@ -386,7 +386,7 @@ public:
 	 */
 	FORCEINLINE static DELEGATE_INSTANCE_INTERFACE_CLASS<FUNC_TEMPLATE_ARGS>* Create( UserClass* InUserObject, FMethodPtr InFunc DELEGATE_COMMA_PAYLOAD_LIST )
 	{
-		return new RAW_METHOD_DELEGATE_INSTANCE_CLASS<UserClass, FUNC_PAYLOAD_TEMPLATE_ARGS>(InUserObject, InFunc DELEGATE_COMMA_PAYLOAD_PASSTHRU);
+		return new RAW_METHOD_DELEGATE_INSTANCE_CLASS(InUserObject, InFunc DELEGATE_COMMA_PAYLOAD_PASSTHRU);
 	}
 
 protected:
@@ -561,7 +561,7 @@ public:
 	 */
 	FORCEINLINE static DELEGATE_INSTANCE_INTERFACE_CLASS<FUNC_TEMPLATE_ARGS>* Create( UserClass* InUserObject, FMethodPtr InFunc DELEGATE_COMMA_PAYLOAD_LIST )
 	{
-		return new UOBJECT_METHOD_DELEGATE_INSTANCE_CLASS<UserClass, FUNC_PAYLOAD_TEMPLATE_ARGS>(InUserObject, InFunc DELEGATE_COMMA_PAYLOAD_PASSTHRU);
+		return new UOBJECT_METHOD_DELEGATE_INSTANCE_CLASS(InUserObject, InFunc DELEGATE_COMMA_PAYLOAD_PASSTHRU);
 	}
 
 protected:
@@ -717,7 +717,7 @@ public:
 	 */
 	FORCEINLINE static DELEGATE_INSTANCE_INTERFACE_CLASS<FUNC_TEMPLATE_ARGS>* Create( FFuncPtr InFunc DELEGATE_COMMA_PAYLOAD_LIST )
 	{
-		return new STATIC_DELEGATE_INSTANCE_CLASS<FUNC_PAYLOAD_TEMPLATE_ARGS>(InFunc DELEGATE_COMMA_PAYLOAD_PASSTHRU);
+		return new STATIC_DELEGATE_INSTANCE_CLASS(InFunc DELEGATE_COMMA_PAYLOAD_PASSTHRU);
 	}
 
 private:
@@ -752,11 +752,6 @@ class UFUNCTION_DELEGATE_INSTANCE_CLASS
 #if !FUNC_IS_VOID
 		RetValType Result;		// UFunction return value.
 #endif
-
-		// Constructor.
-		FParmsWithPayload( DELEGATE_PARAM_LIST_COMMA_PAYLOAD_LIST )
-			DELEGATE_PARMS_COLON_INITIALIZER_LIST
-		{ }
 	};
 
 public:
@@ -823,7 +818,7 @@ public:
 
 	// DELEGATE_INSTANCE_INTERFACE_CLASS interface
 
-	virtual DELEGATE_INSTANCE_INTERFACE_CLASS<FUNC_TEMPLATE_ARGS>* CreateCopy( ) override
+	virtual DELEGATE_INSTANCE_INTERFACE_CLASS* CreateCopy( ) override
 	{
 		return Create(UserObjectPtr.Get(), GetFunctionName() DELEGATE_COMMA_PAYLOAD_PASSIN);
 	}
@@ -832,7 +827,7 @@ public:
 	{
 		checkSlow(IsSafeToExecute());
 
-		FParmsWithPayload Parms = FParmsWithPayload(DELEGATE_PARAM_PASSTHRU_COMMA_PAYLOAD_PASSIN);
+		FParmsWithPayload Parms = { DELEGATE_PARAM_PASSTHRU_COMMA_PAYLOAD_PASSIN };
 		UserObjectPtr->ProcessEvent(CachedFunction, &Parms);
 
 #if !FUNC_IS_VOID
@@ -876,7 +871,7 @@ public:
 	 */
 	FORCEINLINE static DELEGATE_INSTANCE_INTERFACE_CLASS<FUNC_TEMPLATE_ARGS>* Create( UserClass* InUserObject, const FName& InFunctionName DELEGATE_COMMA_PAYLOAD_LIST )
 	{
-		return new UFUNCTION_DELEGATE_INSTANCE_CLASS<UserClass, FUNC_PAYLOAD_TEMPLATE_ARGS>(InUserObject, InFunctionName DELEGATE_COMMA_PAYLOAD_PASSTHRU);
+		return new UFUNCTION_DELEGATE_INSTANCE_CLASS(InUserObject, InFunctionName DELEGATE_COMMA_PAYLOAD_PASSTHRU);
 	}
 
 public:
diff --git a/Engine/Source/Runtime/Core/Public/Delegates/DelegateInstancesImpl_Variadics.inl b/Engine/Source/Runtime/Core/Public/Delegates/DelegateInstancesImpl_Variadics.inl
new file mode 100644
--- /dev/null
+++ b/Engine/Source/Runtime/Core/Public/Delegates/DelegateInstancesImpl_Variadics.inl
@@ -0,0 +1,879 @@
+// Copyright 1998-2014 Epic Games, Inc. All Rights Reserved.
+
+/*================================================================================
+	DelegateInstancesImpl.inl: Inline implementation of delegate bindings.
+
+	This file in re-included in DelegateCombinations.h for EVERY supported delegate signature.
+	Every combination of parameter count, return value presence or other function modifier will
+	include this file to generate a delegate interface type and implementation type for that signature.
+
+	The types declared in this file are for internal use only. 
+================================================================================*/
+
+#pragma once
+
+
+/* Macros for function parameter and delegate payload lists
+ *****************************************************************************/
+
+/**
+ * Implements a delegate binding for UFunctions.
+ *
+ * @params UserClass Must be an UObject derived class.
+ */
+template <class UserClass, typename FuncType, typename... VarTypes>
+class TBaseUFunctionDelegateInstance;
+
+template <class UserClass, typename WrappedRetValType, typename... ParamTypes, typename... VarTypes>
+class TBaseUFunctionDelegateInstance<UserClass, WrappedRetValType (ParamTypes...), VarTypes...> : public IBaseDelegateInstance<typename TUnwrapType<WrappedRetValType>::Type (ParamTypes...)>
+{
+	typedef IBaseDelegateInstance<typename TUnwrapType<WrappedRetValType>::Type (ParamTypes...)> Super;
+
+	static_assert((CanConvertPointerFromTo<UserClass, const UObjectBase>::Result), "You cannot use UFunction delegates with non UObject classes.");
+
+public:
+	typedef typename TUnwrapType<WrappedRetValType>::Type RetValType;
+
+	TBaseUFunctionDelegateInstance(UserClass* InUserObject, const FName& InFunctionName, VarTypes... Vars)
+		: FunctionName (InFunctionName)
+		, UserObjectPtr(InUserObject)
+		, Payload      (Vars...)
+	{
+		check(InFunctionName != NAME_None);
+		
+		if (InUserObject != nullptr)
+		{
+			CachedFunction = UserObjectPtr->FindFunctionChecked(InFunctionName);
+		}
+	}
+
+	// IDelegateInstance interface
+
+	virtual FName GetFunctionName( ) const override
+	{
+		return FunctionName;
+	}
+
+	virtual const void* GetRawMethodPtr( ) const override
+	{
+		return nullptr;
+	}
+
+	virtual const void* GetRawUserObject( ) const override
+	{
+		return UserObjectPtr.Get();
+	}
+
+	virtual EDelegateInstanceType::Type GetType() const override
+	{
+		return EDelegateInstanceType::UFunction;
+	}
+
+	virtual bool HasSameObject( const void* InUserObject ) const override
+	{
+		return (UserObjectPtr.Get() == InUserObject);
+	}
+
+	virtual bool IsCompactable( ) const override
+	{
+		return !UserObjectPtr.Get(true);
+	}
+
+	virtual bool IsSafeToExecute( ) const override
+	{
+		return UserObjectPtr.IsValid();
+	}
+
+public:
+
+	// IBaseDelegateInstance interface
+
+	virtual Super* CreateCopy( ) override
+	{
+		return Payload.ApplyAfter(&Create, UserObjectPtr.Get(), GetFunctionName());
+	}
+
+	virtual RetValType Execute(ParamTypes... Params) const override
+	{
+		typedef TPayload<RetValType (ParamTypes..., VarTypes...)> FParmsWithPayload;
+
+		checkSlow(IsSafeToExecute());
+
+		TPlacementNewer<FParmsWithPayload> PayloadAndParams;
+		Payload.ApplyAfter(PayloadAndParams, Params...);
+		UserObjectPtr->ProcessEvent(CachedFunction, &PayloadAndParams);
+		return PayloadAndParams->GetResult();
+	}
+
+	virtual bool IsSameFunction( const Super& Other ) const override
+	{
+		// NOTE: Payload data is not currently considered when comparing delegate instances.
+		// See the comment in multi-cast delegate's Remove() method for more information.
+		return ((Other.GetType() == EDelegateInstanceType::UFunction) &&
+				(Other.GetRawUserObject() == GetRawUserObject()) &&
+				(Other.GetFunctionName() == GetFunctionName()));
+	}
+
+public:
+
+	/**
+	 * Creates a new UFunction delegate binding for the given user object and function name.
+	 *
+	 * @param InObject The user object to call the function on.
+	 * @param InFunctionName The name of the function call.
+	 * @return The new delegate.
+	 */
+	FORCEINLINE static Super* Create(UserClass* InUserObject, const FName& InFunctionName, VarTypes... Vars)
+	{
+		return new TBaseUFunctionDelegateInstance<UserClass, RetValType (ParamTypes...), VarTypes...>(InUserObject, InFunctionName, Vars...);
+	}
+
+	template <typename TupleType, typename IntSeq>
+	struct TCreateFromTuple;
+
+	template <typename... TupleArgs, uint32... Indices>
+	struct TCreateFromTuple<TTuple<TupleArgs...>, TIntegerSequence<Indices...>>
+	{
+		static Super* CreateImpl(UserClass* InUserObject, const FName& InFunctionName, const TTuple<TupleArgs...>& Tuple)
+		{
+			return Create(InUserObject, InFunctionName, Tuple.template Get<Indices>()...);
+		}
+	};
+
+	template <typename... TupleArgs>
+	FORCEINLINE static Super* CreateFromTuple(UserClass* InUserObject, const FName& InFunctionName, const TTuple<TupleArgs...>& Tuple)
+	{
+		return TCreateFromTuple<TTuple<TupleArgs...>, TMakeIntegerSequence<sizeof...(TupleArgs)>>::CreateImpl(InUserObject, InFunctionName, Tuple);
+	}
+
+public:
+
+	// Holds the cached UFunction to call.
+	UFunction* CachedFunction;
+
+	// Holds the name of the function to call.
+	FName FunctionName;
+
+	// The user object to call the function on.
+	TWeakObjectPtr<UserClass> UserObjectPtr;
+
+	TTuple<VarTypes...> Payload;
+};
+
+template <class UserClass, typename... ParamTypes, typename... VarTypes>
+class TBaseUFunctionDelegateInstance<UserClass, void (ParamTypes...), VarTypes...> : public TBaseUFunctionDelegateInstance<UserClass, TTypeWrapper<void> (ParamTypes...), VarTypes...>
+{
+	typedef TBaseUFunctionDelegateInstance<UserClass, TTypeWrapper<void> (ParamTypes...), VarTypes...> Super;
+
+public:
+	/**
+	 * Creates and initializes a new instance.
+	 *
+	 * @param InUserObject The UObject to call the function on.
+	 * @param InFunctionName The name of the function call.
+	 */
+	TBaseUFunctionDelegateInstance(UserClass* InUserObject, const FName& InFunctionName, VarTypes... Vars)
+		: Super(InUserObject, InFunctionName, Vars...)
+	{
+	}
+
+	virtual bool ExecuteIfSafe(ParamTypes... Params) const override
+	{
+		if (Super::IsSafeToExecute())
+		{
+			Super::Execute(Params...);
+
+			return true;
+		}
+
+		return false;
+	}
+};
+
+
+/* Delegate binding types
+ *****************************************************************************/
+
+/**
+ * Implements a delegate binding for shared pointer member functions.
+ */
+template <bool bConst, class UserClass, ESPMode::Type SPMode, typename FuncType, typename... VarTypes>
+class TBaseSPMethodDelegateInstance;
+
+template <bool bConst, class UserClass, ESPMode::Type SPMode, typename WrappedRetValType, typename... ParamTypes, typename... VarTypes>
+class TBaseSPMethodDelegateInstance<bConst, UserClass, SPMode, WrappedRetValType (ParamTypes...), VarTypes...> : public IBaseDelegateInstance<typename TUnwrapType<WrappedRetValType>::Type (ParamTypes...)>
+{
+	typedef IBaseDelegateInstance<typename TUnwrapType<WrappedRetValType>::Type (ParamTypes...)> Super;
+
+public:
+	typedef typename TUnwrapType<WrappedRetValType>::Type RetValType;
+	typedef typename TMemFunPtrType<bConst, UserClass, RetValType (ParamTypes..., VarTypes...)>::Type FMethodPtr;
+
+	TBaseSPMethodDelegateInstance(const TSharedPtr<UserClass, SPMode>& InUserObject, FMethodPtr InMethodPtr, VarTypes... Vars)
+		: UserObject(InUserObject)
+		, MethodPtr (InMethodPtr)
+		, Payload   (Vars...)
+	{
+		// NOTE: Shared pointer delegates are allowed to have a null incoming object pointer.  Weak pointers can expire,
+		//       an it is possible for a copy of a delegate instance to end up with a null pointer.
+		checkSlow(MethodPtr != nullptr);
+	}
+
+	// IDelegateInstance interface
+
+	virtual FName GetFunctionName() const override
+	{
+		return NAME_None;
+	}
+
+	virtual const void* GetRawMethodPtr() const override
+	{
+		return GetRawMethodPtrInternal();
+	}
+
+	virtual const void* GetRawUserObject() const override
+	{
+		return GetRawUserObjectInternal();
+	}
+
+	virtual EDelegateInstanceType::Type GetType() const override
+	{
+		return SPMode == ESPMode::ThreadSafe ? EDelegateInstanceType::ThreadSafeSharedPointerMethod : EDelegateInstanceType::SharedPointerMethod;
+	}
+
+	virtual bool HasSameObject(const void* InUserObject) const override
+	{
+		return UserObject.HasSameObject(InUserObject);
+	}
+
+	virtual bool IsSafeToExecute() const override
+	{
+		return UserObject.IsValid();
+	}
+
+public:
+
+	// IBaseDelegateInstance interface
+
+	virtual Super* CreateCopy() override
+	{
+		return Payload.template ApplyAfter_ExplicitReturnType<Super*>(static_cast<Super*(*)(const TSharedPtr<UserClass, SPMode>&, FMethodPtr, VarTypes...)>(&Create), UserObject.Pin(), MethodPtr);
+	}
+
+	virtual RetValType Execute(ParamTypes... Params) const override
+	{
+		typedef typename TRemoveConst<UserClass>::Type MutableUserClass;
+
+		// Verify that the user object is still valid.  We only have a weak reference to it.
+		auto SharedUserObject = UserObject.Pin();
+		checkSlow(SharedUserObject.IsValid());
+
+		// Safely remove const to work around a compiler issue with instantiating template permutations for 
+		// overloaded functions that take a function pointer typedef as a member of a templated class.  In
+		// all cases where this code is actually invoked, the UserClass will already be a const pointer.
+		auto MutableUserObject = const_cast<MutableUserClass*>(SharedUserObject.Get());
+
+		// Call the member function on the user's object.  And yes, this is the correct C++ syntax for calling a
+		// pointer-to-member function.
+		checkSlow(MethodPtr != nullptr);
+
+		return Payload.template ApplyAfter_ExplicitReturnType<RetValType>(TMemberFunctionCaller<MutableUserClass, FMethodPtr>(MutableUserObject, MethodPtr), Params...);
+	}
+
+	virtual bool IsSameFunction( const Super& InOtherDelegate ) const override
+	{
+		// NOTE: Payload data is not currently considered when comparing delegate instances.
+		// See the comment in multi-cast delegate's Remove() method for more information.
+		if ((InOtherDelegate.GetType() == EDelegateInstanceType::SharedPointerMethod) || 
+			(InOtherDelegate.GetType() == EDelegateInstanceType::ThreadSafeSharedPointerMethod) ||
+			(InOtherDelegate.GetType() == EDelegateInstanceType::RawMethod) ||
+			(InOtherDelegate.GetType() == EDelegateInstanceType::UObjectMethod))
+		{
+			return (GetRawMethodPtrInternal() == InOtherDelegate.GetRawMethodPtr() && UserObject.HasSameObject(InOtherDelegate.GetRawUserObject()));
+		}
+
+		return false;
+	}
+
+public:
+
+	/**
+	 * Creates a new shared pointer delegate binding for the given user object and method pointer.
+	 *
+	 * @param InUserObjectRef Shared reference to the user's object that contains the class method.
+	 * @param InFunc Member function pointer to your class method.
+	 * @return The new delegate.
+	 */
+	FORCEINLINE static Super* Create(const TSharedPtr<UserClass, SPMode>& InUserObjectRef, FMethodPtr InFunc, VarTypes... Vars)
+	{
+		return new TBaseSPMethodDelegateInstance<bConst, UserClass, SPMode, RetValType (ParamTypes...), VarTypes...>(InUserObjectRef, InFunc, Vars...);
+	}
+
+	/**
+	 * Creates a new shared pointer delegate binding for the given user object and method pointer.
+	 *
+	 * This overload requires that the supplied object derives from TSharedFromThis.
+	 *
+	 * @param InUserObject  The user's object that contains the class method.  Must derive from TSharedFromThis.
+	 * @param InFunc  Member function pointer to your class method.
+	 * @return The new delegate.
+	 */
+	FORCEINLINE static Super* Create(UserClass* InUserObject, FMethodPtr InFunc, VarTypes... Vars)
+	{
+		// We expect the incoming InUserObject to derived from TSharedFromThis.
+		auto UserObjectRef = StaticCastSharedRef<UserClass>(InUserObject->AsShared());
+		return Create(UserObjectRef, InFunc, Vars...);
+	}
+
+protected:
+
+	/**
+	 * Internal, inlined and non-virtual version of GetRawUserObject interface.
+	 */
+	FORCEINLINE const void* GetRawUserObjectInternal( ) const
+	{
+		return UserObject.Pin().Get();
+	}
+
+	/**
+	 * Internal, inlined and non-virtual version of GetRawMethod interface.
+	 */
+	FORCEINLINE const void* GetRawMethodPtrInternal( ) const
+	{
+		return *(const void**)&MethodPtr;
+	}
+
+	// Weak reference to an instance of the user's class which contains a method we would like to call.
+	TWeakPtr<UserClass, SPMode> UserObject;
+
+	// C++ member function pointer.
+	FMethodPtr MethodPtr;
+
+	// Payload member variables, if any.
+	TTuple<VarTypes...> Payload;
+};
+
+template <bool bConst, class UserClass, ESPMode::Type SPMode, typename... ParamTypes, typename... VarTypes>
+class TBaseSPMethodDelegateInstance<bConst, UserClass, SPMode, void (ParamTypes...), VarTypes...> : public TBaseSPMethodDelegateInstance<bConst, UserClass, SPMode, TTypeWrapper<void> (ParamTypes...), VarTypes...>
+{
+	typedef TBaseSPMethodDelegateInstance<bConst, UserClass, SPMode, TTypeWrapper<void> (ParamTypes...), VarTypes...> Super;
+
+public:
+	/**
+	 * Creates and initializes a new instance.
+	 *
+	 * @param InUserObject A shared reference to an arbitrary object (templated) that hosts the member function.
+	 * @param InMethodPtr C++ member function pointer for the method to bind.
+	 */
+	TBaseSPMethodDelegateInstance(const TSharedPtr<UserClass, SPMode>& InUserObject, typename Super::FMethodPtr InMethodPtr, VarTypes... Vars)
+		: Super(InUserObject, InMethodPtr, Vars...)
+	{
+	}
+
+	virtual bool ExecuteIfSafe(ParamTypes... Params) const override
+	{
+		// Verify that the user object is still valid.  We only have a weak reference to it.
+		auto SharedUserObject = Super::UserObject.Pin();
+		if (SharedUserObject.IsValid())
+		{
+			Super::Execute(Params...);
+
+			return true;
+		}
+
+		return false;
+	}
+};
+
+
+/**
+ * Implements a delegate binding for C++ member functions.
+ */
+template <bool bConst, class UserClass, typename FuncType, typename... VarTypes>
+class TBaseRawMethodDelegateInstance;
+
+template <bool bConst, class UserClass, typename WrappedRetValType, typename... ParamTypes, typename... VarTypes>
+class TBaseRawMethodDelegateInstance<bConst, UserClass, WrappedRetValType (ParamTypes...), VarTypes...> : public IBaseDelegateInstance<typename TUnwrapType<WrappedRetValType>::Type (ParamTypes...)>
+{
+	static_assert((!CanConvertPointerFromTo<UserClass, const UObjectBase>::Result), "You cannot use raw method delegates with UObjects.");
+
+	typedef IBaseDelegateInstance<typename TUnwrapType<WrappedRetValType>::Type (ParamTypes...)> Super;
+
+public:
+	typedef typename TUnwrapType<WrappedRetValType>::Type RetValType;
+	typedef typename TMemFunPtrType<bConst, UserClass, RetValType (ParamTypes..., VarTypes...)>::Type FMethodPtr;
+
+	/**
+	 * Creates and initializes a new instance.
+	 *
+	 * @param InUserObject An arbitrary object (templated) that hosts the member function.
+	 * @param InMethodPtr C++ member function pointer for the method to bind.
+	 */
+	TBaseRawMethodDelegateInstance(UserClass* InUserObject, FMethodPtr InMethodPtr, VarTypes... Vars)
+		: UserObject(InUserObject)
+		, MethodPtr (InMethodPtr)
+		, Payload   (Vars...)
+	{
+		// Non-expirable delegates must always have a non-null object pointer on creation (otherwise they could never execute.)
+		check(InUserObject != nullptr && MethodPtr != nullptr);
+	}
+
+	// IDelegateInstance interface
+
+	virtual FName GetFunctionName( ) const override
+	{
+		return NAME_None;
+	}
+
+	virtual const void* GetRawMethodPtr( ) const override
+	{
+		return GetRawMethodPtrInternal();
+	}
+
+	virtual const void* GetRawUserObject( ) const override
+	{
+		return GetRawUserObjectInternal();
+	}
+
+	virtual EDelegateInstanceType::Type GetType( ) const override
+	{
+		return EDelegateInstanceType::RawMethod;
+	}
+
+	virtual bool HasSameObject( const void* InUserObject ) const override
+	{
+		return UserObject == InUserObject;
+	}
+
+	virtual bool IsSafeToExecute( ) const override
+	{
+		// We never know whether or not it is safe to deference a C++ pointer, but we have to
+		// trust the user in this case.  Prefer using a shared-pointer based delegate type instead!
+		return true;
+	}
+
+public:
+
+	// IBaseDelegateInstance interface
+
+	virtual Super* CreateCopy( ) override
+	{
+		return Payload.template ApplyAfter_ExplicitReturnType<Super*>(&Create, UserObject, MethodPtr);
+	}
+
+	virtual RetValType Execute(ParamTypes... Params) const override
+	{
+		typedef typename TRemoveConst<UserClass>::Type MutableUserClass;
+
+		// Safely remove const to work around a compiler issue with instantiating template permutations for 
+		// overloaded functions that take a function pointer typedef as a member of a templated class.  In
+		// all cases where this code is actually invoked, the UserClass will already be a const pointer.
+		auto MutableUserObject = const_cast<MutableUserClass*>(UserObject);
+
+		// Call the member function on the user's object.  And yes, this is the correct C++ syntax for calling a
+		// pointer-to-member function.
+		checkSlow(MethodPtr != nullptr);
+
+		return Payload.template ApplyAfter_ExplicitReturnType<RetValType>(TMemberFunctionCaller<MutableUserClass, FMethodPtr>(MutableUserObject, MethodPtr), Params...);
+	}
+
+	virtual bool IsSameFunction( const Super& InOtherDelegate ) const override
+	{
+		// NOTE: Payload data is not currently considered when comparing delegate instances.
+		// See the comment in multi-cast delegate's Remove() method for more information.
+		if ((InOtherDelegate.GetType() == EDelegateInstanceType::RawMethod) || 
+			(InOtherDelegate.GetType() == EDelegateInstanceType::UObjectMethod) ||
+			(InOtherDelegate.GetType() == EDelegateInstanceType::SharedPointerMethod) ||
+			(InOtherDelegate.GetType() == EDelegateInstanceType::ThreadSafeSharedPointerMethod))
+		{
+			return (GetRawMethodPtrInternal() == InOtherDelegate.GetRawMethodPtr()) && (UserObject == InOtherDelegate.GetRawUserObject());
+		}
+
+		return false;
+	}
+
+public:
+
+	/**
+	 * Creates a new raw method delegate binding for the given user object and function pointer.
+	 *
+	 * @param InUserObject User's object that contains the class method.
+	 * @param InFunc Member function pointer to your class method.
+	 * @return The new delegate.
+	 */
+	FORCEINLINE static Super* Create(UserClass* InUserObject, FMethodPtr InFunc, VarTypes... Vars)
+	{
+		return new TBaseRawMethodDelegateInstance<bConst, UserClass, RetValType (ParamTypes...), VarTypes...>(InUserObject, InFunc, Vars...);
+	}
+
+protected:
+
+	/**
+	 * Internal, inlined and non-virtual version of GetRawUserObject interface.
+	 */
+	FORCEINLINE const void* GetRawUserObjectInternal( ) const
+	{
+		return UserObject;
+	}
+
+	/**
+	 * Internal, inlined and non-virtual version of GetRawMethodPtr interface.
+	 */
+	FORCEINLINE const void* GetRawMethodPtrInternal( ) const
+	{
+		return *(const void**)&MethodPtr;
+	}
+
+	// Pointer to the user's class which contains a method we would like to call.
+	UserClass* UserObject;
+
+	// C++ member function pointer.
+	FMethodPtr MethodPtr;
+
+	// Payload member variables (if any).
+	TTuple<VarTypes...> Payload;
+};
+
+template <bool bConst, class UserClass, typename... ParamTypes, typename... VarTypes>
+class TBaseRawMethodDelegateInstance<bConst, UserClass, void (ParamTypes...), VarTypes...> : public TBaseRawMethodDelegateInstance<bConst, UserClass, TTypeWrapper<void> (ParamTypes...), VarTypes...>
+{
+	typedef TBaseRawMethodDelegateInstance<bConst, UserClass, TTypeWrapper<void> (ParamTypes...), VarTypes...> Super;
+
+public:
+	/**
+	 * Creates and initializes a new instance.
+	 *
+	 * @param InUserObject An arbitrary object (templated) that hosts the member function.
+	 * @param InMethodPtr C++ member function pointer for the method to bind.
+	 */
+	TBaseRawMethodDelegateInstance(UserClass* InUserObject, typename Super::FMethodPtr InMethodPtr, VarTypes... Vars)
+		: Super(InUserObject, InMethodPtr, Vars...)
+	{
+	}
+
+	virtual bool ExecuteIfSafe(ParamTypes... Params) const override
+	{
+		// We never know whether or not it is safe to deference a C++ pointer, but we have to
+		// trust the user in this case.  Prefer using a shared-pointer based delegate type instead!
+		Super::Execute(Params...);
+
+		return true;
+	}
+};
+
+
+/**
+ * Implements a delegate binding for UObject methods.
+ */
+template <bool bConst, class UserClass, typename FuncType, typename... VarTypes>
+class TBaseUObjectMethodDelegateInstance;
+
+template <bool bConst, class UserClass, typename WrappedRetValType, typename... ParamTypes, typename... VarTypes>
+class TBaseUObjectMethodDelegateInstance<bConst, UserClass, WrappedRetValType (ParamTypes...), VarTypes...> : public IBaseDelegateInstance<typename TUnwrapType<WrappedRetValType>::Type (ParamTypes...)>
+{
+	typedef IBaseDelegateInstance<typename TUnwrapType<WrappedRetValType>::Type (ParamTypes...)> Super;
+
+	static_assert((CanConvertPointerFromTo<UserClass, const UObjectBase>::Result), "You cannot use UObject method delegates with raw pointers.");
+
+public:
+	typedef typename TUnwrapType<WrappedRetValType>::Type RetValType;
+	typedef typename TMemFunPtrType<bConst, UserClass, RetValType (ParamTypes..., VarTypes...)>::Type FMethodPtr;
+
+	TBaseUObjectMethodDelegateInstance(UserClass* InUserObject, FMethodPtr InMethodPtr, VarTypes... Vars)
+		: UserObject(InUserObject)
+		, MethodPtr (InMethodPtr)
+		, Payload   (Vars...)
+	{
+		// NOTE: UObject delegates are allowed to have a null incoming object pointer.  UObject weak pointers can expire,
+		//       an it is possible for a copy of a delegate instance to end up with a null pointer.
+		checkSlow(MethodPtr != nullptr);
+	}
+
+	// IDelegateInstance interface
+
+	virtual FName GetFunctionName( ) const override
+	{
+		return NAME_None;
+	}
+
+	virtual const void* GetRawMethodPtr( ) const override
+	{
+		return GetRawMethodPtrInternal();
+	}
+
+	virtual const void* GetRawUserObject( ) const override
+	{
+		return GetRawUserObjectInternal();
+	}
+
+	virtual EDelegateInstanceType::Type GetType( ) const override
+	{
+		return EDelegateInstanceType::UObjectMethod;
+	}
+
+	virtual bool HasSameObject( const void* InUserObject ) const override
+	{
+		return (UserObject.Get() == InUserObject);
+	}
+
+	virtual bool IsCompactable( ) const override
+	{
+		return !UserObject.Get(true);
+	}
+
+	virtual bool IsSafeToExecute( ) const override
+	{
+		return !!UserObject.Get();
+	}
+
+public:
+
+	// IBaseDelegateInstance interface
+
+	virtual Super* CreateCopy( ) override
+	{
+		return Payload.template ApplyAfter_ExplicitReturnType<Super*>(&Create, UserObject.Get(), MethodPtr);
+	}
+
+	virtual RetValType Execute(ParamTypes... Params) const override
+	{
+		typedef typename TRemoveConst<UserClass>::Type MutableUserClass;
+
+		// Verify that the user object is still valid.  We only have a weak reference to it.
+		checkSlow(UserObject.IsValid());
+
+		// Safely remove const to work around a compiler issue with instantiating template permutations for 
+		// overloaded functions that take a function pointer typedef as a member of a templated class.  In
+		// all cases where this code is actually invoked, the UserClass will already be a const pointer.
+		auto MutableUserObject = const_cast<MutableUserClass*>(UserObject.Get());
+
+		// Call the member function on the user's object.  And yes, this is the correct C++ syntax for calling a
+		// pointer-to-member function.
+		checkSlow(MethodPtr != nullptr);
+
+		return Payload.template ApplyAfter_ExplicitReturnType<RetValType>(TMemberFunctionCaller<MutableUserClass, FMethodPtr>(MutableUserObject, MethodPtr), Params...);
+	}
+
+	virtual bool IsSameFunction( const Super& InOtherDelegate ) const override
+	{
+		// NOTE: Payload data is not currently considered when comparing delegate instances.
+		// See the comment in multi-cast delegate's Remove() method for more information.
+		if ((InOtherDelegate.GetType() == EDelegateInstanceType::UObjectMethod) || 
+			(InOtherDelegate.GetType() == EDelegateInstanceType::RawMethod) ||
+			(InOtherDelegate.GetType() == EDelegateInstanceType::SharedPointerMethod) ||
+			(InOtherDelegate.GetType() == EDelegateInstanceType::ThreadSafeSharedPointerMethod))
+		{
+			return (GetRawMethodPtrInternal() == InOtherDelegate.GetRawMethodPtr()) && (UserObject.Get() == InOtherDelegate.GetRawUserObject());
+		}
+
+		return false;
+	}
+
+public:
+
+	/**
+	 * Creates a new UObject delegate binding for the given user object and method pointer.
+	 *
+	 * @param InUserObject User's object that contains the class method.
+	 * @param InFunc Member function pointer to your class method.
+	 * @return The new delegate.
+	 */
+	FORCEINLINE static Super* Create(UserClass* InUserObject, FMethodPtr InFunc, VarTypes... Vars)
+	{
+		return new TBaseUObjectMethodDelegateInstance<bConst, UserClass, RetValType (ParamTypes...), VarTypes...>(InUserObject, InFunc, Vars...);
+	}
+
+protected:
+
+	/**
+	 * Internal, inlined and non-virtual version of GetRawUserObject interface.
+	 */
+	FORCEINLINE const void* GetRawUserObjectInternal( ) const
+	{
+		return UserObject.Get();
+	}
+
+	/**
+	 * Internal, inlined and non-virtual version of GetRawMethodPtr interface.
+	 */
+	FORCEINLINE const void* GetRawMethodPtrInternal( ) const
+	{
+		return *(const void**)&MethodPtr;
+	}
+
+	// Pointer to the user's class which contains a method we would like to call.
+	TWeakObjectPtr<UserClass> UserObject;
+
+	// C++ member function pointer.
+	FMethodPtr MethodPtr;
+
+	// Payload member variables (if any).
+	TTuple<VarTypes...> Payload;
+};
+
+template <bool bConst, class UserClass, typename... ParamTypes, typename... VarTypes>
+class TBaseUObjectMethodDelegateInstance<bConst, UserClass, void (ParamTypes...), VarTypes...> : public TBaseUObjectMethodDelegateInstance<bConst, UserClass, TTypeWrapper<void> (ParamTypes...), VarTypes...>
+{
+	typedef TBaseUObjectMethodDelegateInstance<bConst, UserClass, TTypeWrapper<void> (ParamTypes...), VarTypes...> Super;
+
+public:
+	/**
+	 * Creates and initializes a new instance.
+	 *
+	 * @param InUserObject An arbitrary object (templated) that hosts the member function.
+	 * @param InMethodPtr C++ member function pointer for the method to bind.
+	 */
+	TBaseUObjectMethodDelegateInstance(UserClass* InUserObject, typename Super::FMethodPtr InMethodPtr, VarTypes... Vars)
+		: Super(InUserObject, InMethodPtr, Vars...)
+	{
+	}
+
+	virtual bool ExecuteIfSafe(ParamTypes... Params) const override
+	{
+		// Verify that the user object is still valid.  We only have a weak reference to it.
+		auto ActualUserObject = Super::UserObject.Get();
+		if (ActualUserObject != nullptr)
+		{
+			Super::Execute(Params...);
+
+			return true;
+		}
+		return false;
+	}
+};
+
+
+/**
+ * Implements a delegate binding for regular C++ functions.
+ */
+template <typename FuncType, typename... VarTypes>
+class TBaseStaticDelegateInstance;
+
+template <typename WrappedRetValType, typename... ParamTypes, typename... VarTypes>
+class TBaseStaticDelegateInstance<WrappedRetValType (ParamTypes...), VarTypes...> : public IBaseDelegateInstance<typename TUnwrapType<WrappedRetValType>::Type (ParamTypes...)>
+{
+	typedef IBaseDelegateInstance<typename TUnwrapType<WrappedRetValType>::Type (ParamTypes...)> Super;
+
+public:
+	typedef typename TUnwrapType<WrappedRetValType>::Type RetValType;
+	typedef RetValType (*FFuncPtr)(ParamTypes..., VarTypes...);
+
+	TBaseStaticDelegateInstance(FFuncPtr InStaticFuncPtr, VarTypes... Vars)
+		: StaticFuncPtr(InStaticFuncPtr)
+		, Payload      (Vars...)
+	{
+		check(StaticFuncPtr != nullptr);
+	}
+
+	// IDelegateInstance interface
+
+	virtual FName GetFunctionName( ) const override
+	{
+		return NAME_None;
+	}
+
+	virtual const void* GetRawMethodPtr( ) const override
+	{
+		return *(const void**)&StaticFuncPtr;
+	}
+
+	virtual const void* GetRawUserObject( ) const override
+	{
+		return nullptr;
+	}
+
+	virtual EDelegateInstanceType::Type GetType( ) const override
+	{
+		return EDelegateInstanceType::Raw;
+	}
+
+	virtual bool HasSameObject( const void* UserObject ) const override
+	{
+		// Raw Delegates aren't bound to an object so they can never match
+		return false;
+	}
+
+	virtual bool IsSafeToExecute( ) const override
+	{
+		// Static functions are always safe to execute!
+		return true;
+	}
+
+public:
+
+	// IBaseDelegateInstance interface
+
+	virtual Super* CreateCopy( ) override
+	{
+		return Payload.template ApplyAfter_ExplicitReturnType<Super*>(&Create, StaticFuncPtr);
+	}
+
+	virtual RetValType Execute(ParamTypes... Params) const override
+	{
+		// Call the member function on the user's object.  And yes, this is the correct C++ syntax for calling a
+		// pointer-to-member function.
+		checkSlow(StaticFuncPtr != nullptr);
+
+		return Payload.template ApplyAfter_ExplicitReturnType<RetValType>(StaticFuncPtr, Params...);
+	}
+
+	virtual bool IsSameFunction( const Super& InOtherDelegate ) const override
+	{
+		// NOTE: Payload data is not currently considered when comparing delegate instances.
+		// See the comment in multi-cast delegate's Remove() method for more information.
+		if (InOtherDelegate.GetType() == EDelegateInstanceType::Raw)
+		{
+			// Downcast to our delegate type and compare
+			const auto& OtherStaticDelegate = static_cast<const TBaseStaticDelegateInstance<RetValType (ParamTypes...), VarTypes...>&>(InOtherDelegate);
+
+			return (StaticFuncPtr == OtherStaticDelegate.StaticFuncPtr);
+		}
+
+		return false;
+	}
+
+public:
+
+	/**
+	 * Creates a new static function delegate binding for the given function pointer.
+	 *
+	 * @param InFunc Static function pointer.
+	 * @return The new delegate.
+	 */
+	FORCEINLINE static Super* Create(FFuncPtr InFunc, VarTypes... Vars)
+	{
+		return new TBaseStaticDelegateInstance<RetValType (ParamTypes...), VarTypes...>(InFunc, Vars...);
+	}
+
+private:
+
+	// C++ function pointer.
+	FFuncPtr StaticFuncPtr;
+
+	// Payload member variables, if any.
+	TTuple<VarTypes...> Payload;
+};
+
+template <typename... ParamTypes, typename... VarTypes>
+class TBaseStaticDelegateInstance<void (ParamTypes...), VarTypes...> : public TBaseStaticDelegateInstance<TTypeWrapper<void> (ParamTypes...), VarTypes...>
+{
+	typedef TBaseStaticDelegateInstance<TTypeWrapper<void> (ParamTypes...), VarTypes...> Super;
+
+public:
+	/**
+	 * Creates and initializes a new instance.
+	 *
+	 * @param InStaticFuncPtr C++ function pointer.
+	 */
+	TBaseStaticDelegateInstance(typename Super::FFuncPtr InStaticFuncPtr, VarTypes... Vars)
+		: Super(InStaticFuncPtr, Vars...)
+	{
+	}
+
+	virtual bool ExecuteIfSafe(ParamTypes... Params) const override
+	{
+		Super::Execute(Params...);
+
+		return true;
+	}
+};
diff --git a/Engine/Source/Runtime/Core/Public/Delegates/DelegateSignatureImpl.inl b/Engine/Source/Runtime/Core/Public/Delegates/DelegateSignatureImpl.inl
--- a/Engine/Source/Runtime/Core/Public/Delegates/DelegateSignatureImpl.inl
+++ b/Engine/Source/Runtime/Core/Public/Delegates/DelegateSignatureImpl.inl
@@ -173,29 +173,29 @@ public:
 	/**
 	 * Static: Creates a raw C++ pointer global function delegate
 	 */
-	inline static DELEGATE_CLASS< FUNC_TEMPLATE_ARGS > CreateStatic( typename FStaticDelegate::FFuncPtr InFunc )
+	inline static DELEGATE_CLASS CreateStatic( typename FStaticDelegate::FFuncPtr InFunc )
 	{
-		return DELEGATE_CLASS< FUNC_TEMPLATE_ARGS >( FStaticDelegate::Create( InFunc ) );
+		return DELEGATE_CLASS( FStaticDelegate::Create( InFunc ) );
 	}
 	template< PAYLOAD_TEMPLATE_DECL_OneVar >
-	inline static DELEGATE_CLASS< FUNC_TEMPLATE_ARGS > CreateStatic( typename TStaticDelegate_OneVar< PAYLOAD_TEMPLATE_LIST_OneVar >::FFuncPtr InFunc, PAYLOAD_TEMPLATE_ARGS_OneVar )
+	inline static DELEGATE_CLASS CreateStatic( typename TStaticDelegate_OneVar< PAYLOAD_TEMPLATE_LIST_OneVar >::FFuncPtr InFunc, PAYLOAD_TEMPLATE_ARGS_OneVar )
 	{
-		return DELEGATE_CLASS< FUNC_TEMPLATE_ARGS >( TStaticDelegate_OneVar< PAYLOAD_TEMPLATE_LIST_OneVar >::Create( InFunc, PAYLOAD_TEMPLATE_PASSIN_OneVar ) );
+		return DELEGATE_CLASS( TStaticDelegate_OneVar< PAYLOAD_TEMPLATE_LIST_OneVar >::Create( InFunc, PAYLOAD_TEMPLATE_PASSIN_OneVar ) );
 	}
 	template< PAYLOAD_TEMPLATE_DECL_TwoVars >
-	inline static DELEGATE_CLASS< FUNC_TEMPLATE_ARGS > CreateStatic( typename TStaticDelegate_TwoVars< PAYLOAD_TEMPLATE_LIST_TwoVars >::FFuncPtr InFunc, PAYLOAD_TEMPLATE_ARGS_TwoVars )
+	inline static DELEGATE_CLASS CreateStatic( typename TStaticDelegate_TwoVars< PAYLOAD_TEMPLATE_LIST_TwoVars >::FFuncPtr InFunc, PAYLOAD_TEMPLATE_ARGS_TwoVars )
 	{
-		return DELEGATE_CLASS< FUNC_TEMPLATE_ARGS >( TStaticDelegate_TwoVars< PAYLOAD_TEMPLATE_LIST_TwoVars >::Create( InFunc, PAYLOAD_TEMPLATE_PASSIN_TwoVars ) );
+		return DELEGATE_CLASS( TStaticDelegate_TwoVars< PAYLOAD_TEMPLATE_LIST_TwoVars >::Create( InFunc, PAYLOAD_TEMPLATE_PASSIN_TwoVars ) );
 	}
 	template< PAYLOAD_TEMPLATE_DECL_ThreeVars >
-	inline static DELEGATE_CLASS< FUNC_TEMPLATE_ARGS > CreateStatic( typename TStaticDelegate_ThreeVars< PAYLOAD_TEMPLATE_LIST_ThreeVars >::FFuncPtr InFunc, PAYLOAD_TEMPLATE_ARGS_ThreeVars )
+	inline static DELEGATE_CLASS CreateStatic( typename TStaticDelegate_ThreeVars< PAYLOAD_TEMPLATE_LIST_ThreeVars >::FFuncPtr InFunc, PAYLOAD_TEMPLATE_ARGS_ThreeVars )
 	{
-		return DELEGATE_CLASS< FUNC_TEMPLATE_ARGS >( TStaticDelegate_ThreeVars< PAYLOAD_TEMPLATE_LIST_ThreeVars >::Create( InFunc, PAYLOAD_TEMPLATE_PASSIN_ThreeVars ) );
+		return DELEGATE_CLASS( TStaticDelegate_ThreeVars< PAYLOAD_TEMPLATE_LIST_ThreeVars >::Create( InFunc, PAYLOAD_TEMPLATE_PASSIN_ThreeVars ) );
 	}
 	template< PAYLOAD_TEMPLATE_DECL_FourVars >
-	inline static DELEGATE_CLASS< FUNC_TEMPLATE_ARGS > CreateStatic( typename TStaticDelegate_FourVars< PAYLOAD_TEMPLATE_LIST_FourVars >::FFuncPtr InFunc, PAYLOAD_TEMPLATE_ARGS_FourVars )
+	inline static DELEGATE_CLASS CreateStatic( typename TStaticDelegate_FourVars< PAYLOAD_TEMPLATE_LIST_FourVars >::FFuncPtr InFunc, PAYLOAD_TEMPLATE_ARGS_FourVars )
 	{
-		return DELEGATE_CLASS< FUNC_TEMPLATE_ARGS >( TStaticDelegate_FourVars< PAYLOAD_TEMPLATE_LIST_FourVars >::Create( InFunc, PAYLOAD_TEMPLATE_PASSIN_FourVars ) );
+		return DELEGATE_CLASS( TStaticDelegate_FourVars< PAYLOAD_TEMPLATE_LIST_FourVars >::Create( InFunc, PAYLOAD_TEMPLATE_PASSIN_FourVars ) );
 	}
 
 	/**
@@ -205,54 +205,54 @@ public:
 	 * deleted out from underneath your delegate. Be careful when calling Execute()!
 	 */
 	template< class UserClass >
-	inline static DELEGATE_CLASS< FUNC_TEMPLATE_ARGS > CreateRaw( UserClass* InUserObject, typename TRawMethodDelegate< UserClass >::FMethodPtr InFunc )
+	inline static DELEGATE_CLASS CreateRaw( UserClass* InUserObject, typename TRawMethodDelegate< UserClass >::FMethodPtr InFunc )
 	{
-		return DELEGATE_CLASS< FUNC_TEMPLATE_ARGS >( TRawMethodDelegate< UserClass >::Create( InUserObject, InFunc ) );
+		return DELEGATE_CLASS( TRawMethodDelegate< UserClass >::Create( InUserObject, InFunc ) );
 	}
 	template< class UserClass >
-	inline static DELEGATE_CLASS< FUNC_TEMPLATE_ARGS > CreateRaw( UserClass* InUserObject, typename TRawMethodDelegate_Const< UserClass >::FMethodPtr InFunc )
+	inline static DELEGATE_CLASS CreateRaw( UserClass* InUserObject, typename TRawMethodDelegate_Const< UserClass >::FMethodPtr InFunc )
 	{
-		return DELEGATE_CLASS< FUNC_TEMPLATE_ARGS >( TRawMethodDelegate_Const< UserClass >::Create( InUserObject, InFunc ) );
+		return DELEGATE_CLASS( TRawMethodDelegate_Const< UserClass >::Create( InUserObject, InFunc ) );
 	}
 	template< class UserClass, PAYLOAD_TEMPLATE_DECL_OneVar >
-	inline static DELEGATE_CLASS< FUNC_TEMPLATE_ARGS > CreateRaw( UserClass* InUserObject, typename TRawMethodDelegate_OneVar< UserClass, PAYLOAD_TEMPLATE_LIST_OneVar >::FMethodPtr InFunc, PAYLOAD_TEMPLATE_ARGS_OneVar )
+	inline static DELEGATE_CLASS CreateRaw( UserClass* InUserObject, typename TRawMethodDelegate_OneVar< UserClass, PAYLOAD_TEMPLATE_LIST_OneVar >::FMethodPtr InFunc, PAYLOAD_TEMPLATE_ARGS_OneVar )
 	{
-		return DELEGATE_CLASS< FUNC_TEMPLATE_ARGS >( TRawMethodDelegate_OneVar< UserClass, PAYLOAD_TEMPLATE_LIST_OneVar >::Create( InUserObject, InFunc, PAYLOAD_TEMPLATE_PASSIN_OneVar ) );
+		return DELEGATE_CLASS( TRawMethodDelegate_OneVar< UserClass, PAYLOAD_TEMPLATE_LIST_OneVar >::Create( InUserObject, InFunc, PAYLOAD_TEMPLATE_PASSIN_OneVar ) );
 	}
 	template< class UserClass, PAYLOAD_TEMPLATE_DECL_OneVar >
-	inline static DELEGATE_CLASS< FUNC_TEMPLATE_ARGS > CreateRaw( UserClass* InUserObject, typename TRawMethodDelegate_OneVar_Const< UserClass, PAYLOAD_TEMPLATE_LIST_OneVar >::FMethodPtr InFunc, PAYLOAD_TEMPLATE_ARGS_OneVar )
+	inline static DELEGATE_CLASS CreateRaw( UserClass* InUserObject, typename TRawMethodDelegate_OneVar_Const< UserClass, PAYLOAD_TEMPLATE_LIST_OneVar >::FMethodPtr InFunc, PAYLOAD_TEMPLATE_ARGS_OneVar )
 	{
-		return DELEGATE_CLASS< FUNC_TEMPLATE_ARGS >( TRawMethodDelegate_OneVar_Const< UserClass, PAYLOAD_TEMPLATE_LIST_OneVar >::Create( InUserObject, InFunc, PAYLOAD_TEMPLATE_PASSIN_OneVar ) );
+		return DELEGATE_CLASS( TRawMethodDelegate_OneVar_Const< UserClass, PAYLOAD_TEMPLATE_LIST_OneVar >::Create( InUserObject, InFunc, PAYLOAD_TEMPLATE_PASSIN_OneVar ) );
 	}
 	template< class UserClass, PAYLOAD_TEMPLATE_DECL_TwoVars >
-	inline static DELEGATE_CLASS< FUNC_TEMPLATE_ARGS > CreateRaw( UserClass* InUserObject, typename TRawMethodDelegate_TwoVars< UserClass, PAYLOAD_TEMPLATE_LIST_TwoVars >::FMethodPtr InFunc, PAYLOAD_TEMPLATE_ARGS_TwoVars )
+	inline static DELEGATE_CLASS CreateRaw( UserClass* InUserObject, typename TRawMethodDelegate_TwoVars< UserClass, PAYLOAD_TEMPLATE_LIST_TwoVars >::FMethodPtr InFunc, PAYLOAD_TEMPLATE_ARGS_TwoVars )
 	{
-		return DELEGATE_CLASS< FUNC_TEMPLATE_ARGS >( TRawMethodDelegate_TwoVars< UserClass, PAYLOAD_TEMPLATE_LIST_TwoVars >::Create( InUserObject, InFunc, PAYLOAD_TEMPLATE_PASSIN_TwoVars ) );
+		return DELEGATE_CLASS( TRawMethodDelegate_TwoVars< UserClass, PAYLOAD_TEMPLATE_LIST_TwoVars >::Create( InUserObject, InFunc, PAYLOAD_TEMPLATE_PASSIN_TwoVars ) );
 	}
 	template< class UserClass, PAYLOAD_TEMPLATE_DECL_TwoVars >
-	inline static DELEGATE_CLASS< FUNC_TEMPLATE_ARGS > CreateRaw( UserClass* InUserObject, typename TRawMethodDelegate_TwoVars_Const< UserClass, PAYLOAD_TEMPLATE_LIST_TwoVars >::FMethodPtr InFunc, PAYLOAD_TEMPLATE_ARGS_TwoVars )
+	inline static DELEGATE_CLASS CreateRaw( UserClass* InUserObject, typename TRawMethodDelegate_TwoVars_Const< UserClass, PAYLOAD_TEMPLATE_LIST_TwoVars >::FMethodPtr InFunc, PAYLOAD_TEMPLATE_ARGS_TwoVars )
 	{
-		return DELEGATE_CLASS< FUNC_TEMPLATE_ARGS >( TRawMethodDelegate_TwoVars_Const< UserClass, PAYLOAD_TEMPLATE_LIST_TwoVars >::Create( InUserObject, InFunc, PAYLOAD_TEMPLATE_PASSIN_TwoVars ) );
+		return DELEGATE_CLASS( TRawMethodDelegate_TwoVars_Const< UserClass, PAYLOAD_TEMPLATE_LIST_TwoVars >::Create( InUserObject, InFunc, PAYLOAD_TEMPLATE_PASSIN_TwoVars ) );
 	}
 	template< class UserClass, PAYLOAD_TEMPLATE_DECL_ThreeVars >
-	inline static DELEGATE_CLASS< FUNC_TEMPLATE_ARGS > CreateRaw( UserClass* InUserObject, typename TRawMethodDelegate_ThreeVars< UserClass, PAYLOAD_TEMPLATE_LIST_ThreeVars >::FMethodPtr InFunc, PAYLOAD_TEMPLATE_ARGS_ThreeVars )
+	inline static DELEGATE_CLASS CreateRaw( UserClass* InUserObject, typename TRawMethodDelegate_ThreeVars< UserClass, PAYLOAD_TEMPLATE_LIST_ThreeVars >::FMethodPtr InFunc, PAYLOAD_TEMPLATE_ARGS_ThreeVars )
 	{
-		return DELEGATE_CLASS< FUNC_TEMPLATE_ARGS >( TRawMethodDelegate_ThreeVars< UserClass, PAYLOAD_TEMPLATE_LIST_ThreeVars >::Create( InUserObject, InFunc, PAYLOAD_TEMPLATE_PASSIN_ThreeVars ) );
+		return DELEGATE_CLASS( TRawMethodDelegate_ThreeVars< UserClass, PAYLOAD_TEMPLATE_LIST_ThreeVars >::Create( InUserObject, InFunc, PAYLOAD_TEMPLATE_PASSIN_ThreeVars ) );
 	}
 	template< class UserClass, PAYLOAD_TEMPLATE_DECL_ThreeVars >
-	inline static DELEGATE_CLASS< FUNC_TEMPLATE_ARGS > CreateRaw( UserClass* InUserObject, typename TRawMethodDelegate_ThreeVars_Const< UserClass, PAYLOAD_TEMPLATE_LIST_ThreeVars >::FMethodPtr InFunc, PAYLOAD_TEMPLATE_ARGS_ThreeVars )
+	inline static DELEGATE_CLASS CreateRaw( UserClass* InUserObject, typename TRawMethodDelegate_ThreeVars_Const< UserClass, PAYLOAD_TEMPLATE_LIST_ThreeVars >::FMethodPtr InFunc, PAYLOAD_TEMPLATE_ARGS_ThreeVars )
 	{
-		return DELEGATE_CLASS< FUNC_TEMPLATE_ARGS >( TRawMethodDelegate_ThreeVars_Const< UserClass, PAYLOAD_TEMPLATE_LIST_ThreeVars >::Create( InUserObject, InFunc, PAYLOAD_TEMPLATE_PASSIN_ThreeVars ) );
+		return DELEGATE_CLASS( TRawMethodDelegate_ThreeVars_Const< UserClass, PAYLOAD_TEMPLATE_LIST_ThreeVars >::Create( InUserObject, InFunc, PAYLOAD_TEMPLATE_PASSIN_ThreeVars ) );
 	}
 	template< class UserClass, PAYLOAD_TEMPLATE_DECL_FourVars >
-	inline static DELEGATE_CLASS< FUNC_TEMPLATE_ARGS > CreateRaw( UserClass* InUserObject, typename TRawMethodDelegate_FourVars< UserClass, PAYLOAD_TEMPLATE_LIST_FourVars >::FMethodPtr InFunc, PAYLOAD_TEMPLATE_ARGS_FourVars )
+	inline static DELEGATE_CLASS CreateRaw( UserClass* InUserObject, typename TRawMethodDelegate_FourVars< UserClass, PAYLOAD_TEMPLATE_LIST_FourVars >::FMethodPtr InFunc, PAYLOAD_TEMPLATE_ARGS_FourVars )
 	{
-		return DELEGATE_CLASS< FUNC_TEMPLATE_ARGS >( TRawMethodDelegate_FourVars< UserClass, PAYLOAD_TEMPLATE_LIST_FourVars >::Create( InUserObject, InFunc, PAYLOAD_TEMPLATE_PASSIN_FourVars ) );
+		return DELEGATE_CLASS( TRawMethodDelegate_FourVars< UserClass, PAYLOAD_TEMPLATE_LIST_FourVars >::Create( InUserObject, InFunc, PAYLOAD_TEMPLATE_PASSIN_FourVars ) );
 	}
 	template< class UserClass, PAYLOAD_TEMPLATE_DECL_FourVars >
-	inline static DELEGATE_CLASS< FUNC_TEMPLATE_ARGS > CreateRaw( UserClass* InUserObject, typename TRawMethodDelegate_FourVars_Const< UserClass, PAYLOAD_TEMPLATE_LIST_FourVars >::FMethodPtr InFunc, PAYLOAD_TEMPLATE_ARGS_FourVars )
+	inline static DELEGATE_CLASS CreateRaw( UserClass* InUserObject, typename TRawMethodDelegate_FourVars_Const< UserClass, PAYLOAD_TEMPLATE_LIST_FourVars >::FMethodPtr InFunc, PAYLOAD_TEMPLATE_ARGS_FourVars )
 	{
-		return DELEGATE_CLASS< FUNC_TEMPLATE_ARGS >( TRawMethodDelegate_FourVars_Const< UserClass, PAYLOAD_TEMPLATE_LIST_FourVars >::Create( InUserObject, InFunc, PAYLOAD_TEMPLATE_PASSIN_FourVars ) );
+		return DELEGATE_CLASS( TRawMethodDelegate_FourVars_Const< UserClass, PAYLOAD_TEMPLATE_LIST_FourVars >::Create( InUserObject, InFunc, PAYLOAD_TEMPLATE_PASSIN_FourVars ) );
 	}
 	
 	/**
@@ -262,54 +262,54 @@ public:
 	 * You can use ExecuteIfBound() to call them.
 	 */
 	template< class UserClass >
-	inline static DELEGATE_CLASS< FUNC_TEMPLATE_ARGS > CreateSP( const TSharedRef< UserClass, ESPMode::Fast >& InUserObjectRef, typename TSPMethodDelegate< UserClass >::FMethodPtr InFunc )
+	inline static DELEGATE_CLASS CreateSP( const TSharedRef< UserClass, ESPMode::Fast >& InUserObjectRef, typename TSPMethodDelegate< UserClass >::FMethodPtr InFunc )
 	{
-		return DELEGATE_CLASS< FUNC_TEMPLATE_ARGS >( TSPMethodDelegate< UserClass >::Create( InUserObjectRef, InFunc ) );
+		return DELEGATE_CLASS( TSPMethodDelegate< UserClass >::Create( InUserObjectRef, InFunc ) );
 	}
 	template< class UserClass >
-	inline static DELEGATE_CLASS< FUNC_TEMPLATE_ARGS > CreateSP( const TSharedRef< UserClass, ESPMode::Fast >& InUserObjectRef, typename TSPMethodDelegate_Const< UserClass >::FMethodPtr InFunc )
+	inline static DELEGATE_CLASS CreateSP( const TSharedRef< UserClass, ESPMode::Fast >& InUserObjectRef, typename TSPMethodDelegate_Const< UserClass >::FMethodPtr InFunc )
 	{
-		return DELEGATE_CLASS< FUNC_TEMPLATE_ARGS >( TSPMethodDelegate_Const< UserClass >::Create( InUserObjectRef, InFunc ) );
+		return DELEGATE_CLASS( TSPMethodDelegate_Const< UserClass >::Create( InUserObjectRef, InFunc ) );
 	}
 	template< class UserClass, PAYLOAD_TEMPLATE_DECL_OneVar >
-	inline static DELEGATE_CLASS< FUNC_TEMPLATE_ARGS > CreateSP( const TSharedRef< UserClass, ESPMode::Fast >& InUserObjectRef, typename TSPMethodDelegate_OneVar< UserClass, PAYLOAD_TEMPLATE_LIST_OneVar >::FMethodPtr InFunc, PAYLOAD_TEMPLATE_ARGS_OneVar )
+	inline static DELEGATE_CLASS CreateSP( const TSharedRef< UserClass, ESPMode::Fast >& InUserObjectRef, typename TSPMethodDelegate_OneVar< UserClass, PAYLOAD_TEMPLATE_LIST_OneVar >::FMethodPtr InFunc, PAYLOAD_TEMPLATE_ARGS_OneVar )
 	{
-		return DELEGATE_CLASS< FUNC_TEMPLATE_ARGS >( TSPMethodDelegate_OneVar< UserClass, PAYLOAD_TEMPLATE_LIST_OneVar >::Create( InUserObjectRef, InFunc, PAYLOAD_TEMPLATE_PASSIN_OneVar ) );
+		return DELEGATE_CLASS( TSPMethodDelegate_OneVar< UserClass, PAYLOAD_TEMPLATE_LIST_OneVar >::Create( InUserObjectRef, InFunc, PAYLOAD_TEMPLATE_PASSIN_OneVar ) );
 	}
 	template< class UserClass, PAYLOAD_TEMPLATE_DECL_OneVar >
-	inline static DELEGATE_CLASS< FUNC_TEMPLATE_ARGS > CreateSP( const TSharedRef< UserClass, ESPMode::Fast >& InUserObjectRef, typename TSPMethodDelegate_OneVar_Const< UserClass, PAYLOAD_TEMPLATE_LIST_OneVar >::FMethodPtr InFunc, PAYLOAD_TEMPLATE_ARGS_OneVar )
+	inline static DELEGATE_CLASS CreateSP( const TSharedRef< UserClass, ESPMode::Fast >& InUserObjectRef, typename TSPMethodDelegate_OneVar_Const< UserClass, PAYLOAD_TEMPLATE_LIST_OneVar >::FMethodPtr InFunc, PAYLOAD_TEMPLATE_ARGS_OneVar )
 	{
-		return DELEGATE_CLASS< FUNC_TEMPLATE_ARGS >( TSPMethodDelegate_OneVar_Const< UserClass, PAYLOAD_TEMPLATE_LIST_OneVar >::Create( InUserObjectRef, InFunc, PAYLOAD_TEMPLATE_PASSIN_OneVar ) );
+		return DELEGATE_CLASS( TSPMethodDelegate_OneVar_Const< UserClass, PAYLOAD_TEMPLATE_LIST_OneVar >::Create( InUserObjectRef, InFunc, PAYLOAD_TEMPLATE_PASSIN_OneVar ) );
 	}
 	template< class UserClass, PAYLOAD_TEMPLATE_DECL_TwoVars >
-	inline static DELEGATE_CLASS< FUNC_TEMPLATE_ARGS > CreateSP( const TSharedRef< UserClass, ESPMode::Fast >& InUserObjectRef, typename TSPMethodDelegate_TwoVars< UserClass, PAYLOAD_TEMPLATE_LIST_TwoVars >::FMethodPtr InFunc, PAYLOAD_TEMPLATE_ARGS_TwoVars )
+	inline static DELEGATE_CLASS CreateSP( const TSharedRef< UserClass, ESPMode::Fast >& InUserObjectRef, typename TSPMethodDelegate_TwoVars< UserClass, PAYLOAD_TEMPLATE_LIST_TwoVars >::FMethodPtr InFunc, PAYLOAD_TEMPLATE_ARGS_TwoVars )
 	{
-		return DELEGATE_CLASS< FUNC_TEMPLATE_ARGS >( TSPMethodDelegate_TwoVars< UserClass, PAYLOAD_TEMPLATE_LIST_TwoVars >::Create( InUserObjectRef, InFunc, PAYLOAD_TEMPLATE_PASSIN_TwoVars ) );
+		return DELEGATE_CLASS( TSPMethodDelegate_TwoVars< UserClass, PAYLOAD_TEMPLATE_LIST_TwoVars >::Create( InUserObjectRef, InFunc, PAYLOAD_TEMPLATE_PASSIN_TwoVars ) );
 	}
 	template< class UserClass, PAYLOAD_TEMPLATE_DECL_TwoVars >
-	inline static DELEGATE_CLASS< FUNC_TEMPLATE_ARGS > CreateSP( const TSharedRef< UserClass, ESPMode::Fast >& InUserObjectRef, typename TSPMethodDelegate_TwoVars_Const< UserClass, PAYLOAD_TEMPLATE_LIST_TwoVars >::FMethodPtr InFunc, PAYLOAD_TEMPLATE_ARGS_TwoVars )
+	inline static DELEGATE_CLASS CreateSP( const TSharedRef< UserClass, ESPMode::Fast >& InUserObjectRef, typename TSPMethodDelegate_TwoVars_Const< UserClass, PAYLOAD_TEMPLATE_LIST_TwoVars >::FMethodPtr InFunc, PAYLOAD_TEMPLATE_ARGS_TwoVars )
 	{
-		return DELEGATE_CLASS< FUNC_TEMPLATE_ARGS >( TSPMethodDelegate_TwoVars_Const< UserClass, PAYLOAD_TEMPLATE_LIST_TwoVars >::Create( InUserObjectRef, InFunc, PAYLOAD_TEMPLATE_PASSIN_TwoVars ) );
+		return DELEGATE_CLASS( TSPMethodDelegate_TwoVars_Const< UserClass, PAYLOAD_TEMPLATE_LIST_TwoVars >::Create( InUserObjectRef, InFunc, PAYLOAD_TEMPLATE_PASSIN_TwoVars ) );
 	}
 	template< class UserClass, PAYLOAD_TEMPLATE_DECL_ThreeVars >
-	inline static DELEGATE_CLASS< FUNC_TEMPLATE_ARGS > CreateSP( const TSharedRef< UserClass, ESPMode::Fast >& InUserObjectRef, typename TSPMethodDelegate_ThreeVars< UserClass, PAYLOAD_TEMPLATE_LIST_ThreeVars >::FMethodPtr InFunc, PAYLOAD_TEMPLATE_ARGS_ThreeVars )
+	inline static DELEGATE_CLASS CreateSP( const TSharedRef< UserClass, ESPMode::Fast >& InUserObjectRef, typename TSPMethodDelegate_ThreeVars< UserClass, PAYLOAD_TEMPLATE_LIST_ThreeVars >::FMethodPtr InFunc, PAYLOAD_TEMPLATE_ARGS_ThreeVars )
 	{
-		return DELEGATE_CLASS< FUNC_TEMPLATE_ARGS >( TSPMethodDelegate_ThreeVars< UserClass, PAYLOAD_TEMPLATE_LIST_ThreeVars >::Create( InUserObjectRef, InFunc, PAYLOAD_TEMPLATE_PASSIN_ThreeVars ) );
+		return DELEGATE_CLASS( TSPMethodDelegate_ThreeVars< UserClass, PAYLOAD_TEMPLATE_LIST_ThreeVars >::Create( InUserObjectRef, InFunc, PAYLOAD_TEMPLATE_PASSIN_ThreeVars ) );
 	}
 	template< class UserClass, PAYLOAD_TEMPLATE_DECL_ThreeVars >
-	inline static DELEGATE_CLASS< FUNC_TEMPLATE_ARGS > CreateSP( const TSharedRef< UserClass, ESPMode::Fast >& InUserObjectRef, typename TSPMethodDelegate_ThreeVars_Const< UserClass, PAYLOAD_TEMPLATE_LIST_ThreeVars >::FMethodPtr InFunc, PAYLOAD_TEMPLATE_ARGS_ThreeVars )
+	inline static DELEGATE_CLASS CreateSP( const TSharedRef< UserClass, ESPMode::Fast >& InUserObjectRef, typename TSPMethodDelegate_ThreeVars_Const< UserClass, PAYLOAD_TEMPLATE_LIST_ThreeVars >::FMethodPtr InFunc, PAYLOAD_TEMPLATE_ARGS_ThreeVars )
 	{
-		return DELEGATE_CLASS< FUNC_TEMPLATE_ARGS >( TSPMethodDelegate_ThreeVars_Const< UserClass, PAYLOAD_TEMPLATE_LIST_ThreeVars >::Create( InUserObjectRef, InFunc, PAYLOAD_TEMPLATE_PASSIN_ThreeVars ) );
+		return DELEGATE_CLASS( TSPMethodDelegate_ThreeVars_Const< UserClass, PAYLOAD_TEMPLATE_LIST_ThreeVars >::Create( InUserObjectRef, InFunc, PAYLOAD_TEMPLATE_PASSIN_ThreeVars ) );
 	}
 	template< class UserClass, PAYLOAD_TEMPLATE_DECL_FourVars >
-	inline static DELEGATE_CLASS< FUNC_TEMPLATE_ARGS > CreateSP( const TSharedRef< UserClass, ESPMode::Fast >& InUserObjectRef, typename TSPMethodDelegate_FourVars< UserClass, PAYLOAD_TEMPLATE_LIST_FourVars >::FMethodPtr InFunc, PAYLOAD_TEMPLATE_ARGS_FourVars )
+	inline static DELEGATE_CLASS CreateSP( const TSharedRef< UserClass, ESPMode::Fast >& InUserObjectRef, typename TSPMethodDelegate_FourVars< UserClass, PAYLOAD_TEMPLATE_LIST_FourVars >::FMethodPtr InFunc, PAYLOAD_TEMPLATE_ARGS_FourVars )
 	{
-		return DELEGATE_CLASS< FUNC_TEMPLATE_ARGS >( TSPMethodDelegate_FourVars< UserClass, PAYLOAD_TEMPLATE_LIST_FourVars >::Create( InUserObjectRef, InFunc, PAYLOAD_TEMPLATE_PASSIN_FourVars ) );
+		return DELEGATE_CLASS( TSPMethodDelegate_FourVars< UserClass, PAYLOAD_TEMPLATE_LIST_FourVars >::Create( InUserObjectRef, InFunc, PAYLOAD_TEMPLATE_PASSIN_FourVars ) );
 	}
 	template< class UserClass, PAYLOAD_TEMPLATE_DECL_FourVars >
-	inline static DELEGATE_CLASS< FUNC_TEMPLATE_ARGS > CreateSP( const TSharedRef< UserClass, ESPMode::Fast >& InUserObjectRef, typename TSPMethodDelegate_FourVars_Const< UserClass, PAYLOAD_TEMPLATE_LIST_FourVars >::FMethodPtr InFunc, PAYLOAD_TEMPLATE_ARGS_FourVars )
+	inline static DELEGATE_CLASS CreateSP( const TSharedRef< UserClass, ESPMode::Fast >& InUserObjectRef, typename TSPMethodDelegate_FourVars_Const< UserClass, PAYLOAD_TEMPLATE_LIST_FourVars >::FMethodPtr InFunc, PAYLOAD_TEMPLATE_ARGS_FourVars )
 	{
-		return DELEGATE_CLASS< FUNC_TEMPLATE_ARGS >( TSPMethodDelegate_FourVars_Const< UserClass, PAYLOAD_TEMPLATE_LIST_FourVars >::Create( InUserObjectRef, InFunc, PAYLOAD_TEMPLATE_PASSIN_FourVars ) );
+		return DELEGATE_CLASS( TSPMethodDelegate_FourVars_Const< UserClass, PAYLOAD_TEMPLATE_LIST_FourVars >::Create( InUserObjectRef, InFunc, PAYLOAD_TEMPLATE_PASSIN_FourVars ) );
 	}
 
 	/**
@@ -319,52 +319,52 @@ public:
 	 * You can use ExecuteIfBound() to call them.
 	 */
 	template< class UserClass >
-	inline static DELEGATE_CLASS< FUNC_TEMPLATE_ARGS > CreateSP( UserClass* InUserObject, typename TSPMethodDelegate< UserClass >::FMethodPtr InFunc )
+	inline static DELEGATE_CLASS CreateSP( UserClass* InUserObject, typename TSPMethodDelegate< UserClass >::FMethodPtr InFunc )
 	{
 		return CreateSP( StaticCastSharedRef< UserClass >( InUserObject->AsShared() ), InFunc );
 	}
 	template< class UserClass >
-	inline static DELEGATE_CLASS< FUNC_TEMPLATE_ARGS > CreateSP( UserClass* InUserObject, typename TSPMethodDelegate_Const< UserClass >::FMethodPtr InFunc )
+	inline static DELEGATE_CLASS CreateSP( UserClass* InUserObject, typename TSPMethodDelegate_Const< UserClass >::FMethodPtr InFunc )
 	{
 		return CreateSP( StaticCastSharedRef< UserClass >( InUserObject->AsShared() ), InFunc );
 	}
 	template< class UserClass, PAYLOAD_TEMPLATE_DECL_OneVar >
-	inline static DELEGATE_CLASS< FUNC_TEMPLATE_ARGS > CreateSP( UserClass* InUserObject, typename TSPMethodDelegate_OneVar< UserClass, PAYLOAD_TEMPLATE_LIST_OneVar >::FMethodPtr InFunc, PAYLOAD_TEMPLATE_ARGS_OneVar )
+	inline static DELEGATE_CLASS CreateSP( UserClass* InUserObject, typename TSPMethodDelegate_OneVar< UserClass, PAYLOAD_TEMPLATE_LIST_OneVar >::FMethodPtr InFunc, PAYLOAD_TEMPLATE_ARGS_OneVar )
 	{
 		return CreateSP( StaticCastSharedRef< UserClass >( InUserObject->AsShared() ), InFunc, PAYLOAD_TEMPLATE_PASSIN_OneVar );
 	}
 	template< class UserClass, PAYLOAD_TEMPLATE_DECL_OneVar >
-	inline static DELEGATE_CLASS< FUNC_TEMPLATE_ARGS > CreateSP( UserClass* InUserObject, typename TSPMethodDelegate_OneVar_Const< UserClass, PAYLOAD_TEMPLATE_LIST_OneVar >::FMethodPtr InFunc, PAYLOAD_TEMPLATE_ARGS_OneVar )
+	inline static DELEGATE_CLASS CreateSP( UserClass* InUserObject, typename TSPMethodDelegate_OneVar_Const< UserClass, PAYLOAD_TEMPLATE_LIST_OneVar >::FMethodPtr InFunc, PAYLOAD_TEMPLATE_ARGS_OneVar )
 	{
 		return CreateSP( StaticCastSharedRef< UserClass >( InUserObject->AsShared() ), InFunc, PAYLOAD_TEMPLATE_PASSIN_OneVar );
 	}
 	template< class UserClass, PAYLOAD_TEMPLATE_DECL_TwoVars >
-	inline static DELEGATE_CLASS< FUNC_TEMPLATE_ARGS > CreateSP( UserClass* InUserObject, typename TSPMethodDelegate_TwoVars< UserClass, PAYLOAD_TEMPLATE_LIST_TwoVars >::FMethodPtr InFunc, PAYLOAD_TEMPLATE_ARGS_TwoVars )
+	inline static DELEGATE_CLASS CreateSP( UserClass* InUserObject, typename TSPMethodDelegate_TwoVars< UserClass, PAYLOAD_TEMPLATE_LIST_TwoVars >::FMethodPtr InFunc, PAYLOAD_TEMPLATE_ARGS_TwoVars )
 	{
 		return CreateSP( StaticCastSharedRef< UserClass >( InUserObject->AsShared() ), InFunc, PAYLOAD_TEMPLATE_PASSIN_TwoVars );
 	}
 	template< class UserClass, PAYLOAD_TEMPLATE_DECL_TwoVars >
-	inline static DELEGATE_CLASS< FUNC_TEMPLATE_ARGS > CreateSP( UserClass* InUserObject, typename TSPMethodDelegate_TwoVars_Const< UserClass, PAYLOAD_TEMPLATE_LIST_TwoVars >::FMethodPtr InFunc, PAYLOAD_TEMPLATE_ARGS_TwoVars )
+	inline static DELEGATE_CLASS CreateSP( UserClass* InUserObject, typename TSPMethodDelegate_TwoVars_Const< UserClass, PAYLOAD_TEMPLATE_LIST_TwoVars >::FMethodPtr InFunc, PAYLOAD_TEMPLATE_ARGS_TwoVars )
 	{
 		return CreateSP( StaticCastSharedRef< UserClass >( InUserObject->AsShared() ), InFunc, PAYLOAD_TEMPLATE_PASSIN_TwoVars );
 	}
 	template< class UserClass, PAYLOAD_TEMPLATE_DECL_ThreeVars >
-	inline static DELEGATE_CLASS< FUNC_TEMPLATE_ARGS > CreateSP( UserClass* InUserObject, typename TSPMethodDelegate_ThreeVars< UserClass, PAYLOAD_TEMPLATE_LIST_ThreeVars >::FMethodPtr InFunc, PAYLOAD_TEMPLATE_ARGS_ThreeVars )
+	inline static DELEGATE_CLASS CreateSP( UserClass* InUserObject, typename TSPMethodDelegate_ThreeVars< UserClass, PAYLOAD_TEMPLATE_LIST_ThreeVars >::FMethodPtr InFunc, PAYLOAD_TEMPLATE_ARGS_ThreeVars )
 	{
 		return CreateSP( StaticCastSharedRef< UserClass >( InUserObject->AsShared() ), InFunc, PAYLOAD_TEMPLATE_PASSIN_ThreeVars );
 	}
 	template< class UserClass, PAYLOAD_TEMPLATE_DECL_ThreeVars >
-	inline static DELEGATE_CLASS< FUNC_TEMPLATE_ARGS > CreateSP( UserClass* InUserObject, typename TSPMethodDelegate_ThreeVars_Const< UserClass, PAYLOAD_TEMPLATE_LIST_ThreeVars >::FMethodPtr InFunc, PAYLOAD_TEMPLATE_ARGS_ThreeVars )
+	inline static DELEGATE_CLASS CreateSP( UserClass* InUserObject, typename TSPMethodDelegate_ThreeVars_Const< UserClass, PAYLOAD_TEMPLATE_LIST_ThreeVars >::FMethodPtr InFunc, PAYLOAD_TEMPLATE_ARGS_ThreeVars )
 	{
 		return CreateSP( StaticCastSharedRef< UserClass >( InUserObject->AsShared() ), InFunc, PAYLOAD_TEMPLATE_PASSIN_ThreeVars );
 	}
 	template< class UserClass, PAYLOAD_TEMPLATE_DECL_FourVars >
-	inline static DELEGATE_CLASS< FUNC_TEMPLATE_ARGS > CreateSP( UserClass* InUserObject, typename TSPMethodDelegate_FourVars< UserClass, PAYLOAD_TEMPLATE_LIST_FourVars >::FMethodPtr InFunc, PAYLOAD_TEMPLATE_ARGS_FourVars )
+	inline static DELEGATE_CLASS CreateSP( UserClass* InUserObject, typename TSPMethodDelegate_FourVars< UserClass, PAYLOAD_TEMPLATE_LIST_FourVars >::FMethodPtr InFunc, PAYLOAD_TEMPLATE_ARGS_FourVars )
 	{
 		return CreateSP( StaticCastSharedRef< UserClass >( InUserObject->AsShared() ), InFunc, PAYLOAD_TEMPLATE_PASSIN_FourVars );
 	}
 	template< class UserClass, PAYLOAD_TEMPLATE_DECL_FourVars >
-	inline static DELEGATE_CLASS< FUNC_TEMPLATE_ARGS > CreateSP( UserClass* InUserObject, typename TSPMethodDelegate_FourVars_Const< UserClass, PAYLOAD_TEMPLATE_LIST_FourVars >::FMethodPtr InFunc, PAYLOAD_TEMPLATE_ARGS_FourVars )
+	inline static DELEGATE_CLASS CreateSP( UserClass* InUserObject, typename TSPMethodDelegate_FourVars_Const< UserClass, PAYLOAD_TEMPLATE_LIST_FourVars >::FMethodPtr InFunc, PAYLOAD_TEMPLATE_ARGS_FourVars )
 	{
 		return CreateSP( StaticCastSharedRef< UserClass >( InUserObject->AsShared() ), InFunc, PAYLOAD_TEMPLATE_PASSIN_FourVars );
 	}
@@ -376,54 +376,54 @@ public:
 	 * You can use ExecuteIfBound() to call them.
 	 */
 	template< class UserClass >
-	inline static DELEGATE_CLASS< FUNC_TEMPLATE_ARGS > CreateThreadSafeSP( const TSharedRef< UserClass, ESPMode::ThreadSafe >& InUserObjectRef, typename TThreadSafeSPMethodDelegate< UserClass >::FMethodPtr InFunc )
+	inline static DELEGATE_CLASS CreateThreadSafeSP( const TSharedRef< UserClass, ESPMode::ThreadSafe >& InUserObjectRef, typename TThreadSafeSPMethodDelegate< UserClass >::FMethodPtr InFunc )
 	{
-		return DELEGATE_CLASS< FUNC_TEMPLATE_ARGS >( TThreadSafeSPMethodDelegate< UserClass >::Create( InUserObjectRef, InFunc ) );
+		return DELEGATE_CLASS( TThreadSafeSPMethodDelegate< UserClass >::Create( InUserObjectRef, InFunc ) );
 	}
 	template< class UserClass >
-	inline static DELEGATE_CLASS< FUNC_TEMPLATE_ARGS > CreateThreadSafeSP( const TSharedRef< UserClass, ESPMode::ThreadSafe >& InUserObjectRef, typename TThreadSafeSPMethodDelegate_Const< UserClass >::FMethodPtr InFunc )
+	inline static DELEGATE_CLASS CreateThreadSafeSP( const TSharedRef< UserClass, ESPMode::ThreadSafe >& InUserObjectRef, typename TThreadSafeSPMethodDelegate_Const< UserClass >::FMethodPtr InFunc )
 	{
-		return DELEGATE_CLASS< FUNC_TEMPLATE_ARGS >( TThreadSafeSPMethodDelegate_Const< UserClass >::Create( InUserObjectRef, InFunc ) );
+		return DELEGATE_CLASS( TThreadSafeSPMethodDelegate_Const< UserClass >::Create( InUserObjectRef, InFunc ) );
 	}
 	template< class UserClass, PAYLOAD_TEMPLATE_DECL_OneVar >
-	inline static DELEGATE_CLASS< FUNC_TEMPLATE_ARGS > CreateThreadSafeSP( const TSharedRef< UserClass, ESPMode::ThreadSafe >& InUserObjectRef, typename TThreadSafeSPMethodDelegate_OneVar< UserClass, PAYLOAD_TEMPLATE_LIST_OneVar >::FMethodPtr InFunc, PAYLOAD_TEMPLATE_ARGS_OneVar )
+	inline static DELEGATE_CLASS CreateThreadSafeSP( const TSharedRef< UserClass, ESPMode::ThreadSafe >& InUserObjectRef, typename TThreadSafeSPMethodDelegate_OneVar< UserClass, PAYLOAD_TEMPLATE_LIST_OneVar >::FMethodPtr InFunc, PAYLOAD_TEMPLATE_ARGS_OneVar )
 	{
-		return DELEGATE_CLASS< FUNC_TEMPLATE_ARGS >( TThreadSafeSPMethodDelegate_OneVar< UserClass, PAYLOAD_TEMPLATE_LIST_OneVar >::Create( InUserObjectRef, InFunc, PAYLOAD_TEMPLATE_PASSIN_OneVar ) );
+		return DELEGATE_CLASS( TThreadSafeSPMethodDelegate_OneVar< UserClass, PAYLOAD_TEMPLATE_LIST_OneVar >::Create( InUserObjectRef, InFunc, PAYLOAD_TEMPLATE_PASSIN_OneVar ) );
 	}
 	template< class UserClass, PAYLOAD_TEMPLATE_DECL_OneVar >
-	inline static DELEGATE_CLASS< FUNC_TEMPLATE_ARGS > CreateThreadSafeSP( const TSharedRef< UserClass, ESPMode::ThreadSafe >& InUserObjectRef, typename TThreadSafeSPMethodDelegate_OneVar_Const< UserClass, PAYLOAD_TEMPLATE_LIST_OneVar >::FMethodPtr InFunc, PAYLOAD_TEMPLATE_ARGS_OneVar )
+	inline static DELEGATE_CLASS CreateThreadSafeSP( const TSharedRef< UserClass, ESPMode::ThreadSafe >& InUserObjectRef, typename TThreadSafeSPMethodDelegate_OneVar_Const< UserClass, PAYLOAD_TEMPLATE_LIST_OneVar >::FMethodPtr InFunc, PAYLOAD_TEMPLATE_ARGS_OneVar )
 	{
-		return DELEGATE_CLASS< FUNC_TEMPLATE_ARGS >( TThreadSafeSPMethodDelegate_OneVar_Const< UserClass, PAYLOAD_TEMPLATE_LIST_OneVar >::Create( InUserObjectRef, InFunc, PAYLOAD_TEMPLATE_PASSIN_OneVar ) );
+		return DELEGATE_CLASS( TThreadSafeSPMethodDelegate_OneVar_Const< UserClass, PAYLOAD_TEMPLATE_LIST_OneVar >::Create( InUserObjectRef, InFunc, PAYLOAD_TEMPLATE_PASSIN_OneVar ) );
 	}
 	template< class UserClass, PAYLOAD_TEMPLATE_DECL_TwoVars >
-	inline static DELEGATE_CLASS< FUNC_TEMPLATE_ARGS > CreateThreadSafeSP( const TSharedRef< UserClass, ESPMode::ThreadSafe >& InUserObjectRef, typename TThreadSafeSPMethodDelegate_TwoVars< UserClass, PAYLOAD_TEMPLATE_LIST_TwoVars >::FMethodPtr InFunc, PAYLOAD_TEMPLATE_ARGS_TwoVars )
+	inline static DELEGATE_CLASS CreateThreadSafeSP( const TSharedRef< UserClass, ESPMode::ThreadSafe >& InUserObjectRef, typename TThreadSafeSPMethodDelegate_TwoVars< UserClass, PAYLOAD_TEMPLATE_LIST_TwoVars >::FMethodPtr InFunc, PAYLOAD_TEMPLATE_ARGS_TwoVars )
 	{
-		return DELEGATE_CLASS< FUNC_TEMPLATE_ARGS >( TThreadSafeSPMethodDelegate_TwoVars< UserClass, PAYLOAD_TEMPLATE_LIST_TwoVars >::Create( InUserObjectRef, InFunc, PAYLOAD_TEMPLATE_PASSIN_TwoVars ) );
+		return DELEGATE_CLASS( TThreadSafeSPMethodDelegate_TwoVars< UserClass, PAYLOAD_TEMPLATE_LIST_TwoVars >::Create( InUserObjectRef, InFunc, PAYLOAD_TEMPLATE_PASSIN_TwoVars ) );
 	}
 	template< class UserClass, PAYLOAD_TEMPLATE_DECL_TwoVars >
-	inline static DELEGATE_CLASS< FUNC_TEMPLATE_ARGS > CreateThreadSafeSP( const TSharedRef< UserClass, ESPMode::ThreadSafe >& InUserObjectRef, typename TThreadSafeSPMethodDelegate_TwoVars_Const< UserClass, PAYLOAD_TEMPLATE_LIST_TwoVars >::FMethodPtr InFunc, PAYLOAD_TEMPLATE_ARGS_TwoVars )
+	inline static DELEGATE_CLASS CreateThreadSafeSP( const TSharedRef< UserClass, ESPMode::ThreadSafe >& InUserObjectRef, typename TThreadSafeSPMethodDelegate_TwoVars_Const< UserClass, PAYLOAD_TEMPLATE_LIST_TwoVars >::FMethodPtr InFunc, PAYLOAD_TEMPLATE_ARGS_TwoVars )
 	{
-		return DELEGATE_CLASS< FUNC_TEMPLATE_ARGS >( TThreadSafeSPMethodDelegate_TwoVars_Const< UserClass, PAYLOAD_TEMPLATE_LIST_TwoVars >::Create( InUserObjectRef, InFunc, PAYLOAD_TEMPLATE_PASSIN_TwoVars ) );
+		return DELEGATE_CLASS( TThreadSafeSPMethodDelegate_TwoVars_Const< UserClass, PAYLOAD_TEMPLATE_LIST_TwoVars >::Create( InUserObjectRef, InFunc, PAYLOAD_TEMPLATE_PASSIN_TwoVars ) );
 	}
 	template< class UserClass, PAYLOAD_TEMPLATE_DECL_ThreeVars >
-	inline static DELEGATE_CLASS< FUNC_TEMPLATE_ARGS > CreateThreadSafeSP( const TSharedRef< UserClass, ESPMode::ThreadSafe >& InUserObjectRef, typename TThreadSafeSPMethodDelegate_ThreeVars< UserClass, PAYLOAD_TEMPLATE_LIST_ThreeVars >::FMethodPtr InFunc, PAYLOAD_TEMPLATE_ARGS_ThreeVars )
+	inline static DELEGATE_CLASS CreateThreadSafeSP( const TSharedRef< UserClass, ESPMode::ThreadSafe >& InUserObjectRef, typename TThreadSafeSPMethodDelegate_ThreeVars< UserClass, PAYLOAD_TEMPLATE_LIST_ThreeVars >::FMethodPtr InFunc, PAYLOAD_TEMPLATE_ARGS_ThreeVars )
 	{
-		return DELEGATE_CLASS< FUNC_TEMPLATE_ARGS >( TThreadSafeSPMethodDelegate_ThreeVars< UserClass, PAYLOAD_TEMPLATE_LIST_ThreeVars >::Create( InUserObjectRef, InFunc, PAYLOAD_TEMPLATE_PASSIN_ThreeVars ) );
+		return DELEGATE_CLASS( TThreadSafeSPMethodDelegate_ThreeVars< UserClass, PAYLOAD_TEMPLATE_LIST_ThreeVars >::Create( InUserObjectRef, InFunc, PAYLOAD_TEMPLATE_PASSIN_ThreeVars ) );
 	}
 	template< class UserClass, PAYLOAD_TEMPLATE_DECL_ThreeVars >
-	inline static DELEGATE_CLASS< FUNC_TEMPLATE_ARGS > CreateThreadSafeSP( const TSharedRef< UserClass, ESPMode::ThreadSafe >& InUserObjectRef, typename TThreadSafeSPMethodDelegate_ThreeVars_Const< UserClass, PAYLOAD_TEMPLATE_LIST_ThreeVars >::FMethodPtr InFunc, PAYLOAD_TEMPLATE_ARGS_ThreeVars )
+	inline static DELEGATE_CLASS CreateThreadSafeSP( const TSharedRef< UserClass, ESPMode::ThreadSafe >& InUserObjectRef, typename TThreadSafeSPMethodDelegate_ThreeVars_Const< UserClass, PAYLOAD_TEMPLATE_LIST_ThreeVars >::FMethodPtr InFunc, PAYLOAD_TEMPLATE_ARGS_ThreeVars )
 	{
-		return DELEGATE_CLASS< FUNC_TEMPLATE_ARGS >( TThreadSafeSPMethodDelegate_ThreeVars_Const< UserClass, PAYLOAD_TEMPLATE_LIST_ThreeVars >::Create( InUserObjectRef, InFunc, PAYLOAD_TEMPLATE_PASSIN_ThreeVars ) );
+		return DELEGATE_CLASS( TThreadSafeSPMethodDelegate_ThreeVars_Const< UserClass, PAYLOAD_TEMPLATE_LIST_ThreeVars >::Create( InUserObjectRef, InFunc, PAYLOAD_TEMPLATE_PASSIN_ThreeVars ) );
 	}
 	template< class UserClass, PAYLOAD_TEMPLATE_DECL_FourVars >
-	inline static DELEGATE_CLASS< FUNC_TEMPLATE_ARGS > CreateThreadSafeSP( const TSharedRef< UserClass, ESPMode::ThreadSafe >& InUserObjectRef, typename TThreadSafeSPMethodDelegate_FourVars< UserClass, PAYLOAD_TEMPLATE_LIST_FourVars >::FMethodPtr InFunc, PAYLOAD_TEMPLATE_ARGS_FourVars )
+	inline static DELEGATE_CLASS CreateThreadSafeSP( const TSharedRef< UserClass, ESPMode::ThreadSafe >& InUserObjectRef, typename TThreadSafeSPMethodDelegate_FourVars< UserClass, PAYLOAD_TEMPLATE_LIST_FourVars >::FMethodPtr InFunc, PAYLOAD_TEMPLATE_ARGS_FourVars )
 	{
-		return DELEGATE_CLASS< FUNC_TEMPLATE_ARGS >( TThreadSafeSPMethodDelegate_FourVars< UserClass, PAYLOAD_TEMPLATE_LIST_FourVars >::Create( InUserObjectRef, InFunc, PAYLOAD_TEMPLATE_PASSIN_FourVars ) );
+		return DELEGATE_CLASS( TThreadSafeSPMethodDelegate_FourVars< UserClass, PAYLOAD_TEMPLATE_LIST_FourVars >::Create( InUserObjectRef, InFunc, PAYLOAD_TEMPLATE_PASSIN_FourVars ) );
 	}
 	template< class UserClass, PAYLOAD_TEMPLATE_DECL_FourVars >
-	inline static DELEGATE_CLASS< FUNC_TEMPLATE_ARGS > CreateThreadSafeSP( const TSharedRef< UserClass, ESPMode::ThreadSafe >& InUserObjectRef, typename TThreadSafeSPMethodDelegate_FourVars_Const< UserClass, PAYLOAD_TEMPLATE_LIST_FourVars >::FMethodPtr InFunc, PAYLOAD_TEMPLATE_ARGS_FourVars )
+	inline static DELEGATE_CLASS CreateThreadSafeSP( const TSharedRef< UserClass, ESPMode::ThreadSafe >& InUserObjectRef, typename TThreadSafeSPMethodDelegate_FourVars_Const< UserClass, PAYLOAD_TEMPLATE_LIST_FourVars >::FMethodPtr InFunc, PAYLOAD_TEMPLATE_ARGS_FourVars )
 	{
-		return DELEGATE_CLASS< FUNC_TEMPLATE_ARGS >( TThreadSafeSPMethodDelegate_FourVars_Const< UserClass, PAYLOAD_TEMPLATE_LIST_FourVars >::Create( InUserObjectRef, InFunc, PAYLOAD_TEMPLATE_PASSIN_FourVars ) );
+		return DELEGATE_CLASS( TThreadSafeSPMethodDelegate_FourVars_Const< UserClass, PAYLOAD_TEMPLATE_LIST_FourVars >::Create( InUserObjectRef, InFunc, PAYLOAD_TEMPLATE_PASSIN_FourVars ) );
 	}
 
 	/**
@@ -433,52 +433,52 @@ public:
 	 * You can use ExecuteIfBound() to call them.
 	 */
 	template< class UserClass >
-	inline static DELEGATE_CLASS< FUNC_TEMPLATE_ARGS > CreateThreadSafeSP( UserClass* InUserObject, typename TThreadSafeSPMethodDelegate< UserClass >::FMethodPtr InFunc )
+	inline static DELEGATE_CLASS CreateThreadSafeSP( UserClass* InUserObject, typename TThreadSafeSPMethodDelegate< UserClass >::FMethodPtr InFunc )
 	{
 		return CreateThreadSafeSP( StaticCastSharedRef< UserClass >( InUserObject->AsShared() ), InFunc );
 	}
 	template< class UserClass >
-	inline static DELEGATE_CLASS< FUNC_TEMPLATE_ARGS > CreateThreadSafeSP( UserClass* InUserObject, typename TThreadSafeSPMethodDelegate_Const< UserClass >::FMethodPtr InFunc )
+	inline static DELEGATE_CLASS CreateThreadSafeSP( UserClass* InUserObject, typename TThreadSafeSPMethodDelegate_Const< UserClass >::FMethodPtr InFunc )
 	{
 		return CreateThreadSafeSP( StaticCastSharedRef< UserClass >( InUserObject->AsShared() ), InFunc );
 	}
 	template< class UserClass, PAYLOAD_TEMPLATE_DECL_OneVar >
-	inline static DELEGATE_CLASS< FUNC_TEMPLATE_ARGS > CreateThreadSafeSP( UserClass* InUserObject, typename TThreadSafeSPMethodDelegate_OneVar< UserClass, PAYLOAD_TEMPLATE_LIST_OneVar >::FMethodPtr InFunc, PAYLOAD_TEMPLATE_ARGS_OneVar )
+	inline static DELEGATE_CLASS CreateThreadSafeSP( UserClass* InUserObject, typename TThreadSafeSPMethodDelegate_OneVar< UserClass, PAYLOAD_TEMPLATE_LIST_OneVar >::FMethodPtr InFunc, PAYLOAD_TEMPLATE_ARGS_OneVar )
 	{
 		return CreateThreadSafeSP( StaticCastSharedRef< UserClass >( InUserObject->AsShared() ), InFunc, PAYLOAD_TEMPLATE_PASSIN_OneVar );
 	}
 	template< class UserClass, PAYLOAD_TEMPLATE_DECL_OneVar >
-	inline static DELEGATE_CLASS< FUNC_TEMPLATE_ARGS > CreateThreadSafeSP( UserClass* InUserObject, typename TThreadSafeSPMethodDelegate_OneVar_Const< UserClass, PAYLOAD_TEMPLATE_LIST_OneVar >::FMethodPtr InFunc, PAYLOAD_TEMPLATE_ARGS_OneVar )
+	inline static DELEGATE_CLASS CreateThreadSafeSP( UserClass* InUserObject, typename TThreadSafeSPMethodDelegate_OneVar_Const< UserClass, PAYLOAD_TEMPLATE_LIST_OneVar >::FMethodPtr InFunc, PAYLOAD_TEMPLATE_ARGS_OneVar )
 	{
 		return CreateThreadSafeSP( StaticCastSharedRef< UserClass >( InUserObject->AsShared() ), InFunc, PAYLOAD_TEMPLATE_PASSIN_OneVar );
 	}
 	template< class UserClass, PAYLOAD_TEMPLATE_DECL_TwoVars >
-	inline static DELEGATE_CLASS< FUNC_TEMPLATE_ARGS > CreateThreadSafeSP( UserClass* InUserObject, typename TThreadSafeSPMethodDelegate_TwoVars< UserClass, PAYLOAD_TEMPLATE_LIST_TwoVars >::FMethodPtr InFunc, PAYLOAD_TEMPLATE_ARGS_TwoVars )
+	inline static DELEGATE_CLASS CreateThreadSafeSP( UserClass* InUserObject, typename TThreadSafeSPMethodDelegate_TwoVars< UserClass, PAYLOAD_TEMPLATE_LIST_TwoVars >::FMethodPtr InFunc, PAYLOAD_TEMPLATE_ARGS_TwoVars )
 	{
 		return CreateThreadSafeSP( StaticCastSharedRef< UserClass >( InUserObject->AsShared() ), InFunc, PAYLOAD_TEMPLATE_PASSIN_TwoVars );
 	}
 	template< class UserClass, PAYLOAD_TEMPLATE_DECL_TwoVars >
-	inline static DELEGATE_CLASS< FUNC_TEMPLATE_ARGS > CreateThreadSafeSP( UserClass* InUserObject, typename TThreadSafeSPMethodDelegate_TwoVars_Const< UserClass, PAYLOAD_TEMPLATE_LIST_TwoVars >::FMethodPtr InFunc, PAYLOAD_TEMPLATE_ARGS_TwoVars )
+	inline static DELEGATE_CLASS CreateThreadSafeSP( UserClass* InUserObject, typename TThreadSafeSPMethodDelegate_TwoVars_Const< UserClass, PAYLOAD_TEMPLATE_LIST_TwoVars >::FMethodPtr InFunc, PAYLOAD_TEMPLATE_ARGS_TwoVars )
 	{
 		return CreateThreadSafeSP( StaticCastSharedRef< UserClass >( InUserObject->AsShared() ), InFunc, PAYLOAD_TEMPLATE_PASSIN_TwoVars );
 	}
 	template< class UserClass, PAYLOAD_TEMPLATE_DECL_ThreeVars >
-	inline static DELEGATE_CLASS< FUNC_TEMPLATE_ARGS > CreateThreadSafeSP( UserClass* InUserObject, typename TThreadSafeSPMethodDelegate_ThreeVars< UserClass, PAYLOAD_TEMPLATE_LIST_ThreeVars >::FMethodPtr InFunc, PAYLOAD_TEMPLATE_ARGS_ThreeVars )
+	inline static DELEGATE_CLASS CreateThreadSafeSP( UserClass* InUserObject, typename TThreadSafeSPMethodDelegate_ThreeVars< UserClass, PAYLOAD_TEMPLATE_LIST_ThreeVars >::FMethodPtr InFunc, PAYLOAD_TEMPLATE_ARGS_ThreeVars )
 	{
 		return CreateThreadSafeSP( StaticCastSharedRef< UserClass >( InUserObject->AsShared() ), InFunc, PAYLOAD_TEMPLATE_PASSIN_ThreeVars );
 	}
 	template< class UserClass, PAYLOAD_TEMPLATE_DECL_ThreeVars >
-	inline static DELEGATE_CLASS< FUNC_TEMPLATE_ARGS > CreateThreadSafeSP( UserClass* InUserObject, typename TThreadSafeSPMethodDelegate_ThreeVars_Const< UserClass, PAYLOAD_TEMPLATE_LIST_ThreeVars >::FMethodPtr InFunc, PAYLOAD_TEMPLATE_ARGS_ThreeVars )
+	inline static DELEGATE_CLASS CreateThreadSafeSP( UserClass* InUserObject, typename TThreadSafeSPMethodDelegate_ThreeVars_Const< UserClass, PAYLOAD_TEMPLATE_LIST_ThreeVars >::FMethodPtr InFunc, PAYLOAD_TEMPLATE_ARGS_ThreeVars )
 	{
 		return CreateThreadSafeSP( StaticCastSharedRef< UserClass >( InUserObject->AsShared() ), InFunc, PAYLOAD_TEMPLATE_PASSIN_ThreeVars );
 	}
 	template< class UserClass, PAYLOAD_TEMPLATE_DECL_FourVars >
-	inline static DELEGATE_CLASS< FUNC_TEMPLATE_ARGS > CreateThreadSafeSP( UserClass* InUserObject, typename TThreadSafeSPMethodDelegate_FourVars< UserClass, PAYLOAD_TEMPLATE_LIST_FourVars >::FMethodPtr InFunc, PAYLOAD_TEMPLATE_ARGS_FourVars )
+	inline static DELEGATE_CLASS CreateThreadSafeSP( UserClass* InUserObject, typename TThreadSafeSPMethodDelegate_FourVars< UserClass, PAYLOAD_TEMPLATE_LIST_FourVars >::FMethodPtr InFunc, PAYLOAD_TEMPLATE_ARGS_FourVars )
 	{
 		return CreateThreadSafeSP( StaticCastSharedRef< UserClass >( InUserObject->AsShared() ), InFunc, PAYLOAD_TEMPLATE_PASSIN_FourVars );
 	}
 	template< class UserClass, PAYLOAD_TEMPLATE_DECL_FourVars >
-	inline static DELEGATE_CLASS< FUNC_TEMPLATE_ARGS > CreateThreadSafeSP( UserClass* InUserObject, typename TThreadSafeSPMethodDelegate_FourVars_Const< UserClass, PAYLOAD_TEMPLATE_LIST_FourVars >::FMethodPtr InFunc, PAYLOAD_TEMPLATE_ARGS_FourVars )
+	inline static DELEGATE_CLASS CreateThreadSafeSP( UserClass* InUserObject, typename TThreadSafeSPMethodDelegate_FourVars_Const< UserClass, PAYLOAD_TEMPLATE_LIST_FourVars >::FMethodPtr InFunc, PAYLOAD_TEMPLATE_ARGS_FourVars )
 	{
 		return CreateThreadSafeSP( StaticCastSharedRef< UserClass >( InUserObject->AsShared() ), InFunc, PAYLOAD_TEMPLATE_PASSIN_FourVars );
 	}
@@ -490,29 +490,29 @@ public:
 	 * You can use ExecuteIfBound() to call them.
 	 */
 	template< class UObjectTemplate >
-	inline static DELEGATE_CLASS< FUNC_TEMPLATE_ARGS > CreateUFunction( UObjectTemplate* InUserObject, const FName& InFunctionName )
+	inline static DELEGATE_CLASS CreateUFunction( UObjectTemplate* InUserObject, const FName& InFunctionName )
 	{
-		return DELEGATE_CLASS< FUNC_TEMPLATE_ARGS >( TUFunctionDelegateBinding< UObjectTemplate >::Create( InUserObject, InFunctionName ) );
+		return DELEGATE_CLASS( TUFunctionDelegateBinding< UObjectTemplate >::Create( InUserObject, InFunctionName ) );
 	}
 	template< class UObjectTemplate, PAYLOAD_TEMPLATE_DECL_OneVar >
-	inline static DELEGATE_CLASS< FUNC_TEMPLATE_ARGS > CreateUFunction( UObjectTemplate* InUserObject, const FName& InFunctionName, PAYLOAD_TEMPLATE_ARGS_OneVar )
+	inline static DELEGATE_CLASS CreateUFunction( UObjectTemplate* InUserObject, const FName& InFunctionName, PAYLOAD_TEMPLATE_ARGS_OneVar )
 	{
-		return DELEGATE_CLASS< FUNC_TEMPLATE_ARGS >( TUFunctionDelegateBinding_OneVar< UObjectTemplate, PAYLOAD_TEMPLATE_LIST_OneVar >::Create( InUserObject, InFunctionName, PAYLOAD_TEMPLATE_PASSIN_OneVar ) );
+		return DELEGATE_CLASS( TUFunctionDelegateBinding_OneVar< UObjectTemplate, PAYLOAD_TEMPLATE_LIST_OneVar >::Create( InUserObject, InFunctionName, PAYLOAD_TEMPLATE_PASSIN_OneVar ) );
 	}
 	template< class UObjectTemplate, PAYLOAD_TEMPLATE_DECL_TwoVars >
-	inline static DELEGATE_CLASS< FUNC_TEMPLATE_ARGS > CreateUFunction( UObjectTemplate* InUserObject, const FName& InFunctionName, PAYLOAD_TEMPLATE_ARGS_TwoVars )
+	inline static DELEGATE_CLASS CreateUFunction( UObjectTemplate* InUserObject, const FName& InFunctionName, PAYLOAD_TEMPLATE_ARGS_TwoVars )
 	{
-		return DELEGATE_CLASS< FUNC_TEMPLATE_ARGS >( TUFunctionDelegateBinding_TwoVars< UObjectTemplate, PAYLOAD_TEMPLATE_LIST_TwoVars >::Create( InUserObject, InFunctionName, PAYLOAD_TEMPLATE_PASSIN_TwoVars ) );
+		return DELEGATE_CLASS( TUFunctionDelegateBinding_TwoVars< UObjectTemplate, PAYLOAD_TEMPLATE_LIST_TwoVars >::Create( InUserObject, InFunctionName, PAYLOAD_TEMPLATE_PASSIN_TwoVars ) );
 	}
 	template< class UObjectTemplate, PAYLOAD_TEMPLATE_DECL_ThreeVars >
-	inline static DELEGATE_CLASS< FUNC_TEMPLATE_ARGS > CreateUFunction( UObjectTemplate* InUserObject, const FName& InFunctionName, PAYLOAD_TEMPLATE_ARGS_ThreeVars )
+	inline static DELEGATE_CLASS CreateUFunction( UObjectTemplate* InUserObject, const FName& InFunctionName, PAYLOAD_TEMPLATE_ARGS_ThreeVars )
 	{
-		return DELEGATE_CLASS< FUNC_TEMPLATE_ARGS >( TUFunctionDelegateBinding_ThreeVars< UObjectTemplate, PAYLOAD_TEMPLATE_LIST_ThreeVars >::Create( InUserObject, InFunctionName, PAYLOAD_TEMPLATE_PASSIN_ThreeVars ) );
+		return DELEGATE_CLASS( TUFunctionDelegateBinding_ThreeVars< UObjectTemplate, PAYLOAD_TEMPLATE_LIST_ThreeVars >::Create( InUserObject, InFunctionName, PAYLOAD_TEMPLATE_PASSIN_ThreeVars ) );
 	}
 	template< class UObjectTemplate, PAYLOAD_TEMPLATE_DECL_FourVars >
-	inline static DELEGATE_CLASS< FUNC_TEMPLATE_ARGS > CreateUFunction( UObjectTemplate* InUserObject, const FName& InFunctionName, PAYLOAD_TEMPLATE_ARGS_FourVars )
+	inline static DELEGATE_CLASS CreateUFunction( UObjectTemplate* InUserObject, const FName& InFunctionName, PAYLOAD_TEMPLATE_ARGS_FourVars )
 	{
-		return DELEGATE_CLASS< FUNC_TEMPLATE_ARGS >( TUFunctionDelegateBinding_FourVars< UObjectTemplate, PAYLOAD_TEMPLATE_LIST_FourVars >::Create( InUserObject, InFunctionName, PAYLOAD_TEMPLATE_PASSIN_FourVars ) );
+		return DELEGATE_CLASS( TUFunctionDelegateBinding_FourVars< UObjectTemplate, PAYLOAD_TEMPLATE_LIST_FourVars >::Create( InUserObject, InFunctionName, PAYLOAD_TEMPLATE_PASSIN_FourVars ) );
 	}
 
 	/**
@@ -522,54 +522,54 @@ public:
 	 * You can use ExecuteIfBound() to call them.
 	 */
 	template< class UserClass >
-	inline static DELEGATE_CLASS< FUNC_TEMPLATE_ARGS > CreateUObject( UserClass* InUserObject, typename TUObjectMethodDelegate< UserClass >::FMethodPtr InFunc )
+	inline static DELEGATE_CLASS CreateUObject( UserClass* InUserObject, typename TUObjectMethodDelegate< UserClass >::FMethodPtr InFunc )
 	{
-		return DELEGATE_CLASS< FUNC_TEMPLATE_ARGS >( TUObjectMethodDelegate< UserClass >::Create( InUserObject, InFunc ) );
+		return DELEGATE_CLASS( TUObjectMethodDelegate< UserClass >::Create( InUserObject, InFunc ) );
 	}
 	template< class UserClass >
-	inline static DELEGATE_CLASS< FUNC_TEMPLATE_ARGS > CreateUObject( UserClass* InUserObject, typename TUObjectMethodDelegate_Const< UserClass >::FMethodPtr InFunc )
+	inline static DELEGATE_CLASS CreateUObject( UserClass* InUserObject, typename TUObjectMethodDelegate_Const< UserClass >::FMethodPtr InFunc )
 	{
-		return DELEGATE_CLASS< FUNC_TEMPLATE_ARGS >( TUObjectMethodDelegate_Const< UserClass >::Create( InUserObject, InFunc ) );
+		return DELEGATE_CLASS( TUObjectMethodDelegate_Const< UserClass >::Create( InUserObject, InFunc ) );
 	}
 	template< class UserClass, PAYLOAD_TEMPLATE_DECL_OneVar >
-	inline static DELEGATE_CLASS< FUNC_TEMPLATE_ARGS > CreateUObject( UserClass* InUserObject, typename TUObjectMethodDelegate_OneVar< UserClass, PAYLOAD_TEMPLATE_LIST_OneVar >::FMethodPtr InFunc, PAYLOAD_TEMPLATE_ARGS_OneVar )
+	inline static DELEGATE_CLASS CreateUObject( UserClass* InUserObject, typename TUObjectMethodDelegate_OneVar< UserClass, PAYLOAD_TEMPLATE_LIST_OneVar >::FMethodPtr InFunc, PAYLOAD_TEMPLATE_ARGS_OneVar )
 	{
-		return DELEGATE_CLASS< FUNC_TEMPLATE_ARGS >( TUObjectMethodDelegate_OneVar< UserClass, PAYLOAD_TEMPLATE_LIST_OneVar >::Create( InUserObject, InFunc, PAYLOAD_TEMPLATE_PASSIN_OneVar ) );
+		return DELEGATE_CLASS( TUObjectMethodDelegate_OneVar< UserClass, PAYLOAD_TEMPLATE_LIST_OneVar >::Create( InUserObject, InFunc, PAYLOAD_TEMPLATE_PASSIN_OneVar ) );
 	}
 	template< class UserClass, PAYLOAD_TEMPLATE_DECL_OneVar >
-	inline static DELEGATE_CLASS< FUNC_TEMPLATE_ARGS > CreateUObject( UserClass* InUserObject, typename TUObjectMethodDelegate_OneVar_Const< UserClass, PAYLOAD_TEMPLATE_LIST_OneVar >::FMethodPtr InFunc, PAYLOAD_TEMPLATE_ARGS_OneVar )
+	inline static DELEGATE_CLASS CreateUObject( UserClass* InUserObject, typename TUObjectMethodDelegate_OneVar_Const< UserClass, PAYLOAD_TEMPLATE_LIST_OneVar >::FMethodPtr InFunc, PAYLOAD_TEMPLATE_ARGS_OneVar )
 	{
-		return DELEGATE_CLASS< FUNC_TEMPLATE_ARGS >( TUObjectMethodDelegate_OneVar_Const< UserClass, PAYLOAD_TEMPLATE_LIST_OneVar >::Create( InUserObject, InFunc, PAYLOAD_TEMPLATE_PASSIN_OneVar ) );
+		return DELEGATE_CLASS( TUObjectMethodDelegate_OneVar_Const< UserClass, PAYLOAD_TEMPLATE_LIST_OneVar >::Create( InUserObject, InFunc, PAYLOAD_TEMPLATE_PASSIN_OneVar ) );
 	}
 	template< class UserClass, PAYLOAD_TEMPLATE_DECL_TwoVars >
-	inline static DELEGATE_CLASS< FUNC_TEMPLATE_ARGS > CreateUObject( UserClass* InUserObject, typename TUObjectMethodDelegate_TwoVars< UserClass, PAYLOAD_TEMPLATE_LIST_TwoVars >::FMethodPtr InFunc, PAYLOAD_TEMPLATE_ARGS_TwoVars )
+	inline static DELEGATE_CLASS CreateUObject( UserClass* InUserObject, typename TUObjectMethodDelegate_TwoVars< UserClass, PAYLOAD_TEMPLATE_LIST_TwoVars >::FMethodPtr InFunc, PAYLOAD_TEMPLATE_ARGS_TwoVars )
 	{
-		return DELEGATE_CLASS< FUNC_TEMPLATE_ARGS >( TUObjectMethodDelegate_TwoVars< UserClass, PAYLOAD_TEMPLATE_LIST_TwoVars >::Create( InUserObject, InFunc, PAYLOAD_TEMPLATE_PASSIN_TwoVars ) );
+		return DELEGATE_CLASS( TUObjectMethodDelegate_TwoVars< UserClass, PAYLOAD_TEMPLATE_LIST_TwoVars >::Create( InUserObject, InFunc, PAYLOAD_TEMPLATE_PASSIN_TwoVars ) );
 	}
 	template< class UserClass, PAYLOAD_TEMPLATE_DECL_TwoVars >
-	inline static DELEGATE_CLASS< FUNC_TEMPLATE_ARGS > CreateUObject( UserClass* InUserObject, typename TUObjectMethodDelegate_TwoVars_Const< UserClass, PAYLOAD_TEMPLATE_LIST_TwoVars >::FMethodPtr InFunc, PAYLOAD_TEMPLATE_ARGS_TwoVars )
+	inline static DELEGATE_CLASS CreateUObject( UserClass* InUserObject, typename TUObjectMethodDelegate_TwoVars_Const< UserClass, PAYLOAD_TEMPLATE_LIST_TwoVars >::FMethodPtr InFunc, PAYLOAD_TEMPLATE_ARGS_TwoVars )
 	{
-		return DELEGATE_CLASS< FUNC_TEMPLATE_ARGS >( TUObjectMethodDelegate_TwoVars_Const< UserClass, PAYLOAD_TEMPLATE_LIST_TwoVars >::Create( InUserObject, InFunc, PAYLOAD_TEMPLATE_PASSIN_TwoVars ) );
+		return DELEGATE_CLASS( TUObjectMethodDelegate_TwoVars_Const< UserClass, PAYLOAD_TEMPLATE_LIST_TwoVars >::Create( InUserObject, InFunc, PAYLOAD_TEMPLATE_PASSIN_TwoVars ) );
 	}
 	template< class UserClass, PAYLOAD_TEMPLATE_DECL_ThreeVars >
-	inline static DELEGATE_CLASS< FUNC_TEMPLATE_ARGS > CreateUObject( UserClass* InUserObject, typename TUObjectMethodDelegate_ThreeVars< UserClass, PAYLOAD_TEMPLATE_LIST_ThreeVars >::FMethodPtr InFunc, PAYLOAD_TEMPLATE_ARGS_ThreeVars )
+	inline static DELEGATE_CLASS CreateUObject( UserClass* InUserObject, typename TUObjectMethodDelegate_ThreeVars< UserClass, PAYLOAD_TEMPLATE_LIST_ThreeVars >::FMethodPtr InFunc, PAYLOAD_TEMPLATE_ARGS_ThreeVars )
 	{
-		return DELEGATE_CLASS< FUNC_TEMPLATE_ARGS >( TUObjectMethodDelegate_ThreeVars< UserClass, PAYLOAD_TEMPLATE_LIST_ThreeVars >::Create( InUserObject, InFunc, PAYLOAD_TEMPLATE_PASSIN_ThreeVars ) );
+		return DELEGATE_CLASS( TUObjectMethodDelegate_ThreeVars< UserClass, PAYLOAD_TEMPLATE_LIST_ThreeVars >::Create( InUserObject, InFunc, PAYLOAD_TEMPLATE_PASSIN_ThreeVars ) );
 	}
 	template< class UserClass, PAYLOAD_TEMPLATE_DECL_ThreeVars >
-	inline static DELEGATE_CLASS< FUNC_TEMPLATE_ARGS > CreateUObject( UserClass* InUserObject, typename TUObjectMethodDelegate_ThreeVars_Const< UserClass, PAYLOAD_TEMPLATE_LIST_ThreeVars >::FMethodPtr InFunc, PAYLOAD_TEMPLATE_ARGS_ThreeVars )
+	inline static DELEGATE_CLASS CreateUObject( UserClass* InUserObject, typename TUObjectMethodDelegate_ThreeVars_Const< UserClass, PAYLOAD_TEMPLATE_LIST_ThreeVars >::FMethodPtr InFunc, PAYLOAD_TEMPLATE_ARGS_ThreeVars )
 	{
-		return DELEGATE_CLASS< FUNC_TEMPLATE_ARGS >( TUObjectMethodDelegate_ThreeVars_Const< UserClass, PAYLOAD_TEMPLATE_LIST_ThreeVars >::Create( InUserObject, InFunc, PAYLOAD_TEMPLATE_PASSIN_ThreeVars ) );
+		return DELEGATE_CLASS( TUObjectMethodDelegate_ThreeVars_Const< UserClass, PAYLOAD_TEMPLATE_LIST_ThreeVars >::Create( InUserObject, InFunc, PAYLOAD_TEMPLATE_PASSIN_ThreeVars ) );
 	}
 	template< class UserClass, PAYLOAD_TEMPLATE_DECL_FourVars >
-	inline static DELEGATE_CLASS< FUNC_TEMPLATE_ARGS > CreateUObject( UserClass* InUserObject, typename TUObjectMethodDelegate_FourVars< UserClass, PAYLOAD_TEMPLATE_LIST_FourVars >::FMethodPtr InFunc, PAYLOAD_TEMPLATE_ARGS_FourVars )
+	inline static DELEGATE_CLASS CreateUObject( UserClass* InUserObject, typename TUObjectMethodDelegate_FourVars< UserClass, PAYLOAD_TEMPLATE_LIST_FourVars >::FMethodPtr InFunc, PAYLOAD_TEMPLATE_ARGS_FourVars )
 	{
-		return DELEGATE_CLASS< FUNC_TEMPLATE_ARGS >( TUObjectMethodDelegate_FourVars< UserClass, PAYLOAD_TEMPLATE_LIST_FourVars >::Create( InUserObject, InFunc, PAYLOAD_TEMPLATE_PASSIN_FourVars ) );
+		return DELEGATE_CLASS( TUObjectMethodDelegate_FourVars< UserClass, PAYLOAD_TEMPLATE_LIST_FourVars >::Create( InUserObject, InFunc, PAYLOAD_TEMPLATE_PASSIN_FourVars ) );
 	}
 	template< class UserClass, PAYLOAD_TEMPLATE_DECL_FourVars >
-	inline static DELEGATE_CLASS< FUNC_TEMPLATE_ARGS > CreateUObject( UserClass* InUserObject, typename TUObjectMethodDelegate_FourVars_Const< UserClass, PAYLOAD_TEMPLATE_LIST_FourVars >::FMethodPtr InFunc, PAYLOAD_TEMPLATE_ARGS_FourVars )
+	inline static DELEGATE_CLASS CreateUObject( UserClass* InUserObject, typename TUObjectMethodDelegate_FourVars_Const< UserClass, PAYLOAD_TEMPLATE_LIST_FourVars >::FMethodPtr InFunc, PAYLOAD_TEMPLATE_ARGS_FourVars )
 	{
-		return DELEGATE_CLASS< FUNC_TEMPLATE_ARGS >( TUObjectMethodDelegate_FourVars_Const< UserClass, PAYLOAD_TEMPLATE_LIST_FourVars >::Create( InUserObject, InFunc, PAYLOAD_TEMPLATE_PASSIN_FourVars ) );
+		return DELEGATE_CLASS( TUObjectMethodDelegate_FourVars_Const< UserClass, PAYLOAD_TEMPLATE_LIST_FourVars >::Create( InUserObject, InFunc, PAYLOAD_TEMPLATE_PASSIN_FourVars ) );
 	}
 
 public:
@@ -1081,7 +1081,7 @@ public:
 	 *
 	 * @return true if this delegate equals the other, false otherwise.
 	 */
-	bool operator==( const DELEGATE_CLASS< FUNC_TEMPLATE_ARGS >& Other ) const
+	bool operator==( const DELEGATE_CLASS& Other ) const
 	{
 		TDelegateInstanceInterface* DelegateInstance = (TDelegateInstanceInterface*)GetDelegateInstance();
 		TDelegateInstanceInterface* OtherInstance = (TDelegateInstanceInterface*)Other.GetDelegateInstance();
diff --git a/Engine/Source/Runtime/Core/Public/Delegates/DelegateSignatureImpl_Variadics.inl b/Engine/Source/Runtime/Core/Public/Delegates/DelegateSignatureImpl_Variadics.inl
new file mode 100644
--- /dev/null
+++ b/Engine/Source/Runtime/Core/Public/Delegates/DelegateSignatureImpl_Variadics.inl
@@ -0,0 +1,1219 @@
+// Copyright 1998-2014 Epic Games, Inc. All Rights Reserved.
+
+// Only designed to be included directly by Delegate.h
+#if !defined( __Delegate_h__ ) || !defined( FUNC_INCLUDING_INLINE_IMPL )
+	#error "This inline header must only be included by Delegate.h"
+#endif
+
+#pragma once
+
+/**
+ * Unicast delegate base object.
+ *
+ * Use the various DECLARE_DELEGATE macros to create the actual delegate type, templated to
+ * the function signature the delegate is compatible with. Then, you can create an instance
+ * of that class when you want to bind a function to the delegate.
+ */
+template <typename WrappedRetValType, typename... ParamTypes>
+class TBaseDelegate : public FDelegateBase<>
+{
+	typedef typename TUnwrapType<WrappedRetValType>::Type RetValType;
+	typedef RetValType TFuncType(ParamTypes...);
+
+public:
+	/** Type definition for the shared interface of delegate instance types compatible with this delegate class. */
+	typedef IBaseDelegateInstance<TFuncType> TDelegateInstanceInterface;
+
+	/** Declare the user's "fast" shared pointer-based delegate instance types. */
+	template <typename UserClass                                                                            > struct TSPMethodDelegate                 : TBaseSPMethodDelegateInstance<false, UserClass, ESPMode::Fast, TFuncType                                        > { typedef TBaseSPMethodDelegateInstance<false, UserClass, ESPMode::Fast, TFuncType                                        > Super; TSPMethodDelegate                (const TSharedRef< UserClass, ESPMode::Fast >& InUserObject, typename Super::FMethodPtr InMethodPtr                                                            ) : Super(InUserObject, InMethodPtr                        ) {} };
+	template <typename UserClass                                                                            > struct TSPMethodDelegate_Const           : TBaseSPMethodDelegateInstance<true , UserClass, ESPMode::Fast, TFuncType                                        > { typedef TBaseSPMethodDelegateInstance<true , UserClass, ESPMode::Fast, TFuncType                                        > Super; TSPMethodDelegate_Const          (const TSharedRef< UserClass, ESPMode::Fast >& InUserObject, typename Super::FMethodPtr InMethodPtr                                                            ) : Super(InUserObject, InMethodPtr                        ) {} };
+	template <typename UserClass, typename Var1Type                                                         > struct TSPMethodDelegate_OneVar          : TBaseSPMethodDelegateInstance<false, UserClass, ESPMode::Fast, TFuncType, Var1Type                              > { typedef TBaseSPMethodDelegateInstance<false, UserClass, ESPMode::Fast, TFuncType, Var1Type                              > Super; TSPMethodDelegate_OneVar         (const TSharedRef< UserClass, ESPMode::Fast >& InUserObject, typename Super::FMethodPtr InMethodPtr, Var1Type Var1                                             ) : Super(InUserObject, InMethodPtr, Var1                  ) {} };
+	template <typename UserClass, typename Var1Type                                                         > struct TSPMethodDelegate_OneVar_Const    : TBaseSPMethodDelegateInstance<true , UserClass, ESPMode::Fast, TFuncType, Var1Type                              > { typedef TBaseSPMethodDelegateInstance<true , UserClass, ESPMode::Fast, TFuncType, Var1Type                              > Super; TSPMethodDelegate_OneVar_Const   (const TSharedRef< UserClass, ESPMode::Fast >& InUserObject, typename Super::FMethodPtr InMethodPtr, Var1Type Var1                                             ) : Super(InUserObject, InMethodPtr, Var1                  ) {} };
+	template <typename UserClass, typename Var1Type, typename Var2Type                                      > struct TSPMethodDelegate_TwoVars         : TBaseSPMethodDelegateInstance<false, UserClass, ESPMode::Fast, TFuncType, Var1Type, Var2Type                    > { typedef TBaseSPMethodDelegateInstance<false, UserClass, ESPMode::Fast, TFuncType, Var1Type, Var2Type                    > Super; TSPMethodDelegate_TwoVars        (const TSharedRef< UserClass, ESPMode::Fast >& InUserObject, typename Super::FMethodPtr InMethodPtr, Var1Type Var1, Var2Type Var2                              ) : Super(InUserObject, InMethodPtr, Var1, Var2            ) {} };
+	template <typename UserClass, typename Var1Type, typename Var2Type                                      > struct TSPMethodDelegate_TwoVars_Const   : TBaseSPMethodDelegateInstance<true , UserClass, ESPMode::Fast, TFuncType, Var1Type, Var2Type                    > { typedef TBaseSPMethodDelegateInstance<true , UserClass, ESPMode::Fast, TFuncType, Var1Type, Var2Type                    > Super; TSPMethodDelegate_TwoVars_Const  (const TSharedRef< UserClass, ESPMode::Fast >& InUserObject, typename Super::FMethodPtr InMethodPtr, Var1Type Var1, Var2Type Var2                              ) : Super(InUserObject, InMethodPtr, Var1, Var2            ) {} };
+	template <typename UserClass, typename Var1Type, typename Var2Type, typename Var3Type                   > struct TSPMethodDelegate_ThreeVars       : TBaseSPMethodDelegateInstance<false, UserClass, ESPMode::Fast, TFuncType, Var1Type, Var2Type, Var3Type          > { typedef TBaseSPMethodDelegateInstance<false, UserClass, ESPMode::Fast, TFuncType, Var1Type, Var2Type, Var3Type          > Super; TSPMethodDelegate_ThreeVars      (const TSharedRef< UserClass, ESPMode::Fast >& InUserObject, typename Super::FMethodPtr InMethodPtr, Var1Type Var1, Var2Type Var2, Var3Type Var3               ) : Super(InUserObject, InMethodPtr, Var1, Var2, Var3      ) {} };
+	template <typename UserClass, typename Var1Type, typename Var2Type, typename Var3Type                   > struct TSPMethodDelegate_ThreeVars_Const : TBaseSPMethodDelegateInstance<true , UserClass, ESPMode::Fast, TFuncType, Var1Type, Var2Type, Var3Type          > { typedef TBaseSPMethodDelegateInstance<true , UserClass, ESPMode::Fast, TFuncType, Var1Type, Var2Type, Var3Type          > Super; TSPMethodDelegate_ThreeVars_Const(const TSharedRef< UserClass, ESPMode::Fast >& InUserObject, typename Super::FMethodPtr InMethodPtr, Var1Type Var1, Var2Type Var2, Var3Type Var3               ) : Super(InUserObject, InMethodPtr, Var1, Var2, Var3      ) {} };
+	template <typename UserClass, typename Var1Type, typename Var2Type, typename Var3Type, typename Var4Type> struct TSPMethodDelegate_FourVars        : TBaseSPMethodDelegateInstance<false, UserClass, ESPMode::Fast, TFuncType, Var1Type, Var2Type, Var3Type, Var4Type> { typedef TBaseSPMethodDelegateInstance<false, UserClass, ESPMode::Fast, TFuncType, Var1Type, Var2Type, Var3Type, Var4Type> Super; TSPMethodDelegate_FourVars       (const TSharedRef< UserClass, ESPMode::Fast >& InUserObject, typename Super::FMethodPtr InMethodPtr, Var1Type Var1, Var2Type Var2, Var3Type Var3, Var4Type Var4) : Super(InUserObject, InMethodPtr, Var1, Var2, Var3, Var4) {} };
+	template <typename UserClass, typename Var1Type, typename Var2Type, typename Var3Type, typename Var4Type> struct TSPMethodDelegate_FourVars_Const  : TBaseSPMethodDelegateInstance<true , UserClass, ESPMode::Fast, TFuncType, Var1Type, Var2Type, Var3Type, Var4Type> { typedef TBaseSPMethodDelegateInstance<true , UserClass, ESPMode::Fast, TFuncType, Var1Type, Var2Type, Var3Type, Var4Type> Super; TSPMethodDelegate_FourVars_Const (const TSharedRef< UserClass, ESPMode::Fast >& InUserObject, typename Super::FMethodPtr InMethodPtr, Var1Type Var1, Var2Type Var2, Var3Type Var3, Var4Type Var4) : Super(InUserObject, InMethodPtr, Var1, Var2, Var3, Var4) {} };
+
+	/** Declare the user's "thread-safe" shared pointer-based delegate instance types. */
+	template <typename UserClass                                                                            > struct TThreadSafeSPMethodDelegate                 : TBaseSPMethodDelegateInstance<false, UserClass, ESPMode::ThreadSafe, TFuncType                                        > { typedef TBaseSPMethodDelegateInstance<false, UserClass, ESPMode::ThreadSafe, TFuncType                                        > Super; TThreadSafeSPMethodDelegate                (const TSharedRef<UserClass, ESPMode::ThreadSafe>& InUserObject, typename Super::FMethodPtr InMethodPtr                                                            ) : Super(InUserObject, InMethodPtr                        ) {} };
+	template <typename UserClass                                                                            > struct TThreadSafeSPMethodDelegate_Const           : TBaseSPMethodDelegateInstance<true , UserClass, ESPMode::ThreadSafe, TFuncType                                        > { typedef TBaseSPMethodDelegateInstance<true , UserClass, ESPMode::ThreadSafe, TFuncType                                        > Super; TThreadSafeSPMethodDelegate_Const          (const TSharedRef<UserClass, ESPMode::ThreadSafe>& InUserObject, typename Super::FMethodPtr InMethodPtr                                                            ) : Super(InUserObject, InMethodPtr                        ) {} };
+	template <typename UserClass, typename Var1Type                                                         > struct TThreadSafeSPMethodDelegate_OneVar          : TBaseSPMethodDelegateInstance<false, UserClass, ESPMode::ThreadSafe, TFuncType, Var1Type                              > { typedef TBaseSPMethodDelegateInstance<false, UserClass, ESPMode::ThreadSafe, TFuncType, Var1Type                              > Super; TThreadSafeSPMethodDelegate_OneVar         (const TSharedRef<UserClass, ESPMode::ThreadSafe>& InUserObject, typename Super::FMethodPtr InMethodPtr, Var1Type Var1                                             ) : Super(InUserObject, InMethodPtr, Var1                  ) {} };
+	template <typename UserClass, typename Var1Type                                                         > struct TThreadSafeSPMethodDelegate_OneVar_Const    : TBaseSPMethodDelegateInstance<true , UserClass, ESPMode::ThreadSafe, TFuncType, Var1Type                              > { typedef TBaseSPMethodDelegateInstance<true , UserClass, ESPMode::ThreadSafe, TFuncType, Var1Type                              > Super; TThreadSafeSPMethodDelegate_OneVar_Const   (const TSharedRef<UserClass, ESPMode::ThreadSafe>& InUserObject, typename Super::FMethodPtr InMethodPtr, Var1Type Var1                                             ) : Super(InUserObject, InMethodPtr, Var1                  ) {} };
+	template <typename UserClass, typename Var1Type, typename Var2Type                                      > struct TThreadSafeSPMethodDelegate_TwoVars         : TBaseSPMethodDelegateInstance<false, UserClass, ESPMode::ThreadSafe, TFuncType, Var1Type, Var2Type                    > { typedef TBaseSPMethodDelegateInstance<false, UserClass, ESPMode::ThreadSafe, TFuncType, Var1Type, Var2Type                    > Super; TThreadSafeSPMethodDelegate_TwoVars        (const TSharedRef<UserClass, ESPMode::ThreadSafe>& InUserObject, typename Super::FMethodPtr InMethodPtr, Var1Type Var1, Var2Type Var2                              ) : Super(InUserObject, InMethodPtr, Var1, Var2            ) {} };
+	template <typename UserClass, typename Var1Type, typename Var2Type                                      > struct TThreadSafeSPMethodDelegate_TwoVars_Const   : TBaseSPMethodDelegateInstance<true , UserClass, ESPMode::ThreadSafe, TFuncType, Var1Type, Var2Type                    > { typedef TBaseSPMethodDelegateInstance<true , UserClass, ESPMode::ThreadSafe, TFuncType, Var1Type, Var2Type                    > Super; TThreadSafeSPMethodDelegate_TwoVars_Const  (const TSharedRef<UserClass, ESPMode::ThreadSafe>& InUserObject, typename Super::FMethodPtr InMethodPtr, Var1Type Var1, Var2Type Var2                              ) : Super(InUserObject, InMethodPtr, Var1, Var2            ) {} };
+	template <typename UserClass, typename Var1Type, typename Var2Type, typename Var3Type                   > struct TThreadSafeSPMethodDelegate_ThreeVars       : TBaseSPMethodDelegateInstance<false, UserClass, ESPMode::ThreadSafe, TFuncType, Var1Type, Var2Type, Var3Type          > { typedef TBaseSPMethodDelegateInstance<false, UserClass, ESPMode::ThreadSafe, TFuncType, Var1Type, Var2Type, Var3Type          > Super; TThreadSafeSPMethodDelegate_ThreeVars      (const TSharedRef<UserClass, ESPMode::ThreadSafe>& InUserObject, typename Super::FMethodPtr InMethodPtr, Var1Type Var1, Var2Type Var2, Var3Type Var3               ) : Super(InUserObject, InMethodPtr, Var1, Var2, Var3      ) {} };
+	template <typename UserClass, typename Var1Type, typename Var2Type, typename Var3Type                   > struct TThreadSafeSPMethodDelegate_ThreeVars_Const : TBaseSPMethodDelegateInstance<true , UserClass, ESPMode::ThreadSafe, TFuncType, Var1Type, Var2Type, Var3Type          > { typedef TBaseSPMethodDelegateInstance<true , UserClass, ESPMode::ThreadSafe, TFuncType, Var1Type, Var2Type, Var3Type          > Super; TThreadSafeSPMethodDelegate_ThreeVars_Const(const TSharedRef<UserClass, ESPMode::ThreadSafe>& InUserObject, typename Super::FMethodPtr InMethodPtr, Var1Type Var1, Var2Type Var2, Var3Type Var3               ) : Super(InUserObject, InMethodPtr, Var1, Var2, Var3      ) {} };
+	template <typename UserClass, typename Var1Type, typename Var2Type, typename Var3Type, typename Var4Type> struct TThreadSafeSPMethodDelegate_FourVars        : TBaseSPMethodDelegateInstance<false, UserClass, ESPMode::ThreadSafe, TFuncType, Var1Type, Var2Type, Var3Type, Var4Type> { typedef TBaseSPMethodDelegateInstance<false, UserClass, ESPMode::ThreadSafe, TFuncType, Var1Type, Var2Type, Var3Type, Var4Type> Super; TThreadSafeSPMethodDelegate_FourVars       (const TSharedRef<UserClass, ESPMode::ThreadSafe>& InUserObject, typename Super::FMethodPtr InMethodPtr, Var1Type Var1, Var2Type Var2, Var3Type Var3, Var4Type Var4) : Super(InUserObject, InMethodPtr, Var1, Var2, Var3, Var4) {} };
+	template <typename UserClass, typename Var1Type, typename Var2Type, typename Var3Type, typename Var4Type> struct TThreadSafeSPMethodDelegate_FourVars_Const  : TBaseSPMethodDelegateInstance<true , UserClass, ESPMode::ThreadSafe, TFuncType, Var1Type, Var2Type, Var3Type, Var4Type> { typedef TBaseSPMethodDelegateInstance<true , UserClass, ESPMode::ThreadSafe, TFuncType, Var1Type, Var2Type, Var3Type, Var4Type> Super; TThreadSafeSPMethodDelegate_FourVars_Const (const TSharedRef<UserClass, ESPMode::ThreadSafe>& InUserObject, typename Super::FMethodPtr InMethodPtr, Var1Type Var1, Var2Type Var2, Var3Type Var3, Var4Type Var4) : Super(InUserObject, InMethodPtr, Var1, Var2, Var3, Var4) {} };
+
+	/** Declare the user's C++ pointer-based delegate instance types. */
+	template <typename UserClass                                                                            > struct TRawMethodDelegate                 : TBaseRawMethodDelegateInstance<false, UserClass, TFuncType                                        > { typedef TBaseRawMethodDelegateInstance<false, UserClass, TFuncType                                        > Super; TRawMethodDelegate                (UserClass* InUserObject, typename Super::FMethodPtr InMethodPtr                                                            ) : Super(InUserObject, InMethodPtr                        ) {} };
+	template <typename UserClass                                                                            > struct TRawMethodDelegate_Const           : TBaseRawMethodDelegateInstance<true , UserClass, TFuncType                                        > { typedef TBaseRawMethodDelegateInstance<true , UserClass, TFuncType                                        > Super; TRawMethodDelegate_Const          (UserClass* InUserObject, typename Super::FMethodPtr InMethodPtr                                                            ) : Super(InUserObject, InMethodPtr                        ) {} };
+	template <typename UserClass, typename Var1Type                                                         > struct TRawMethodDelegate_OneVar          : TBaseRawMethodDelegateInstance<false, UserClass, TFuncType, Var1Type                              > { typedef TBaseRawMethodDelegateInstance<false, UserClass, TFuncType, Var1Type                              > Super; TRawMethodDelegate_OneVar         (UserClass* InUserObject, typename Super::FMethodPtr InMethodPtr, Var1Type Var1                                             ) : Super(InUserObject, InMethodPtr, Var1                  ) {} };
+	template <typename UserClass, typename Var1Type                                                         > struct TRawMethodDelegate_OneVar_Const    : TBaseRawMethodDelegateInstance<true , UserClass, TFuncType, Var1Type                              > { typedef TBaseRawMethodDelegateInstance<true , UserClass, TFuncType, Var1Type                              > Super; TRawMethodDelegate_OneVar_Const   (UserClass* InUserObject, typename Super::FMethodPtr InMethodPtr, Var1Type Var1                                             ) : Super(InUserObject, InMethodPtr, Var1                  ) {} };
+	template <typename UserClass, typename Var1Type, typename Var2Type                                      > struct TRawMethodDelegate_TwoVars         : TBaseRawMethodDelegateInstance<false, UserClass, TFuncType, Var1Type, Var2Type                    > { typedef TBaseRawMethodDelegateInstance<false, UserClass, TFuncType, Var1Type, Var2Type                    > Super; TRawMethodDelegate_TwoVars        (UserClass* InUserObject, typename Super::FMethodPtr InMethodPtr, Var1Type Var1, Var2Type Var2                              ) : Super(InUserObject, InMethodPtr, Var1, Var2            ) {} };
+	template <typename UserClass, typename Var1Type, typename Var2Type                                      > struct TRawMethodDelegate_TwoVars_Const   : TBaseRawMethodDelegateInstance<true , UserClass, TFuncType, Var1Type, Var2Type                    > { typedef TBaseRawMethodDelegateInstance<true , UserClass, TFuncType, Var1Type, Var2Type                    > Super; TRawMethodDelegate_TwoVars_Const  (UserClass* InUserObject, typename Super::FMethodPtr InMethodPtr, Var1Type Var1, Var2Type Var2                              ) : Super(InUserObject, InMethodPtr, Var1, Var2            ) {} };
+	template <typename UserClass, typename Var1Type, typename Var2Type, typename Var3Type                   > struct TRawMethodDelegate_ThreeVars       : TBaseRawMethodDelegateInstance<false, UserClass, TFuncType, Var1Type, Var2Type, Var3Type          > { typedef TBaseRawMethodDelegateInstance<false, UserClass, TFuncType, Var1Type, Var2Type, Var3Type          > Super; TRawMethodDelegate_ThreeVars      (UserClass* InUserObject, typename Super::FMethodPtr InMethodPtr, Var1Type Var1, Var2Type Var2, Var3Type Var3               ) : Super(InUserObject, InMethodPtr, Var1, Var2, Var3      ) {} };
+	template <typename UserClass, typename Var1Type, typename Var2Type, typename Var3Type                   > struct TRawMethodDelegate_ThreeVars_Const : TBaseRawMethodDelegateInstance<true , UserClass, TFuncType, Var1Type, Var2Type, Var3Type          > { typedef TBaseRawMethodDelegateInstance<true , UserClass, TFuncType, Var1Type, Var2Type, Var3Type          > Super; TRawMethodDelegate_ThreeVars_Const(UserClass* InUserObject, typename Super::FMethodPtr InMethodPtr, Var1Type Var1, Var2Type Var2, Var3Type Var3               ) : Super(InUserObject, InMethodPtr, Var1, Var2, Var3      ) {} };
+	template <typename UserClass, typename Var1Type, typename Var2Type, typename Var3Type, typename Var4Type> struct TRawMethodDelegate_FourVars        : TBaseRawMethodDelegateInstance<false, UserClass, TFuncType, Var1Type, Var2Type, Var3Type, Var4Type> { typedef TBaseRawMethodDelegateInstance<false, UserClass, TFuncType, Var1Type, Var2Type, Var3Type, Var4Type> Super; TRawMethodDelegate_FourVars       (UserClass* InUserObject, typename Super::FMethodPtr InMethodPtr, Var1Type Var1, Var2Type Var2, Var3Type Var3, Var4Type Var4) : Super(InUserObject, InMethodPtr, Var1, Var2, Var3, Var4) {} };
+	template <typename UserClass, typename Var1Type, typename Var2Type, typename Var3Type, typename Var4Type> struct TRawMethodDelegate_FourVars_Const  : TBaseRawMethodDelegateInstance<true , UserClass, TFuncType, Var1Type, Var2Type, Var3Type, Var4Type> { typedef TBaseRawMethodDelegateInstance<true , UserClass, TFuncType, Var1Type, Var2Type, Var3Type, Var4Type> Super; TRawMethodDelegate_FourVars_Const (UserClass* InUserObject, typename Super::FMethodPtr InMethodPtr, Var1Type Var1, Var2Type Var2, Var3Type Var3, Var4Type Var4) : Super(InUserObject, InMethodPtr, Var1, Var2, Var3, Var4) {} };
+	
+	/** Declare the user's UFunction-based delegate instance types. */
+	template <typename UObjectTemplate                                                                            > struct TUFunctionDelegateBinding           : TBaseUFunctionDelegateInstance<UObjectTemplate, TFuncType                                        > { typedef TBaseUFunctionDelegateInstance<UObjectTemplate, TFuncType                                        > Super; TUFunctionDelegateBinding          (UObjectTemplate* InUserObject, const FName& InFunctionName                                                            ) : Super(InUserObject, InFunctionName                        ) {} };
+	template <typename UObjectTemplate, typename Var1Type                                                         > struct TUFunctionDelegateBinding_OneVar    : TBaseUFunctionDelegateInstance<UObjectTemplate, TFuncType, Var1Type                              > { typedef TBaseUFunctionDelegateInstance<UObjectTemplate, TFuncType, Var1Type                              > Super; TUFunctionDelegateBinding_OneVar   (UObjectTemplate* InUserObject, const FName& InFunctionName, Var1Type Var1                                             ) : Super(InUserObject, InFunctionName, Var1                  ) {} };
+	template <typename UObjectTemplate, typename Var1Type, typename Var2Type                                      > struct TUFunctionDelegateBinding_TwoVars   : TBaseUFunctionDelegateInstance<UObjectTemplate, TFuncType, Var1Type, Var2Type                    > { typedef TBaseUFunctionDelegateInstance<UObjectTemplate, TFuncType, Var1Type, Var2Type                    > Super; TUFunctionDelegateBinding_TwoVars  (UObjectTemplate* InUserObject, const FName& InFunctionName, Var1Type Var1, Var2Type Var2                              ) : Super(InUserObject, InFunctionName, Var1, Var2            ) {} };
+	template <typename UObjectTemplate, typename Var1Type, typename Var2Type, typename Var3Type                   > struct TUFunctionDelegateBinding_ThreeVars : TBaseUFunctionDelegateInstance<UObjectTemplate, TFuncType, Var1Type, Var2Type, Var3Type          > { typedef TBaseUFunctionDelegateInstance<UObjectTemplate, TFuncType, Var1Type, Var2Type, Var3Type          > Super; TUFunctionDelegateBinding_ThreeVars(UObjectTemplate* InUserObject, const FName& InFunctionName, Var1Type Var1, Var2Type Var2, Var3Type Var3               ) : Super(InUserObject, InFunctionName, Var1, Var2, Var3      ) {} };
+	template <typename UObjectTemplate, typename Var1Type, typename Var2Type, typename Var3Type, typename Var4Type> struct TUFunctionDelegateBinding_FourVars  : TBaseUFunctionDelegateInstance<UObjectTemplate, TFuncType, Var1Type, Var2Type, Var3Type, Var4Type> { typedef TBaseUFunctionDelegateInstance<UObjectTemplate, TFuncType, Var1Type, Var2Type, Var3Type, Var4Type> Super; TUFunctionDelegateBinding_FourVars (UObjectTemplate* InUserObject, const FName& InFunctionName, Var1Type Var1, Var2Type Var2, Var3Type Var3, Var4Type Var4) : Super(InUserObject, InFunctionName, Var1, Var2, Var3, Var4) {} };
+
+	/** Declare the user's UObject-based delegate instance types. */
+	template <typename UserClass                                                                            > struct TUObjectMethodDelegate                 : TBaseUObjectMethodDelegateInstance<false, UserClass, TFuncType                                        > { typedef TBaseUObjectMethodDelegateInstance<false, UserClass, TFuncType                                        > Super; TUObjectMethodDelegate                (UserClass* InUserObject, typename Super::FMethodPtr InMethodPtr                                                            ) : Super(InUserObject, InMethodPtr                        ) {} };
+	template <typename UserClass                                                                            > struct TUObjectMethodDelegate_Const           : TBaseUObjectMethodDelegateInstance<true , UserClass, TFuncType                                        > { typedef TBaseUObjectMethodDelegateInstance<true , UserClass, TFuncType                                        > Super; TUObjectMethodDelegate_Const          (UserClass* InUserObject, typename Super::FMethodPtr InMethodPtr                                                            ) : Super(InUserObject, InMethodPtr                        ) {} };
+	template <typename UserClass, typename Var1Type                                                         > struct TUObjectMethodDelegate_OneVar          : TBaseUObjectMethodDelegateInstance<false, UserClass, TFuncType, Var1Type                              > { typedef TBaseUObjectMethodDelegateInstance<false, UserClass, TFuncType, Var1Type                              > Super; TUObjectMethodDelegate_OneVar         (UserClass* InUserObject, typename Super::FMethodPtr InMethodPtr, Var1Type Var1                                             ) : Super(InUserObject, InMethodPtr, Var1                  ) {} };
+	template <typename UserClass, typename Var1Type                                                         > struct TUObjectMethodDelegate_OneVar_Const    : TBaseUObjectMethodDelegateInstance<true , UserClass, TFuncType, Var1Type                              > { typedef TBaseUObjectMethodDelegateInstance<true , UserClass, TFuncType, Var1Type                              > Super; TUObjectMethodDelegate_OneVar_Const   (UserClass* InUserObject, typename Super::FMethodPtr InMethodPtr, Var1Type Var1                                             ) : Super(InUserObject, InMethodPtr, Var1                  ) {} };
+	template <typename UserClass, typename Var1Type, typename Var2Type                                      > struct TUObjectMethodDelegate_TwoVars         : TBaseUObjectMethodDelegateInstance<false, UserClass, TFuncType, Var1Type, Var2Type                    > { typedef TBaseUObjectMethodDelegateInstance<false, UserClass, TFuncType, Var1Type, Var2Type                    > Super; TUObjectMethodDelegate_TwoVars        (UserClass* InUserObject, typename Super::FMethodPtr InMethodPtr, Var1Type Var1, Var2Type Var2                              ) : Super(InUserObject, InMethodPtr, Var1, Var2            ) {} };
+	template <typename UserClass, typename Var1Type, typename Var2Type                                      > struct TUObjectMethodDelegate_TwoVars_Const   : TBaseUObjectMethodDelegateInstance<true , UserClass, TFuncType, Var1Type, Var2Type                    > { typedef TBaseUObjectMethodDelegateInstance<true , UserClass, TFuncType, Var1Type, Var2Type                    > Super; TUObjectMethodDelegate_TwoVars_Const  (UserClass* InUserObject, typename Super::FMethodPtr InMethodPtr, Var1Type Var1, Var2Type Var2                              ) : Super(InUserObject, InMethodPtr, Var1, Var2            ) {} };
+	template <typename UserClass, typename Var1Type, typename Var2Type, typename Var3Type                   > struct TUObjectMethodDelegate_ThreeVars       : TBaseUObjectMethodDelegateInstance<false, UserClass, TFuncType, Var1Type, Var2Type, Var3Type          > { typedef TBaseUObjectMethodDelegateInstance<false, UserClass, TFuncType, Var1Type, Var2Type, Var3Type          > Super; TUObjectMethodDelegate_ThreeVars      (UserClass* InUserObject, typename Super::FMethodPtr InMethodPtr, Var1Type Var1, Var2Type Var2, Var3Type Var3               ) : Super(InUserObject, InMethodPtr, Var1, Var2, Var3      ) {} };
+	template <typename UserClass, typename Var1Type, typename Var2Type, typename Var3Type                   > struct TUObjectMethodDelegate_ThreeVars_Const : TBaseUObjectMethodDelegateInstance<true , UserClass, TFuncType, Var1Type, Var2Type, Var3Type          > { typedef TBaseUObjectMethodDelegateInstance<true , UserClass, TFuncType, Var1Type, Var2Type, Var3Type          > Super; TUObjectMethodDelegate_ThreeVars_Const(UserClass* InUserObject, typename Super::FMethodPtr InMethodPtr, Var1Type Var1, Var2Type Var2, Var3Type Var3               ) : Super(InUserObject, InMethodPtr, Var1, Var2, Var3      ) {} };
+	template <typename UserClass, typename Var1Type, typename Var2Type, typename Var3Type, typename Var4Type> struct TUObjectMethodDelegate_FourVars        : TBaseUObjectMethodDelegateInstance<false, UserClass, TFuncType, Var1Type, Var2Type, Var3Type, Var4Type> { typedef TBaseUObjectMethodDelegateInstance<false, UserClass, TFuncType, Var1Type, Var2Type, Var3Type, Var4Type> Super; TUObjectMethodDelegate_FourVars       (UserClass* InUserObject, typename Super::FMethodPtr InMethodPtr, Var1Type Var1, Var2Type Var2, Var3Type Var3, Var4Type Var4) : Super(InUserObject, InMethodPtr, Var1, Var2, Var3, Var4) {} };
+	template <typename UserClass, typename Var1Type, typename Var2Type, typename Var3Type, typename Var4Type> struct TUObjectMethodDelegate_FourVars_Const  : TBaseUObjectMethodDelegateInstance<true , UserClass, TFuncType, Var1Type, Var2Type, Var3Type, Var4Type> { typedef TBaseUObjectMethodDelegateInstance<true , UserClass, TFuncType, Var1Type, Var2Type, Var3Type, Var4Type> Super; TUObjectMethodDelegate_FourVars_Const (UserClass* InUserObject, typename Super::FMethodPtr InMethodPtr, Var1Type Var1, Var2Type Var2, Var3Type Var3, Var4Type Var4) : Super(InUserObject, InMethodPtr, Var1, Var2, Var3, Var4) {} };
+	
+	/** Declare the user's static function pointer delegate instance types. */
+	                                                                                      struct FStaticDelegate           : TBaseStaticDelegateInstance<TFuncType                                        > { typedef TBaseStaticDelegateInstance<TFuncType                                        > Super; FStaticDelegate          ( typename Super::FFuncPtr InFuncPtr                                                            ) : Super(InFuncPtr                        ) {} };
+	template <typename Var1Type                                                         > struct TStaticDelegate_OneVar    : TBaseStaticDelegateInstance<TFuncType, Var1Type                              > { typedef TBaseStaticDelegateInstance<TFuncType, Var1Type                              > Super; TStaticDelegate_OneVar   ( typename Super::FFuncPtr InFuncPtr, Var1Type Var1                                             ) : Super(InFuncPtr, Var1                  ) {} };
+	template <typename Var1Type, typename Var2Type                                      > struct TStaticDelegate_TwoVars   : TBaseStaticDelegateInstance<TFuncType, Var1Type, Var2Type                    > { typedef TBaseStaticDelegateInstance<TFuncType, Var1Type, Var2Type                    > Super; TStaticDelegate_TwoVars  ( typename Super::FFuncPtr InFuncPtr, Var1Type Var1, Var2Type Var2                              ) : Super(InFuncPtr, Var1, Var2            ) {} };
+	template <typename Var1Type, typename Var2Type, typename Var3Type                   > struct TStaticDelegate_ThreeVars : TBaseStaticDelegateInstance<TFuncType, Var1Type, Var2Type, Var3Type          > { typedef TBaseStaticDelegateInstance<TFuncType, Var1Type, Var2Type, Var3Type          > Super; TStaticDelegate_ThreeVars( typename Super::FFuncPtr InFuncPtr, Var1Type Var1, Var2Type Var2, Var3Type Var3               ) : Super(InFuncPtr, Var1, Var2, Var3      ) {} };
+	template <typename Var1Type, typename Var2Type, typename Var3Type, typename Var4Type> struct TStaticDelegate_FourVars  : TBaseStaticDelegateInstance<TFuncType, Var1Type, Var2Type, Var3Type, Var4Type> { typedef TBaseStaticDelegateInstance<TFuncType, Var1Type, Var2Type, Var3Type, Var4Type> Super; TStaticDelegate_FourVars ( typename Super::FFuncPtr InFuncPtr, Var1Type Var1, Var2Type Var2, Var3Type Var3, Var4Type Var4) : Super(InFuncPtr, Var1, Var2, Var3, Var4) {} };
+
+public:
+
+	/**
+	 * Static: Creates a raw C++ pointer global function delegate
+	 */
+	template <typename... VarTypes>
+	inline static TBaseDelegate<RetValType, ParamTypes...> CreateStatic(typename TBaseStaticDelegateInstance<TFuncType, VarTypes...>::FFuncPtr InFunc, VarTypes... Vars)
+	{
+		return TBaseStaticDelegateInstance<TFuncType, VarTypes...>::Create(InFunc, Vars...);
+	}
+
+	/**
+	 * Static: Creates a raw C++ pointer member function delegate.
+	 *
+	 * Raw pointer doesn't use any sort of reference, so may be unsafe to call if the object was
+	 * deleted out from underneath your delegate. Be careful when calling Execute()!
+	 */
+	template <typename UserClass, typename... VarTypes>
+	inline static TBaseDelegate<RetValType, ParamTypes...> CreateRaw(UserClass* InUserObject, typename TMemFunPtrType<false, UserClass, RetValType (ParamTypes..., VarTypes...)>::Type InFunc, VarTypes... Vars)
+	{
+		return TBaseRawMethodDelegateInstance<false, UserClass, TFuncType, VarTypes...>::Create(InUserObject, InFunc, Vars...);
+	}
+	template <typename UserClass, typename... VarTypes>
+	inline static TBaseDelegate<RetValType, ParamTypes...> CreateRaw(UserClass* InUserObject, typename TMemFunPtrType<true, UserClass, RetValType (ParamTypes..., VarTypes...)>::Type InFunc, VarTypes... Vars)
+	{
+		return TBaseRawMethodDelegateInstance<true , UserClass, TFuncType, VarTypes...>::Create( InUserObject, InFunc, Vars...);
+	}
+	
+	/**
+	 * Static: Creates a shared pointer-based (fast, not thread-safe) member function delegate.
+	 *
+	 * Shared pointer delegates keep a weak reference to your object.
+	 * You can use ExecuteIfBound() to call them.
+	 */
+	template <typename UserClass, typename... VarTypes>
+	inline static TBaseDelegate<RetValType, ParamTypes...> CreateSP(const TSharedRef<UserClass, ESPMode::Fast>& InUserObjectRef, typename TMemFunPtrType<false, UserClass, RetValType (ParamTypes..., VarTypes...)>::Type InFunc, VarTypes... Vars)
+	{
+		return TBaseSPMethodDelegateInstance<false, UserClass, ESPMode::Fast, TFuncType, VarTypes...>::Create(InUserObjectRef, InFunc, Vars...);
+	}
+	template <typename UserClass, typename... VarTypes>
+	inline static TBaseDelegate<RetValType, ParamTypes...> CreateSP(const TSharedRef<UserClass, ESPMode::Fast>& InUserObjectRef, typename TMemFunPtrType<true, UserClass, RetValType (ParamTypes..., VarTypes...)>::Type InFunc, VarTypes... Vars)
+	{
+		return TBaseSPMethodDelegateInstance<true , UserClass, ESPMode::Fast, TFuncType, VarTypes...>::Create(InUserObjectRef, InFunc, Vars...);
+	}
+
+	/**
+	 * Static: Creates a shared pointer-based (fast, not thread-safe) member function delegate.
+	 *
+	 * Shared pointer delegates keep a weak reference to your object.
+	 * You can use ExecuteIfBound() to call them.
+	 */
+	template <typename UserClass, typename... VarTypes>
+	inline static TBaseDelegate<RetValType, ParamTypes...> CreateSP(UserClass* InUserObject, typename TMemFunPtrType<false, UserClass, RetValType (ParamTypes..., VarTypes...)>::Type InFunc, VarTypes... Vars)
+	{
+		return CreateSP(StaticCastSharedRef<UserClass>(InUserObject->AsShared()), InFunc, Vars...);
+	}
+	template <typename UserClass, typename... VarTypes>
+	inline static TBaseDelegate<RetValType, ParamTypes...> CreateSP(UserClass* InUserObject, typename TMemFunPtrType<true, UserClass, RetValType (ParamTypes..., VarTypes...)>::Type InFunc, VarTypes... Vars)
+	{
+		return CreateSP(StaticCastSharedRef<UserClass>(InUserObject->AsShared()), InFunc, Vars...);
+	}
+	
+	/**
+	 * Static: Creates a shared pointer-based (slower, conditionally thread-safe) member function delegate.
+	 *
+	 * Shared pointer delegates keep a weak reference to your object.
+	 * You can use ExecuteIfBound() to call them.
+	 */
+	template <typename UserClass, typename... VarTypes>
+	inline static TBaseDelegate<RetValType, ParamTypes...> CreateThreadSafeSP(const TSharedRef<UserClass, ESPMode::ThreadSafe>& InUserObjectRef, typename TMemFunPtrType<false, UserClass, RetValType (ParamTypes..., VarTypes...)>::Type InFunc, VarTypes... Vars)
+	{
+		return TBaseSPMethodDelegateInstance<false, UserClass, ESPMode::ThreadSafe, TFuncType, VarTypes...>::Create(InUserObjectRef, InFunc, Vars...);
+	}
+	template <typename UserClass, typename... VarTypes>
+	inline static TBaseDelegate<RetValType, ParamTypes...> CreateThreadSafeSP(const TSharedRef<UserClass, ESPMode::ThreadSafe>& InUserObjectRef, typename TMemFunPtrType<true, UserClass, RetValType (ParamTypes..., VarTypes...)>::Type InFunc, VarTypes... Vars)
+	{
+		return TBaseSPMethodDelegateInstance<true , UserClass, ESPMode::ThreadSafe, TFuncType, VarTypes...>::Create(InUserObjectRef, InFunc, Vars...);
+	}
+
+	/**
+	 * Static: Creates a shared pointer-based (slower, conditionally thread-safe) member function delegate.
+	 *
+	 * Shared pointer delegates keep a weak reference to your object.
+	 * You can use ExecuteIfBound() to call them.
+	 */
+	template <typename UserClass, typename... VarTypes>
+	inline static TBaseDelegate<RetValType, ParamTypes...> CreateThreadSafeSP(UserClass* InUserObject, typename TMemFunPtrType<false, UserClass, RetValType (ParamTypes..., VarTypes...)>::Type InFunc, VarTypes... Vars)
+	{
+		return CreateThreadSafeSP(StaticCastSharedRef<UserClass>(InUserObject->AsShared()), InFunc, Vars...);
+	}
+	template <typename UserClass, typename... VarTypes>
+	inline static TBaseDelegate<RetValType, ParamTypes...> CreateThreadSafeSP(UserClass* InUserObject, typename TMemFunPtrType<true, UserClass, RetValType (ParamTypes..., VarTypes...)>::Type InFunc, VarTypes... Vars)
+	{
+		return CreateThreadSafeSP(StaticCastSharedRef<UserClass>(InUserObject->AsShared()), InFunc, Vars...);
+	}
+
+	/**
+	 * Static: Creates a UFunction-based member function delegate.
+	 *
+	 * UFunction delegates keep a weak reference to your object.
+	 * You can use ExecuteIfBound() to call them.
+	 */
+	template <typename UObjectTemplate, typename... VarTypes>
+	inline static TBaseDelegate<RetValType, ParamTypes...> CreateUFunction(UObjectTemplate* InUserObject, const FName& InFunctionName, VarTypes... Vars)
+	{
+		return TBaseUFunctionDelegateInstance<UObjectTemplate, TFuncType, VarTypes...>::Create(InUserObject, InFunctionName, Vars...);
+	}
+
+	/**
+	 * Static: Creates a UObject-based member function delegate.
+	 *
+	 * UObject delegates keep a weak reference to your object.
+	 * You can use ExecuteIfBound() to call them.
+	 */
+	template <typename UserClass, typename... VarTypes>
+	inline static TBaseDelegate<RetValType, ParamTypes...> CreateUObject(UserClass* InUserObject, typename TMemFunPtrType<false, UserClass, RetValType (ParamTypes..., VarTypes...)>::Type InFunc, VarTypes... Vars)
+	{
+		return TBaseUObjectMethodDelegateInstance<false, UserClass, TFuncType, VarTypes...>::Create(InUserObject, InFunc, Vars...);
+	}
+	template <typename UserClass, typename... VarTypes>
+	inline static TBaseDelegate<RetValType, ParamTypes...> CreateUObject(UserClass* InUserObject, typename TMemFunPtrType<true, UserClass, RetValType (ParamTypes..., VarTypes...)>::Type InFunc, VarTypes... Vars)
+	{
+		return TBaseUObjectMethodDelegateInstance<true , UserClass, TFuncType, VarTypes...>::Create(InUserObject, InFunc, Vars...);
+	}
+
+public:
+
+	/**
+	 * Default constructor
+	 */
+	inline TBaseDelegate()
+		: FDelegateBase(nullptr)
+	{ }
+
+	/**
+	 * Destructor.
+	 */
+	inline ~TBaseDelegate()
+	{
+		Unbind();
+	}
+
+	/**
+	 * Creates and initializes a new instance with the given delegate instance.
+	 *
+	 * The delegate will assume ownership of the incoming delegate instance!
+	 * IMPORTANT: This is a system-internal function and you should never be using this in regular C++ code
+	 *
+	 * @param InDelegateInstance The delegate instance to assign.
+	 */
+	inline TBaseDelegate( TDelegateInstanceInterface* InDelegateInstance )
+		: FDelegateBase(InDelegateInstance)
+	{ }
+
+	/**
+	 * Creates and initializes a new instance from an existing delegate object.
+	 *
+	 * @param Other The delegate object to copy from.
+	 */
+	inline TBaseDelegate( const TBaseDelegate& Other )
+		: FDelegateBase(nullptr)
+	{
+		*this = Other;
+	}
+
+	/**
+	 * Assignment operator.
+	 *
+	 * @param	OtherDelegate	Delegate object to copy from
+	 */
+	inline TBaseDelegate& operator=( const TBaseDelegate& Other )
+	{
+		if (&Other != this)
+		{
+			// this down-cast is OK! allows for managing invocation list in the base class without requiring virtual functions
+			TDelegateInstanceInterface* OtherInstance = (TDelegateInstanceInterface*)Other.GetDelegateInstance();
+
+			if (OtherInstance != nullptr)
+			{
+				SetDelegateInstance(OtherInstance->CreateCopy());
+			}
+			else
+			{
+				SetDelegateInstance(nullptr);
+			}
+		}
+
+		return *this;
+	}
+
+public:
+
+	/**
+	 * Creates a raw C++ pointer global function delegate
+	 */
+	template <typename... VarTypes>
+	inline void BindStatic(typename TBaseStaticDelegateInstance<TFuncType, VarTypes...>::FFuncPtr InFunc, VarTypes... Vars)
+	{
+		*this = CreateStatic(InFunc, Vars...);
+	}
+
+	/**
+	 * Binds a raw C++ pointer delegate.
+	 *
+	 * Raw pointer doesn't use any sort of reference, so may be unsafe to call if the object was
+	 * deleted out from underneath your delegate. Be careful when calling Execute()!
+	 */
+	template <typename UserClass, typename... VarTypes>
+	inline void BindRaw(UserClass* InUserObject, typename TMemFunPtrType<false, UserClass, RetValType (ParamTypes..., VarTypes...)>::Type InFunc, VarTypes... Vars)
+	{
+		*this = CreateRaw(InUserObject, InFunc, Vars...);
+	}
+	template <typename UserClass, typename... VarTypes>
+	inline void BindRaw(UserClass* InUserObject, typename TMemFunPtrType<true, UserClass, RetValType (ParamTypes..., VarTypes...)>::Type InFunc, VarTypes... Vars)
+	{
+		*this = CreateRaw(InUserObject, InFunc, Vars...);
+	}
+
+	/**
+	 * Binds a shared pointer-based (fast, not thread-safe) member function delegate.  Shared pointer delegates keep a weak reference to your object.  You can use ExecuteIfBound() to call them.
+	 */
+	template <typename UserClass, typename... VarTypes>
+	inline void BindSP(const TSharedRef<UserClass, ESPMode::Fast>& InUserObjectRef, typename TMemFunPtrType<false, UserClass, RetValType (ParamTypes..., VarTypes...)>::Type InFunc, VarTypes... Vars)
+	{
+		*this = CreateSP(InUserObjectRef, InFunc, Vars...);
+	}
+	template <typename UserClass, typename... VarTypes>
+	inline void BindSP(const TSharedRef<UserClass, ESPMode::Fast>& InUserObjectRef, typename TMemFunPtrType<true, UserClass, RetValType (ParamTypes..., VarTypes...)>::Type InFunc, VarTypes... Vars)
+	{
+		*this = CreateSP(InUserObjectRef, InFunc, Vars...);
+	}
+
+	/**
+	 * Creates a shared pointer-based (fast, not thread-safe) member function delegate.
+	 *
+	 * Shared pointer delegates keep a weak reference to your object.
+	 * You can use ExecuteIfBound() to call them.
+	 */
+	template <typename UserClass, typename... VarTypes>
+	inline void BindSP(UserClass* InUserObject, typename TMemFunPtrType<false, UserClass, RetValType (ParamTypes..., VarTypes...)>::Type InFunc, VarTypes... Vars)
+	{
+		*this = CreateSP(InUserObject, InFunc, Vars...);
+	}
+	template <typename UserClass, typename... VarTypes>
+	inline void BindSP(UserClass* InUserObject, typename TMemFunPtrType<true, UserClass, RetValType (ParamTypes..., VarTypes...)>::Type InFunc, VarTypes... Vars)
+	{
+		*this = CreateSP(InUserObject, InFunc, Vars...);
+	}
+
+	/**
+	 * Binds a shared pointer-based (slower, conditionally thread-safe) member function delegate.
+	 *
+	 * Shared pointer delegates keep a weak reference to your object.
+	 * You can use ExecuteIfBound() to call them.
+	 */
+	template <typename UserClass, typename... VarTypes>
+	inline void BindThreadSafeSP(const TSharedRef<UserClass, ESPMode::ThreadSafe>& InUserObjectRef, typename TMemFunPtrType<false, UserClass, RetValType (ParamTypes..., VarTypes...)>::Type InFunc, VarTypes... Vars)
+	{
+		*this = CreateThreadSafeSP(InUserObjectRef, InFunc, Vars...);
+	}
+	template <typename UserClass, typename... VarTypes>
+	inline void BindThreadSafeSP(const TSharedRef<UserClass, ESPMode::ThreadSafe>& InUserObjectRef, typename TMemFunPtrType<true, UserClass, RetValType (ParamTypes..., VarTypes...)>::Type InFunc, VarTypes... Vars)
+	{
+		*this = CreateThreadSafeSP(InUserObjectRef, InFunc, Vars...);
+	}
+
+	/**
+	 * Creates a shared pointer-based (slower, conditionally thread-safe) member function delegate.
+	 *
+	 * Shared pointer delegates keep a weak reference to your object.
+	 * You can use ExecuteIfBound() to call them.
+	 */
+	template <typename UserClass, typename... VarTypes>
+	inline void BindThreadSafeSP(UserClass* InUserObject, typename TMemFunPtrType<false, UserClass, RetValType (ParamTypes..., VarTypes...)>::Type InFunc, VarTypes... Vars)
+	{
+		*this = CreateThreadSafeSP(InUserObject, InFunc, Vars...);
+	}
+	template <typename UserClass, typename... VarTypes>
+	inline void BindThreadSafeSP(UserClass* InUserObject, typename TMemFunPtrType<true, UserClass, RetValType (ParamTypes..., VarTypes...)>::Type InFunc, VarTypes... Vars)
+	{
+		*this = CreateThreadSafeSP(InUserObject, InFunc, Vars...);
+	}
+
+	/**
+	 * Creates a UFunction-based member function delegate.
+	 *
+	 * UFunction delegates keep a weak reference to your object.
+	 * You can use ExecuteIfBound() to call them.
+	 */
+	template <typename UObjectTemplate, typename... VarTypes>
+	inline void BindUFunction(UObjectTemplate* InUserObject, const FName& InFunctionName, VarTypes... Vars)
+	{
+		*this = CreateUFunction(InUserObject, InFunctionName, Vars...);
+	}
+
+	/**
+	 * Creates a UObject-based member function delegate.
+	 *
+	 * UObject delegates keep a weak reference to your object.
+	 * You can use ExecuteIfBound() to call them.
+	 */
+	template <typename UserClass, typename... VarTypes>
+	inline void BindUObject(UserClass* InUserObject, typename TMemFunPtrType<false, UserClass, RetValType (ParamTypes..., VarTypes...)>::Type InFunc, VarTypes... Vars)
+	{
+		*this = CreateUObject(InUserObject, InFunc, Vars...);
+	}
+	template <typename UserClass, typename... VarTypes>
+	inline void BindUObject(UserClass* InUserObject, typename TMemFunPtrType<true, UserClass, RetValType (ParamTypes..., VarTypes...)>::Type InFunc, VarTypes... Vars)
+	{
+		*this = CreateUObject(InUserObject, InFunc, Vars...);
+	}
+
+public:
+	/**
+	 * Execute the delegate.
+	 *
+	 * If the function pointer is not valid, an error will occur. Check IsBound() before
+	 * calling this method or use ExecuteIfBound() instead.
+	 *
+	 * @see ExecuteIfBound
+	 */
+	inline RetValType Execute(ParamTypes... Params) const
+	{
+		TDelegateInstanceInterface* DelegateInstance = (TDelegateInstanceInterface*)GetDelegateInstance();
+
+		// If this assert goes off, Execute() was called before a function was bound to the delegate.
+		// Consider using ExecuteIfSafe() instead.
+		checkSlow(DelegateInstance != nullptr);
+
+		return DelegateInstance->Execute(Params...);
+	}
+
+	/**
+	 * Equality operator.
+	 *
+	 * @return true if this delegate equals the other, false otherwise.
+	 */
+	bool operator==( const TBaseDelegate& Other ) const
+	{
+		TDelegateInstanceInterface* DelegateInstance = (TDelegateInstanceInterface*)GetDelegateInstance();
+		TDelegateInstanceInterface* OtherInstance = (TDelegateInstanceInterface*)Other.GetDelegateInstance();
+
+		// The function these delegates point to must be the same
+		if ((DelegateInstance != nullptr) && (OtherInstance != nullptr))
+		{
+			return DelegateInstance->IsSameFunction(*OtherInstance);
+		}
+
+		// If neither delegate is initialized to anything yet, then we treat them as equal
+		if ((DelegateInstance == nullptr) && (OtherInstance == nullptr))
+		{
+			return true;
+		}
+
+		// No match!
+		return false;
+	}
+};
+
+template <typename... ParamTypes>
+class TBaseDelegate<void, ParamTypes...> : public TBaseDelegate<TTypeWrapper<void>, ParamTypes...>
+{
+	typedef TBaseDelegate<TTypeWrapper<void>, ParamTypes...> Super;
+
+public:
+	typedef typename Super::TDelegateInstanceInterface TDelegateInstanceInterface;
+
+	/**
+	 * Default constructor
+	 */
+	TBaseDelegate()
+	{
+	}
+
+	/**
+	 * Creates and initializes a new instance with the given delegate instance.
+	 *
+	 * The delegate will assume ownership of the incoming delegate instance!
+	 * IMPORTANT: This is a system-internal function and you should never be using this in regular C++ code
+	 *
+	 * @param InDelegateInstance The delegate instance to assign.
+	 */
+	TBaseDelegate(TDelegateInstanceInterface* InDelegateInstance)
+		: Super(InDelegateInstance)
+	{
+	}
+
+	/**
+	 * Execute the delegate, but only if the function pointer is still valid
+	 *
+	 * @return  Returns true if the function was executed
+	 */
+	// NOTE: Currently only delegates with no return value support ExecuteIfBound() 
+	inline bool ExecuteIfBound(ParamTypes... Params) const
+	{
+		if (Super::IsBound())
+		{
+			return ((TDelegateInstanceInterface*)Super::GetDelegateInstance())->ExecuteIfSafe(Params...);
+		}
+
+		return false;
+	}
+};
+
+
+/**
+ * Multicast delegate base class.
+ *
+ * This class implements the functionality of multicast delegates. It is templated to the function signature
+ * that it is compatible with. Use the various DECLARE_MULTICAST_DELEGATE and DECLARE_EVENT macros to create
+ * actual delegate types.
+ *
+ * Multicast delegates offer no guarantees for the calling order of bound functions. As bindings get added
+ * and removed over time, the calling order may change. Only bindings without return values are supported.
+ */
+template <typename RetValType, typename... ParamTypes>
+class TBaseMulticastDelegate;
+
+template <typename... ParamTypes>
+class TBaseMulticastDelegate<void, ParamTypes...> : public FMulticastDelegateBase<>
+{
+public:
+	/** DEPRECATED: Type definition for unicast delegate classes whose delegate instances are compatible with this delegate. */
+	typedef TBaseDelegate< void, ParamTypes... > FDelegate;
+
+	/** Type definition for the shared interface of delegate instance types compatible with this delegate class. */
+	typedef IBaseDelegateInstance<void (ParamTypes...)> TDelegateInstanceInterface;
+
+public:
+
+	/**
+	 * Adds a delegate instance to this multicast delegate's invocation list.
+	 *
+	 * This delegate will take over ownership of the given delegate instance.
+	 *
+	 * @param DelegateInstance The delegate instance to add.
+	 */
+	void Add( TDelegateInstanceInterface* DelegateInstance )
+	{
+		if (DelegateInstance != nullptr)
+		{
+			AddDelegateInstance(DelegateInstance);
+		}
+	}
+
+	/**
+	 * DEPRECATED: Adds a unicast delegate to this multi-cast delegate's invocation list.
+	 *
+	 * This method is retained for backwards compatibility.
+	 *
+	 * @param Delegate The delegate to add.
+	 */
+	void Add( const FDelegate& Delegate )
+	{
+		TDelegateInstanceInterface* DelegateInstance = (TDelegateInstanceInterface*)Delegate.GetDelegateInstance();
+
+		if (DelegateInstance != nullptr)
+		{
+			AddDelegateInstance(DelegateInstance->CreateCopy());
+		}
+	}
+
+	/**
+	 * Adds a raw C++ pointer global function delegate
+	 *
+	 * @param	InFunc	Function pointer
+	 */
+	template <typename... VarTypes>
+	inline void AddStatic(typename TBaseStaticDelegateInstance<void (ParamTypes...), VarTypes...>::FFuncPtr InFunc, VarTypes... Vars)
+	{
+		Add(FDelegate::CreateStatic(InFunc, Vars...));
+	}
+
+	/**
+	 * Adds a raw C++ pointer delegate.
+	 *
+	 * Raw pointer doesn't use any sort of reference, so may be unsafe to call if the object was
+	 * deleted out from underneath your delegate. Be careful when calling Execute()!
+	 *
+	 * @param	InUserObject	User object to bind to
+	 * @param	InFunc			Class method function address
+	 */
+	template <typename UserClass, typename... VarTypes>
+	inline void AddRaw(UserClass* InUserObject, typename TMemFunPtrType<false, UserClass, void (ParamTypes..., VarTypes...)>::Type InFunc, VarTypes... Vars)
+	{
+		Add(FDelegate::CreateRaw(InUserObject, InFunc, Vars...));
+	}
+	template <typename UserClass, typename... VarTypes>
+	inline void AddRaw(UserClass* InUserObject, typename TMemFunPtrType<true, UserClass, void (ParamTypes..., VarTypes...)>::Type InFunc, VarTypes... Vars)
+	{
+		Add(FDelegate::CreateRaw(InUserObject, InFunc, Vars...));
+	}
+
+	/**
+	 * Adds a shared pointer-based (fast, not thread-safe) member function delegate.
+	 *
+	 * Shared pointer delegates keep a weak reference to your object.
+	 *
+	 * @param	InUserObjectRef	User object to bind to
+	 * @param	InFunc			Class method function address
+	 */
+	template <typename UserClass, typename... VarTypes>
+	inline void AddSP(const TSharedRef<UserClass, ESPMode::Fast>& InUserObjectRef, typename TMemFunPtrType<false, UserClass, void (ParamTypes..., VarTypes...)>::Type InFunc, VarTypes... Vars)
+	{
+		Add(FDelegate::CreateSP(InUserObjectRef, InFunc, Vars...));
+	}
+	template <typename UserClass, typename... VarTypes>
+	inline void AddSP(const TSharedRef<UserClass, ESPMode::Fast>& InUserObjectRef, typename TMemFunPtrType<true, UserClass, void (ParamTypes..., VarTypes...)>::Type InFunc, VarTypes... Vars)
+	{
+		Add(FDelegate::CreateSP(InUserObjectRef, InFunc, Vars...));
+	}
+
+	/**
+	 * Adds a shared pointer-based (fast, not thread-safe) member function delegate.
+	 *
+	 * Shared pointer delegates keep a weak reference to your object.
+	 *
+	 * @param	InUserObject	User object to bind to
+	 * @param	InFunc			Class method function address
+	 */
+	template <typename UserClass, typename... VarTypes>
+	inline void AddSP(UserClass* InUserObject, typename TMemFunPtrType<false, UserClass, void (ParamTypes..., VarTypes...)>::Type InFunc, VarTypes... Vars)
+	{
+		Add(FDelegate::CreateSP(InUserObject, InFunc, Vars...));
+	}
+	template <typename UserClass, typename... VarTypes>
+	inline void AddSP(UserClass* InUserObject, typename TMemFunPtrType<true, UserClass, void (ParamTypes..., VarTypes...)>::Type InFunc, VarTypes... Vars)
+	{
+		Add(FDelegate::CreateSP(InUserObject, InFunc, Vars...));
+	}
+
+	/**
+	 * Adds a shared pointer-based (slower, conditionally thread-safe) member function delegate.  Shared pointer delegates keep a weak reference to your object.
+	 *
+	 * @param	InUserObjectRef	User object to bind to
+	 * @param	InFunc			Class method function address
+	 */
+	template <typename UserClass, typename... VarTypes>
+	inline void AddThreadSafeSP(const TSharedRef<UserClass, ESPMode::ThreadSafe>& InUserObjectRef, typename TMemFunPtrType<false, UserClass, void (ParamTypes..., VarTypes...)>::Type InFunc, VarTypes... Vars)
+	{
+		Add(FDelegate::CreateThreadSafeSP(InUserObjectRef, InFunc, Vars...));
+	}
+	template <typename UserClass, typename... VarTypes>
+	inline void AddThreadSafeSP(const TSharedRef<UserClass, ESPMode::ThreadSafe>& InUserObjectRef, typename TMemFunPtrType<true, UserClass, void (ParamTypes..., VarTypes...)>::Type InFunc, VarTypes... Vars)
+	{
+		Add(FDelegate::CreateThreadSafeSP(InUserObjectRef, InFunc, Vars...));
+	}
+
+	/**
+	 * Adds a shared pointer-based (slower, conditionally thread-safe) member function delegate.
+	 *
+	 * Shared pointer delegates keep a weak reference to your object.
+	 *
+	 * @param	InUserObject	User object to bind to
+	 * @param	InFunc			Class method function address
+	 */
+	template <typename UserClass, typename... VarTypes>
+	inline void AddThreadSafeSP(UserClass* InUserObject, typename TMemFunPtrType<false, UserClass, void (ParamTypes..., VarTypes...)>::Type InFunc, VarTypes... Vars)
+	{
+		Add(FDelegate::CreateThreadSafeSP(InUserObject, InFunc, Vars...));
+	}
+	template <typename UserClass, typename... VarTypes>
+	inline void AddThreadSafeSP(UserClass* InUserObject, typename TMemFunPtrType<true, UserClass, void (ParamTypes..., VarTypes...)>::Type InFunc, VarTypes... Vars)
+	{
+		Add(FDelegate::CreateThreadSafeSP(InUserObject, InFunc, Vars...));
+	}
+
+	/**
+	 * Adds a UFunction-based member function delegate.
+	 *
+	 * UFunction delegates keep a weak reference to your object.
+	 *
+	 * @param	InUserObject	User object to bind to
+	 * @param	InFunctionName			Class method function address
+	 */
+	template <typename UObjectTemplate, typename... VarTypes>
+	inline void AddUFunction(UObjectTemplate* InUserObject, const FName& InFunctionName, VarTypes... Vars)
+	{
+		Add(FDelegate::CreateUFunction(InUserObject, InFunctionName, Vars...));
+	}
+
+	/**
+	 * Adds a UObject-based member function delegate.
+	 *
+	 * UObject delegates keep a weak reference to your object.
+	 *
+	 * @param	InUserObject	User object to bind to
+	 * @param	InFunc			Class method function address
+	 */
+	template <typename UserClass, typename... VarTypes>
+	inline void AddUObject(UserClass* InUserObject, typename TMemFunPtrType<false, UserClass, void (ParamTypes..., VarTypes...)>::Type InFunc, VarTypes... Vars)
+	{
+		Add(FDelegate::CreateUObject(InUserObject, InFunc, Vars...));
+	}
+	template <typename UserClass, typename... VarTypes>
+	inline void AddUObject(UserClass* InUserObject, typename TMemFunPtrType<true, UserClass, void (ParamTypes..., VarTypes...)>::Type InFunc, VarTypes... Vars)
+	{
+		Add(FDelegate::CreateUObject(InUserObject, InFunc, Vars...));
+	}
+
+public:
+
+	/**
+	 * Removes a delegate instance from this multi-cast delegate's invocation list (performance is O(N)).
+	 *
+	 * Note that the order of the delegate instances may not be preserved!
+	 *
+	 * @param DelegateInstance The delegate instance to remove.
+	 */
+	void Remove( const TDelegateInstanceInterface& DelegateInstance )
+	{
+		// NOTE: For delegates with PAYLOAD, the payload variables are not considered when comparing delegate objects!
+		// This means that two delegate instances bound to the same object/method, but with different payload data,
+		// cannot be bound to the same single multi-cast delegate.  We could support this by changing IsSameFunction()
+		// to do comparison of payload data, but there are some trade-offs if we do that.
+
+		RemoveDelegateInstance(&DelegateInstance);
+	}
+
+	/**
+	 * DEPRECATED: Removes a unicast delegate from this multi-cast delegate's invocation list (performance is O(N)).
+	 *
+	 * The order of the delegates may not be preserved!
+	 * This function is retained for backwards compatibility.
+	 *
+	 * @param Delegate The delegate to remove.
+	 */
+	void Remove( const FDelegate& Delegate )
+	{
+		// NOTE: For delegates with PAYLOAD, the payload variables are not considered when comparing delegate objects!
+		// This means that two delegate instances bound to the same object/method, but with different payload data,
+		// cannot be bound to the same single multi-cast delegate.  We could support this by changing IsSameFunction()
+		// to do comparison of payload data, but there are some trade-offs if we do that.
+
+		TDelegateInstanceInterface* DelegateInstance = (TDelegateInstanceInterface*)Delegate.GetDelegateInstance();
+
+		if (DelegateInstance != nullptr)
+		{
+			RemoveDelegateInstance(DelegateInstance);
+		}		
+	}
+
+	// NOTE: These direct Remove methods are only supported for multi-cast delegates with no payload attached.
+	// See the comment in the multi-cast delegate Remove() method above for more details.
+
+	/**
+	 * Removes a raw C++ pointer global function delegate (performance is O(N)).  Note that the order of the
+	 * delegates may not be preserved!
+	 *
+	 * @param	InFunc	Function pointer
+	 */
+	inline void RemoveStatic( typename TBaseStaticDelegateInstance< void (ParamTypes...) >::FFuncPtr InFunc )
+	{
+		Remove( FDelegate::CreateStatic( InFunc ) );
+	}
+
+	/**
+	 * Removes a raw C++ pointer delegate (performance is O(N)).  Note that the order of the delegates may not
+	 * be preserved!
+	 *
+	 * @param	InUserObject	User object to unbind from
+	 * @param	InFunc			Class method function address
+	 */
+	template< class UserClass >
+	inline void RemoveRaw( UserClass* InUserObject, typename TBaseRawMethodDelegateInstance<false, UserClass, void (ParamTypes...)>::FMethodPtr InFunc )
+	{
+		Remove( FDelegate::CreateRaw( InUserObject, InFunc ) );
+	}
+	template< class UserClass >
+	inline void RemoveRaw( UserClass* InUserObject, typename TBaseRawMethodDelegateInstance<true , UserClass, void (ParamTypes...)>::FMethodPtr InFunc )
+	{
+		Remove( FDelegate::CreateRaw( InUserObject, InFunc ) );
+	}
+
+	/**
+	 * Removes a shared pointer-based member function delegate (performance is O(N)).  Note that the order of
+	 * the delegates may not be preserved!
+	 *
+	 * @param	InUserObjectRef	User object to unbind from
+	 * @param	InFunc			Class method function address
+	 */
+	template< class UserClass >
+	inline void RemoveSP( const TSharedRef< UserClass, ESPMode::Fast >& InUserObjectRef, typename TBaseSPMethodDelegateInstance<false, UserClass, ESPMode::Fast, void (ParamTypes...)>::FMethodPtr InFunc )
+	{
+		Remove( FDelegate::CreateSP( InUserObjectRef, InFunc ) );
+	}
+	template< class UserClass >
+	inline void RemoveSP( const TSharedRef< UserClass, ESPMode::Fast >& InUserObjectRef, typename TBaseSPMethodDelegateInstance<true , UserClass, ESPMode::Fast, void (ParamTypes...)>::FMethodPtr InFunc )
+	{
+		Remove( FDelegate::CreateSP( InUserObjectRef, InFunc ) );
+	}
+
+	/**
+	 * Removes a shared pointer-based member function delegate (performance is O(N)).  Note that the order of
+	 * the delegates may not be preserved!
+	 *
+	 * @param	InUserObject	User object to unbind from
+	 * @param	InFunc			Class method function address
+	 */
+	template< class UserClass >
+	inline void RemoveSP( UserClass* InUserObject, typename TBaseSPMethodDelegateInstance<false, UserClass, ESPMode::Fast, void (ParamTypes...)>::FMethodPtr InFunc )
+	{
+		Remove( FDelegate::CreateSP( InUserObject, InFunc ) );
+	}
+	template< class UserClass >
+	inline void RemoveSP( UserClass* InUserObject, typename TBaseSPMethodDelegateInstance<true , UserClass, ESPMode::Fast, void (ParamTypes...)>::FMethodPtr InFunc )
+	{
+		Remove( FDelegate::CreateSP( InUserObject, InFunc ) );
+	}
+
+	/**
+	 * Removes a shared pointer-based member function delegate (performance is O(N)).  Note that the order of
+	 * the delegates may not be preserved!
+	 *
+	 * @param	InUserObjectRef	User object to unbind from
+	 * @param	InFunc			Class method function address
+	 */
+	template< class UserClass >
+	inline void RemoveThreadSafeSP( const TSharedRef< UserClass, ESPMode::ThreadSafe >& InUserObjectRef, typename TBaseSPMethodDelegateInstance<false, UserClass, ESPMode::ThreadSafe, void (ParamTypes...)>::FMethodPtr InFunc )
+	{
+		Remove( FDelegate::CreateThreadSafeSP( InUserObjectRef, InFunc ) );
+	}
+	template< class UserClass >
+	inline void RemoveThreadSafeSP( const TSharedRef< UserClass, ESPMode::ThreadSafe >& InUserObjectRef, typename TBaseSPMethodDelegateInstance<true , UserClass, ESPMode::ThreadSafe, void (ParamTypes...)>::FMethodPtr InFunc )
+	{
+		Remove( FDelegate::CreateThreadSafeSP( InUserObjectRef, InFunc ) );
+	}
+
+	/**
+	 * Removes a shared pointer-based member function delegate (performance is O(N)).  Note that the order of
+	 * the delegates may not be preserved!
+	 *
+	 * @param	InUserObject	User object to unbind from
+	 * @param	InFunc			Class method function address
+	 */
+	template< class UserClass >
+	inline void RemoveThreadSafeSP( UserClass* InUserObject, typename TBaseSPMethodDelegateInstance<false, UserClass, ESPMode::ThreadSafe, void (ParamTypes...)>::FMethodPtr InFunc )
+	{
+		Remove( FDelegate::CreateThreadSafeSP( InUserObject, InFunc ) );
+	}
+	template< class UserClass >
+	inline void RemoveThreadSafeSP( UserClass* InUserObject, typename TBaseSPMethodDelegateInstance<true , UserClass, ESPMode::ThreadSafe, void (ParamTypes...)>::FMethodPtr InFunc )
+	{
+		Remove( FDelegate::CreateThreadSafeSP( InUserObject, InFunc ) );
+	}
+
+	/**
+	 * Removes a UFunction-based member function delegate (performance is O(N)).  Note that the order of the
+	 * delegates may not be preserved!
+	 *
+	 * @param	InUserObject	User object to unbind from
+	 * @param	InFunc			Class method function address
+	 */
+	template< class UObjectTemplate >
+	inline void RemoveUFunction( UObjectTemplate* InUserObject, const FName& InFunctionName )
+	{
+		Remove( FDelegate::CreateUFunction( InUserObject, InFunctionName ) );
+	}
+
+	/**
+	 * Removes a UObject-based member function delegate (performance is O(N)).  Note that the order of the
+	 * delegates may not be preserved!
+	 *
+	 * @param	InUserObject	User object to unbind from
+	 * @param	InFunc			Class method function address
+	 */
+	template< class UserClass >
+	inline void RemoveUObject( UserClass* InUserObject, typename TBaseUObjectMethodDelegateInstance<false, UserClass, void (ParamTypes...)>::FMethodPtr InFunc )
+	{
+		Remove( FDelegate::CreateUObject( InUserObject, InFunc ) );
+	}
+	template< class UserClass >
+	inline void RemoveUObject( UserClass* InUserObject, typename TBaseUObjectMethodDelegateInstance<true , UserClass, void (ParamTypes...)>::FMethodPtr InFunc )
+	{
+		Remove( FDelegate::CreateUObject( InUserObject, InFunc ) );
+	}
+
+protected:
+
+	/** 
+	 * Hidden default constructor.
+	 */
+	inline TBaseMulticastDelegate( ) { }
+
+	/**
+	 * Hidden copy constructor (for proper deep copies).
+	 *
+	 * @param Other The multicast delegate to copy from.
+	 */
+	TBaseMulticastDelegate( const TBaseMulticastDelegate& Other )
+	{
+		*this = Other;
+	}
+
+	/**
+	 * Hidden assignment operator (for proper deep copies).
+	 *
+	 * @param Other The delegate to assign from.
+	 * @return This instance.
+	 */
+	TBaseMulticastDelegate& operator=( const TBaseMulticastDelegate& Other )
+	{
+		if (&Other != this)
+		{
+			Clear();
+
+			for (IDelegateInstance* DelegateInstance : Other.GetInvocationList())
+			{
+				// this down-cast is OK! allows for managing invocation list in the base class without requiring virtual functions
+				TDelegateInstanceInterface* OtherInstance = (TDelegateInstanceInterface*)DelegateInstance;
+
+				if (OtherInstance != nullptr)
+				{
+					AddInternal(OtherInstance->CreateCopy());
+				}			
+			}
+		}
+
+		return *this;
+	}
+
+protected:
+
+	/**
+	 * Adds a function delegate to this multi-cast delegate's invocation list.
+	 *
+	 * This method will assert if the same function has already been bound.
+	 *
+	 * @param	InDelegate	Delegate to add
+	 */
+	void AddDelegateInstance( TDelegateInstanceInterface* InDelegateInstance )
+	{
+#if !UE_BUILD_SHIPPING && !UE_BUILD_TEST
+		// verify that the same function isn't already bound
+		for (IDelegateInstance* DelegateInstance : GetInvocationList())
+		{
+			if (DelegateInstance != nullptr)
+			{
+				// this down-cast is OK! allows for managing invocation list in the base class without requiring virtual functions
+				TDelegateInstanceInterface* DelegateInstanceInterface = (TDelegateInstanceInterface*)DelegateInstance;
+				check(!DelegateInstanceInterface->IsSameFunction(*InDelegateInstance));
+			}
+		}
+#endif
+
+		AddInternal(InDelegateInstance);
+	}
+
+public:
+	/**
+	 * Broadcasts this delegate to all bound objects, except to those that may have expired.
+	 *
+	 * The constness of this method is a lie, but it allows for broadcasting from const functions.
+	 */
+	void Broadcast(ParamTypes... Params) const
+	{
+		bool NeedsCompaction = false;
+
+		LockInvocationList();
+		{
+			const TArray<IDelegateInstance*>& InvocationList = GetInvocationList();
+
+			// call bound functions in reverse order, so we ignore any instances that may be added by callees
+			for (int32 InvocationListIndex = InvocationList.Num() - 1; InvocationListIndex >= 0; --InvocationListIndex)
+			{
+				// this down-cast is OK! allows for managing invocation list in the base class without requiring virtual functions
+				TDelegateInstanceInterface* DelegateInstanceInterface = (TDelegateInstanceInterface*)InvocationList[InvocationListIndex];
+
+				if ((DelegateInstanceInterface == nullptr) || !DelegateInstanceInterface->ExecuteIfSafe(Params...))
+				{
+					NeedsCompaction = true;
+				}
+			}
+		}
+		UnlockInvocationList();
+
+		if (NeedsCompaction)
+		{
+			const_cast<TBaseMulticastDelegate*>(this)->CompactInvocationList();
+		}
+	}
+
+protected:
+	/**
+	 * Removes a function from this multi-cast delegate's invocation list (performance is O(N)).
+	 *
+	 * The function is not actually removed, but deleted and marked as removed.
+	 * It will be removed next time the invocation list is compacted within Broadcast().
+	 *
+	 * @param InDelegateInstance The delegate instance to remove.
+	 */
+	void RemoveDelegateInstance( const TDelegateInstanceInterface* InDelegateInstance )
+	{
+		const TArray<IDelegateInstance*>& InvocationList = GetInvocationList();
+
+		// NOTE: We assume that this method is never called with a nullptr object, in which case the
+		//       the following algorithm would break down (it would remove the first found instance
+		//       of a matching function binding, which is not necessarily the instance we wish to remove).
+
+		for (int32 InvocationListIndex = 0; InvocationListIndex < InvocationList.Num(); ++InvocationListIndex)
+		{
+			// this down-cast is OK! allows for managing invocation list in the base class without requiring virtual functions
+			TDelegateInstanceInterface*& DelegateInstanceRef = (TDelegateInstanceInterface*&)InvocationList[InvocationListIndex];
+
+			// NOTE: We must do a deep compare here, not just compare delegate pointers, because multiple
+			//       delegate pointers can refer to the exact same object and method
+			if ((DelegateInstanceRef != nullptr) && DelegateInstanceRef->IsSameFunction(*InDelegateInstance))
+			{
+				delete DelegateInstanceRef;
+				DelegateInstanceRef = nullptr;
+				
+				break;	// no need to continue, as we never allow the same delegate to be bound twice
+			}
+		}
+
+		const_cast<TBaseMulticastDelegate*>(this)->CompactInvocationList();
+	}
+};
+
+
+/**
+ * Implements a multicast delegate.
+ *
+ * This class should not be instantiated directly. Use the DECLARE_MULTICAST_DELEGATE macros instead.
+ */
+template <typename RetValType, typename... ParamTypes>
+class TMulticastDelegate;
+
+template <typename... ParamTypes>
+class TMulticastDelegate<void, ParamTypes...> : public TBaseMulticastDelegate<void, ParamTypes... >
+{
+private:
+	// Prevents erroneous use by other classes.
+	typedef TBaseMulticastDelegate< void, ParamTypes... > Super;
+};
+
+
+/**
+ * Dynamic delegate base object (UObject-based, serializable).  You'll use the various DECLARE_DYNAMIC_DELEGATE
+ * macros to create the actual delegate type, templated to the function signature the delegate is compatible with.
+ * Then, you can create an instance of that class when you want to assign functions to the delegate.
+ */
+template <typename TWeakPtr, typename RetValType, typename... ParamTypes>
+class TBaseDynamicDelegate : public TScriptDelegate<TWeakPtr>
+{
+public:
+	/**
+	 * Default constructor
+	 */
+	TBaseDynamicDelegate() { }
+
+	/**
+	 * Construction from an FScriptDelegate must be explicit.  This is really only used by UObject system internals.
+	 *
+	 * @param	InScriptDelegate	The delegate to construct from by copying
+	 */
+	explicit TBaseDynamicDelegate( const TScriptDelegate<TWeakPtr>& InScriptDelegate )
+		: TScriptDelegate<TWeakPtr>( InScriptDelegate )
+	{ }
+
+	/**
+	 * Templated helper class to define a typedef for user's method pointer, then used below
+	 */
+	template< class UserClass >
+	class TMethodPtrResolver
+	{
+	public:
+		typedef RetValType (UserClass::*FMethodPtr)(ParamTypes... Params);
+	};
+
+	/**
+	 * Binds a UObject instance and a UObject method address to this delegate.
+	 *
+	 * @param	InUserObject			UObject instance
+	 * @param	InMethodPtr				Member function address pointer
+	 * @param	InMacroFunctionName		Name of member function, including class name (generated by a macro)
+	 *
+	 * NOTE:  Do not call this function directly.  Instead, call BindDynamic() which is a macro proxy function that
+	 *        automatically sets the function name string for the caller.
+	 */
+	template< class UserClass >
+	void __Internal_BindDynamic( UserClass* InUserObject, typename TMethodPtrResolver< UserClass >::FMethodPtr InMethodPtr, const FString& InMacroFunctionName )
+	{
+		check( InUserObject != nullptr && InMethodPtr != nullptr && !InMacroFunctionName.IsEmpty() );
+
+		// NOTE: We're not actually storing the incoming method pointer or calling it.  We simply require it for type-safety reasons.
+
+		// NOTE: If you hit a compile error on the following line, it means you're trying to use a non-UObject type
+		//       with this delegate, which is not supported
+		this->Object = InUserObject;
+
+		// Store the function name.  The incoming function name was generated by a macro and includes the method's class name.
+		// We'll need to strip off that class prefix and just store the function name by itself.
+		this->FunctionName = *InMacroFunctionName.Mid( InMacroFunctionName.Find( TEXT( "::" ), ESearchCase::CaseSensitive, ESearchDir::FromEnd) + 2 );
+
+		ensureMsgf(this->IsBound(), TEXT("Unable to bind delegate to '%s' (function might not be marked as a UFUNCTION)"), *InMacroFunctionName);
+	}
+
+	friend uint32 GetTypeHash(const TBaseDynamicDelegate& Key)
+	{
+		return FCrc::MemCrc_DEPRECATED(&Key,sizeof(Key));
+	}
+
+	// NOTE:  Execute() method must be defined in derived classes
+
+	// NOTE:  ExecuteIfBound() method must be defined in derived classes
+};
+
+
+/**
+ * Dynamic multi-cast delegate base object (UObject-based, serializable).  You'll use the various
+ * DECLARE_DYNAMIC_MULTICAST_DELEGATE macros to create the actual delegate type, templated to the function
+ * signature the delegate is compatible with.   Then, you can create an instance of that class when you
+ * want to assign functions to the delegate.
+ */
+template <typename TWeakPtr, typename RetValType, typename... ParamTypes>
+class TBaseDynamicMulticastDelegate : public TMulticastScriptDelegate<TWeakPtr>
+{
+public:
+	/** The actual single-cast delegate class for this multi-cast delegate */
+	typedef TBaseDynamicDelegate<FWeakObjectPtr, RetValType, ParamTypes...> FDelegate;
+
+	/**
+	 * Default constructor
+	 */
+	TBaseDynamicMulticastDelegate() { }
+
+	/**
+	 * Construction from an FMulticastScriptDelegate must be explicit.  This is really only used by UObject system internals.
+	 *
+	 * @param	InScriptDelegate	The delegate to construct from by copying
+	 */
+	explicit TBaseDynamicMulticastDelegate( const TMulticastScriptDelegate<TWeakPtr>& InMulticastScriptDelegate )
+		: TMulticastScriptDelegate<TWeakPtr>( InMulticastScriptDelegate )
+	{ }
+
+	/**
+	 * Tests if a UObject instance and a UObject method address pair are already bound to this multi-cast delegate.
+	 *
+	 * @param	InUserObject			UObject instance
+	 * @param	InMethodPtr				Member function address pointer
+	 * @param	InMacroFunctionName		Name of member function, including class name (generated by a macro)
+	 * @return	True if the instance/method is already bound.
+	 *
+	 * NOTE:  Do not call this function directly.  Instead, call IsAlreadyBound() which is a macro proxy function that
+	 *        automatically sets the function name string for the caller.
+	 */
+	template< class UserClass >
+	bool __Internal_IsAlreadyBound( UserClass* InUserObject, typename FDelegate::template TMethodPtrResolver< UserClass >::FMethodPtr InMethodPtr, const FString& InMacroFunctionName ) const
+	{
+		check( InUserObject != nullptr && InMethodPtr != nullptr && !InMacroFunctionName.IsEmpty() );
+
+		// NOTE: We're not actually using the incoming method pointer or calling it.  We simply require it for type-safety reasons.
+
+		FDelegate NewDelegate;
+		NewDelegate.__Internal_BindDynamic( InUserObject, InMethodPtr, InMacroFunctionName );
+
+		return this->Contains( NewDelegate );
+	}
+
+	/**
+	 * Binds a UObject instance and a UObject method address to this multi-cast delegate.
+	 *
+	 * @param	InUserObject			UObject instance
+	 * @param	InMethodPtr				Member function address pointer
+	 * @param	InMacroFunctionName		Name of member function, including class name (generated by a macro)
+	 *
+	 * NOTE:  Do not call this function directly.  Instead, call AddDynamic() which is a macro proxy function that
+	 *        automatically sets the function name string for the caller.
+	 */
+	template< class UserClass >
+	void __Internal_AddDynamic( UserClass* InUserObject, typename FDelegate::template TMethodPtrResolver< UserClass >::FMethodPtr InMethodPtr, const FString& InMacroFunctionName )
+	{
+		check( InUserObject != nullptr && InMethodPtr != nullptr && !InMacroFunctionName.IsEmpty() );
+
+		// NOTE: We're not actually storing the incoming method pointer or calling it.  We simply require it for type-safety reasons.
+
+		FDelegate NewDelegate;
+		NewDelegate.__Internal_BindDynamic( InUserObject, InMethodPtr, InMacroFunctionName );
+
+		this->Add( NewDelegate );
+	}
+
+	/**
+	 * Binds a UObject instance and a UObject method address to this multi-cast delegate, but only if it hasn't been bound before.
+	 *
+	 * @param	InUserObject			UObject instance
+	 * @param	InMethodPtr				Member function address pointer
+	 * @param	InMacroFunctionName		Name of member function, including class name (generated by a macro)
+	 *
+	 * NOTE:  Do not call this function directly.  Instead, call AddUniqueDynamic() which is a macro proxy function that
+	 *        automatically sets the function name string for the caller.
+	 */
+	template< class UserClass >
+	void __Internal_AddUniqueDynamic( UserClass* InUserObject, typename FDelegate::template TMethodPtrResolver< UserClass >::FMethodPtr InMethodPtr, const FString& InMacroFunctionName )
+	{
+		check( InUserObject != nullptr && InMethodPtr != nullptr && !InMacroFunctionName.IsEmpty() );
+
+		// NOTE: We're not actually storing the incoming method pointer or calling it.  We simply require it for type-safety reasons.
+
+		FDelegate NewDelegate;
+		NewDelegate.__Internal_BindDynamic( InUserObject, InMethodPtr, InMacroFunctionName );
+
+		this->AddUnique( NewDelegate );
+	}
+
+	/**
+	 * Unbinds a UObject instance and a UObject method address from this multi-cast delegate.
+	 *
+	 * @param	InUserObject			UObject instance
+	 * @param	InMethodPtr				Member function address pointer
+	 * @param	InMacroFunctionName		Name of member function, including class name (generated by a macro)
+	 *
+	 * NOTE:  Do not call this function directly.  Instead, call RemoveDynamic() which is a macro proxy function that
+	 *        automatically sets the function name string for the caller.
+	 */
+	template< class UserClass >
+	void __Internal_RemoveDynamic( UserClass* InUserObject, typename FDelegate::template TMethodPtrResolver< UserClass >::FMethodPtr InMethodPtr, const FString& InMacroFunctionName )
+	{
+		check( InUserObject != nullptr && InMethodPtr != nullptr && !InMacroFunctionName.IsEmpty() );
+
+		// NOTE: We're not actually storing the incoming method pointer or calling it.  We simply require it for type-safety reasons.
+
+		FDelegate NewDelegate;
+		NewDelegate.__Internal_BindDynamic( InUserObject, InMethodPtr, InMacroFunctionName );
+
+		this->Remove( NewDelegate );
+	}
+
+	// NOTE:  Broadcast() method must be defined in derived classes
+};
diff --git a/Engine/Source/Runtime/Core/Public/Delegates/IntegerSequence.h b/Engine/Source/Runtime/Core/Public/Delegates/IntegerSequence.h
new file mode 100644
--- /dev/null
+++ b/Engine/Source/Runtime/Core/Public/Delegates/IntegerSequence.h
@@ -0,0 +1,20 @@
+// Copyright 1998-2014 Epic Games, Inc. All Rights Reserved.
+
+#pragma once
+
+template <uint32... Indices>
+struct TIntegerSequence;
+
+template <uint32 ToAdd, uint32... Values>
+struct TMakeIntegerSequenceImpl : TMakeIntegerSequenceImpl<ToAdd - 1, Values..., sizeof...(Values)>
+{
+};
+
+template <uint32... Values>
+struct TMakeIntegerSequenceImpl<0, Values...>
+{
+	typedef TIntegerSequence<Values...> Type;
+};
+
+template <uint32 ToAdd>
+using TMakeIntegerSequence = typename TMakeIntegerSequenceImpl<ToAdd>::Type;
diff --git a/Engine/Source/Runtime/Core/Public/Delegates/Tuple.h b/Engine/Source/Runtime/Core/Public/Delegates/Tuple.h
new file mode 100644
--- /dev/null
+++ b/Engine/Source/Runtime/Core/Public/Delegates/Tuple.h
@@ -0,0 +1,133 @@
+// Copyright 1998-2014 Epic Games, Inc. All Rights Reserved.
+
+#pragma once
+
+#include "IntegerSequence.h"
+
+template <typename T, uint32 Index>
+struct TTupleElement
+{
+	template <typename... ArgTypes>
+	explicit TTupleElement(ArgTypes&&... Args)
+		: Value(Forward<ArgTypes>(Args)...)
+	{
+	}
+
+	T Value;
+};
+
+template <uint32 IterIndex, uint32 Index, typename... Types>
+struct TTupleElementHelperImpl;
+
+template <uint32 IterIndex, uint32 Index, typename ElementType, typename... Types>
+struct TTupleElementHelperImpl<IterIndex, Index, ElementType, Types...> : TTupleElementHelperImpl<IterIndex + 1, Index, Types...>
+{
+};
+
+template <uint32 Index, typename ElementType, typename... Types>
+struct TTupleElementHelperImpl<Index, Index, ElementType, Types...>
+{
+	typedef ElementType Type;
+
+	template <typename TupleType>
+	static FORCEINLINE ElementType& Get(TupleType& Tuple)
+	{
+		return static_cast<TTupleElement<ElementType, Index>&>(Tuple).Value;
+	}
+
+	template <typename TupleType>
+	static FORCEINLINE const ElementType& Get(const TupleType& Tuple)
+	{
+		return Get((TupleType&)Tuple);
+	}
+};
+
+template <uint32 WantedIndex, typename... Types>
+struct TTupleElementHelper : TTupleElementHelperImpl<0, WantedIndex, Types...>
+{
+};
+
+template <typename... Types>
+struct TTuple;
+
+template <typename Indices, typename... Types>
+struct TTupleImpl;
+
+template <uint32... Indices, typename... Types>
+struct TTupleImpl<TIntegerSequence<Indices...>, Types...> : TTupleElement<Types, Indices>...
+{
+	template <typename... ArgTypes>
+	explicit TTupleImpl(ArgTypes&&... Args)
+		: TTupleElement<Types, Indices>(Forward<ArgTypes>(Args))...
+	{
+	}
+
+	template <uint32 Index> FORCEINLINE const typename TTupleElementHelper<Index, Types...>::Type& Get() const { return TTupleElementHelper<Index, Types...>::Get(*this); }
+	template <uint32 Index> FORCEINLINE       typename TTupleElementHelper<Index, Types...>::Type& Get()       { return TTupleElementHelper<Index, Types...>::Get(*this); }
+
+	template <typename FuncType, typename... ArgTypes>
+	auto ApplyAfter(FuncType&& Func, ArgTypes&&... Args) const -> decltype(Func(Forward<ArgTypes>(Args)..., Get<Indices>()...))
+	{
+		return Func(Forward<ArgTypes>(Args)..., Get<Indices>()...);
+	}
+
+	// This exists because VC isn't so great at deducing return types sometimes, so we can specify it if we know it
+	template <typename RetType, typename FuncType, typename... ArgTypes>
+	RetType ApplyAfter_ExplicitReturnType(FuncType&& Func, ArgTypes&&... Args) const
+	{
+		return Func(Forward<ArgTypes>(Args)..., Get<Indices>()...);
+	}
+
+	template <typename FuncType, typename... ArgTypes>
+	auto ApplyBefore(FuncType&& Func, ArgTypes&&... Args) const -> decltype(Func(Get<Indices>()..., Forward<ArgTypes>(Args)...))
+	{
+		return Func(Get<Indices>()..., Forward<ArgTypes>(Args)...);
+	}
+};
+
+#ifdef _MSC_VER
+
+	// Not strictly necessary, but some VC versions give a 'syntax error: <fake-expression>' error
+	// for empty tuples.
+	template <>
+	struct TTupleImpl<TIntegerSequence<>>
+	{
+		explicit TTupleImpl()
+		{
+		}
+
+		// Doesn't matter what these return, or even have a function body, but they need to be declared
+		template <uint32 Index> FORCEINLINE const int32& Get() const;
+		template <uint32 Index> FORCEINLINE       int32& Get();
+
+		template <typename FuncType, typename... ArgTypes>
+		auto ApplyAfter(FuncType&& Func, ArgTypes&&... Args) const -> decltype(Func(Forward<ArgTypes>(Args)...))
+		{
+			return Func(Forward<ArgTypes>(Args)...);
+		}
+
+		// This exists because VC isn't so great at deducing return types sometimes, so we can specify it if we know it
+		template <typename RetType, typename FuncType, typename... ArgTypes>
+		RetType ApplyAfter_ExplicitReturnType(FuncType&& Func, ArgTypes&&... Args) const
+		{
+			return Func(Forward<ArgTypes>(Args)...);
+		}
+
+		template <typename FuncType, typename... ArgTypes>
+		auto ApplyBefore(FuncType&& Func, ArgTypes&&... Args) const -> decltype(Func(Forward<ArgTypes>(Args)...))
+		{
+			return Func(Forward<ArgTypes>(Args)...);
+		}
+	};
+
+#endif
+
+template <typename... Types>
+struct TTuple : TTupleImpl<TMakeIntegerSequence<sizeof...(Types)>, Types...>
+{
+	template <typename... ArgTypes>
+	explicit TTuple(ArgTypes&&... Args)
+		: TTupleImpl<TMakeIntegerSequence<sizeof...(Types)>, Types...>(Forward<ArgTypes>(Args)...)
+	{
+	}
+};
diff --git a/Engine/Source/Runtime/Core/Public/Templates/UnrealTemplate.h b/Engine/Source/Runtime/Core/Public/Templates/UnrealTemplate.h
--- a/Engine/Source/Runtime/Core/Public/Templates/UnrealTemplate.h
+++ b/Engine/Source/Runtime/Core/Public/Templates/UnrealTemplate.h
@@ -517,3 +517,68 @@ struct TIdentity
 {
 	typedef T Type;
 };
+
+/**
+ * Wraps a type.
+ *
+ * The intended use of this template is to allow types to be passed around where the unwrapped type would give different behavior.
+ * An example of this is when you want a template specialization to refer to the primary template, but doing that would cause
+ * infinite recursion through the specialization, e.g.:
+ *
+ * // Before
+ * template <typename T>
+ * struct TThing
+ * {
+ *     void f(T t)
+ *     {
+ *          DoSomething(t);
+ *     }
+ * };
+ *
+ * template <>
+ * struct TThing<int>
+ * {
+ *     void f(int t)
+ *     {
+ *         DoSomethingElseFirst(t);
+ *         TThing<int>::f(t); // Infinite recursion
+ *     }
+ * };
+ *
+ *
+ * // After
+ * template <typename T>
+ * struct TThing
+ * {
+ *     typedef typename TUnwrapType<T>::Type RealT;
+ *
+ *     void f(RealT t)
+ *     {
+ *          DoSomething(t);
+ *     }
+ * };
+ *
+ * template <>
+ * struct TThing<int>
+ * {
+ *     void f(int t)
+ *     {
+ *         DoSomethingElseFirst(t);
+ *         TThing<TTypeWrapper<int>>::f(t); // works
+ *     }
+ * };
+ */
+template <typename T>
+struct TTypeWrapper;
+
+template <typename T>
+struct TUnwrapType
+{
+	typedef T Type;
+};
+
+template <typename T>
+struct TUnwrapType<TTypeWrapper<T>>
+{
+	typedef T Type;
+};
diff --git a/Engine/Source/Runtime/Core/Public/Windows/WindowsPlatformCompilerSetup.h b/Engine/Source/Runtime/Core/Public/Windows/WindowsPlatformCompilerSetup.h
--- a/Engine/Source/Runtime/Core/Public/Windows/WindowsPlatformCompilerSetup.h
+++ b/Engine/Source/Runtime/Core/Public/Windows/WindowsPlatformCompilerSetup.h
@@ -102,6 +102,7 @@ static_assert(_MSC_VER >= 1700, "Visual Studio 2012 or newer required.");
 //#pragma warning(disable : 4305) // 'argument' : truncation from 'double' to 'float' --- fp:precise builds only!!!!!!!
 #pragma warning(disable : 4324) // structure was padded due to __declspec(align())						
 #pragma warning(disable : 4355) // this used in base initializer list                                   
+#pragma warning(disable : 4373) // '%$S': virtual function overrides '%$pS', previous versions of the compiler did not override when parameters only differed by const/volatile qualifiers
 #pragma warning(disable : 4389) // signed/unsigned mismatch                                             
 #pragma warning(disable : 4511) // copy constructor could not be generated                              
 #pragma warning(disable : 4512) // assignment operator could not be generated                           
diff --git a/Engine/Source/Runtime/Messaging/Private/Bus/MessageRouter.cpp b/Engine/Source/Runtime/Messaging/Private/Bus/MessageRouter.cpp
--- a/Engine/Source/Runtime/Messaging/Private/Bus/MessageRouter.cpp
+++ b/Engine/Source/Runtime/Messaging/Private/Bus/MessageRouter.cpp
@@ -41,7 +41,7 @@ uint32 FMessageRouter::Run( )
 		{
 			CurrentTime = FDateTime::UtcNow();
 
-			TBaseDelegate_NoParams<void> Command;
+			CommandDelegate Command;
 
 			while (Commands.Dequeue(Command))
 			{
diff --git a/Engine/Source/Runtime/Messaging/Private/Bus/MessageRouter.h b/Engine/Source/Runtime/Messaging/Private/Bus/MessageRouter.h
--- a/Engine/Source/Runtime/Messaging/Private/Bus/MessageRouter.h
+++ b/Engine/Source/Runtime/Messaging/Private/Bus/MessageRouter.h
@@ -10,6 +10,7 @@ class FMessageRouter
 	: public FRunnable
 {
 public:
+	DECLARE_DELEGATE(CommandDelegate)
 
 	/** Default constructor. */
 	FMessageRouter( );
@@ -128,7 +129,7 @@ protected:
 	 * @param Command The command to queue up.
 	 * @return true if the command was enqueued, false otherwise.
 	 */
-	FORCEINLINE bool EnqueueCommand( TBaseDelegate_NoParams<void> Command )
+	FORCEINLINE bool EnqueueCommand( CommandDelegate Command )
 	{
 		if (!Commands.Enqueue(Command))
 		{
@@ -230,7 +231,7 @@ private:
 	TMap<FName, TArray<IMessageSubscriptionPtr>> ActiveSubscriptions;
 
 	/** Holds the router command queue. */
-	TQueue<TBaseDelegate_NoParams<void>, EQueueMode::Mpsc> Commands;
+	TQueue<CommandDelegate, EQueueMode::Mpsc> Commands;
 
 	/** Holds the current time. */
 	FDateTime CurrentTime;
diff --git a/Engine/Source/Runtime/Messaging/Private/Bus/MessageTracer.cpp b/Engine/Source/Runtime/Messaging/Private/Bus/MessageTracer.cpp
--- a/Engine/Source/Runtime/Messaging/Private/Bus/MessageTracer.cpp
+++ b/Engine/Source/Runtime/Messaging/Private/Bus/MessageTracer.cpp
@@ -62,7 +62,7 @@ bool FMessageTracer::Tick( float DeltaTime )
 	}
 
 	// process new traces
-	TBaseDelegate_NoParams<void> Trace;
+	TraceDelegate Trace;
 
 	while (Traces.Dequeue(Trace))
 	{
diff --git a/Engine/Source/Runtime/Messaging/Private/Bus/MessageTracer.h b/Engine/Source/Runtime/Messaging/Private/Bus/MessageTracer.h
--- a/Engine/Source/Runtime/Messaging/Private/Bus/MessageTracer.h
+++ b/Engine/Source/Runtime/Messaging/Private/Bus/MessageTracer.h
@@ -40,6 +40,7 @@ class FMessageTracer
 	};
 
 public:
+	DECLARE_DELEGATE(TraceDelegate)
 
 	/** Default constructor. */
 	FMessageTracer( );
@@ -317,7 +318,7 @@ protected:
 	 *
 	 * @param Trace The action to enqueue.
 	 */
-	FORCEINLINE void EnqueueTrace( TBaseDelegate_NoParams<void> Trace )
+	FORCEINLINE void EnqueueTrace( TraceDelegate Trace )
 	{
 		Traces.Enqueue(Trace);
 	}
@@ -452,7 +453,7 @@ private:
 	bool Running;
 
 	/** Holds the trace actions queue. */
-	TQueue<TBaseDelegate_NoParams<void>, EQueueMode::Mpsc> Traces;
+	TQueue<TraceDelegate, EQueueMode::Mpsc> Traces;
 
 private:
 
