Enum class codegen fixes for UFunctions, as raised here:

https://github.com/EpicGames/UnrealEngine/pull/433

diff --git a/Engine/Source/Programs/UnrealHeaderTool/Resources/UHTDebugging/TestObject.h b/Engine/Source/Programs/UnrealHeaderTool/Resources/UHTDebugging/TestObject.h
--- a/Engine/Source/Programs/UnrealHeaderTool/Resources/UHTDebugging/TestObject.h
+++ b/Engine/Source/Programs/UnrealHeaderTool/Resources/UHTDebugging/TestObject.h
@@ -17,6 +17,9 @@ class UTestObject : public UObject
 	UPROPERTY()
 	ECppEnum EnumProperty;
 
+	UFUNCTION()
+	void CodeGenTestForEnumClasses(ECppEnum Val);
+
 #if BLAH
 	UFUNCTION() int x; // This should not compile if UHT parses it, which it shouldn't
 #elif BLAH2
diff --git a/Engine/Source/Runtime/Core/Public/Containers/EnumAsByte.h b/Engine/Source/Runtime/Core/Public/Containers/EnumAsByte.h
--- a/Engine/Source/Runtime/Core/Public/Containers/EnumAsByte.h
+++ b/Engine/Source/Runtime/Core/Public/Containers/EnumAsByte.h
@@ -32,7 +32,7 @@ public:
 	 * @param InValue value to construct with.
 	 */
 	FORCEINLINE TEnumAsByte( TEnum InValue )
-		: Value(InValue)
+		: Value(static_cast<uint8>(InValue))
 	{ }
 
 	/**
@@ -41,7 +41,7 @@ public:
 	 * @param InValue value to construct with.
 	 */
 	explicit FORCEINLINE TEnumAsByte( int32 InValue )
-		: Value(InValue)
+		: Value(static_cast<uint8>(InValue))
 	{ }
 
 	/**
@@ -72,7 +72,7 @@ public:
 	 */
 	FORCEINLINE TEnumAsByte& operator=( TEnum InValue )
 	{
-		Value = InValue;
+		Value = static_cast<uint8>(InValue);
 		return *this;
 	}
 
@@ -83,7 +83,7 @@ public:
 	 */
 	bool operator==( TEnum InValue ) const
 	{
-		return (Value == InValue);
+		return static_cast<TEnum>(Value) == InValue;
 	}
 
 	/**
@@ -93,7 +93,7 @@ public:
 	 */
 	bool operator==(TEnumAsByte InValue) const
 	{
-		return (Value == InValue.Value);
+		return Value == InValue.Value;
 	}
 
 	/**
@@ -101,7 +101,7 @@ public:
 	 */
 	operator TEnum( ) const
 	{
-		return TEnum(Value);
+		return (TEnum)Value;
 	}
 
 public:
@@ -111,7 +111,7 @@ public:
 	 */
 	TEnum GetValue( ) const
 	{
-		return TEnum(Value);
+		return (TEnum)Value;
 	}
 
 private:
diff --git a/Engine/Source/Runtime/CoreUObject/Private/UObject/PropertyByte.cpp b/Engine/Source/Runtime/CoreUObject/Private/UObject/PropertyByte.cpp
--- a/Engine/Source/Runtime/CoreUObject/Private/UObject/PropertyByte.cpp
+++ b/Engine/Source/Runtime/CoreUObject/Private/UObject/PropertyByte.cpp
@@ -97,11 +97,11 @@ FString UByteProperty::GetCPPType( FString* ExtendedTypeText/*=NULL*/, uint32 CP
 
 		if ((CPPExportFlags & CPPF_ArgumentOrReturnValue) && ((PropertyFlags & CPF_ReturnParm) || !(PropertyFlags & CPF_OutParm)))
 		{
-			return FString::Printf(TEXT("enum %s"), *FullyQualifiedEnumName);
+			return FullyQualifiedEnumName;
 		}
 		else
 		{
-			return FString::Printf(TEXT("TEnumAsByte<enum %s>"), *FullyQualifiedEnumName);
+			return FString::Printf(TEXT("TEnumAsByte<%s>"), *FullyQualifiedEnumName);
 		}
 	}
 	return Super::GetCPPType(ExtendedTypeText, CPPExportFlags);
