CrashReport - Fixed bad engine version stored as Parameter0

diff --git a/Engine/Source/Runtime/Core/Private/Linux/LinuxPlatformMisc.cpp b/Engine/Source/Runtime/Core/Private/Linux/LinuxPlatformMisc.cpp
--- a/Engine/Source/Runtime/Core/Private/Linux/LinuxPlatformMisc.cpp
+++ b/Engine/Source/Runtime/Core/Private/Linux/LinuxPlatformMisc.cpp
@@ -1136,7 +1136,7 @@ void GenerateWindowsErrorReport(const FString & WERPath)
 		WriteLine(ReportFile, TEXT("\t<ProblemSignatures>"));
 		WriteLine(ReportFile, TEXT("\t\t<EventType>APPCRASH</EventType>"));
 		WriteLine(ReportFile, *FString::Printf(TEXT("\t\t<Parameter0>UE4-%s</Parameter0>"), FApp::GetGameName()));
-		WriteLine(ReportFile, *FString::Printf(TEXT("\t\t<Parameter1>1.0.%d.%d</Parameter1>"), ENGINE_VERSION_HIWORD, ENGINE_VERSION_LOWORD));
+		WriteLine(ReportFile, *FString::Printf(TEXT("\t\t<Parameter1>%d.%d.%d.0</Parameter1>"), ENGINE_MAJOR_VERSION, ENGINE_MINOR_VERSION, ENGINE_PATCH_VERSION));
 		WriteLine(ReportFile, TEXT("\t\t<Parameter2>0</Parameter2>"));													// FIXME: supply valid?
 		WriteLine(ReportFile, TEXT("\t\t<Parameter3>Unknown Fault Module</Parameter3>"));										// FIXME: supply valid?
 		WriteLine(ReportFile, TEXT("\t\t<Parameter4>0.0.0.0</Parameter4>"));													// FIXME: supply valid?
