Fixed an issue with CocoaWindow.performClose: closing a window without asking Slate if it should

diff --git a/Engine/Source/Runtime/Core/Private/Mac/CocoaWindow.cpp b/Engine/Source/Runtime/Core/Private/Mac/CocoaWindow.cpp
--- a/Engine/Source/Runtime/Core/Private/Mac/CocoaWindow.cpp
+++ b/Engine/Source/Runtime/Core/Private/Mac/CocoaWindow.cpp
@@ -270,10 +270,12 @@ NSString* NSPerformDragOperation = @"NSPerformDragOperation";
 
 - (void)performClose:(id)Sender
 {
-	SCOPED_AUTORELEASE_POOL;
-	bDeferOrderFront = false;
-	
-	[self close];
+	GameThreadCall(^{
+		if (MacApplication)
+		{
+			MacApplication->OnWindowClose(self);
+		}
+	}, @[ NSDefaultRunLoopMode ], false);
 }
 
 - (void)performMiniaturize:(id)Sender
@@ -289,6 +291,13 @@ NSString* NSPerformDragOperation = @"NSPerformDragOperation";
 	[self zoom: self];
 }
 
+- (void)destroy
+{
+	SCOPED_AUTORELEASE_POOL;
+	bDeferOrderFront = false;
+	[self close];
+}
+
 - (void)setFrame:(NSRect)FrameRect display:(BOOL)Flag
 {
 	SCOPED_AUTORELEASE_POOL;
@@ -510,6 +519,12 @@ NSString* NSPerformDragOperation = @"NSPerformDragOperation";
 	}
 	self.bForwardEvents = false;
 	[self setDelegate:nil];
+
+	if ([[NSApp windows] count] == 1)
+	{
+		// It's the last window. App will quit, so stop updating menu state cache.
+		GameThreadCall(^{ FPlatformMisc::UpdateCachedMacMenuState = nullptr; }, @[ NSDefaultRunLoopMode ], true);
+	}
 }
 
 - (void)mouseDown:(NSEvent*)Event
diff --git a/Engine/Source/Runtime/Core/Private/Mac/MacApplication.cpp b/Engine/Source/Runtime/Core/Private/Mac/MacApplication.cpp
--- a/Engine/Source/Runtime/Core/Private/Mac/MacApplication.cpp
+++ b/Engine/Source/Runtime/Core/Private/Mac/MacApplication.cpp
@@ -1292,6 +1292,15 @@ bool FMacApplication::OnWindowDestroyed( FCocoaWindow* Window )
 	return false;
 }
 
+void FMacApplication::OnWindowClose( FCocoaWindow* Window )
+{
+	TSharedPtr< FMacWindow > EventWindow = FindWindowByNSWindow( Windows, &WindowsMutex, Window );
+	if( EventWindow.IsValid() )
+	{
+		MessageHandler->OnWindowClose( EventWindow.ToSharedRef() );
+	}
+}
+
 void FMacApplication::OnMouseCursorLock( bool bLockEnabled )
 {
 	bIsMouseCursorLocked = bLockEnabled;
diff --git a/Engine/Source/Runtime/Core/Private/Mac/MacWindow.cpp b/Engine/Source/Runtime/Core/Private/Mac/MacWindow.cpp
--- a/Engine/Source/Runtime/Core/Private/Mac/MacWindow.cpp
+++ b/Engine/Source/Runtime/Core/Private/Mac/MacWindow.cpp
@@ -309,7 +309,7 @@ void FMacWindow::Destroy()
 			// Close the window
 			MainThreadCall(^{
 				SCOPED_AUTORELEASE_POOL;
-				[Window performClose:nil];
+				[Window destroy];
 			}, UE4CloseEventMode, true);
 		}
 	}
diff --git a/Engine/Source/Runtime/Core/Public/Mac/CocoaWindow.h b/Engine/Source/Runtime/Core/Public/Mac/CocoaWindow.h
--- a/Engine/Source/Runtime/Core/Public/Mac/CocoaWindow.h
+++ b/Engine/Source/Runtime/Core/Public/Mac/CocoaWindow.h
@@ -66,6 +66,9 @@
 /** Order window to the front. */
 - (void)orderFrontAndMakeMain:(bool)bMain andKey:(bool)bKey;
 
+/** Destroy the window */
+- (void)destroy;
+
 @end
 
 /**
diff --git a/Engine/Source/Runtime/Core/Public/Mac/MacApplication.h b/Engine/Source/Runtime/Core/Public/Mac/MacApplication.h
--- a/Engine/Source/Runtime/Core/Public/Mac/MacApplication.h
+++ b/Engine/Source/Runtime/Core/Public/Mac/MacApplication.h
@@ -102,6 +102,7 @@ public:
 	void OnWindowRedrawContents( FCocoaWindow* Window );
 	void OnWindowDidClose( FCocoaWindow* Window );
 	bool OnWindowDestroyed( FCocoaWindow* Window );
+	void OnWindowClose( FCocoaWindow* Window );
 
 	void OnMouseCursorLock( bool bLockEnabled );
 
