Fixed a freeze on quitting the editor with Cmd+W on Mac

diff --git a/Engine/Source/Runtime/Core/Private/Mac/CocoaWindow.cpp b/Engine/Source/Runtime/Core/Private/Mac/CocoaWindow.cpp
--- a/Engine/Source/Runtime/Core/Private/Mac/CocoaWindow.cpp
+++ b/Engine/Source/Runtime/Core/Private/Mac/CocoaWindow.cpp
@@ -523,7 +523,7 @@ NSString* NSPerformDragOperation = @"NSPerformDragOperation";
 	if ([[NSApp windows] count] == 1)
 	{
 		// It's the last window. App will quit, so stop updating menu state cache.
-		GameThreadCall(^{ FPlatformMisc::UpdateCachedMacMenuState = nullptr; }, @[ NSDefaultRunLoopMode ], true);
+		GameThreadCall(^{ FPlatformMisc::UpdateCachedMacMenuState = nullptr; }, @[ NSDefaultRunLoopMode ], false);
 	}
 }
 
