On OS X when a UE4 application loses focus push modal windows down to the normal window level, restore them on gaining focus, so that they don't float over other applications. Notifications still float while tooltips and menus will hide on deactivate and reappear on activation.

diff --git a/Engine/Source/Runtime/Core/Private/Mac/MacApplication.cpp b/Engine/Source/Runtime/Core/Private/Mac/MacApplication.cpp
--- a/Engine/Source/Runtime/Core/Private/Mac/MacApplication.cpp
+++ b/Engine/Source/Runtime/Core/Private/Mac/MacApplication.cpp
@@ -142,6 +142,31 @@ FMacApplication::FMacApplication()
 	}
 
 	AppActivationObserver = [[NSNotificationCenter defaultCenter] addObserverForName:NSApplicationDidBecomeActiveNotification object:[NSApplication sharedApplication] queue:[NSOperationQueue mainQueue] usingBlock:^(NSNotification* Notification){
+								// Change modal window levels
+								NSWindow* BaseWindow = nil;
+								for( auto Window : Windows )
+								{
+									NSWindow* CocoaWindow = (NSWindow*)Window->GetOSWindowHandle();
+									if ( CocoaWindow && Window->IsVisible() && [CocoaWindow level] == NSModalPanelWindowLevel )
+									{
+										BaseWindow = CocoaWindow;
+										break;
+									}
+								}
+							 
+							 	for( auto Window : Windows )
+								{
+									NSWindow* CocoaWindow = (NSWindow*)Window->GetOSWindowHandle();
+									if ( CocoaWindow && Window->GetDefinition().IsModalWindow )
+									{
+										[CocoaWindow setLevel:NSModalPanelWindowLevel];
+										if ( BaseWindow )
+										{
+											[CocoaWindow orderWindow:NSWindowBelow relativeTo:[BaseWindow windowNumber]];
+										}
+									}
+								}
+							 
 								// If editor thread doesn't have the focus, don't suck up too much CPU time.
 								if( GIsEditor )
 								{
@@ -157,6 +182,35 @@ FMacApplication::FMacApplication()
 							}];
 
 	AppDeactivationObserver = [[NSNotificationCenter defaultCenter] addObserverForName:NSApplicationDidResignActiveNotification object:[NSApplication sharedApplication] queue:[NSOperationQueue mainQueue] usingBlock:^(NSNotification* Notification){
+									// Change modal window levels
+									if (Windows.Num() > 0)
+									{
+										NSWindow* BaseWindow = nil;
+										for( auto Window : Windows )
+										{
+											NSWindow* CocoaWindow = (NSWindow*)Window->GetOSWindowHandle();
+											if ( CocoaWindow && Window->IsVisible() && [CocoaWindow level] == NSNormalWindowLevel )
+											{
+												BaseWindow = CocoaWindow;
+											}
+										}
+										
+										check(BaseWindow);
+										
+										for( auto Window : Windows )
+										{
+											NSWindow* CocoaWindow = (NSWindow*)Window->GetOSWindowHandle();
+											if ( CocoaWindow && Window->GetDefinition().IsModalWindow )
+											{
+												[CocoaWindow setLevel:NSNormalWindowLevel];
+												[CocoaWindow orderWindow:NSWindowAbove relativeTo:[BaseWindow windowNumber]];
+												BaseWindow = CocoaWindow;
+											}
+										}
+									}
+							   
+									CFRunLoopRunInMode(kCFRunLoopDefaultMode, 0, true);
+							   
 									// If editor thread doesn't have the focus, don't suck up too much CPU time.
 									if( GIsEditor )
 									{
diff --git a/Engine/Source/Runtime/Core/Private/Mac/MacWindow.cpp b/Engine/Source/Runtime/Core/Private/Mac/MacWindow.cpp
--- a/Engine/Source/Runtime/Core/Private/Mac/MacWindow.cpp
+++ b/Engine/Source/Runtime/Core/Private/Mac/MacWindow.cpp
@@ -105,7 +105,8 @@ void FMacWindow::Initialize( FMacApplication* const Application, const TSharedRe
 			// @todo: We really need a window type in tab manager to know what window level to use and whether or not a window should hide on deactivate
 			if (Definition->IsModalWindow)
 			{
-				[WindowHandle setLevel: NSModalPanelWindowLevel];
+				NSInteger Level = [NSApp isActive] ? NSModalPanelWindowLevel : NSNormalWindowLevel;
+				[WindowHandle setLevel: Level];
 			}
 			else if (Definition->IsRegularWindow)
 			{
@@ -164,6 +165,15 @@ void FMacWindow::Initialize( FMacApplication* const Application, const TSharedRe
 			}
 			else if(Definition->AppearsInTaskbar)
 			{
+				if ( !Definition->Title.IsEmpty() )
+				{
+					CFStringRef CFName = FPlatformString::TCHARToCFString( *Definition->Title );
+					[WindowHandle setTitle: ( NSString *)CFName];
+					CFRelease( CFName );
+					
+					[NSApp addWindowsItem: WindowHandle title: [WindowHandle title] filename: NO];
+				}
+				
 				[WindowHandle setCollectionBehavior:NSWindowCollectionBehaviorFullScreenAuxiliary|NSWindowCollectionBehaviorDefault|NSWindowCollectionBehaviorManaged|NSWindowCollectionBehaviorParticipatesInCycle];
 			}
 			else
