Added support for negative exit codes from UBT launched by FMacPlatformProcess::CreateProc() using a shell script

diff --git a/Engine/Source/Runtime/Core/Private/Mac/MacPlatformProcess.cpp b/Engine/Source/Runtime/Core/Private/Mac/MacPlatformProcess.cpp
--- a/Engine/Source/Runtime/Core/Private/Mac/MacPlatformProcess.cpp
+++ b/Engine/Source/Runtime/Core/Private/Mac/MacPlatformProcess.cpp
@@ -344,10 +344,11 @@ FProcHandle FMacPlatformProcess::CreateProc( const TCHAR* URL, const TCHAR* Parm
 	
 	if(LaunchPath == NULL)
 	{
-		return FProcHandle(NULL);
+		return FProcHandle(NULL, false);
 	}
 
 	NSTask* ProcessHandle = [[NSTask alloc] init];
+	bool bIsShellScript = false;
 
 	if (ProcessHandle)
 	{
@@ -361,6 +362,7 @@ FProcHandle FMacPlatformProcess::CreateProc( const TCHAR* URL, const TCHAR* Parm
 			[Arguments addObject: @"-c"];
 			[Arguments addObject: Arg];
 			CFRelease((CFStringRef)Arg);
+			bIsShellScript = true;
 		}
 		else
 		{
@@ -453,22 +455,22 @@ FProcHandle FMacPlatformProcess::CreateProc( const TCHAR* URL, const TCHAR* Parm
 
 	CFRelease((CFStringRef)LaunchPath);
 
-	return FProcHandle(ProcessHandle);
+	return FProcHandle(ProcessHandle, bIsShellScript);
 }
 
-bool FMacPlatformProcess::IsProcRunning( FProcHandle & ProcessHandle )
+bool FMacPlatformProcess::IsProcRunning( FProcHandle& ProcessHandle )
 {
 	SCOPED_AUTORELEASE_POOL;
 	return [(NSTask*)ProcessHandle.Get() isRunning];
 }
 
-void FMacPlatformProcess::WaitForProc( FProcHandle & ProcessHandle )
+void FMacPlatformProcess::WaitForProc( FProcHandle& ProcessHandle )
 {
 	SCOPED_AUTORELEASE_POOL;
 	[(NSTask*)ProcessHandle.Get() waitUntilExit];
 }
 
-void FMacPlatformProcess::TerminateProc( FProcHandle & ProcessHandle, bool KillTree )
+void FMacPlatformProcess::TerminateProc( FProcHandle& ProcessHandle, bool KillTree )
 {
 	SCOPED_AUTORELEASE_POOL;
 
@@ -506,7 +508,7 @@ uint32 FMacPlatformProcess::GetCurrentProcessId()
 	return getpid();
 }
 
-bool FMacPlatformProcess::GetProcReturnCode( FProcHandle & ProcessHandle, int32* ReturnCode )
+bool FMacPlatformProcess::GetProcReturnCode( FProcHandle& ProcessHandle, int32* ReturnCode )
 {
 	SCOPED_AUTORELEASE_POOL;
 
@@ -516,6 +518,10 @@ bool FMacPlatformProcess::GetProcReturnCode( FProcHandle & ProcessHandle, int32*
 	}
 
 	*ReturnCode = [(NSTask*)ProcessHandle.Get() terminationStatus];
+	if (ProcessHandle.IsShellScript && *ReturnCode > 128)
+	{
+		*ReturnCode = *ReturnCode - 256;
+	}
 	return true;
 }
 
@@ -883,7 +889,7 @@ FString FMacPlatformProcess::ReadPipe( void* ReadPipe )
 	return Output;
 }
 
-bool FMacPlatformProcess::ReadPipeToArray(void* ReadPipe, TArray<uint8> & Output)
+bool FMacPlatformProcess::ReadPipeToArray(void* ReadPipe, TArray<uint8>& Output)
 {
 	SCOPED_AUTORELEASE_POOL;
 
diff --git a/Engine/Source/Runtime/Core/Public/Mac/MacPlatformProcess.h b/Engine/Source/Runtime/Core/Public/Mac/MacPlatformProcess.h
--- a/Engine/Source/Runtime/Core/Public/Mac/MacPlatformProcess.h
+++ b/Engine/Source/Runtime/Core/Public/Mac/MacPlatformProcess.h
@@ -14,12 +14,16 @@ public:
 	/** Default constructor. */
 	FORCEINLINE FProcHandle()
 		: TProcHandle()
+		, IsShellScript( false )
 	{}
 
 	/** Initialization constructor. */
-	FORCEINLINE explicit FProcHandle( HandleType Other )
+	FORCEINLINE explicit FProcHandle( HandleType Other, bool InIsShellScript )
 		: TProcHandle( Other )
+		, IsShellScript( InIsShellScript )
 	{}
+
+	bool IsShellScript;
 };
 
 /**
