Core: Updated FDateTime documentation

diff --git a/Engine/Source/Runtime/Core/Public/Misc/DateTime.h b/Engine/Source/Runtime/Core/Public/Misc/DateTime.h
--- a/Engine/Source/Runtime/Core/Public/Misc/DateTime.h
+++ b/Engine/Source/Runtime/Core/Public/Misc/DateTime.h
@@ -593,8 +593,12 @@ public:
 	/**
 	 * Gets the local date and time on this computer.
 	 *
+	 * This method takes into account the local computer's time zone and daylight saving
+	 * settings. For time zone independent time comparisons, and when comparing times
+	 * between different computers, please use UtcNow() instead.
+	 *
 	 * @return Current date and time.
-	 * @see Today
+	 * @see Today, UtcNow
 	 */
 	static CORE_API FDateTime Now( );
 
@@ -626,7 +630,7 @@ public:
 	 * The time component is set to 00:00:00
 	 *
 	 * @return Current date.
-	 * @see Now
+	 * @see Now, UtcNow
 	 */
 	static FDateTime Today( )
 	{
@@ -636,7 +640,13 @@ public:
 	/**
 	 * Gets the UTC date and time on this computer.
 	 *
+	 * This method returns the Coordinated Universal Time (UTC), which does not take the
+	 * local computer's time zone and daylight savings settings into account. It should be
+	 * used when comparing dates and times that should be independent of the user's locale.
+	 * To get the date and time in the current locale, use Now() instead.
+	 *
 	 * @return Current date and time.
+	 * @see Now
 	 */
 	static CORE_API FDateTime UtcNow( );
 
