#docs - Minor cleanup to FMatrix, FQuat.

diff --git a/Engine/Source/Runtime/Core/Public/Math/Matrix.h b/Engine/Source/Runtime/Core/Public/Math/Matrix.h
--- a/Engine/Source/Runtime/Core/Public/Math/Matrix.h
+++ b/Engine/Source/Runtime/Core/Public/Math/Matrix.h
@@ -3,6 +3,7 @@
 #pragma once
 
 /**
+ * 4x4 matrix of floating point values.
  * Matrix-matrix multiplication happens with a pre-multiple of the transpose --
  * in other words, Res = Mat1.operator*(Mat2) means Res = Mat2^T * Mat1, as
  * opposed to Res = Mat1 * Mat2.
diff --git a/Engine/Source/Runtime/Core/Public/Math/Transform.h b/Engine/Source/Runtime/Core/Public/Math/Transform.h
--- a/Engine/Source/Runtime/Core/Public/Math/Transform.h
+++ b/Engine/Source/Runtime/Core/Public/Math/Transform.h
@@ -4,7 +4,7 @@
 #pragma once
 
 /** 
- * FTransform class for Quat/Translation/Scale.
+ * Transform composed of Quat/Translation/Scale.
  */
 class FTransform
 {
@@ -16,15 +16,15 @@ class FTransform
 	friend MAYBE_COREUOBJECT_API class UScriptStruct* Z_Construct_UScriptStruct_UObject_FTransform();
 
 protected:
-	/** Rotation of this transformation, as a quaternion */
+	/** Rotation of this transformation, as a quaternion. */
 	FQuat	Rotation;
-	/** Translation of this transformation, as a vector */
+	/** Translation of this transformation, as a vector. */
 	FVector	Translation;
-	/** 3D scale (always applied in local space) as a vector */
+	/** 3D scale (always applied in local space) as a vector. */
 	FVector	Scale3D;
 public:
 	/**
-	 * The identity transformation (Rotation = FQuat::Identity, Translation = FVector::ZeroVector, Scale3D = (1,1,1))
+	 * The identity transformation (Rotation = FQuat::Identity, Translation = FVector::ZeroVector, Scale3D = (1,1,1)).
 	 */
 	static CORE_API const FTransform Identity;
 
