HAVE_RUNTIME_THREADING_SWITCHES needs to be true for Android and iOS (iPhone 4)
#ue4
#android
#ios

diff --git a/Engine/Source/Runtime/Core/Public/Misc/App.h b/Engine/Source/Runtime/Core/Public/Misc/App.h
--- a/Engine/Source/Runtime/Core/Public/Misc/App.h
+++ b/Engine/Source/Runtime/Core/Public/Misc/App.h
@@ -7,7 +7,7 @@
 #pragma once
 
 // platforms which can have runtime threading switches
-#define HAVE_RUNTIME_THREADING_SWITCHES			(PLATFORM_DESKTOP)
+#define HAVE_RUNTIME_THREADING_SWITCHES			(PLATFORM_DESKTOP || PLATFORM_ANDROID || PLATFORM_IOS)
 
 /**
  * Provides information about the application.
