Alleviate windowed resolution problems on large displays.

- If screen is 4000+ pixels wide, none of attempted "convenient" resolutions may fit. This provides a safe fallback for such cases.

diff --git a/Engine/Source/Runtime/Core/Private/Misc/CoreMisc.cpp b/Engine/Source/Runtime/Core/Private/Misc/CoreMisc.cpp
--- a/Engine/Source/Runtime/Core/Private/Misc/CoreMisc.cpp
+++ b/Engine/Source/Runtime/Core/Private/Misc/CoreMisc.cpp
@@ -957,6 +957,12 @@ void GenerateConvenientWindowedResolutions(const FDisplayMetrics& InDisplayMetri
 			}
 		}
 	}
+	
+	// if no convenient resolutions have been found, add a minimum one (if it fits)
+	if (OutResolutions.Num() == 0 && InDisplayMetrics.PrimaryDisplayHeight > MinHeight && InDisplayMetrics.PrimaryDisplayWidth > MinWidth)
+	{
+		OutResolutions.Add(FIntPoint(MinWidth, MinHeight));
+	}
 }
 
 /*----------------------------------------------------------------------------
