CrashReport - Added detection of OOM

diff --git a/Engine/Source/Runtime/Core/Private/GenericPlatform/GenericPlatformMemory.cpp b/Engine/Source/Runtime/Core/Private/GenericPlatform/GenericPlatformMemory.cpp
--- a/Engine/Source/Runtime/Core/Private/GenericPlatform/GenericPlatformMemory.cpp
+++ b/Engine/Source/Runtime/Core/Private/GenericPlatform/GenericPlatformMemory.cpp
@@ -33,6 +33,8 @@ FGenericPlatformMemoryStats::FGenericPlatformMemoryStats()
 	, PeakUsedVirtual( 0 )
 {}
 
+bool FGenericPlatformMemory::bIsOOM = false;
+
 void FGenericPlatformMemory::SetupMemoryPools()
 {
 	SET_MEMORY_STAT(MCR_Physical, 0); // "unlimited" physical memory, we still need to make this call to set the short name, etc
@@ -48,6 +50,7 @@ void FGenericPlatformMemory::Init()
 
 void FGenericPlatformMemory::OnOutOfMemory(uint64 Size, uint32 Alignment)
 {
+	bIsOOM = true;
 	UE_LOG(LogMemory, Fatal, TEXT("Ran out of memory allocating %llu bytes with alignment %u"), Size, Alignment);
 }
 
diff --git a/Engine/Source/Runtime/Core/Public/GenericPlatform/GenericPlatformMemory.h b/Engine/Source/Runtime/Core/Public/GenericPlatform/GenericPlatformMemory.h
--- a/Engine/Source/Runtime/Core/Public/GenericPlatform/GenericPlatformMemory.h
+++ b/Engine/Source/Runtime/Core/Public/GenericPlatform/GenericPlatformMemory.h
@@ -92,6 +92,9 @@ struct FPlatformMemoryStats;
 /** Generic implementation for most platforms, these tend to be unused and unimplemented. */
 struct CORE_API FGenericPlatformMemory
 {
+	/** Set to true if we encounters out of memory. */
+	static bool bIsOOM;
+
 	/**
 	 * Various memory regions that can be used with memory stats. The exact meaning of
 	 * the enums are relatively platform-dependent, although the general ones (Physical, GPU)
