CrashReport - Added misc comments, updated formating etc (mostly no code change)

diff --git a/Engine/Source/Runtime/Core/Public/HAL/ExceptionHandling.h b/Engine/Source/Runtime/Core/Public/HAL/ExceptionHandling.h
--- a/Engine/Source/Runtime/Core/Public/HAL/ExceptionHandling.h
+++ b/Engine/Source/Runtime/Core/Public/HAL/ExceptionHandling.h
@@ -13,6 +13,7 @@ extern CORE_API bool GUseCrashReportClient;
 
 extern CORE_API TCHAR MiniDumpFilenameW[1024];
 
+// @TODO yrx 2014-09-11 Move to PlatformExceptionHandling
 #if PLATFORM_WINDOWS
 extern CORE_API int32 NullReportCrash( LPEXCEPTION_POINTERS ExceptionInfo );
 extern CORE_API int32 ReportCrash( LPEXCEPTION_POINTERS ExceptionInfo );
