CrashReport - Fixed bad engine version stored as Parameter0 and added commandline as Param8, to be the same as on other platforms

diff --git a/Engine/Source/Runtime/Core/Private/Mac/MacPlatformMisc.cpp b/Engine/Source/Runtime/Core/Private/Mac/MacPlatformMisc.cpp
--- a/Engine/Source/Runtime/Core/Private/Mac/MacPlatformMisc.cpp
+++ b/Engine/Source/Runtime/Core/Private/Mac/MacPlatformMisc.cpp
@@ -15,6 +15,7 @@
 #include "Runtime/Launch/Resources/Version.h"
 #include "EngineVersion.h"
 #include "MacMallocZone.h"
+//#include "MacPlatformCrashContext.h"
 
 #include <dlfcn.h>
 #include <IOKit/IOKitLib.h>
@@ -1235,11 +1236,13 @@ void FMacCrashContext::GenerateWindowsErrorReport(char const* WERPath) const
 		WriteUTF16String(ReportFile, *GMacAppInfo.AppName);
 		WriteLine(ReportFile, TEXT("</Parameter0>"));
 		
-		WriteUTF16String(ReportFile, TEXT("\t\t<Parameter1>1.0."));
-		WriteUTF16String(ReportFile, ItoTCHAR(ENGINE_VERSION_HIWORD, 10));
+		WriteUTF16String(ReportFile, TEXT("\t\t<Parameter1>"));
+		WriteUTF16String(ReportFile, ItoTCHAR(ENGINE_MAJOR_VERSION, 10));
 		WriteUTF16String(ReportFile, TEXT("."));
-		WriteUTF16String(ReportFile, ItoTCHAR(ENGINE_VERSION_LOWORD, 10));
-		WriteLine(ReportFile, TEXT("</Parameter1>"));
+		WriteUTF16String(ReportFile, ItoTCHAR(ENGINE_MINOR_VERSION, 10));
+		WriteUTF16String(ReportFile, TEXT("."));
+		WriteUTF16String(ReportFile, ItoTCHAR(ENGINE_PATCH_VERSION, 10));
+		WriteLine(ReportFile, TEXT(".0</Parameter1>"));
 
 		// App time stamp
 		WriteLine(ReportFile, TEXT("\t\t<Parameter2>528f2d37</Parameter2>"));													// FIXME: supply valid?
@@ -1307,12 +1310,10 @@ void FMacCrashContext::GenerateWindowsErrorReport(char const* WERPath) const
 			WriteLine(ReportFile, TEXT("</Parameter7>"));
 		}
 		
-		// Fault type -> Signal
-		WriteUTF16String(ReportFile, TEXT("\t\t<Parameter8>"));
-		FMemory::Memzero(Line, PATH_MAX * sizeof(TCHAR));
-		FUTF8ToTCHAR_Convert::Convert(Line, PATH_MAX, SignalDescription, FCStringAnsi::Strlen(SignalDescription));
-		WriteUTF16String(ReportFile, Line);
-		WriteLine(ReportFile, TEXT("</Parameter8>"));
+		// Command line, must match the Windows version.
+		WriteUTF16String(ReportFile, TEXT("\t\t<Parameter8>!"));
+		WriteUTF16String(ReportFile, FCommandLine::Get();
+		WriteLine(ReportFile, TEXT("!</Parameter8>"));
 		
 		WriteUTF16String(ReportFile, TEXT("\t\t<Parameter9>"));
 		WriteUTF16String(ReportFile, *GMacAppInfo.BranchBaseDir);
@@ -1693,6 +1694,12 @@ void FMacCrashContext::GenerateCrashInfoAndLaunchReporter() const
 			
 			close(ReportFile);
 		}
+
+		// Introduces a new runtime crash context. Will replace all Windows related crash reporting.
+		//FCStringAnsi::Strncpy(FilePath, CrashInfoFolder, PATH_MAX);
+		//FCStringAnsi::Strcat(FilePath, PATH_MAX, "/" );
+		//FCStringAnsi::Strcat(FilePath, PATH_MAX, FGenericCrashContext::CrashContextRuntimeXMLNameA );
+		//SerializeAsXML( FilePath ); @todo uncomment after verification
 		
 		// copy log
 		FCStringAnsi::Strncpy(FilePath, CrashInfoFolder, PATH_MAX);
