Fix Mac build & remove debug code.

diff --git a/Engine/Source/Runtime/Core/Private/Mac/MacPlatformMisc.cpp b/Engine/Source/Runtime/Core/Private/Mac/MacPlatformMisc.cpp
--- a/Engine/Source/Runtime/Core/Private/Mac/MacPlatformMisc.cpp
+++ b/Engine/Source/Runtime/Core/Private/Mac/MacPlatformMisc.cpp
@@ -15,7 +15,6 @@
 #include "Runtime/Launch/Resources/Version.h"
 #include "EngineVersion.h"
 #include "MacMallocZone.h"
-//#include "MacPlatformCrashContext.h"
 
 #include <dlfcn.h>
 #include <IOKit/IOKitLib.h>
@@ -1312,7 +1311,7 @@ void FMacCrashContext::GenerateWindowsErrorReport(char const* WERPath) const
 		
 		// Command line, must match the Windows version.
 		WriteUTF16String(ReportFile, TEXT("\t\t<Parameter8>!"));
-		WriteUTF16String(ReportFile, FCommandLine::Get();
+		WriteUTF16String(ReportFile, FCommandLine::Get());
 		WriteLine(ReportFile, TEXT("!</Parameter8>"));
 		
 		WriteUTF16String(ReportFile, TEXT("\t\t<Parameter9>"));
@@ -1694,12 +1693,6 @@ void FMacCrashContext::GenerateCrashInfoAndLaunchReporter() const
 			
 			close(ReportFile);
 		}
-
-		// Introduces a new runtime crash context. Will replace all Windows related crash reporting.
-		//FCStringAnsi::Strncpy(FilePath, CrashInfoFolder, PATH_MAX);
-		//FCStringAnsi::Strcat(FilePath, PATH_MAX, "/" );
-		//FCStringAnsi::Strcat(FilePath, PATH_MAX, FGenericCrashContext::CrashContextRuntimeXMLNameA );
-		//SerializeAsXML( FilePath ); @todo uncomment after verification
 		
 		// copy log
 		FCStringAnsi::Strncpy(FilePath, CrashInfoFolder, PATH_MAX);
