CrashReport - Added comment regarding backtrace_symbols

diff --git a/Engine/Source/Runtime/Core/Private/Linux/LinuxPlatformStackWalk.cpp b/Engine/Source/Runtime/Core/Private/Linux/LinuxPlatformStackWalk.cpp
--- a/Engine/Source/Runtime/Core/Private/Linux/LinuxPlatformStackWalk.cpp
+++ b/Engine/Source/Runtime/Core/Private/Linux/LinuxPlatformStackWalk.cpp
@@ -232,6 +232,7 @@ void FLinuxPlatformStackWalk::CaptureStackBackTrace( uint64* BackTrace, uint32 M
 
 		if (LinuxContext->BacktraceSymbols == NULL)
 		{
+			// @TODO yrx 2014-09-29 Replace with backtrace_symbols_fd due to malloc()
 			LinuxContext->BacktraceSymbols = backtrace_symbols(reinterpret_cast< void** >( BackTrace ), MaxDepth);
 		}
 	}
