Core - Added traits for SIZE_T

diff --git a/Engine/Source/Runtime/Core/Public/Templates/UnrealTypeTraits.h b/Engine/Source/Runtime/Core/Public/Templates/UnrealTypeTraits.h
--- a/Engine/Source/Runtime/Core/Public/Templates/UnrealTypeTraits.h
+++ b/Engine/Source/Runtime/Core/Public/Templates/UnrealTypeTraits.h
@@ -75,6 +75,7 @@ template<> struct TIsIntegralType<bool> { enum { Value = true }; };
 
 template<> struct TIsIntegralType<WIDECHAR> { enum { Value = true }; };
 template<> struct TIsIntegralType<ANSICHAR> { enum { Value = true }; };
+template<> struct TIsIntegralType<SIZE_T> { enum { Value = true }; };
 
 /**
  * TIsSignedIntegralType
@@ -140,6 +141,12 @@ Expose_TFormatSpecifier(float, "%f")
 Expose_TFormatSpecifier(double, "%f")
 Expose_TFormatSpecifier(long double, "%f")
 
+#if PLATFORM_32BITS
+Expose_TFormatSpecifier(SIZE_T, "%u")
+#else
+Expose_TFormatSpecifier(SIZE_T, "%llu")
+#endif // PLATFORM_32BITS
+
 
 
 /**
