UrlEncode converts string to UTF-8 before doing encoding as required by standard

diff --git a/Engine/Source/Runtime/Core/Private/GenericPlatform/GenericPlatformHttp.cpp b/Engine/Source/Runtime/Core/Private/GenericPlatform/GenericPlatformHttp.cpp
--- a/Engine/Source/Runtime/Core/Private/GenericPlatform/GenericPlatformHttp.cpp
+++ b/Engine/Source/Runtime/Core/Private/GenericPlatform/GenericPlatformHttp.cpp
@@ -5,6 +5,7 @@
 #include "Runtime/Online/HTTP/Public/Interfaces/IHttpBase.h"
 #include "Runtime/Online/HTTP/Public/Interfaces/IHttpRequest.h"
 #include "Runtime/Online/HTTP/Public/Interfaces/IHttpResponse.h"
+#include "StringConv.h"
 
 /**
  * A generic http request
@@ -54,20 +55,49 @@ IHttpRequest* FGenericPlatformHttp::ConstructRequest()
 }
 
 
+static bool IsAllowedChar(ANSICHAR LookupChar)
+{
+	static char AllowedChars[] = "ABCDEFGHIJKLMNOPQRSTUVWXYZabcdefghijklmnopqrstuvwxyz0123456789-_.~";
+	static bool bTableFilled = false;
+	static bool AllowedTable[256] = { false };
+
+	if (!bTableFilled)
+	{
+		for (int32 Idx = 0; Idx < ARRAY_COUNT(AllowedChars) - 1; ++Idx)	// -1 to avoid trailing 0
+		{
+			uint8 AllowedCharIdx = static_cast<uint8>(AllowedChars[Idx]);
+			check(AllowedCharIdx < ARRAY_COUNT(AllowedTable));
+			AllowedTable[AllowedCharIdx] = true;
+		}
+
+		bTableFilled = true;
+	}
+
+	return AllowedTable[LookupChar];
+}
+
 FString FGenericPlatformHttp::UrlEncode(const FString &UnencodedString)
 {
+	FTCHARToUTF8 Converter(*UnencodedString);	//url encoding must be encoded over each utf-8 byte
+	const UTF8CHAR* UTF8Data = (UTF8CHAR*) Converter.Get();	//converter uses ANSI instead of UTF8CHAR - not sure why - but other code seems to just do this cast. In this case it really doesn't matter
 	FString EncodedString = TEXT("");
-	const FString UnreservedChars = "ABCDEFGHIJKLMNOPQRSTUVWXYZabcdefghijklmnopqrstuvwxyz0123456789-_.~";
-	for (auto EncodeIt = UnencodedString.CreateConstIterator(); EncodeIt; ++EncodeIt)
+	
+	TCHAR Buffer[2] = { 0, 0 };
+
+	for (int32 ByteIdx = 0, Length = Converter.Length(); ByteIdx < Length; ++ByteIdx)
 	{
-		if (UnreservedChars.Contains(FString::Chr(*EncodeIt)))
+		UTF8CHAR ByteToEncode = UTF8Data[ByteIdx];
+		
+		if (IsAllowedChar(ByteToEncode))
 		{
-			EncodedString += *EncodeIt;
+			Buffer[0] = ByteToEncode;
+			FString TmpString = Buffer;
+			EncodedString += TmpString;
 		}
-		else if (*EncodeIt != '\0')
+		else if (ByteToEncode != '\0')
 		{
 			EncodedString += TEXT("%");
-			EncodedString += FString::Printf(TEXT("%.2X"), *EncodeIt);
+			EncodedString += FString::Printf(TEXT("%.2X"), ByteToEncode);
 		}
 	}
 	return EncodedString;
