CIS fix

diff --git a/Engine/Source/Runtime/Core/Public/Mac/MacPlatformProcess.h b/Engine/Source/Runtime/Core/Public/Mac/MacPlatformProcess.h
--- a/Engine/Source/Runtime/Core/Public/Mac/MacPlatformProcess.h
+++ b/Engine/Source/Runtime/Core/Public/Mac/MacPlatformProcess.h
@@ -18,7 +18,7 @@ public:
 	{}
 
 	/** Initialization constructor. */
-	FORCEINLINE explicit FProcHandle( HandleType Other, bool InIsShellScript )
+	FORCEINLINE explicit FProcHandle( HandleType Other, bool InIsShellScript = false )
 		: TProcHandle( Other )
 		, IsShellScript( InIsShellScript )
 	{}
