Fix clang build

diff --git a/Engine/Source/Runtime/Core/Private/GenericPlatform/GenericPlatformHttp.cpp b/Engine/Source/Runtime/Core/Private/GenericPlatform/GenericPlatformHttp.cpp
--- a/Engine/Source/Runtime/Core/Private/GenericPlatform/GenericPlatformHttp.cpp
+++ b/Engine/Source/Runtime/Core/Private/GenericPlatform/GenericPlatformHttp.cpp
@@ -55,7 +55,7 @@ IHttpRequest* FGenericPlatformHttp::ConstructRequest()
 }
 
 
-static bool IsAllowedChar(ANSICHAR LookupChar)
+static bool IsAllowedChar(UTF8CHAR LookupChar)
 {
 	static char AllowedChars[] = "ABCDEFGHIJKLMNOPQRSTUVWXYZabcdefghijklmnopqrstuvwxyz0123456789-_.~";
 	static bool bTableFilled = false;
