added comments

diff --git a/Engine/Shaders/ShaderVersion.usf b/Engine/Shaders/ShaderVersion.usf
--- a/Engine/Shaders/ShaderVersion.usf
+++ b/Engine/Shaders/ShaderVersion.usf
@@ -2,4 +2,4 @@
 // Each time the console command is executed it generates a new GUID. As this file is included
 // in common.usf (which should be included in any shader) it allows to invalidate the shader DDC.
 // 
-// GUID = 55058CB5415760DB2475E8ABED874656
\ No newline at end of file
+// GUID = CEB51B844C0F1E837BB98A92CFC5BE4C
\ No newline at end of file
diff --git a/Engine/Source/Runtime/Core/Private/Misc/ConfigCacheIni.cpp b/Engine/Source/Runtime/Core/Private/Misc/ConfigCacheIni.cpp
--- a/Engine/Source/Runtime/Core/Private/Misc/ConfigCacheIni.cpp
+++ b/Engine/Source/Runtime/Core/Private/Misc/ConfigCacheIni.cpp
@@ -2603,8 +2603,8 @@ static FString GetSourceIniFilename(const TCHAR* BaseIniName, const TCHAR* Platf
  * Creates a chain of ini filenames to load and combine.
  *
  * @param InBaseIniName Ini name.
- * @param InPlatformName Platform name.
- * @param InGameName Game name.
+ * @param InPlatformName Platform name, NULL means to use the current platform
+ * @param InGameName Game name, NULL means to use the current
  * @param OutHierarchy An array which is to receive the generated hierachy of ini filenames.
  */
 static void GetSourceIniHierarchyFilenames(const TCHAR* InBaseIniName, const TCHAR* InPlatformName, const TCHAR* InGameName, const TCHAR* EngineConfigDir, const TCHAR* SourceConfigDir, TArray<FIniFilename>& OutHierarchy, bool bRequireDefaultIni)
