Remove extraneous posix_spawn_file_actions_destroy.

- PR #441 by amigo (tracked as TTP #346841).

diff --git a/Engine/Source/Runtime/Core/Private/Linux/LinuxPlatformProcess.cpp b/Engine/Source/Runtime/Core/Private/Linux/LinuxPlatformProcess.cpp
--- a/Engine/Source/Runtime/Core/Private/Linux/LinuxPlatformProcess.cpp
+++ b/Engine/Source/Runtime/Core/Private/Linux/LinuxPlatformProcess.cpp
@@ -558,8 +558,6 @@ FProcHandle FLinuxPlatformProcess::CreateProc(const TCHAR* URL, const TCHAR* Par
 		*OutProcessID = ChildPid;
 	}
 
-	posix_spawn_file_actions_destroy(&FileActions);
-
 	return FProcHandle( ChildPid );
 }
 
