Skip Mac menu cache update while pumping messages in SlateTick(). Also, call highlightKeyEquivalent: on main thread to fix a crash caused by both main and game therad trying to redraw menu item at the same time.

diff --git a/Engine/Source/Runtime/Core/Private/Mac/MacApplication.cpp b/Engine/Source/Runtime/Core/Private/Mac/MacApplication.cpp
--- a/Engine/Source/Runtime/Core/Private/Mac/MacApplication.cpp
+++ b/Engine/Source/Runtime/Core/Private/Mac/MacApplication.cpp
@@ -110,6 +110,7 @@ FMacApplication::FMacApplication()
 	, bIsMouseCaptureEnabled( false )
 	, bIsMouseCursorLocked( false )
 	, bSystemModalMode( false )
+	, bIsProcessingNSEvent( false )
 	, ModifierKeysFlags( 0 )
 	, CurrentModifierFlags( 0 )
 	, EventMonitor( NULL )
@@ -511,6 +512,8 @@ void FMacApplication::ProcessNSEvent(NSEvent* const Event, TSharedPtr< FMacWindo
 		return;
 	}
 
+	bIsProcessingNSEvent = true;
+
 	if( CurrentModifierFlags != [Event modifierFlags] )
 	{
 		NSUInteger ModifierFlags = [Event modifierFlags];
@@ -816,7 +819,7 @@ void FMacApplication::ProcessNSEvent(NSEvent* const Event, TSharedPtr< FMacWindo
 				FCocoaMenu* MainMenu = [[NSApp mainMenu] isKindOfClass:[FCocoaMenu class]] ? (FCocoaMenu*)[NSApp mainMenu]: nil;
 				if ( MainMenu )
 				{
-					[MainMenu highlightKeyEquivalent:Event];
+					MainThreadCall(^{ [MainMenu highlightKeyEquivalent:Event]; }, NSDefaultRunLoopMode, true);
 				}
 			}
 			else
@@ -848,6 +851,8 @@ void FMacApplication::ProcessNSEvent(NSEvent* const Event, TSharedPtr< FMacWindo
 			break;
 		}
 	}
+
+	bIsProcessingNSEvent = false;
 }
 
 void FMacApplication::ProcessEvent( NSEvent* Event )
diff --git a/Engine/Source/Runtime/Core/Private/Mac/MacPlatformMisc.cpp b/Engine/Source/Runtime/Core/Private/Mac/MacPlatformMisc.cpp
--- a/Engine/Source/Runtime/Core/Private/Mac/MacPlatformMisc.cpp
+++ b/Engine/Source/Runtime/Core/Private/Mac/MacPlatformMisc.cpp
@@ -463,7 +463,7 @@ void FMacPlatformMisc::PumpMessages( bool bFromMainLoop )
 	{
 		ProcessGameThreadEvents();
 
-		if (UpdateCachedMacMenuState && bChachedMacMenuStateNeedsUpdate && IsInGameThread())
+		if (UpdateCachedMacMenuState && bChachedMacMenuStateNeedsUpdate && MacApplication && !MacApplication->IsProcessingNSEvent() && IsInGameThread())
 		{
 			UpdateCachedMacMenuState();
 			bChachedMacMenuStateNeedsUpdate = false;
diff --git a/Engine/Source/Runtime/Core/Public/Mac/MacApplication.h b/Engine/Source/Runtime/Core/Public/Mac/MacApplication.h
--- a/Engine/Source/Runtime/Core/Public/Mac/MacApplication.h
+++ b/Engine/Source/Runtime/Core/Public/Mac/MacApplication.h
@@ -76,6 +76,8 @@ public:
 
 	void UseMouseCaptureWindow(bool bUseMouseCaptureWindow);
 
+	bool IsProcessingNSEvent() const { return bIsProcessingNSEvent; }
+
 #if WITH_EDITOR
     virtual void SendAnalytics(IAnalyticsProvider* Provider) override;
 
@@ -163,6 +165,7 @@ private:
 	bool bIsMouseCursorLocked;
 
 	bool bSystemModalMode;
+	bool bIsProcessingNSEvent;
 
 	TSharedPtr< FMacWindow > LastEventWindow;
 
