Mark class with CORE_API allowing use from other modules

diff --git a/Engine/Source/Runtime/Core/Public/Misc/AES.h b/Engine/Source/Runtime/Core/Public/Misc/AES.h
--- a/Engine/Source/Runtime/Core/Public/Misc/AES.h
+++ b/Engine/Source/Runtime/Core/Public/Misc/AES.h
@@ -8,7 +8,7 @@
 
 #define AES_BLOCK_SIZE				16
 
-struct FAES
+struct CORE_API FAES
 {
 	static const uint32 AESBlockSize = 16;
 
