Remove unused parameters
Flag ProcessInputFileContents with CORE_API to allow use from other modules

diff --git a/Engine/Source/Runtime/Core/Private/Misc/ConfigCacheIni.cpp b/Engine/Source/Runtime/Core/Private/Misc/ConfigCacheIni.cpp
--- a/Engine/Source/Runtime/Core/Private/Misc/ConfigCacheIni.cpp
+++ b/Engine/Source/Runtime/Core/Private/Misc/ConfigCacheIni.cpp
@@ -160,7 +160,7 @@ bool FConfigFile::Combine(const FString& Filename)
 	// needs file ops), and the other caller of this is already checking for disabled file ops
 	if( FFileHelper::LoadFileToString( Text, *Filename ) )
 	{
-		CombineFromBuffer(Filename,Text);
+		CombineFromBuffer(Text);
 		return true;
 	}
 
@@ -168,7 +168,7 @@ bool FConfigFile::Combine(const FString& Filename)
 }
 
 
-void FConfigFile::CombineFromBuffer(const FString& Filename,const FString& Buffer)
+void FConfigFile::CombineFromBuffer(const FString& Buffer)
 {
 	// Replace %GAME% with game name.
 	FString Text = Buffer.Replace( TEXT("%GAME%"), GGameName, ESearchCase::CaseSensitive );
@@ -376,10 +376,9 @@ void FConfigFile::CombineFromBuffer(const FString& Filename,const FString& Buffe
 /**
  * Process the contents of an .ini file that has been read into an FString
  * 
- * @param Filename Name of the .ini file the contents came from
  * @param Contents Contents of the .ini file
  */
-void FConfigFile::ProcessInputFileContents(const FString& Filename, FString& Contents)
+void FConfigFile::ProcessInputFileContents(const FString& Contents)
 {
 	// Replace %GAME% with game name.
 	FString Text = Contents.Replace(TEXT("%GAME%"), GGameName, ESearchCase::CaseSensitive);
@@ -545,7 +544,7 @@ void FConfigFile::Read( const FString& Filename )
 		if( FFileHelper::LoadFileToString( Text, *Filename ) )
 		{
 			// process the contents of the string
-			ProcessInputFileContents(Filename, Text);
+			ProcessInputFileContents(Text);
 		}
 	}
 }
diff --git a/Engine/Source/Runtime/Core/Private/Misc/RemoteConfigIni.cpp b/Engine/Source/Runtime/Core/Private/Misc/RemoteConfigIni.cpp
--- a/Engine/Source/Runtime/Core/Private/Misc/RemoteConfigIni.cpp
+++ b/Engine/Source/Runtime/Core/Private/Misc/RemoteConfigIni.cpp
@@ -482,11 +482,11 @@ void ProcessIniContents(const TCHAR* FilenameToLoad, const TCHAR* IniFileName, F
 		
 		if (bDoCombine)
 		{
-			Config->CombineFromBuffer(IniFileName, RemoteInfo->Buffer);
+			Config->CombineFromBuffer(RemoteInfo->Buffer);
 		}
 		else
 		{
-			Config->ProcessInputFileContents(IniFileName, RemoteInfo->Buffer);
+			Config->ProcessInputFileContents(RemoteInfo->Buffer);
 		}
 	}
 }
diff --git a/Engine/Source/Runtime/Core/Public/Misc/ConfigCacheIni.h b/Engine/Source/Runtime/Core/Public/Misc/ConfigCacheIni.h
--- a/Engine/Source/Runtime/Core/Public/Misc/ConfigCacheIni.h
+++ b/Engine/Source/Runtime/Core/Public/Misc/ConfigCacheIni.h
@@ -80,7 +80,7 @@ public:
 	bool operator!=( const FConfigFile& Other ) const;
 
 	CORE_API bool Combine( const FString& Filename);
-	CORE_API void CombineFromBuffer(const FString& Filename,const FString& Buffer);
+	CORE_API void CombineFromBuffer(const FString& Buffer);
 	CORE_API void Read( const FString& Filename );
 	CORE_API bool Write( const FString& Filename, bool bDoRemoteWrite=true, const FString& InitialText=FString() );
 	CORE_API void Dump(FOutputDevice& Ar);
@@ -99,7 +99,7 @@ public:
 	 * @param Filename Name of the .ini file the contents came from
 	 * @param Contents Contents of the .ini file
 	 */
-	void ProcessInputFileContents(const FString& Filename, FString& Contents);
+	CORE_API void ProcessInputFileContents(const FString& Contents);
 
 
 	/** Adds any properties that exist in InSourceFile that this config file is missing */
