Core - Fixed bad comments regarding SIZE_T

diff --git a/Engine/Source/Runtime/Core/Public/GenericPlatform/GenericPlatform.h b/Engine/Source/Runtime/Core/Public/GenericPlatform/GenericPlatform.h
--- a/Engine/Source/Runtime/Core/Public/GenericPlatform/GenericPlatform.h
+++ b/Engine/Source/Runtime/Core/Public/GenericPlatform/GenericPlatform.h
@@ -55,8 +55,8 @@ struct FGenericPlatformTypes
 	typedef WIDECHAR			TCHAR;		// A switchable character  - In-memory only.  Either ANSICHAR or WIDECHAR, depending on a licensee's requirements.
 
 	typedef SelectIntPointerType<uint32, uint64, sizeof(void*)>::TIntPointer UPTRINT;	// unsigned int the same size as a pointer
-	typedef SelectIntPointerType<int32, int64, sizeof(void*)>::TIntPointer PTRINT;	// signed int the same size as a pointer
-	typedef UPTRINT SIZE_T;															// signed int the same size as a pointer
+	typedef SelectIntPointerType<int32, int64, sizeof(void*)>::TIntPointer PTRINT;		// signed int the same size as a pointer
+	typedef UPTRINT SIZE_T;																// unsigned int the same size as a pointer
 
 	typedef int32					TYPE_OF_NULL;
 	typedef decltype(nullptr)		TYPE_OF_NULLPTR;
diff --git a/Engine/Source/Runtime/Core/Public/HAL/Platform.h b/Engine/Source/Runtime/Core/Public/HAL/Platform.h
--- a/Engine/Source/Runtime/Core/Public/HAL/Platform.h
+++ b/Engine/Source/Runtime/Core/Public/HAL/Platform.h
@@ -583,12 +583,12 @@ typedef FPlatformTypes::CHAR16		UCS2CHAR;	///< A 16-bit character containing a U
 typedef FPlatformTypes::CHAR16		UTF16CHAR;	///< A 16-bit character containing a UTF16 (Unicode, 16-bit, variable-width) code unit.
 typedef FPlatformTypes::CHAR32		UTF32CHAR;	///< A 32-bit character containing a UTF32 (Unicode, 32-bit, fixed-width) code unit.
 
-typedef FPlatformTypes::UPTRINT UPTRINT;	///< An unsigned integer the same size as a pointer
-typedef FPlatformTypes::PTRINT PTRINT;		///< A signed integer the same size as a pointer
-typedef FPlatformTypes::SIZE_T SIZE_T;		///< A signed integer the same size as a pointer
+typedef FPlatformTypes::UPTRINT UPTRINT;		///< An unsigned integer the same size as a pointer
+typedef FPlatformTypes::PTRINT PTRINT;			///< A signed integer the same size as a pointer
+typedef FPlatformTypes::SIZE_T SIZE_T;			///< An unsigned integer the same size as a pointer, the same as UPTRINT
 
-typedef FPlatformTypes::TYPE_OF_NULL	TYPE_OF_NULL;	///< The type of the NULL constant.
-typedef FPlatformTypes::TYPE_OF_NULLPTR	TYPE_OF_NULLPTR; ///< The type of the C++ nullptr keyword.
+typedef FPlatformTypes::TYPE_OF_NULL	TYPE_OF_NULL;		///< The type of the NULL constant.
+typedef FPlatformTypes::TYPE_OF_NULLPTR	TYPE_OF_NULLPTR;	///< The type of the C++ nullptr keyword.
 
 //------------------------------------------------------------------
 // Test the global types
