Stats - Removed unused macros

diff --git a/Engine/Source/Runtime/Core/Public/Stats/Stats.h b/Engine/Source/Runtime/Core/Public/Stats/Stats.h
--- a/Engine/Source/Runtime/Core/Public/Stats/Stats.h
+++ b/Engine/Source/Runtime/Core/Public/Stats/Stats.h
@@ -1,8 +1,5 @@
 // Copyright 1998-2014 Epic Games, Inc. All Rights Reserved.
 
-/*=============================================================================
-	UnStats.h: Performance stats framework.
-=============================================================================*/
 #pragma once
 
 #define FORCEINLINE_STATS FORCEINLINE
@@ -96,10 +93,6 @@ FORCEINLINE void StatsMasterEnableSubtract(int32 Value = 1)
 }
 
 // Remove all the macros
-#define STAT_IS_COLLECTING(Stat) (false)
-
-#define REGISTER_DYNAMIC_CYCLE_STAT(Stat, Group) (TStatId())
-
 
 #define DEFINE_STAT(Stat)
 #define SCOPE_CYCLE_COUNTER(Stat)
diff --git a/Engine/Source/Runtime/Core/Public/Stats/Stats2.h b/Engine/Source/Runtime/Core/Public/Stats/Stats2.h
--- a/Engine/Source/Runtime/Core/Public/Stats/Stats2.h
+++ b/Engine/Source/Runtime/Core/Public/Stats/Stats2.h
@@ -1435,8 +1435,6 @@ struct FStat_##StatName\
 
 #define STAT_GROUP_TO_FStatGroup(Group) FStatGroup_##Group
 
-#define STAT_IS_COLLECTING(Stat) (FThreadStats::IsCollectingData(GET_STATID(Stat)))
-
 #define DEFINE_STAT(Stat) \
 	struct FThreadSafeStaticStat<FStat_##Stat> StatPtr_##Stat;
 
diff --git a/Engine/Source/Runtime/CoreUObject/Public/UObject/UObjectBase.h b/Engine/Source/Runtime/CoreUObject/Public/UObject/UObjectBase.h
--- a/Engine/Source/Runtime/CoreUObject/Public/UObject/UObjectBase.h
+++ b/Engine/Source/Runtime/CoreUObject/Public/UObject/UObjectBase.h
@@ -127,7 +127,7 @@ public:
 	{
 #if STATS
 		// this is done to avoid even registering stats for a disabled group (unless we plan on using it later)
-		if (bForDeferredUse || STAT_IS_COLLECTING(STAT_UObjectsStatGroupTester))
+		if (bForDeferredUse || FThreadStats::IsCollectingData(GET_STATID(STAT_UObjectsStatGroupTester)))
 		{
 			if (!StatID.IsValidStat())
 			{
