Fixed a problem with FMacApplication::bIsProcessingNSEvent being reset when ProcessNSEvent() was called recursively

diff --git a/Engine/Source/Runtime/Core/Private/Mac/MacApplication.cpp b/Engine/Source/Runtime/Core/Private/Mac/MacApplication.cpp
--- a/Engine/Source/Runtime/Core/Private/Mac/MacApplication.cpp
+++ b/Engine/Source/Runtime/Core/Private/Mac/MacApplication.cpp
@@ -512,6 +512,7 @@ void FMacApplication::ProcessNSEvent(NSEvent* const Event, TSharedPtr< FMacWindo
 		return;
 	}
 
+	const bool bWasProcessingNSEvent = bIsProcessingNSEvent; // ProcessNSEvent can be called recursively
 	bIsProcessingNSEvent = true;
 
 	if( CurrentModifierFlags != [Event modifierFlags] )
@@ -852,7 +853,7 @@ void FMacApplication::ProcessNSEvent(NSEvent* const Event, TSharedPtr< FMacWindo
 		}
 	}
 
-	bIsProcessingNSEvent = false;
+	bIsProcessingNSEvent = bWasProcessingNSEvent;
 }
 
 void FMacApplication::ProcessEvent( NSEvent* Event )
