#docs - Fix mistaken comments about composed transform order of transforms.

diff --git a/Engine/Source/Runtime/Core/Public/Math/Transform.h b/Engine/Source/Runtime/Core/Public/Math/Transform.h
--- a/Engine/Source/Runtime/Core/Public/Math/Transform.h
+++ b/Engine/Source/Runtime/Core/Public/Math/Transform.h
@@ -13,10 +13,10 @@
  * Transformation of direction vectors is applied in the order: Scale -> Rotate.
  *
  * Order matters when composing transforms: C = A * B will yield a transform C that logically
- * first applies B then A to any subsequent transformation.
+ * first applies A then B to any subsequent transformation. Note that this is the opposite order of Quaternion multiplication.
  *
- * Example: LocalToWorld = (LocalToWorld * DeltaRotation) will change rotation in local space by DeltaRotation.
- * Example: LocalToWorld = (DeltaRotation * LocalToWorld) will change rotation in world space by DeltaRotation.
+ * Example: LocalToWorld = (DeltaRotation * LocalToWorld) will change rotation in local space by DeltaRotation.
+ * Example: LocalToWorld = (LocalToWorld * DeltaRotation) will change rotation in world space by DeltaRotation.
  */
 
 class FTransform
@@ -453,7 +453,7 @@ public:
 
 	/**
 	 * Return a transform that is the result of this multiplied by another transform.
-	 * Order matters when composing transforms : C = A * B will yield a transform C that logically first applies B then A to any subsequent transformation.
+	 * Order matters when composing transforms : C = A * B will yield a transform C that logically first applies A then B to any subsequent transformation.
 	 * 
 	 * @param  Other other transform by which to multiply.
 	 * @return new transform: this * Other
@@ -463,7 +463,7 @@ public:
 
 	/**
 	 * Sets this transform to the result of this multiplied by another transform.
-	 * Order matters when composing transforms : C = A * B will yield a transform C that logically first applies B then A to any subsequent transformation.
+	 * Order matters when composing transforms : C = A * B will yield a transform C that logically first applies A then B to any subsequent transformation.
 	 *
 	 * @param  Other other transform by which to multiply.
 	 */
@@ -471,7 +471,7 @@ public:
 	
 	/**
 	* Return a transform that is the result of this multiplied by another transform (made only from a rotation).
-	* Order matters when composing transforms : C = A * B will yield a transform C that logically first applies B then A to any subsequent transformation.
+	* Order matters when composing transforms : C = A * B will yield a transform C that logically first applies A then B to any subsequent transformation.
 	*
 	* @param  Other other quaternion rotation by which to multiply.
 	* @return new transform: this * FTransform(Other)
@@ -480,7 +480,7 @@ public:
 	
 	/**
 	* Sets this transform to the result of this multiplied by another transform (made only from a rotation).
-	* Order matters when composing transforms : C = A * B will yield a transform C that logically first applies B then A to any subsequent transformation.
+	* Order matters when composing transforms : C = A * B will yield a transform C that logically first applies A then B to any subsequent transformation.
 	*
 	* @param  Other other quaternion rotation by which to multiply.
 	*/
@@ -619,7 +619,7 @@ public:
 	/**
 	 * Create a new transform: OutTransform = A * B.
 	 *
-	 * Order matters when composing transforms : A * B will yield a transform that logically first applies B then A to any subsequent transformation.
+	 * Order matters when composing transforms : A * B will yield a transform that logically first applies A then B to any subsequent transformation.
 	 * 
 	 * @param  OutTransform pointer to transform that will store the result of A * B.
 	 * @param  A Transform A.
diff --git a/Engine/Source/Runtime/Core/Public/Math/TransformVectorized.h b/Engine/Source/Runtime/Core/Public/Math/TransformVectorized.h
--- a/Engine/Source/Runtime/Core/Public/Math/TransformVectorized.h
+++ b/Engine/Source/Runtime/Core/Public/Math/TransformVectorized.h
@@ -7,20 +7,20 @@
 #define MAKE_QUATINV_VECTORREGISTER(X) VectorMultiply(GlobalVectorConstants::QINV_SIGN_MASK, X)
 
 /**
-* Transform composed of Scale, Rotation (as a quaternion), and Translation.
-*
-* Transforms can be used to convert from one space to another, for example by transforming
-* positions and directions from local space to world space.
-*
-* Transformation of position vectors is applied in the order:  Scale -> Rotate -> Translate.
-* Transformation of direction vectors is applied in the order: Scale -> Rotate.
-*
-* Order matters when composing transforms: C = A * B will yield a transform C that logically
-* first applies B then A to any subsequent transformation.
-*
-* Example: LocalToWorld = (LocalToWorld * DeltaRotation) will change rotation in local space by DeltaRotation.
-* Example: LocalToWorld = (DeltaRotation * LocalToWorld) will change rotation in world space by DeltaRotation.
-*/
+ * Transform composed of Scale, Rotation (as a quaternion), and Translation.
+ *
+ * Transforms can be used to convert from one space to another, for example by transforming
+ * positions and directions from local space to world space.
+ *
+ * Transformation of position vectors is applied in the order:  Scale -> Rotate -> Translate.
+ * Transformation of direction vectors is applied in the order: Scale -> Rotate.
+ *
+ * Order matters when composing transforms: C = A * B will yield a transform C that logically
+ * first applies A then B to any subsequent transformation. Note that this is the opposite order of Quaternion multiplication.
+ *
+ * Example: LocalToWorld = (DeltaRotation * LocalToWorld) will change rotation in local space by DeltaRotation.
+ * Example: LocalToWorld = (LocalToWorld * DeltaRotation) will change rotation in world space by DeltaRotation.
+ */
 
 MS_ALIGN(16) class FTransform
 {
diff --git a/Engine/Source/Runtime/Engine/Classes/Kismet/KismetMathLibrary.h b/Engine/Source/Runtime/Engine/Classes/Kismet/KismetMathLibrary.h
--- a/Engine/Source/Runtime/Engine/Classes/Kismet/KismetMathLibrary.h
+++ b/Engine/Source/Runtime/Engine/Classes/Kismet/KismetMathLibrary.h
@@ -1221,10 +1221,10 @@ class UKismetMathLibrary : public UBlueprintFunctionLibrary
 	 * Compose two transforms in order: A * B.
 	 *
 	 * Order matters when composing transforms:
-	 * A * B will yield a transform that logically first applies B then A to any subsequent transformation.
+	 * A * B will yield a transform that logically first applies A then B to any subsequent transformation.
 	 *
-	 * Example: LocalToWorld = ComposeTransforms(LocalToWorld, DeltaRotation) will change rotation in local space by DeltaRotation.
-	 * Example: LocalToWorld = ComposeTransforms(DeltaRotation, LocalToWorld) will change rotation in world space by DeltaRotation.
+	 * Example: LocalToWorld = ComposeTransforms(DeltaRotation, LocalToWorld) will change rotation in local space by DeltaRotation.
+	 * Example: LocalToWorld = ComposeTransforms(LocalToWorld, DeltaRotation) will change rotation in world space by DeltaRotation.
 	 *
 	 * @return New transform: A * B
 	 */
