UE4: Fixed a mistake in an error message pointing the user to the wrong ini file.

diff --git a/Engine/Source/Runtime/Core/Private/Misc/ConfigCacheIni.cpp b/Engine/Source/Runtime/Core/Private/Misc/ConfigCacheIni.cpp
--- a/Engine/Source/Runtime/Core/Private/Misc/ConfigCacheIni.cpp
+++ b/Engine/Source/Runtime/Core/Private/Misc/ConfigCacheIni.cpp
@@ -2711,7 +2711,7 @@ void FConfigCacheIni::InitializeConfigSystem()
 			else
 			{
 				const FText Message = FText::Format(
-					NSLOCTEXT("Core", "MismatchedGameNames", "The name of the .uproject file ('{0}') must match the GameName key in the [URL] section of Config/DefaultGame.ini (currently '{1}').  Please either change the ini or rename the .uproject to match each other (case-insensitive match)."),
+					NSLOCTEXT("Core", "MismatchedGameNames", "The name of the .uproject file ('{0}') must match the GameName key in the [URL] section of Config/DefaultEngine.ini (currently '{1}').  Please either change the ini or rename the .uproject to match each other (case-insensitive match)."),
 					FText::FromString(GGameName),
 					FText::FromString(GameName));
 				if (!GIsBuildMachine)
