Handle a case where Mac app is deactivated while there are no visible windows (rarely caused crashes on quit or between hiding the splash screen and showing the editor window)

diff --git a/Engine/Source/Runtime/Core/Private/Mac/MacApplication.cpp b/Engine/Source/Runtime/Core/Private/Mac/MacApplication.cpp
--- a/Engine/Source/Runtime/Core/Private/Mac/MacApplication.cpp
+++ b/Engine/Source/Runtime/Core/Private/Mac/MacApplication.cpp
@@ -155,15 +155,18 @@ FMacApplication::FMacApplication()
 									}
 								}
 							 
-							 	for( auto Window : Windows )
+								if (BaseWindow)
 								{
-									NSWindow* CocoaWindow = (NSWindow*)Window->GetOSWindowHandle();
-									if ( CocoaWindow && Window->GetDefinition().IsModalWindow )
+									for( auto Window : Windows )
 									{
-										[CocoaWindow setLevel:NSModalPanelWindowLevel];
-										if ( BaseWindow )
+										NSWindow* CocoaWindow = (NSWindow*)Window->GetOSWindowHandle();
+										if ( CocoaWindow && Window->GetDefinition().IsModalWindow )
 										{
-											[CocoaWindow orderWindow:NSWindowBelow relativeTo:[BaseWindow windowNumber]];
+											[CocoaWindow setLevel:NSModalPanelWindowLevel];
+											if ( BaseWindow )
+											{
+												[CocoaWindow orderWindow:NSWindowBelow relativeTo:[BaseWindow windowNumber]];
+											}
 										}
 									}
 								}
@@ -196,20 +199,21 @@ FMacApplication::FMacApplication()
 											}
 										}
 										
-										check(BaseWindow);
-										
-										for( auto Window : Windows )
+										if (BaseWindow)
 										{
-											NSWindow* CocoaWindow = (NSWindow*)Window->GetOSWindowHandle();
-											if ( CocoaWindow && Window->GetDefinition().IsModalWindow )
+											for( auto Window : Windows )
 											{
-												[CocoaWindow setLevel:NSNormalWindowLevel];
-												[CocoaWindow orderWindow:NSWindowAbove relativeTo:[BaseWindow windowNumber]];
-												BaseWindow = CocoaWindow;
+												NSWindow* CocoaWindow = (NSWindow*)Window->GetOSWindowHandle();
+												if ( CocoaWindow && Window->GetDefinition().IsModalWindow )
+												{
+													[CocoaWindow setLevel:NSNormalWindowLevel];
+													[CocoaWindow orderWindow:NSWindowAbove relativeTo:[BaseWindow windowNumber]];
+													BaseWindow = CocoaWindow;
+												}
 											}
 										}
-									}
-							   
+								   }
+
 									CFRunLoopRunInMode(kCFRunLoopDefaultMode, 0, true);
 							   
 									// If editor thread doesn't have the focus, don't suck up too much CPU time.
