#docs - Documentation for quaternions, explaining order of composition during multiplication.

diff --git a/Engine/Source/Runtime/Core/Public/Math/Quat.h b/Engine/Source/Runtime/Core/Public/Math/Quat.h
--- a/Engine/Source/Runtime/Core/Public/Math/Quat.h
+++ b/Engine/Source/Runtime/Core/Public/Math/Quat.h
@@ -8,9 +8,15 @@
 
 
 /**
- * Floating point quaternion.
- *
+ * Floating point quaternion that can represent a rotation about an axis in 3-D space.
  * The X, Y, Z, W components also double as the Axis/Angle format.
+ *
+ * Order matters when composing quaternions: C = A * B will yield a quaternion C that logically
+ * first applies B then A to any subsequent transformation (right first, then left).
+ * Note that this is the opposite order of FTransform multiplication.
+ *
+ * Example: LocalToWorld = (LocalToWorld * DeltaRotation) will change rotation in local space by DeltaRotation.
+ * Example: LocalToWorld = (DeltaRotation * LocalToWorld) will change rotation in world space by DeltaRotation.
  */
 MS_ALIGN(16) class FQuat 
 {
@@ -99,6 +105,7 @@ public:
 
 	/**
 	 * Gets the result of adding a Quaternion to this.
+	 * This is a component-wise addition; composing quaternions should be done via multiplication.
 	 *
 	 * @param Q The Quaternion to add.
 	 * @return The result of addition.
@@ -107,14 +114,16 @@ public:
 
 	/**
 	 * Adds to this quaternion.
+	 * This is a component-wise addition; composing quaternions should be done via multiplication.
 	 *
 	 * @param Other The quaternion to add to this.
-	 * @return Reference to this after addition.
+	 * @return Result after addition.
 	 */
 	FORCEINLINE FQuat operator+=(const FQuat& Q);
 
 	/**
 	 * Gets the result of subtracting a Quaternion to this.
+	 * This is a component-wise subtraction; composing quaternions should be done via multiplication.
 	 *
 	 * @param Q The Quaternion to subtract.
 	 * @return The result of subtraction.
@@ -132,6 +141,7 @@ public:
 
 	/**
 	 * Subtracts another quaternion from this.
+	 * This is a component-wise subtraction; composing quaternions should be done via multiplication.
 	 *
 	 * @param Q The other quaternion.
 	 * @return reference to this after subtraction.
@@ -139,34 +149,44 @@ public:
 	FORCEINLINE FQuat operator-=(const FQuat& Q);
 
 	/**
-	 * Gets the result of multiplying a Quaternion to this.
+	 * Gets the result of multiplying this by another quaternion (this * Q).
+	 *
+	 * Order matters when composing quaternions: C = A * B will yield a quaternion C that logically
+	 * first applies B then A to any subsequent transformation (right first, then left).
 	 *
 	 * @param Q The Quaternion to multiply this by.
-	 * @return The result of multiplication.
+	 * @return The result of multiplication (this * Q).
 	 */
 	FORCEINLINE FQuat operator*( const FQuat& Q ) const;
 
 	/**
-	 * Multiply this by a quaternion.
+	 * Multiply this by a quaternion (this = this * Q).
 	 *
-	 * @param Q the quaternion to multiply by this.
-	 * @return reference to this after multiply.
+	 * Order matters when composing quaternions: C = A * B will yield a quaternion C that logically
+	 * first applies B then A to any subsequent transformation (right first, then left).
+	 *
+	 * @param Q the quaternion to multiply this by.
+	 * @return The result of multiplication (this * Q).
 	 */
 	FORCEINLINE FQuat operator*=(const FQuat& Q);
 
 	/**
-	 * This transforms vector, not dot product
+	 * Rotate a vector by this quaternion.
 	 *
-	 * @param V
+	 * @param V the vector to be rotated
+	 * @return vector after rotation
+	 * @see RotateVector()
 	 */
 	FVector operator*( const FVector& V ) const;
 
 	/** 
+	 * Multiply this by a matrix.
 	 * This matrix conversion came from
 	 * http://www.m-hikari.com/ija/ija-password-2008/ija-password17-20-2008/aristidouIJA17-20-2008.pdf
 	 * used for non-uniform scaling transform.
 	 *
-	 * @param M
+	 * @param M Matrix to multiply by.
+	 * @return Matrix result after multiplication.
 	 */	
 	FMatrix operator*( const FMatrix& M ) const;
 	
@@ -196,9 +216,12 @@ public:
 
  	/**
 	 * Checks whether two quaternions are identical.
+	 * This is an exact comparison per-component;see Equals() for a comparison
+	 * that allows for a small error tolerance and flipped axes of rotation.
 	 *
 	 * @param Q The other quaternion.
 	 * @return true if two quaternion are identical, otherwise false.
+	 * @see Equals()
 	 */
 	bool operator==( const FQuat& Q ) const;
 
@@ -265,7 +288,7 @@ public:
 	void ToAxisAndAngle( FVector& Axis, float& Angle ) const;
 
 	/**
-	 * Return vector rotated by this quaternion
+	 * Rotate a vector by this quaternion.
 	 *
 	 * @param V the vector to be rotated
 	 * @return vector after rotation
diff --git a/Engine/Source/Runtime/Core/Public/Math/Transform.h b/Engine/Source/Runtime/Core/Public/Math/Transform.h
--- a/Engine/Source/Runtime/Core/Public/Math/Transform.h
+++ b/Engine/Source/Runtime/Core/Public/Math/Transform.h
@@ -13,7 +13,7 @@
  * Transformation of direction vectors is applied in the order: Scale -> Rotate.
  *
  * Order matters when composing transforms: C = A * B will yield a transform C that logically
- * first applies A then B to any subsequent transformation. Note that this is the opposite order of Quaternion multiplication.
+ * first applies A then B to any subsequent transformation. Note that this is the opposite order of quaternion (FQuat) multiplication.
  *
  * Example: LocalToWorld = (DeltaRotation * LocalToWorld) will change rotation in local space by DeltaRotation.
  * Example: LocalToWorld = (LocalToWorld * DeltaRotation) will change rotation in world space by DeltaRotation.
diff --git a/Engine/Source/Runtime/Core/Public/Math/TransformVectorized.h b/Engine/Source/Runtime/Core/Public/Math/TransformVectorized.h
--- a/Engine/Source/Runtime/Core/Public/Math/TransformVectorized.h
+++ b/Engine/Source/Runtime/Core/Public/Math/TransformVectorized.h
@@ -16,7 +16,7 @@
  * Transformation of direction vectors is applied in the order: Scale -> Rotate.
  *
  * Order matters when composing transforms: C = A * B will yield a transform C that logically
- * first applies A then B to any subsequent transformation. Note that this is the opposite order of Quaternion multiplication.
+ * first applies A then B to any subsequent transformation. Note that this is the opposite order of quaternion (FQuat) multiplication.
  *
  * Example: LocalToWorld = (DeltaRotation * LocalToWorld) will change rotation in local space by DeltaRotation.
  * Example: LocalToWorld = (LocalToWorld * DeltaRotation) will change rotation in world space by DeltaRotation.
diff --git a/Engine/Source/Runtime/Core/Public/Math/UnrealMathDirectX.h b/Engine/Source/Runtime/Core/Public/Math/UnrealMathDirectX.h
--- a/Engine/Source/Runtime/Core/Public/Math/UnrealMathDirectX.h
+++ b/Engine/Source/Runtime/Core/Public/Math/UnrealMathDirectX.h
@@ -633,11 +633,14 @@ FORCEINLINE uint32 VectorAnyGreaterThan(const VectorRegister& Vec1, const Vector
 #define VECTOR_ROUND_TOWARD_ZERO		_MM_ROUND_TOWARD_ZERO
 
 /**
-* Multiplies two quaternions: The order matters.
+* Multiplies two quaternions; the order matters.
+*
+* Order matters when composing quaternions: C = VectorQuaternionMultiply2(A, B) will yield a quaternion C = A * B
+* that logically first applies B then A to any subsequent transformation (right first, then left).
 *
-* @param Result	Returns Quat1 * Quat2
 * @param Quat1	Pointer to the first quaternion
 * @param Quat2	Pointer to the second quaternion
+* @return Quat1 * Quat2
 */
 FORCEINLINE VectorRegister VectorQuaternionMultiply2( const VectorRegister& Quat1, const VectorRegister& Quat2 )
 {
@@ -653,9 +656,12 @@ FORCEINLINE VectorRegister VectorQuaternionMultiply2( const VectorRegister& Quat
 }
 
 /**
-* Multiplies two quaternions: The order matters.
+* Multiplies two quaternions; the order matters.
+*
+* When composing quaternions: VectorQuaternionMultiply(C, A, B) will yield a quaternion C = A * B
+* that logically first applies B then A to any subsequent transformation (right first, then left).
 *
-* @param Result	Pointer to where the result should be stored
+* @param Result	Pointer to where the result Quat1 * Quat2 should be stored
 * @param Quat1	Pointer to the first quaternion (must not be the destination)
 * @param Quat2	Pointer to the second quaternion (must not be the destination)
 */
@@ -669,11 +675,14 @@ FORCEINLINE void VectorQuaternionMultiply( FQuat *Result, const FQuat* Quat1, co
 }
 
 /**
-* Multiplies two quaternions: The order matters.
+* Multiplies two quaternions; the order matters.
+*
+* When composing quaternions: VectorQuaternionMultiply(C, A, B) will yield a quaternion C = A * B
+* that logically first applies B then A to any subsequent transformation (right first, then left).
 *
-* @param VResult	VectorRegister Pointer to where the result should be stored
-* @param VQuat1	VectorRegister Pointer to the first quaternion (must not be the destination)
-* @param VQuat2	VectorRegister Pointer to the second quaternion (must not be the destination)
+* @param Result	Pointer to where the result Quat1 * Quat2 should be stored
+* @param Quat1	Pointer to the first quaternion (must not be the destination)
+* @param Quat2	Pointer to the second quaternion (must not be the destination)
 */
 FORCEINLINE void VectorQuaternionMultiply( VectorRegister *VResult, const VectorRegister* VQuat1, const VectorRegister* VQuat2)
 {	
diff --git a/Engine/Source/Runtime/Core/Public/Math/UnrealMathFPU.h b/Engine/Source/Runtime/Core/Public/Math/UnrealMathFPU.h
--- a/Engine/Source/Runtime/Core/Public/Math/UnrealMathFPU.h
+++ b/Engine/Source/Runtime/Core/Public/Math/UnrealMathFPU.h
@@ -594,9 +594,12 @@ FORCEINLINE VectorRegister VectorNormalize( const VectorRegister& Vector )
 #define USE_FAST_QUAT_MUL 1
 
 /**
-* Multiplies two quaternions: The order matters.
+* Multiplies two quaternions; the order matters.
 *
-* @param Result	Pointer to where the result should be stored
+* When composing quaternions: VectorQuaternionMultiply(C, A, B) will yield a quaternion C = A * B
+* that logically first applies B then A to any subsequent transformation (right first, then left).
+*
+* @param Result	Pointer to where the result Quat1 * Quat2 should be stored
 * @param Quat1	Pointer to the first quaternion (must not be the destination)
 * @param Quat2	Pointer to the second quaternion (must not be the destination)
 */
@@ -639,11 +642,14 @@ FORCEINLINE void VectorQuaternionMultiply( void *Result, const void* Quat1, cons
 }
 
 /**
-* Multiplies two quaternions: The order matters.
+* Multiplies two quaternions; the order matters.
+*
+* Order matters when composing quaternions: C = VectorQuaternionMultiply2(A, B) will yield a quaternion C = A * B
+* that logically first applies B then A to any subsequent transformation (right first, then left).
 *
-* @param Result	Returns Quat1 * Quat2
 * @param Quat1	Pointer to the first quaternion
 * @param Quat2	Pointer to the second quaternion
+* @return Quat1 * Quat2
 */
 FORCEINLINE VectorRegister VectorQuaternionMultiply2( const VectorRegister& Quat1, const VectorRegister& Quat2 )
 {
diff --git a/Engine/Source/Runtime/Core/Public/Math/UnrealMathNeon.h b/Engine/Source/Runtime/Core/Public/Math/UnrealMathNeon.h
--- a/Engine/Source/Runtime/Core/Public/Math/UnrealMathNeon.h
+++ b/Engine/Source/Runtime/Core/Public/Math/UnrealMathNeon.h
@@ -935,11 +935,14 @@ static const VectorRegister QMULTI_SIGN_MASK1 = MakeVectorRegister( 1.f, 1.f, -1
 static const VectorRegister QMULTI_SIGN_MASK2 = MakeVectorRegister( -1.f, 1.f, 1.f, -1.f );
 
 /**
-* Multiplies two quaternions: The order matters.
+* Multiplies two quaternions; the order matters.
 *
-* @param Result	Returns Quat1 * Quat2
-* @param Quat1	First quaternion
-* @param Quat2	Second quaternion
+* Order matters when composing quaternions: C = VectorQuaternionMultiply2(A, B) will yield a quaternion C = A * B
+* that logically first applies B then A to any subsequent transformation (right first, then left).
+*
+* @param Quat1	Pointer to the first quaternion
+* @param Quat2	Pointer to the second quaternion
+* @return Quat1 * Quat2
 */
 FORCEINLINE VectorRegister VectorQuaternionMultiply2( const VectorRegister& Quat1, const VectorRegister& Quat2 )
 {
@@ -952,9 +955,12 @@ FORCEINLINE VectorRegister VectorQuaternionMultiply2( const VectorRegister& Quat
 }
 
 /**
-* Multiplies two quaternions: The order matters
+* Multiplies two quaternions; the order matters.
+*
+* When composing quaternions: VectorQuaternionMultiply(C, A, B) will yield a quaternion C = A * B
+* that logically first applies B then A to any subsequent transformation (right first, then left).
 *
-* @param Result	Pointer to where the result should be stored
+* @param Result	Pointer to where the result Quat1 * Quat2 should be stored
 * @param Quat1	Pointer to the first quaternion (must not be the destination)
 * @param Quat2	Pointer to the second quaternion (must not be the destination)
 */
diff --git a/Engine/Source/Runtime/Core/Public/Math/UnrealMathSSE.h b/Engine/Source/Runtime/Core/Public/Math/UnrealMathSSE.h
--- a/Engine/Source/Runtime/Core/Public/Math/UnrealMathSSE.h
+++ b/Engine/Source/Runtime/Core/Public/Math/UnrealMathSSE.h
@@ -944,11 +944,14 @@ FORCEINLINE void VectorStoreByte4( const VectorRegister& Vec, void* Ptr )
 #define VECTOR_ROUND_TOWARD_ZERO		_MM_ROUND_TOWARD_ZERO
 
 /**
-* Multiplies two quaternions: The order matters.
+* Multiplies two quaternions; the order matters.
 *
-* @param Result	Returns Quat1 * Quat2
-* @param Quat1	First quaternion
-* @param Quat2	Second quaternion
+* Order matters when composing quaternions: C = VectorQuaternionMultiply2(A, B) will yield a quaternion C = A * B
+* that logically first applies B then A to any subsequent transformation (right first, then left).
+*
+* @param Quat1	Pointer to the first quaternion
+* @param Quat2	Pointer to the second quaternion
+* @return Quat1 * Quat2
 */
 FORCEINLINE VectorRegister VectorQuaternionMultiply2( const VectorRegister& Quat1, const VectorRegister& Quat2 )
 {
@@ -961,9 +964,12 @@ FORCEINLINE VectorRegister VectorQuaternionMultiply2( const VectorRegister& Quat
 }
 
 /**
-* Multiplies two quaternions: The order matters
+* Multiplies two quaternions; the order matters.
+*
+* When composing quaternions: VectorQuaternionMultiply(C, A, B) will yield a quaternion C = A * B
+* that logically first applies B then A to any subsequent transformation (right first, then left).
 *
-* @param Result	Pointer to where the result should be stored
+* @param Result	Pointer to where the result Quat1 * Quat2 should be stored
 * @param Quat1	Pointer to the first quaternion (must not be the destination)
 * @param Quat2	Pointer to the second quaternion (must not be the destination)
 */
