Move references to UObjects from FCoreDelegates to FCoreUObjectDelegates
TTP #343071 CORE: FCoreDelegates should not have references to UObjects.

diff --git a/Engine/Plugins/2D/Paper2D/Source/Paper2DEditor/Private/Paper2DEditorModule.cpp b/Engine/Plugins/2D/Paper2D/Source/Paper2DEditor/Private/Paper2DEditorModule.cpp
--- a/Engine/Plugins/2D/Paper2D/Source/Paper2DEditor/Private/Paper2DEditorModule.cpp
+++ b/Engine/Plugins/2D/Paper2D/Source/Paper2DEditor/Private/Paper2DEditorModule.cpp
@@ -74,7 +74,7 @@ private:
 	TSharedPtr<IComponentAssetBroker> PaperFlipbookBroker;
 	TSharedPtr<IComponentAssetBroker> PaperTileMapBroker;
 
-	FCoreDelegates::FOnObjectPropertyChanged::FDelegate OnPropertyChangedHandle;
+	FCoreUObjectDelegates::FOnObjectPropertyChanged::FDelegate OnPropertyChangedHandle;
 
 public:
 	virtual void StartupModule() override
@@ -121,8 +121,8 @@ public:
 		}
 
 		// Register to be notified when properties are edited
-		OnPropertyChangedHandle = FCoreDelegates::FOnObjectPropertyChanged::FDelegate::CreateRaw(this, &FPaper2DEditor::OnPropertyChanged);
-		FCoreDelegates::OnObjectPropertyChanged.Add(OnPropertyChangedHandle);
+		OnPropertyChangedHandle = FCoreUObjectDelegates::FOnObjectPropertyChanged::FDelegate::CreateRaw(this, &FPaper2DEditor::OnPropertyChanged);
+		FCoreUObjectDelegates::OnObjectPropertyChanged.Add(OnPropertyChangedHandle);
 
 		// Register the thumbnail renderers
 		UThumbnailManager::Get().RegisterCustomRenderer(UPaperSprite::StaticClass(), UPaperSpriteThumbnailRenderer::StaticClass());
@@ -167,7 +167,7 @@ public:
 			UThumbnailManager::Get().UnregisterCustomRenderer(UPaperFlipbook::StaticClass());
 
 			// Unregister the property modification handler
-			FCoreDelegates::OnObjectPropertyChanged.Remove(OnPropertyChangedHandle);
+			FCoreUObjectDelegates::OnObjectPropertyChanged.Remove(OnPropertyChangedHandle);
 		}
 
 		// Unregister the details customization
diff --git a/Engine/Plugins/Runtime/LiveEditorListenServer/Source/LiveEditorListenServer/Private/LiveEditorListenServer.cpp b/Engine/Plugins/Runtime/LiveEditorListenServer/Source/LiveEditorListenServer/Private/LiveEditorListenServer.cpp
--- a/Engine/Plugins/Runtime/LiveEditorListenServer/Source/LiveEditorListenServer/Private/LiveEditorListenServer.cpp
+++ b/Engine/Plugins/Runtime/LiveEditorListenServer/Source/LiveEditorListenServer/Private/LiveEditorListenServer.cpp
@@ -229,8 +229,8 @@ void FLiveEditorListenServer::InstallHooks()
 	TickObject = new nLiveEditorListenServer::FTickObject(this);
 
 	MapLoadObserver = MakeShareable( new nLiveEditorListenServer::FMapLoadObserver(this) );
-	FCoreDelegates::PostLoadMap.AddSP(MapLoadObserver.Get(), &nLiveEditorListenServer::FMapLoadObserver::OnPostLoadMap);
-	FCoreDelegates::PreLoadMap.AddSP(MapLoadObserver.Get(), &nLiveEditorListenServer::FMapLoadObserver::OnPreLoadMap);
+	FCoreUObjectDelegates::PostLoadMap.AddSP(MapLoadObserver.Get(), &nLiveEditorListenServer::FMapLoadObserver::OnPostLoadMap);
+	FCoreUObjectDelegates::PreLoadMap.AddSP(MapLoadObserver.Get(), &nLiveEditorListenServer::FMapLoadObserver::OnPreLoadMap);
 }
 
 void FLiveEditorListenServer::RemoveHooks()
@@ -270,8 +270,8 @@ void FLiveEditorListenServer::RemoveHooks()
 	delete ObjectDeleteListener;
 	ObjectDeleteListener = NULL;
 
-	FCoreDelegates::PostLoadMap.RemoveAll(MapLoadObserver.Get());
-	FCoreDelegates::PreLoadMap.RemoveAll(MapLoadObserver.Get());
+	FCoreUObjectDelegates::PostLoadMap.RemoveAll(MapLoadObserver.Get());
+	FCoreUObjectDelegates::PreLoadMap.RemoveAll(MapLoadObserver.Get());
 	MapLoadObserver = NULL;
 }
 
diff --git a/Engine/Source/Editor/BlueprintGraph/Private/BlueprintActionDatabase.cpp b/Engine/Source/Editor/BlueprintGraph/Private/BlueprintActionDatabase.cpp
--- a/Engine/Source/Editor/BlueprintGraph/Private/BlueprintActionDatabase.cpp
+++ b/Engine/Source/Editor/BlueprintGraph/Private/BlueprintActionDatabase.cpp
@@ -864,7 +864,7 @@ FBlueprintActionDatabase& FBlueprintActionDatabase::Get()
 FBlueprintActionDatabase::FBlueprintActionDatabase()
 {
 	RefreshAll();
-	FCoreDelegates::OnAssetLoaded.AddStatic(&BlueprintActionDatabaseImpl::OnAssetLoaded);
+	FCoreUObjectDelegates::OnAssetLoaded.AddStatic(&BlueprintActionDatabaseImpl::OnAssetLoaded);
 
 	IAssetRegistry& AssetRegistry = FModuleManager::LoadModuleChecked<FAssetRegistryModule>(TEXT("AssetRegistry")).Get();
 	AssetRegistry.OnAssetAdded().AddStatic(&BlueprintActionDatabaseImpl::OnAssetAdded);
diff --git a/Engine/Source/Editor/ContentBrowser/Private/AssetViewWidgets.cpp b/Engine/Source/Editor/ContentBrowser/Private/AssetViewWidgets.cpp
--- a/Engine/Source/Editor/ContentBrowser/Private/AssetViewWidgets.cpp
+++ b/Engine/Source/Editor/ContentBrowser/Private/AssetViewWidgets.cpp
@@ -933,7 +933,7 @@ bool SAssetViewItem::OnVisualizeTooltip(const TSharedPtr<SWidget>& TooltipConten
 
 SAssetListItem::~SAssetListItem()
 {
-	FCoreDelegates::OnAssetLoaded.RemoveAll(this);
+	FCoreUObjectDelegates::OnAssetLoaded.RemoveAll(this);
 }
 
 BEGIN_SLATE_FUNCTION_BUILD_OPTIMIZATION
@@ -1105,7 +1105,7 @@ void SAssetListItem::Construct( const FArguments& InArgs )
 	SetForceMipLevelsToBeResident(true);
 
 	// listen for asset loads so we can force mips to stream in if required
-	FCoreDelegates::OnAssetLoaded.AddSP(this, &SAssetViewItem::HandleAssetLoaded);
+	FCoreUObjectDelegates::OnAssetLoaded.AddSP(this, &SAssetViewItem::HandleAssetLoaded);
 }
 END_SLATE_FUNCTION_BUILD_OPTIMIZATION
 
@@ -1138,7 +1138,7 @@ FOptionalSize SAssetListItem::GetSCCImageSize() const
 
 SAssetTileItem::~SAssetTileItem()
 {
-	FCoreDelegates::OnAssetLoaded.RemoveAll(this);
+	FCoreUObjectDelegates::OnAssetLoaded.RemoveAll(this);
 }
 
 BEGIN_SLATE_FUNCTION_BUILD_OPTIMIZATION
@@ -1309,7 +1309,7 @@ void SAssetTileItem::Construct( const FArguments& InArgs )
 	SetForceMipLevelsToBeResident(true);
 
 	// listen for asset loads so we can force mips to stream in if required
-	FCoreDelegates::OnAssetLoaded.AddSP(this, &SAssetViewItem::HandleAssetLoaded);
+	FCoreUObjectDelegates::OnAssetLoaded.AddSP(this, &SAssetViewItem::HandleAssetLoaded);
 }
 END_SLATE_FUNCTION_BUILD_OPTIMIZATION
 
diff --git a/Engine/Source/Editor/ContentBrowser/Private/SAssetView.cpp b/Engine/Source/Editor/ContentBrowser/Private/SAssetView.cpp
--- a/Engine/Source/Editor/ContentBrowser/Private/SAssetView.cpp
+++ b/Engine/Source/Editor/ContentBrowser/Private/SAssetView.cpp
@@ -41,8 +41,8 @@ SAssetView::~SAssetView()
 	}
 
 	// Unregister listener for asset loading and object property changes
-	FCoreDelegates::OnAssetLoaded.RemoveAll( this );
-	FCoreDelegates::OnObjectPropertyChanged.RemoveAll( this );
+	FCoreUObjectDelegates::OnAssetLoaded.RemoveAll(this);
+	FCoreUObjectDelegates::OnObjectPropertyChanged.RemoveAll(this);
 
 	// Remove the listener for when view settings are changed
 	UContentBrowserSettings::OnSettingChanged().RemoveAll(this);
@@ -85,8 +85,8 @@ void SAssetView::Construct( const FArguments& InArgs )
 	CollectionManagerModule.Get().OnCollectionRenamed().AddSP( this, &SAssetView::OnCollectionRenamed );
 
 	// Listen for when assets are loaded or changed to update item data
-	FCoreDelegates::OnAssetLoaded.AddSP(this, &SAssetView::OnAssetLoaded);
-	FCoreDelegates::OnObjectPropertyChanged.AddSP(this, &SAssetView::OnObjectPropertyChanged);
+	FCoreUObjectDelegates::OnAssetLoaded.AddSP(this, &SAssetView::OnAssetLoaded);
+	FCoreUObjectDelegates::OnObjectPropertyChanged.AddSP(this, &SAssetView::OnObjectPropertyChanged);
 
 	// Listen for when view settings are changed
 	UContentBrowserSettings::OnSettingChanged().AddSP(this, &SAssetView::HandleSettingChanged);
diff --git a/Engine/Source/Editor/EditorStyle/Private/SlateEditorStyle.h b/Engine/Source/Editor/EditorStyle/Private/SlateEditorStyle.h
--- a/Engine/Source/Editor/EditorStyle/Private/SlateEditorStyle.h
+++ b/Engine/Source/Editor/EditorStyle/Private/SlateEditorStyle.h
@@ -133,7 +133,7 @@ public:
 		TSharedRef< class FSlateEditorStyle::FStyle > NewStyle = MakeShareable( new FSlateEditorStyle::FStyle( InCustomization ) );
 		NewStyle->Initialize();
 
-		FCoreDelegates::OnObjectPropertyChanged.AddSP( NewStyle, &FSlateEditorStyle::FStyle::SettingsChanged );
+		FCoreUObjectDelegates::OnObjectPropertyChanged.AddSP(NewStyle, &FSlateEditorStyle::FStyle::SettingsChanged);
 
 		return NewStyle;
 }
diff --git a/Engine/Source/Editor/IntroTutorials/Private/InteractiveTutorials.cpp b/Engine/Source/Editor/IntroTutorials/Private/InteractiveTutorials.cpp
--- a/Engine/Source/Editor/IntroTutorials/Private/InteractiveTutorials.cpp
+++ b/Engine/Source/Editor/IntroTutorials/Private/InteractiveTutorials.cpp
@@ -30,7 +30,7 @@ void FInteractiveTutorials::SetupEditorHooks()
 	USelection::SelectionChangedEvent.AddSP(this, &FInteractiveTutorials::OnDelegateTriggered<FOnSelectionChanged, FBaseTriggerDelegate::OnSelectionChanged>);
 	USelection::SelectObjectEvent.AddSP(this, &FInteractiveTutorials::OnDelegateTriggered<FOnSelectObject, FBaseTriggerDelegate::OnSelectObject>);
 
-	FCoreDelegates::OnObjectPropertyChanged.AddSP( this, &FInteractiveTutorials::OnDelegateTriggered<FObjectPropertyChangedTrigger, FBaseTriggerDelegate::ObjectPropertyChanged> );
+	FCoreUObjectDelegates::OnObjectPropertyChanged.AddSP( this, &FInteractiveTutorials::OnDelegateTriggered<FObjectPropertyChangedTrigger, FBaseTriggerDelegate::ObjectPropertyChanged> );
 
 	FAssetEditorManager::Get().OnAssetEditorRequestedOpen().AddSP(this, &FInteractiveTutorials::OnDelegateTriggered<FAssetEditorRequestedOpenTrigger, FBaseTriggerDelegate::AssetEditorRequestedOpen>);
 
diff --git a/Engine/Source/Editor/Kismet/Private/FindInBlueprintManager.cpp b/Engine/Source/Editor/Kismet/Private/FindInBlueprintManager.cpp
--- a/Engine/Source/Editor/Kismet/Private/FindInBlueprintManager.cpp
+++ b/Engine/Source/Editor/Kismet/Private/FindInBlueprintManager.cpp
@@ -589,9 +589,9 @@ FFindInBlueprintSearchManager::~FFindInBlueprintSearchManager()
 	AssetRegistryModule->Get().OnAssetAdded().RemoveAll(this);
 	AssetRegistryModule->Get().OnAssetRemoved().RemoveAll(this);
 	AssetRegistryModule->Get().OnAssetRenamed().RemoveAll(this);
-	FCoreDelegates::PreGarbageCollect.RemoveAll(this);
-	FCoreDelegates::PostGarbageCollect.RemoveAll(this);
-	FCoreDelegates::OnAssetLoaded.RemoveAll(this);
+	FCoreUObjectDelegates::PreGarbageCollect.RemoveAll(this);
+	FCoreUObjectDelegates::PostGarbageCollect.RemoveAll(this);
+	FCoreUObjectDelegates::OnAssetLoaded.RemoveAll(this);
 }
 
 void FFindInBlueprintSearchManager::Initialize()
@@ -600,9 +600,9 @@ void FFindInBlueprintSearchManager::Initialize()
 	AssetRegistryModule->Get().OnAssetAdded().AddRaw(this, &FFindInBlueprintSearchManager::OnAssetAdded);
 	AssetRegistryModule->Get().OnAssetRemoved().AddRaw(this, &FFindInBlueprintSearchManager::OnAssetRemoved);
 	AssetRegistryModule->Get().OnAssetRenamed().AddRaw(this, &FFindInBlueprintSearchManager::OnAssetRenamed);
-	FCoreDelegates::PreGarbageCollect.AddRaw(this, &FFindInBlueprintSearchManager::PauseFindInBlueprintSearch);
-	FCoreDelegates::PostGarbageCollect.AddRaw(this, &FFindInBlueprintSearchManager::UnpauseFindInBlueprintSearch);
-	FCoreDelegates::OnAssetLoaded.AddRaw(this, &FFindInBlueprintSearchManager::OnAssetLoaded);
+	FCoreUObjectDelegates::PreGarbageCollect.AddRaw(this, &FFindInBlueprintSearchManager::PauseFindInBlueprintSearch);
+	FCoreUObjectDelegates::PostGarbageCollect.AddRaw(this, &FFindInBlueprintSearchManager::UnpauseFindInBlueprintSearch);
+	FCoreUObjectDelegates::OnAssetLoaded.AddRaw(this, &FFindInBlueprintSearchManager::OnAssetLoaded);
 
 	// Do an immediate load of the cache to catch any Blueprints that were discovered by the asset registry before we initialized.
 	BuildCache();
diff --git a/Engine/Source/Editor/Persona/Private/Persona.cpp b/Engine/Source/Editor/Persona/Private/Persona.cpp
--- a/Engine/Source/Editor/Persona/Private/Persona.cpp
+++ b/Engine/Source/Editor/Persona/Private/Persona.cpp
@@ -182,8 +182,8 @@ FPersona::FPersona()
 	, PreviewScene(FPreviewScene::ConstructionValues().AllowAudioPlayback(true).ShouldSimulatePhysics(true))
 {
 	// Register to be notified when properties are edited
-	OnPropertyChangedHandle = FCoreDelegates::FOnObjectPropertyChanged::FDelegate::CreateRaw(this, &FPersona::OnPropertyChanged);
-	FCoreDelegates::OnObjectPropertyChanged.Add(OnPropertyChangedHandle);
+	OnPropertyChangedHandle = FCoreUObjectDelegates::FOnObjectPropertyChanged::FDelegate::CreateRaw(this, &FPersona::OnPropertyChanged);
+	FCoreUObjectDelegates::OnObjectPropertyChanged.Add(OnPropertyChangedHandle);
 }
 
 FPersona::~FPersona()
@@ -202,7 +202,7 @@ FPersona::~FPersona()
 		Viewport.Pin()->CleanupPersonaReferences();
 	}
 
-	FCoreDelegates::OnObjectPropertyChanged.Remove(OnPropertyChangedHandle);
+	FCoreUObjectDelegates::OnObjectPropertyChanged.Remove(OnPropertyChangedHandle);
 
 	if(PreviewComponent)
 	{
diff --git a/Engine/Source/Editor/Persona/Private/Persona.h b/Engine/Source/Editor/Persona/Private/Persona.h
--- a/Engine/Source/Editor/Persona/Private/Persona.h
+++ b/Engine/Source/Editor/Persona/Private/Persona.h
@@ -280,7 +280,7 @@ public:
 	TWeakPtr<class SAnimationEditorViewportTabBody> Viewport;
 
 	// Property changed delegate
-	FCoreDelegates::FOnObjectPropertyChanged::FDelegate OnPropertyChangedHandle;
+	FCoreUObjectDelegates::FOnObjectPropertyChanged::FDelegate OnPropertyChangedHandle;
 	void OnPropertyChanged(UObject* ObjectBeingModified, FPropertyChangedEvent& PropertyChangedEvent);
 
 	/** Shared data between modes - for now only used for viewport **/
diff --git a/Engine/Source/Editor/Persona/Private/SAnimNotifyPanel.cpp b/Engine/Source/Editor/Persona/Private/SAnimNotifyPanel.cpp
--- a/Engine/Source/Editor/Persona/Private/SAnimNotifyPanel.cpp
+++ b/Engine/Source/Editor/Persona/Private/SAnimNotifyPanel.cpp
@@ -2970,8 +2970,8 @@ void SAnimNotifyPanel::Construct(const FArguments& InArgs)
 		]
 	];
 
-	OnPropertyChangedHandle = FCoreDelegates::FOnObjectPropertyChanged::FDelegate::CreateSP(this, &SAnimNotifyPanel::OnPropertyChanged);
-	FCoreDelegates::OnObjectPropertyChanged.Add(OnPropertyChangedHandle);
+	OnPropertyChangedHandle = FCoreUObjectDelegates::FOnObjectPropertyChanged::FDelegate::CreateSP(this, &SAnimNotifyPanel::OnPropertyChanged);
+	FCoreUObjectDelegates::OnObjectPropertyChanged.Add(OnPropertyChangedHandle);
 
 	// Base notify classes used to search asset data for children.
 	NotifyClassNames.Add(TEXT("Class'/Script/Engine.AnimNotify'"));
@@ -2992,7 +2992,7 @@ SAnimNotifyPanel::~SAnimNotifyPanel()
 		PersonaPtr.Pin()->UnregisterOnGenericDelete(this);
 	}
 
-	FCoreDelegates::OnObjectPropertyChanged.Remove(OnPropertyChangedHandle);
+	FCoreUObjectDelegates::OnObjectPropertyChanged.Remove(OnPropertyChangedHandle);
 }
 
 FReply SAnimNotifyPanel::InsertTrack(int32 TrackIndexToInsert)
diff --git a/Engine/Source/Editor/Persona/Private/SAnimNotifyPanel.h b/Engine/Source/Editor/Persona/Private/SAnimNotifyPanel.h
--- a/Engine/Source/Editor/Persona/Private/SAnimNotifyPanel.h
+++ b/Engine/Source/Editor/Persona/Private/SAnimNotifyPanel.h
@@ -191,7 +191,7 @@ public:
 	void OnPasteNotifies(SAnimNotifyTrack* RequestTrack, float ClickTime, ENotifyPasteMode::Type PasteMode, ENotifyPasteMultipleMode::Type MultiplePasteType);
 
 	/** Handler for properties changing on objects */
-	FCoreDelegates::FOnObjectPropertyChanged::FDelegate OnPropertyChangedHandle;
+	FCoreUObjectDelegates::FOnObjectPropertyChanged::FDelegate OnPropertyChangedHandle;
 	void OnPropertyChanged(UObject* ChangedObject, FPropertyChangedEvent& PropertyEvent);
 
 	/** SWidget Interface */
diff --git a/Engine/Source/Editor/Persona/Private/SAnimationBlendSpace1D.cpp b/Engine/Source/Editor/Persona/Private/SAnimationBlendSpace1D.cpp
--- a/Engine/Source/Editor/Persona/Private/SAnimationBlendSpace1D.cpp
+++ b/Engine/Source/Editor/Persona/Private/SAnimationBlendSpace1D.cpp
@@ -510,7 +510,7 @@ void SBlendSpace1DWidget::ResampleData()
 
 SBlendSpaceEditor1D::~SBlendSpaceEditor1D()
 {
-	FCoreDelegates::OnObjectPropertyChanged.Remove(OnPropertyChangedHandle);
+	FCoreUObjectDelegates::OnObjectPropertyChanged.Remove(OnPropertyChangedHandle);
 }
 
 void SBlendSpaceEditor1D::Construct(const FArguments& InArgs)
@@ -521,8 +521,8 @@ void SBlendSpaceEditor1D::Construct(const FArguments& InArgs)
 		);
 
 	// Register to be notified when properties are edited
-	OnPropertyChangedHandle = FCoreDelegates::FOnObjectPropertyChanged::FDelegate::CreateRaw(this, &SBlendSpaceEditor1D::OnPropertyChanged);
-	FCoreDelegates::OnObjectPropertyChanged.Add(OnPropertyChangedHandle);
+	OnPropertyChangedHandle = FCoreUObjectDelegates::FOnObjectPropertyChanged::FDelegate::CreateRaw(this, &SBlendSpaceEditor1D::OnPropertyChanged);
+	FCoreUObjectDelegates::OnObjectPropertyChanged.Add(OnPropertyChangedHandle);
 
 	bCachedDisplayVerticalValue = GetBlendSpace()->bDisplayEditorVertically;
 
diff --git a/Engine/Source/Editor/Persona/Private/SAnimationBlendSpace1D.h b/Engine/Source/Editor/Persona/Private/SAnimationBlendSpace1D.h
--- a/Engine/Source/Editor/Persona/Private/SAnimationBlendSpace1D.h
+++ b/Engine/Source/Editor/Persona/Private/SAnimationBlendSpace1D.h
@@ -173,7 +173,7 @@ protected:
 private:
 
 	// Property changed delegate
-	FCoreDelegates::FOnObjectPropertyChanged::FDelegate OnPropertyChangedHandle;
+	FCoreUObjectDelegates::FOnObjectPropertyChanged::FDelegate OnPropertyChangedHandle;
 	void OnPropertyChanged(UObject* ObjectBeingModified, FPropertyChangedEvent& PropertyChangedEvent);
 
 	// Updates the UI to reflect the current blend space parameter values */
diff --git a/Engine/Source/Editor/PropertyEditor/Private/UserInterface/PropertyTable/SPropertyTableCell.cpp b/Engine/Source/Editor/PropertyEditor/Private/UserInterface/PropertyTable/SPropertyTableCell.cpp
--- a/Engine/Source/Editor/PropertyEditor/Private/UserInterface/PropertyTable/SPropertyTableCell.cpp
+++ b/Engine/Source/Editor/PropertyEditor/Private/UserInterface/PropertyTable/SPropertyTableCell.cpp
@@ -20,7 +20,7 @@ void SPropertyTableCell::Construct( const FArguments& InArgs, const TSharedRef<
 	Cell->OnEnteredEditMode().AddSP( this, &SPropertyTableCell::EnteredEditMode );
 	Cell->OnExitedEditMode().AddSP( this, &SPropertyTableCell::ExitedEditMode );
 
-	FCoreDelegates::OnObjectPropertyChanged.AddSP( this, &SPropertyTableCell::OnCellValueChanged );
+	FCoreUObjectDelegates::OnObjectPropertyChanged.AddSP(this, &SPropertyTableCell::OnCellValueChanged);
 
 	SetForegroundColor( FEditorStyle::GetSlateColor( "InvertedForeground" ) );
 }
diff --git a/Engine/Source/Editor/Sequencer/Private/SequencerObjectChangeListener.cpp b/Engine/Source/Editor/Sequencer/Private/SequencerObjectChangeListener.cpp
--- a/Engine/Source/Editor/Sequencer/Private/SequencerObjectChangeListener.cpp
+++ b/Engine/Source/Editor/Sequencer/Private/SequencerObjectChangeListener.cpp
@@ -11,16 +11,16 @@ FSequencerObjectChangeListener::FSequencerObjectChangeListener( TSharedRef<ISequ
 	: Sequencer( InSequencer )
 	, bListenForActorsOnly( bInListenForActorsOnly )
 {
-	FCoreDelegates::OnPreObjectPropertyChanged.AddRaw(this, &FSequencerObjectChangeListener::OnObjectPreEditChange);
-	FCoreDelegates::OnObjectPropertyChanged.AddRaw(this, &FSequencerObjectChangeListener::OnObjectPostEditChange);
+	FCoreUObjectDelegates::OnPreObjectPropertyChanged.AddRaw(this, &FSequencerObjectChangeListener::OnObjectPreEditChange);
+	FCoreUObjectDelegates::OnObjectPropertyChanged.AddRaw(this, &FSequencerObjectChangeListener::OnObjectPostEditChange);
 	//GEditor->OnPreActorMoved.AddRaw(this, &FSequencerObjectChangeListener::OnActorPreEditMove);
 	GEditor->OnActorMoved().AddRaw( this, &FSequencerObjectChangeListener::OnActorPostEditMove );
 }
 
 FSequencerObjectChangeListener::~FSequencerObjectChangeListener()
 {
-	FCoreDelegates::OnPreObjectPropertyChanged.RemoveAll( this );
-	FCoreDelegates::OnObjectPropertyChanged.RemoveAll( this );
+	FCoreUObjectDelegates::OnPreObjectPropertyChanged.RemoveAll(this);
+	FCoreUObjectDelegates::OnObjectPropertyChanged.RemoveAll(this);
 	GEditor->OnActorMoved().RemoveAll( this );
 }
 
diff --git a/Engine/Source/Editor/StaticMeshEditor/Private/SStaticMeshEditorViewport.cpp b/Engine/Source/Editor/StaticMeshEditor/Private/SStaticMeshEditorViewport.cpp
--- a/Engine/Source/Editor/StaticMeshEditor/Private/SStaticMeshEditorViewport.cpp
+++ b/Engine/Source/Editor/StaticMeshEditor/Private/SStaticMeshEditorViewport.cpp
@@ -55,13 +55,13 @@ void SStaticMeshEditorViewport::Construct(const FArguments& InArgs)
 			SAssignNew(OverlayTextVerticalBox, SVerticalBox)
 		];
 
-	FCoreDelegates::OnObjectPropertyChanged.AddRaw(this, &SStaticMeshEditorViewport::OnObjectPropertyChanged);
+	FCoreUObjectDelegates::OnObjectPropertyChanged.AddRaw(this, &SStaticMeshEditorViewport::OnObjectPropertyChanged);
 
 }
 
 SStaticMeshEditorViewport::~SStaticMeshEditorViewport()
 {
-	FCoreDelegates::OnObjectPropertyChanged.RemoveAll(this);
+	FCoreUObjectDelegates::OnObjectPropertyChanged.RemoveAll(this);
 	if (EditorViewportClient.IsValid())
 	{
 		EditorViewportClient->Viewport = NULL;
diff --git a/Engine/Source/Editor/UnrealEd/Private/AssetThumbnail.cpp b/Engine/Source/Editor/UnrealEd/Private/AssetThumbnail.cpp
--- a/Engine/Source/Editor/UnrealEd/Private/AssetThumbnail.cpp
+++ b/Engine/Source/Editor/UnrealEd/Private/AssetThumbnail.cpp
@@ -706,8 +706,8 @@ FAssetThumbnailPool::FAssetThumbnailPool( uint32 InNumInPool, const TAttribute<b
 	, MaxRealTimeThumbnailsPerFrame( InMaxRealTimeThumbnailsPerFrame )
 	, MaxFrameTimeAllowance( InMaxFrameTimeAllowance )
 {
-	FCoreDelegates::OnObjectPropertyChanged.AddRaw(this, &FAssetThumbnailPool::OnObjectPropertyChanged);
-	FCoreDelegates::OnAssetLoaded.AddRaw(this, &FAssetThumbnailPool::OnAssetLoaded);
+	FCoreUObjectDelegates::OnObjectPropertyChanged.AddRaw(this, &FAssetThumbnailPool::OnObjectPropertyChanged);
+	FCoreUObjectDelegates::OnAssetLoaded.AddRaw(this, &FAssetThumbnailPool::OnAssetLoaded);
 	if ( GEditor )
 	{
 		GEditor->OnActorMoved().AddRaw( this, &FAssetThumbnailPool::OnActorPostEditMove );
@@ -716,8 +716,8 @@ FAssetThumbnailPool::FAssetThumbnailPool( uint32 InNumInPool, const TAttribute<b
 
 FAssetThumbnailPool::~FAssetThumbnailPool()
 {
-	FCoreDelegates::OnObjectPropertyChanged.RemoveAll(this);
-	FCoreDelegates::OnAssetLoaded.RemoveAll(this);
+	FCoreUObjectDelegates::OnObjectPropertyChanged.RemoveAll(this);
+	FCoreUObjectDelegates::OnAssetLoaded.RemoveAll(this);
 	if ( GEditor )
 	{
 		GEditor->OnActorMoved().RemoveAll(this);
diff --git a/Engine/Source/Editor/UnrealEd/Private/Commandlets/ChunkManifestGenerator.cpp b/Engine/Source/Editor/UnrealEd/Private/Commandlets/ChunkManifestGenerator.cpp
--- a/Engine/Source/Editor/UnrealEd/Private/Commandlets/ChunkManifestGenerator.cpp
+++ b/Engine/Source/Editor/UnrealEd/Private/Commandlets/ChunkManifestGenerator.cpp
@@ -194,7 +194,7 @@ void FChunkManifestGenerator::Initialize(bool InGenerateChunks)
 	// Hook up game delegate
 	if (bGenerateChunks)
 	{
-		FCoreDelegates::OnAssetLoaded.AddRaw(this, &FChunkManifestGenerator::OnAssetLoaded);
+		FCoreUObjectDelegates::OnAssetLoaded.AddRaw(this, &FChunkManifestGenerator::OnAssetLoaded);
 	}
 }
 
diff --git a/Engine/Source/Editor/UnrealEd/Private/Commandlets/CookCommandlet.cpp b/Engine/Source/Editor/UnrealEd/Private/Commandlets/CookCommandlet.cpp
--- a/Engine/Source/Editor/UnrealEd/Private/Commandlets/CookCommandlet.cpp
+++ b/Engine/Source/Editor/UnrealEd/Private/Commandlets/CookCommandlet.cpp
@@ -1212,7 +1212,7 @@ bool UCookCommandlet::Cook(const TArray<ITargetPlatform*>& Platforms, TArray<FSt
 	FParse::Value(*Params, TEXT("SubsetTarget="), SubsetTarget);
 	bool bDoSubset = SubsetMod > 0 && SubsetTarget < SubsetMod;
 
-	FCoreDelegates::PackageCreatedForLoad.AddUObject(this, &UCookCommandlet::MaybeMarkPackageAsAlreadyLoaded);
+	FCoreUObjectDelegates::PackageCreatedForLoad.AddUObject(this, &UCookCommandlet::MaybeMarkPackageAsAlreadyLoaded);
 	
 	SaveGlobalShaderMapFiles(Platforms);
 
diff --git a/Engine/Source/Editor/UnrealEd/Private/Commandlets/DerivedDataCacheCommandlet.cpp b/Engine/Source/Editor/UnrealEd/Private/Commandlets/DerivedDataCacheCommandlet.cpp
--- a/Engine/Source/Editor/UnrealEd/Private/Commandlets/DerivedDataCacheCommandlet.cpp
+++ b/Engine/Source/Editor/UnrealEd/Private/Commandlets/DerivedDataCacheCommandlet.cpp
@@ -49,7 +49,7 @@ int32 UDerivedDataCacheCommandlet::Main( const FString& Params )
 
 	if (!bStartupOnly && bFillCache)
 	{
-		FCoreDelegates::PackageCreatedForLoad.AddUObject(this, &UDerivedDataCacheCommandlet::MaybeMarkPackageAsAlreadyLoaded);
+		FCoreUObjectDelegates::PackageCreatedForLoad.AddUObject(this, &UDerivedDataCacheCommandlet::MaybeMarkPackageAsAlreadyLoaded);
 
 		TArray<FString> FilesInPath;
 
diff --git a/Engine/Source/Editor/UnrealEd/Private/Commandlets/FixupRedirectsCommandlet.cpp b/Engine/Source/Editor/UnrealEd/Private/Commandlets/FixupRedirectsCommandlet.cpp
--- a/Engine/Source/Editor/UnrealEd/Private/Commandlets/FixupRedirectsCommandlet.cpp
+++ b/Engine/Source/Editor/UnrealEd/Private/Commandlets/FixupRedirectsCommandlet.cpp
@@ -254,7 +254,7 @@ int32 UFixupRedirectsCommandlet::Main( const FString& Params )
 	}
 
 	// unregister the callback so we stop getting redirections added
-	FCoreDelegates::RedirectorFollowed.RemoveAll(&GRedirectCollector);
+	FCoreUObjectDelegates::RedirectorFollowed.RemoveAll(&GRedirectCollector);
 	
 	/////////////////////////////////////////////////////////////////////
 	// Explain to user what is happening
diff --git a/Engine/Source/Editor/UnrealEd/Private/Editor.cpp b/Engine/Source/Editor/UnrealEd/Private/Editor.cpp
--- a/Engine/Source/Editor/UnrealEd/Private/Editor.cpp
+++ b/Engine/Source/Editor/UnrealEd/Private/Editor.cpp
@@ -495,9 +495,9 @@ void UEditorEngine::Init(IEngineLoop* InEngineLoop)
 	FSlateApplication::Get().SetAppIcon(FEditorStyle::GetBrush(TEXT("Editor.AppIcon")));
 
 	FCoreDelegates::ModalErrorMessage.BindUObject(this, &UEditorEngine::OnModalMessageDialog);
-	FCoreDelegates::ShouldLoadOnTop.BindUObject(this, &UEditorEngine::OnShouldLoadOnTop);
+	FCoreUObjectDelegates::ShouldLoadOnTop.BindUObject(this, &UEditorEngine::OnShouldLoadOnTop);
 	FCoreDelegates::PreWorldOriginOffset.AddUObject(this, &UEditorEngine::PreWorldOriginOffset);
-	FCoreDelegates::OnAssetLoaded.AddUObject(this, &UEditorEngine::OnAssetLoaded);
+	FCoreUObjectDelegates::OnAssetLoaded.AddUObject(this, &UEditorEngine::OnAssetLoaded);
 	FWorldDelegates::LevelAddedToWorld.AddUObject(this, &UEditorEngine::OnLevelAddedToWorld);
 	FWorldDelegates::LevelRemovedFromWorld.AddUObject(this, &UEditorEngine::OnLevelRemovedFromWorld);
 	FLevelStreamingGCHelper::OnGCStreamedOutLevels.AddUObject(this, &UEditorEngine::OnGCStreamedOutLevels);
@@ -745,9 +745,9 @@ void UEditorEngine::FinishDestroy()
 		// Unregister events
 		FEditorDelegates::MapChange.RemoveAll(this);
 		FCoreDelegates::ModalErrorMessage.Unbind();
-		FCoreDelegates::ShouldLoadOnTop.Unbind();
+		FCoreUObjectDelegates::ShouldLoadOnTop.Unbind();
 		FCoreDelegates::PreWorldOriginOffset.RemoveAll(this);
-		FCoreDelegates::OnAssetLoaded.RemoveAll(this);
+		FCoreUObjectDelegates::OnAssetLoaded.RemoveAll(this);
 		FWorldDelegates::LevelAddedToWorld.RemoveAll(this);
 		FWorldDelegates::LevelRemovedFromWorld.RemoveAll(this);
 		FLevelStreamingGCHelper::OnGCStreamedOutLevels.RemoveAll(this);
diff --git a/Engine/Source/Editor/UnrealEd/Private/ThumbnailHelpers.cpp b/Engine/Source/Editor/UnrealEd/Private/ThumbnailHelpers.cpp
--- a/Engine/Source/Editor/UnrealEd/Private/ThumbnailHelpers.cpp
+++ b/Engine/Source/Editor/UnrealEd/Private/ThumbnailHelpers.cpp
@@ -865,12 +865,12 @@ FBlueprintThumbnailScene::FBlueprintThumbnailScene()
 {
 	CurrentBlueprint = NULL;
 
-	FCoreDelegates::PreGarbageCollect.AddRaw(this, &FBlueprintThumbnailScene::OnPreGarbageCollect);
+	FCoreUObjectDelegates::PreGarbageCollect.AddRaw(this, &FBlueprintThumbnailScene::OnPreGarbageCollect);
 }
 
 FBlueprintThumbnailScene::~FBlueprintThumbnailScene()
 {
-	FCoreDelegates::PreGarbageCollect.RemoveAll(this);
+	FCoreUObjectDelegates::PreGarbageCollect.RemoveAll(this);
 }
 
 void FBlueprintThumbnailScene::SetBlueprint(UBlueprint* Blueprint)
diff --git a/Engine/Source/Editor/UnrealEd/Private/UnrealEdEngine.cpp b/Engine/Source/Editor/UnrealEd/Private/UnrealEdEngine.cpp
--- a/Engine/Source/Editor/UnrealEd/Private/UnrealEdEngine.cpp
+++ b/Engine/Source/Editor/UnrealEd/Private/UnrealEdEngine.cpp
@@ -52,7 +52,7 @@ void UUnrealEdEngine::Init(IEngineLoop* InEngineLoop)
 	
 	// Register to the PostGarbageCollect delegate, as we want to use this to trigger the RefreshAllBroweser delegate from 
 	// here rather then from Core
-	FCoreDelegates::PostGarbageCollect.AddUObject(this, &UUnrealEdEngine::OnPostGarbageCollect);
+	FCoreUObjectDelegates::PostGarbageCollect.AddUObject(this, &UUnrealEdEngine::OnPostGarbageCollect);
 
 	// register to color picker changed event and trigger RedrawAllViewports when that happens */
 	FCoreDelegates::ColorPickerChanged.AddUObject(this, &UUnrealEdEngine::OnColorPickerChanged);
@@ -117,8 +117,8 @@ void UUnrealEdEngine::Init(IEngineLoop* InEngineLoop)
 		CookServer = ConstructObject<UCookOnTheFlyServer>( UCookOnTheFlyServer::StaticClass() );
 		CookServer->Initialize( ECookMode::CookByTheBookFromTheEditor, ECookInitializationFlags::AutoTick | ECookInitializationFlags::AsyncSave );
 
-		FCoreDelegates::OnObjectPropertyChanged.AddUObject(CookServer, &UCookOnTheFlyServer::OnObjectPropertyChanged);
-		FCoreDelegates::OnObjectModified.AddUObject(CookServer, &UCookOnTheFlyServer::OnObjectModified);
+		FCoreUObjectDelegates::OnObjectPropertyChanged.AddUObject(CookServer, &UCookOnTheFlyServer::OnObjectPropertyChanged);
+		FCoreUObjectDelegates::OnObjectModified.AddUObject(CookServer, &UCookOnTheFlyServer::OnObjectModified);
 	}
 	else if ( ExperimentalSettings->bCookOnTheSide )
 	{
@@ -126,8 +126,8 @@ void UUnrealEdEngine::Init(IEngineLoop* InEngineLoop)
 		CookServer->Initialize( ECookMode::CookOnTheFly, ECookInitializationFlags::AutoTick | ECookInitializationFlags::AsyncSave );
 		CookServer->StartNetworkFileServer( false );
 
-		FCoreDelegates::OnObjectPropertyChanged.AddUObject(CookServer, &UCookOnTheFlyServer::OnObjectPropertyChanged);
-		FCoreDelegates::OnObjectModified.AddUObject(CookServer, &UCookOnTheFlyServer::OnObjectModified);
+		FCoreUObjectDelegates::OnObjectPropertyChanged.AddUObject(CookServer, &UCookOnTheFlyServer::OnObjectPropertyChanged);
+		FCoreUObjectDelegates::OnObjectModified.AddUObject(CookServer, &UCookOnTheFlyServer::OnObjectModified);
 	}
 }
 
@@ -254,8 +254,8 @@ void UUnrealEdEngine::FinishDestroy()
 {
 	if( CookServer)
 	{
-		FCoreDelegates::OnObjectPropertyChanged.RemoveAll( CookServer );
-		FCoreDelegates::OnObjectModified.RemoveAll( CookServer );
+		FCoreUObjectDelegates::OnObjectPropertyChanged.RemoveAll(CookServer);
+		FCoreUObjectDelegates::OnObjectModified.RemoveAll(CookServer);
 	}
 
 	if(PackageAutoSaver.Get())
@@ -271,7 +271,7 @@ void UUnrealEdEngine::FinishDestroy()
 	}
 
 	UPackage::PackageDirtyStateChangedEvent.RemoveAll(this);
-	FCoreDelegates::PostGarbageCollect.RemoveAll(this);
+	FCoreUObjectDelegates::PostGarbageCollect.RemoveAll(this);
 	FCoreDelegates::ColorPickerChanged.RemoveAll(this);
 	Super::FinishDestroy();
 }
diff --git a/Engine/Source/Editor/UnrealEd/Private/UnrealEdMisc.cpp b/Engine/Source/Editor/UnrealEd/Private/UnrealEdMisc.cpp
--- a/Engine/Source/Editor/UnrealEd/Private/UnrealEdMisc.cpp
+++ b/Engine/Source/Editor/UnrealEd/Private/UnrealEdMisc.cpp
@@ -180,7 +180,7 @@ void FUnrealEdMisc::OnInit()
 	FEditorSupportDelegates::RedrawAllViewports.AddRaw(this, &FUnrealEdMisc::CB_RedrawAllViewports);
 	GEngine->OnLevelActorAdded().AddRaw( this, &FUnrealEdMisc::CB_LevelActorsAdded );
 
-	FCoreDelegates::OnObjectSaved.AddRaw(this, &FUnrealEdMisc::OnObjectSaved);
+	FCoreUObjectDelegates::OnObjectSaved.AddRaw(this, &FUnrealEdMisc::OnObjectSaved);
 	FEditorDelegates::PreSaveWorld.AddRaw(this, &FUnrealEdMisc::OnWorldSaved);
 
 #if USE_UNIT_TESTS
diff --git a/Engine/Source/Runtime/AIModule/Private/EnvironmentQuery/EnvQueryManager.cpp b/Engine/Source/Runtime/AIModule/Private/EnvironmentQuery/EnvQueryManager.cpp
--- a/Engine/Source/Runtime/AIModule/Private/EnvironmentQuery/EnvQueryManager.cpp
+++ b/Engine/Source/Runtime/AIModule/Private/EnvironmentQuery/EnvQueryManager.cpp
@@ -82,7 +82,7 @@ UEnvQueryManager::UEnvQueryManager(const FPostConstructInitializeProperties& PCI
 {
 	if (!HasAnyFlags(RF_ClassDefaultObject))
 	{
-		FCoreDelegates::PreLoadMap.AddUObject(this, &UEnvQueryManager::OnPreLoadMap);
+		FCoreUObjectDelegates::PreLoadMap.AddUObject(this, &UEnvQueryManager::OnPreLoadMap);
 	}
 
 	NextQueryID = 0;
@@ -95,7 +95,7 @@ class UWorld* UEnvQueryManager::GetWorld() const
 
 void UEnvQueryManager::FinishDestroy()
 {
-	FCoreDelegates::PreLoadMap.RemoveAll(this);
+	FCoreUObjectDelegates::PreLoadMap.RemoveAll(this);
 	Super::FinishDestroy();
 }
 
diff --git a/Engine/Source/Runtime/Core/Private/Misc/CallbackDevice.cpp b/Engine/Source/Runtime/Core/Private/Misc/CallbackDevice.cpp
--- a/Engine/Source/Runtime/Core/Private/Misc/CallbackDevice.cpp
+++ b/Engine/Source/Runtime/Core/Private/Misc/CallbackDevice.cpp
@@ -26,14 +26,9 @@ FCoreDelegates::FOnUserControllerConnectionChange FCoreDelegates::OnControllerCo
 FCoreDelegates::FOnSafeFrameChangedEvent FCoreDelegates::OnSafeFrameChangedEvent;
 FCoreDelegates::FOnHandleSystemEnsure FCoreDelegates::OnHandleSystemEnsure;
 FCoreDelegates::FOnHandleSystemError FCoreDelegates::OnHandleSystemError;
-FCoreDelegates::FOnPreObjectPropertyChanged FCoreDelegates::OnPreObjectPropertyChanged;
-FCoreDelegates::FOnObjectPropertyChanged FCoreDelegates::OnObjectPropertyChanged;
 FCoreDelegates::FOnActorLabelChanged FCoreDelegates::OnActorLabelChanged;
 
 #if WITH_EDITOR
-	FCoreDelegates::FOnObjectModified  FCoreDelegates::OnObjectModified;
-	FCoreDelegates::FOnAssetLoaded FCoreDelegates::OnAssetLoaded;
-	FCoreDelegates::FOnObjectSaved FCoreDelegates::OnObjectSaved;
 	FSimpleMulticastDelegate FCoreDelegates::PreModal;
 	FSimpleMulticastDelegate FCoreDelegates::PostModal;
 #endif	//WITH_EDITOR
@@ -42,20 +37,11 @@ FSimpleMulticastDelegate FCoreDelegates::OnShutdownAfterError;
 FSimpleMulticastDelegate FCoreDelegates::OnInit;
 FSimpleMulticastDelegate FCoreDelegates::OnExit;
 FSimpleMulticastDelegate FCoreDelegates::OnPreExit;
-FSimpleMulticastDelegate FCoreDelegates::PreGarbageCollect;
-FSimpleMulticastDelegate FCoreDelegates::PostGarbageCollect;
 FCoreDelegates::FGetSerialNumberBlocksForDebugVisualizersDelegate FCoreDelegates::GetSerialNumberBlocksDebugVisualizers;
 FCoreDelegates::FObjectArrayForDebugVisualizersDelegate FCoreDelegates::ObjectArrayForDebugVisualizers;
-FCoreDelegates::FOnRedirectorFollowed FCoreDelegates::RedirectorFollowed;
-FSimpleMulticastDelegate FCoreDelegates::PreLoadMap;
-FSimpleMulticastDelegate FCoreDelegates::PostLoadMap;
 FSimpleMulticastDelegate FCoreDelegates::ColorPickerChanged;
 FCoreDelegates::FOnModalMessageBox FCoreDelegates::ModalErrorMessage;
-FCoreDelegates::FOnLoadObjectsOnTop FCoreDelegates::ShouldLoadOnTop;
 FCoreDelegates::FOnInviteAccepted FCoreDelegates::OnInviteAccepted;
-FCoreDelegates::FStringAssetReferenceLoaded FCoreDelegates::StringAssetReferenceLoaded;
-FCoreDelegates::FStringAssetReferenceSaving FCoreDelegates::StringAssetReferenceSaving;
-FCoreDelegates::FPackageCreatedForLoad FCoreDelegates::PackageCreatedForLoad;
 FCoreDelegates::FWorldOriginOffset FCoreDelegates::PreWorldOriginOffset;
 FCoreDelegates::FWorldOriginOffset FCoreDelegates::PostWorldOriginOffset;
 FCoreDelegates::FStarvedGameLoop FCoreDelegates::StarvedGameLoop;
diff --git a/Engine/Source/Runtime/Core/Public/Misc/CallbackDevice.h b/Engine/Source/Runtime/Core/Public/Misc/CallbackDevice.h
--- a/Engine/Source/Runtime/Core/Public/Misc/CallbackDevice.h
+++ b/Engine/Source/Runtime/Core/Public/Misc/CallbackDevice.h
@@ -46,43 +46,23 @@ public:
 	//hot fix delegate
 	DECLARE_DELEGATE_TwoParams(FHotFixDelegate, void *, int32);
 
-	// Callback for object property modifications
-	DECLARE_MULTICAST_DELEGATE_TwoParams(FOnObjectPropertyChanged, UObject*, struct FPropertyChangedEvent&);
-
 	// Callback for object property modifications
 	DECLARE_MULTICAST_DELEGATE_OneParam(FOnActorLabelChanged, AActor*);
 
-	// Callback for PreEditChange
-	DECLARE_MULTICAST_DELEGATE_TwoParams(FOnPreObjectPropertyChanged, UObject*, const class FEditPropertyChain& );
-
-	#if WITH_EDITOR
-	// Callback for all object modifications
-	DECLARE_MULTICAST_DELEGATE_OneParam(FOnObjectModified, UObject*);
-
-	// Callback for when an asset is loaded (Editor)
-	DECLARE_MULTICAST_DELEGATE_OneParam(FOnAssetLoaded, UObject*);
-	
-	// Callback for when an asset is saved (Editor)
-	DECLARE_MULTICAST_DELEGATE_OneParam(FOnObjectSaved, UObject*);
-	#endif	//WITH_EDITOR
-
 	// delegate type for prompting the pak system to mount a new pak
 	DECLARE_DELEGATE_RetVal_TwoParams(bool, FOnMountPak, const FString&, uint32);
 
 	// Callback for PER_MODULE_BOILERPLATE macro's GSerialNumberBlocksForDebugVisualizers
 	DECLARE_DELEGATE_RetVal(int32***, FGetSerialNumberBlocksForDebugVisualizersDelegate);
 
-	// Callback for PER_MODULE_BOILERPLATE macro's GObjectArrayForDebugVisualizers
-	DECLARE_DELEGATE_RetVal(TArray<class UObjectBase*>*, FObjectArrayForDebugVisualizersDelegate);
-
-	// Delegate type for redirector followed events ( Params: const FString& PackageName, UObject* Redirector )
-	DECLARE_MULTICAST_DELEGATE_TwoParams(FOnRedirectorFollowed, const FString&, UObject*);
-
 	/** delegate type for opening a modal message box ( Params: EAppMsgType::Type MessageType, const FText& Text, const FText& Title ) */
 	DECLARE_DELEGATE_RetVal_ThreeParams(EAppReturnType::Type, FOnModalMessageBox, EAppMsgType::Type, const FText&, const FText&);
 
-	/** delegate type for querying whether a loaded object should replace an already existing one */
-	DECLARE_DELEGATE_RetVal_OneParam(bool, FOnLoadObjectsOnTop, const FString&);
+	// Callback for PER_MODULE_BOILERPLATE macro's GObjectArrayForDebugVisualizers
+	DECLARE_DELEGATE_RetVal(TArray<class UObjectBase*>*, FObjectArrayForDebugVisualizersDelegate);
+
+	// Called in PER_MODULE_BOILERPLATE macro.
+	static FObjectArrayForDebugVisualizersDelegate ObjectArrayForDebugVisualizers;
 
 	// Callback for handling an ensure
 	DECLARE_MULTICAST_DELEGATE(FOnHandleSystemEnsure);
@@ -124,25 +104,10 @@ public:
 	// Callback when an error (crash) has occurred
 	static FOnHandleSystemError OnHandleSystemError;
 
-	// Called before a property is changed
-	static FOnPreObjectPropertyChanged OnPreObjectPropertyChanged;
-
-	// Called when a property is changed
-	static FOnObjectPropertyChanged OnObjectPropertyChanged;
-
 	// Called when an actor label is changed
 	static FOnActorLabelChanged OnActorLabelChanged;
 
 #if WITH_EDITOR
-	// Called when any object is modified at all 
-	static FOnObjectModified OnObjectModified;
-
-	// Called when an asset is loaded
-	static FOnAssetLoaded OnAssetLoaded;
-
-	// Called when an asset is saved
-	static FOnObjectSaved OnObjectSaved;
-
 	// Called before the editor displays a modal window, allowing other windows the opportunity to disable themselves to avoid reentrant calls
 	static FSimpleMulticastDelegate PreModal;
 
@@ -165,51 +130,18 @@ public:
 	// Called when before the application is exiting.
 	static FSimpleMulticastDelegate OnPreExit;
 
-	// Called before garbage collection
-	static FSimpleMulticastDelegate PreGarbageCollect;
-
-	// Called after garbage collection
-	static FSimpleMulticastDelegate PostGarbageCollect;
-
 	// Called in PER_MODULE_BOILERPLATE macro.
 	static FGetSerialNumberBlocksForDebugVisualizersDelegate GetSerialNumberBlocksDebugVisualizers;
 
-	// Called in PER_MODULE_BOILERPLATE macro.
-	static FObjectArrayForDebugVisualizersDelegate ObjectArrayForDebugVisualizers;
-
-	// Sent when a UObjectRedirector was followed to find the destination object
-	static FOnRedirectorFollowed	RedirectorFollowed;
-
-	// Sent at the very beginning of LoadMap
-	static FSimpleMulticastDelegate PreLoadMap;
-
-	// Sent at the _successful_ end of LoadMap
-	static FSimpleMulticastDelegate PostLoadMap;
-
 	/** Color picker color has changed, please refresh as needed*/
 	static FSimpleMulticastDelegate ColorPickerChanged;
 
 	/** requests to open a message box */
 	static FOnModalMessageBox ModalErrorMessage;
 
-	/** Queries whether an object should be loaded on top ( replace ) an already existing one */
-	static FOnLoadObjectsOnTop ShouldLoadOnTop;
-
 	/** Called when the user accepts an invitation to the current game */
 	static FOnInviteAccepted OnInviteAccepted;
 
-	/** called when loading a string asset reference */
-	DECLARE_DELEGATE_OneParam(FStringAssetReferenceLoaded, FString const& /*LoadedAssetLongPathname*/);
-	static FStringAssetReferenceLoaded StringAssetReferenceLoaded;
-
-	/** called when path to world root is changed */
-	DECLARE_MULTICAST_DELEGATE_OneParam(FPackageCreatedForLoad, class UPackage*);
-	static FPackageCreatedForLoad PackageCreatedForLoad;
-
-	/** called when saving a string asset reference, can replace the value with something else */
-	DECLARE_DELEGATE_RetVal_OneParam( FString, FStringAssetReferenceSaving, FString const& /*SavingAssetLongPathname*/);
-	static FStringAssetReferenceSaving StringAssetReferenceSaving;
-		
 	DECLARE_MULTICAST_DELEGATE_ThreeParams(FWorldOriginOffset, class UWorld*, FIntVector, FIntVector);
 	/** called before world origin shifting */
 	static FWorldOriginOffset PreWorldOriginOffset;
diff --git a/Engine/Source/Runtime/CoreUObject/Private/Misc/StringAssetReference.cpp b/Engine/Source/Runtime/CoreUObject/Private/Misc/StringAssetReference.cpp
--- a/Engine/Source/Runtime/CoreUObject/Private/Misc/StringAssetReference.cpp
+++ b/Engine/Source/Runtime/CoreUObject/Private/Misc/StringAssetReference.cpp
@@ -16,16 +16,16 @@ FStringAssetReference::FStringAssetReference(const UObject* InObject)
 bool FStringAssetReference::Serialize(FArchive& Ar)
 {
 #if WITH_EDITOR
-	if (Ar.IsSaving() && Ar.IsPersistent() && FCoreDelegates::StringAssetReferenceSaving.IsBound())
+	if (Ar.IsSaving() && Ar.IsPersistent() && FCoreUObjectDelegates::StringAssetReferenceSaving.IsBound())
 	{
-		AssetLongPathname = FCoreDelegates::StringAssetReferenceSaving.Execute(AssetLongPathname);
+		AssetLongPathname = FCoreUObjectDelegates::StringAssetReferenceSaving.Execute(AssetLongPathname);
 	}
 #endif // WITH_EDITOR
 	Ar << *this;
 #if WITH_EDITOR
-	if (Ar.IsLoading() && Ar.IsPersistent() && FCoreDelegates::StringAssetReferenceLoaded.IsBound())
+	if (Ar.IsLoading() && Ar.IsPersistent() && FCoreUObjectDelegates::StringAssetReferenceLoaded.IsBound())
 	{
-		FCoreDelegates::StringAssetReferenceLoaded.Execute(AssetLongPathname);
+		FCoreUObjectDelegates::StringAssetReferenceLoaded.Execute(AssetLongPathname);
 	}
 #endif // WITH_EDITOR
 
diff --git a/Engine/Source/Runtime/CoreUObject/Private/Misc/StringReferenceTemplates.h b/Engine/Source/Runtime/CoreUObject/Private/Misc/StringReferenceTemplates.h
--- a/Engine/Source/Runtime/CoreUObject/Private/Misc/StringReferenceTemplates.h
+++ b/Engine/Source/Runtime/CoreUObject/Private/Misc/StringReferenceTemplates.h
@@ -27,9 +27,9 @@ bool SerializeFromMismatchedTagTemplate(FString& Output, const FPropertyTag& Tag
 			Output = FString();
 		}
 #if WITH_EDITOR
-		if (Ar.IsLoading() && Ar.IsPersistent() && FCoreDelegates::StringAssetReferenceLoaded.IsBound())
+		if (Ar.IsLoading() && Ar.IsPersistent() && FCoreUObjectDelegates::StringAssetReferenceLoaded.IsBound())
 		{
-			FCoreDelegates::StringAssetReferenceLoaded.Execute(Output);
+			FCoreUObjectDelegates::StringAssetReferenceLoaded.Execute(Output);
 		}
 #endif // WITH_EDITOR
 		return true;
@@ -41,9 +41,9 @@ bool SerializeFromMismatchedTagTemplate(FString& Output, const FPropertyTag& Tag
 
 		Output = String;
 #if WITH_EDITOR
-		if (Ar.IsLoading() && Ar.IsPersistent() && FCoreDelegates::StringAssetReferenceLoaded.IsBound())
+		if (Ar.IsLoading() && Ar.IsPersistent() && FCoreUObjectDelegates::StringAssetReferenceLoaded.IsBound())
 		{
-			FCoreDelegates::StringAssetReferenceLoaded.Execute(Output);
+			FCoreUObjectDelegates::StringAssetReferenceLoaded.Execute(Output);
 		}
 #endif // WITH_EDITOR
 		return true;
diff --git a/Engine/Source/Runtime/CoreUObject/Private/UObject/GarbageCollection.cpp b/Engine/Source/Runtime/CoreUObject/Private/UObject/GarbageCollection.cpp
--- a/Engine/Source/Runtime/CoreUObject/Private/UObject/GarbageCollection.cpp
+++ b/Engine/Source/Runtime/CoreUObject/Private/UObject/GarbageCollection.cpp
@@ -1043,7 +1043,7 @@ void CollectGarbage( EObjectFlags KeepFlags, bool bPerformFullPurge )
 	{
 		FAddFlushAsyncLoadingCallback()
 		{
-			FCoreDelegates::PreGarbageCollect.AddStatic(FlushAsyncLoadingCallback);
+			FCoreUObjectDelegates::PreGarbageCollect.AddStatic(FlushAsyncLoadingCallback);
 		}
 		/** Wrapper function to handle default parameter when used as function pointer */
 		static void FlushAsyncLoadingCallback()
@@ -1059,7 +1059,7 @@ void CollectGarbage( EObjectFlags KeepFlags, bool bPerformFullPurge )
 
 	// Route callbacks so we can ensure that we are e.g. not in the middle of loading something by flushing
 	// the async loading, etc...
-	FCoreDelegates::PreGarbageCollect.Broadcast();
+	FCoreUObjectDelegates::PreGarbageCollect.Broadcast();
 	
 	// Set 'I'm garbage collecting' flag - might be checked inside various functions.
 	GIsGarbageCollecting = true; 
@@ -1174,7 +1174,7 @@ void CollectGarbage( EObjectFlags KeepFlags, bool bPerformFullPurge )
 	GIsGarbageCollecting = false;
 
 	// Route callbacks to verify GC assumptions
-	FCoreDelegates::PostGarbageCollect.Broadcast();
+	FCoreUObjectDelegates::PostGarbageCollect.Broadcast();
 }
 
 /**
diff --git a/Engine/Source/Runtime/CoreUObject/Private/UObject/LinkerLoad.cpp b/Engine/Source/Runtime/CoreUObject/Private/UObject/LinkerLoad.cpp
--- a/Engine/Source/Runtime/CoreUObject/Private/UObject/LinkerLoad.cpp
+++ b/Engine/Source/Runtime/CoreUObject/Private/UObject/LinkerLoad.cpp
@@ -351,7 +351,7 @@ ULinkerLoad* ULinkerLoad::CreateLinker( UPackage* Parent, const TCHAR* Filename,
 			return NULL;
 		}
 	}
-	FCoreDelegates::PackageCreatedForLoad.Broadcast(Parent);
+	FCoreUObjectDelegates::PackageCreatedForLoad.Broadcast(Parent);
 	return Linker;
 }
 
@@ -1602,7 +1602,7 @@ ULinkerLoad::ELinkerStatus ULinkerLoad::FindExistingExports()
 			// the objects in those cases, so don't try to find existing exports.
 			//
 			bool bContainsMap			= LinkerRoot ? LinkerRoot->ContainsMap() : false;
-			bool bRequestFindExisting	= FCoreDelegates::ShouldLoadOnTop.IsBound() ? !FCoreDelegates::ShouldLoadOnTop.Execute(Filename) : true;
+			bool bRequestFindExisting = FCoreUObjectDelegates::ShouldLoadOnTop.IsBound() ? !FCoreUObjectDelegates::ShouldLoadOnTop.Execute(Filename) : true;
 			if( (!IsRunningCommandlet() && bRequestFindExisting && !bContainsMap) )
 			{
 				for (int32 ExportIndex = 0; ExportIndex < ExportMap.Num(); ExportIndex++)
@@ -2016,7 +2016,7 @@ ULinkerLoad::EVerifyResult ULinkerLoad::VerifyImport(int32 ImportIndex)
 					Result = VERIFY_Redirected;
 
 					// send a callback saying we followed a redirector successfully
-					FCoreDelegates::RedirectorFollowed.Broadcast(Filename, Redir);
+					FCoreUObjectDelegates::RedirectorFollowed.Broadcast(Filename, Redir);
 
 					// now, fake our Import to be what the redirector pointed to
 					Import.XObject = Redir->DestinationObject;
@@ -2654,7 +2654,7 @@ UObject* ULinkerLoad::Create( UClass* ObjectClass, FName ObjectName, UObject* Ou
 			if (Redir->DestinationObject && Redir->DestinationObject->IsA(ObjectClass))
 			{
 				// send a callback saying we followed a redirector successfully
-				FCoreDelegates::RedirectorFollowed.Broadcast(Filename, Redir);
+				FCoreUObjectDelegates::RedirectorFollowed.Broadcast(Filename, Redir);
 				// and return the object we are being redirected to
 				return Redir->DestinationObject;
 			}
diff --git a/Engine/Source/Runtime/CoreUObject/Private/UObject/Obj.cpp b/Engine/Source/Runtime/CoreUObject/Private/UObject/Obj.cpp
--- a/Engine/Source/Runtime/CoreUObject/Private/UObject/Obj.cpp
+++ b/Engine/Source/Runtime/CoreUObject/Private/UObject/Obj.cpp
@@ -217,7 +217,7 @@ void UObject::PostEditChange(void)
 
 void UObject::PostEditChangeProperty(FPropertyChangedEvent& PropertyChangedEvent)
 {
-	FCoreDelegates::OnObjectPropertyChanged.Broadcast(this, PropertyChangedEvent);
+	FCoreUObjectDelegates::OnObjectPropertyChanged.Broadcast(this, PropertyChangedEvent);
 }
 
 
@@ -226,7 +226,7 @@ void UObject::PreEditChange( FEditPropertyChain& PropertyAboutToChange )
 	// forward the notification to the UProperty* version of PreEditChange
 	PreEditChange(PropertyAboutToChange.GetActiveNode()->GetValue());
 
-	FCoreDelegates::OnPreObjectPropertyChanged.Broadcast(this, PropertyAboutToChange);
+	FCoreUObjectDelegates::OnPreObjectPropertyChanged.Broadcast(this, PropertyAboutToChange);
 
 	if ( HasAnyFlags(RF_ClassDefaultObject|RF_ArchetypeObject) && PropertyAboutToChange.GetActiveMemberNode() == PropertyAboutToChange.GetHead() && !FApp::IsGame())
 	{
@@ -695,7 +695,7 @@ void UObject::ConditionalPostLoadSubobjects( FObjectInstancingGraph* OuterInstan
 void UObject::PreSave()
 {
 #if WITH_EDITOR
-	FCoreDelegates::OnObjectSaved.Broadcast(this);
+	FCoreUObjectDelegates::OnObjectSaved.Broadcast(this);
 #endif
 }
 
@@ -719,7 +719,7 @@ bool UObject::Modify( bool bAlwaysMarkDirty/*=true*/ )
 		}
 	}
 #if WITH_EDITOR
-	FCoreDelegates::OnObjectModified.Broadcast(this);
+	FCoreUObjectDelegates::OnObjectModified.Broadcast(this);
 #endif
 	return bSavedToTransactionBuffer;
 }
@@ -3422,17 +3422,17 @@ void InitUObject()
 	// this is a hack to give fixup redirects insight into the startup packages
 	if (CommandLine.Contains(TEXT("fixupredirects")) )
 	{
-		FCoreDelegates::RedirectorFollowed.AddRaw(&GRedirectCollector, &FRedirectCollector::OnRedirectorFollowed);
-		FCoreDelegates::StringAssetReferenceLoaded.BindRaw(&GRedirectCollector, &FRedirectCollector::OnStringAssetReferenceLoaded);
-		FCoreDelegates::StringAssetReferenceSaving.BindRaw(&GRedirectCollector, &FRedirectCollector::OnStringAssetReferenceSaved);
+		FCoreUObjectDelegates::RedirectorFollowed.AddRaw(&GRedirectCollector, &FRedirectCollector::OnRedirectorFollowed);
+		FCoreUObjectDelegates::StringAssetReferenceLoaded.BindRaw(&GRedirectCollector, &FRedirectCollector::OnStringAssetReferenceLoaded);
+		FCoreUObjectDelegates::StringAssetReferenceSaving.BindRaw(&GRedirectCollector, &FRedirectCollector::OnStringAssetReferenceSaved);
 	}
 
 	// this is a hack to the cooker insight into the startup packages
 	if (CommandLine.Contains(TEXT("cookcommandlet")) || 
 		  CommandLine.Contains(TEXT("run=cook")) )
 	{
-		FCoreDelegates::StringAssetReferenceLoaded.BindRaw(&GRedirectCollector, &FRedirectCollector::OnStringAssetReferenceLoaded);
-		FCoreDelegates::StringAssetReferenceSaving.BindRaw(&GRedirectCollector, &FRedirectCollector::OnStringAssetReferenceSaved);
+		FCoreUObjectDelegates::StringAssetReferenceLoaded.BindRaw(&GRedirectCollector, &FRedirectCollector::OnStringAssetReferenceLoaded);
+		FCoreUObjectDelegates::StringAssetReferenceSaving.BindRaw(&GRedirectCollector, &FRedirectCollector::OnStringAssetReferenceSaved);
 	}
 
 	// Object initialization.
diff --git a/Engine/Source/Runtime/CoreUObject/Private/UObject/ObjectRedirector.cpp b/Engine/Source/Runtime/CoreUObject/Private/UObject/ObjectRedirector.cpp
--- a/Engine/Source/Runtime/CoreUObject/Private/UObject/ObjectRedirector.cpp
+++ b/Engine/Source/Runtime/CoreUObject/Private/UObject/ObjectRedirector.cpp
@@ -109,7 +109,7 @@ FScopedRedirectorCatcher::FScopedRedirectorCatcher(const FString& InObjectPathNa
 , bWasRedirectorFollowed(false)
 {
 	// register itself on construction to see if the object is a redirector 
-	FCoreDelegates::RedirectorFollowed.AddRaw(this, &FScopedRedirectorCatcher::OnRedirectorFollowed);
+	FCoreUObjectDelegates::RedirectorFollowed.AddRaw(this, &FScopedRedirectorCatcher::OnRedirectorFollowed);
 }
 
 /**
@@ -118,7 +118,7 @@ FScopedRedirectorCatcher::FScopedRedirectorCatcher(const FString& InObjectPathNa
 FScopedRedirectorCatcher::~FScopedRedirectorCatcher()
 {
 	// register itself on construction to see if the object is a redirector 
-	FCoreDelegates::RedirectorFollowed.RemoveAll(this);
+	FCoreUObjectDelegates::RedirectorFollowed.RemoveAll(this);
 }
 
 
diff --git a/Engine/Source/Runtime/CoreUObject/Private/UObject/UObjectGlobals.cpp b/Engine/Source/Runtime/CoreUObject/Private/UObject/UObjectGlobals.cpp
--- a/Engine/Source/Runtime/CoreUObject/Private/UObject/UObjectGlobals.cpp
+++ b/Engine/Source/Runtime/CoreUObject/Private/UObject/UObjectGlobals.cpp
@@ -55,6 +55,26 @@ FCoreUObjectDelegates::FReplaceHotReloadClassDelegate FCoreUObjectDelegates::Rep
 FCoreUObjectDelegates::FIsPackageOKToSaveDelegate FCoreUObjectDelegates::IsPackageOKToSaveDelegate;
 FCoreUObjectDelegates::FAutoPackageBackupDelegate FCoreUObjectDelegates::AutoPackageBackupDelegate;
 
+FCoreUObjectDelegates::FOnPreObjectPropertyChanged FCoreUObjectDelegates::OnPreObjectPropertyChanged;
+FCoreUObjectDelegates::FOnObjectPropertyChanged FCoreUObjectDelegates::OnObjectPropertyChanged;
+
+#if WITH_EDITOR
+FCoreUObjectDelegates::FOnObjectModified FCoreUObjectDelegates::OnObjectModified;
+FCoreUObjectDelegates::FOnAssetLoaded FCoreUObjectDelegates::OnAssetLoaded;
+FCoreUObjectDelegates::FOnObjectSaved FCoreUObjectDelegates::OnObjectSaved;
+#endif // WITH_EDITOR
+
+FCoreUObjectDelegates::FOnRedirectorFollowed FCoreUObjectDelegates::RedirectorFollowed;
+
+FSimpleMulticastDelegate FCoreUObjectDelegates::PreGarbageCollect;
+FSimpleMulticastDelegate FCoreUObjectDelegates::PostGarbageCollect;
+FSimpleMulticastDelegate FCoreUObjectDelegates::PreLoadMap;
+FSimpleMulticastDelegate FCoreUObjectDelegates::PostLoadMap;
+FCoreUObjectDelegates::FOnLoadObjectsOnTop FCoreUObjectDelegates::ShouldLoadOnTop;
+
+FCoreUObjectDelegates::FStringAssetReferenceLoaded FCoreUObjectDelegates::StringAssetReferenceLoaded;
+FCoreUObjectDelegates::FStringAssetReferenceSaving FCoreUObjectDelegates::StringAssetReferenceSaving;
+FCoreUObjectDelegates::FPackageCreatedForLoad FCoreUObjectDelegates::PackageCreatedForLoad;
 /**
  * Fast version of StaticFindObject that relies on the passed in FName being the object name
  * without any group/ package qualifiers.
@@ -1137,7 +1157,7 @@ void EndLoad(const TCHAR* LoadContext)
 						check(Obj);
 						if ( Obj->IsAsset() )
 						{
-							FCoreDelegates::OnAssetLoaded.Broadcast(Obj);
+							FCoreUObjectDelegates::OnAssetLoaded.Broadcast(Obj);
 						}
 					}
 				}
diff --git a/Engine/Source/Runtime/CoreUObject/Public/UObject/UObjectGlobals.h b/Engine/Source/Runtime/CoreUObject/Public/UObject/UObjectGlobals.h
--- a/Engine/Source/Runtime/CoreUObject/Public/UObject/UObjectGlobals.h
+++ b/Engine/Source/Runtime/CoreUObject/Public/UObject/UObjectGlobals.h
@@ -1584,9 +1584,48 @@ DECLARE_DELEGATE_OneParam( FAddPackageToDefaultChangelistDelegate, const TCHAR*
  */
 struct COREUOBJECT_API FCoreUObjectDelegates
 {
+	// Callback for object property modifications
+	DECLARE_MULTICAST_DELEGATE_TwoParams(FOnObjectPropertyChanged, UObject*, struct FPropertyChangedEvent&);
+
+	// Called when a property is changed
+	static FOnObjectPropertyChanged OnObjectPropertyChanged;
+
+	// Callback for PreEditChange
+	DECLARE_MULTICAST_DELEGATE_TwoParams(FOnPreObjectPropertyChanged, UObject*, const class FEditPropertyChain&);
+
+	// Called before a property is changed
+	static FOnPreObjectPropertyChanged OnPreObjectPropertyChanged;
+
 	/** Delegate type for making auto backup of package */
 	DECLARE_DELEGATE_RetVal_OneParam(bool, FAutoPackageBackupDelegate, const UPackage&);
 
+#if WITH_EDITOR
+	// Callback for all object modifications
+	DECLARE_MULTICAST_DELEGATE_OneParam(FOnObjectModified, UObject*);
+
+	// Called when any object is modified at all 
+	static FOnObjectModified OnObjectModified;
+
+	// Callback for when an asset is loaded (Editor)
+	DECLARE_MULTICAST_DELEGATE_OneParam(FOnAssetLoaded, UObject*);
+
+	// Called when an asset is loaded
+	static FOnAssetLoaded OnAssetLoaded;
+
+	// Callback for when an asset is saved (Editor)
+	DECLARE_MULTICAST_DELEGATE_OneParam(FOnObjectSaved, UObject*);
+
+	// Called when an asset is saved
+	static FOnObjectSaved OnObjectSaved;
+
+#endif	//WITH_EDITOR
+
+	// Delegate type for redirector followed events ( Params: const FString& PackageName, UObject* Redirector )
+	DECLARE_MULTICAST_DELEGATE_TwoParams(FOnRedirectorFollowed, const FString&, UObject*);
+
+	// Sent when a UObjectRedirector was followed to find the destination object
+	static FOnRedirectorFollowed RedirectorFollowed;
+
 	/** Delegate used by SavePackage() to create the package backup */
 	static FAutoPackageBackupDelegate AutoPackageBackupDelegate;
 
@@ -1599,6 +1638,36 @@ struct COREUOBJECT_API FCoreUObjectDelegates
 	/** Delegate for replacing hot-reloaded classes that changed after hot-reload */
 	DECLARE_DELEGATE_TwoParams(FReplaceHotReloadClassDelegate, UClass*, UClass*);
 	static FReplaceHotReloadClassDelegate ReplaceHotReloadClassDelegate;
+
+	// Sent at the very beginning of LoadMap
+	static FSimpleMulticastDelegate PreLoadMap;
+
+	// Sent at the _successful_ end of LoadMap
+	static FSimpleMulticastDelegate PostLoadMap;
+
+	// Called before garbage collection
+	static FSimpleMulticastDelegate PreGarbageCollect;
+
+	// Called after garbage collection
+	static FSimpleMulticastDelegate PostGarbageCollect;
+
+	/** delegate type for querying whether a loaded object should replace an already existing one */
+	DECLARE_DELEGATE_RetVal_OneParam(bool, FOnLoadObjectsOnTop, const FString&);
+
+	/** Queries whether an object should be loaded on top ( replace ) an already existing one */
+	static FOnLoadObjectsOnTop ShouldLoadOnTop;
+
+	/** called when loading a string asset reference */
+	DECLARE_DELEGATE_OneParam(FStringAssetReferenceLoaded, FString const& /*LoadedAssetLongPathname*/);
+	static FStringAssetReferenceLoaded StringAssetReferenceLoaded;
+
+	/** called when path to world root is changed */
+	DECLARE_MULTICAST_DELEGATE_OneParam(FPackageCreatedForLoad, class UPackage*);
+	static FPackageCreatedForLoad PackageCreatedForLoad;
+
+	/** called when saving a string asset reference, can replace the value with something else */
+	DECLARE_DELEGATE_RetVal_OneParam(FString, FStringAssetReferenceSaving, FString const& /*SavingAssetLongPathname*/);
+	static FStringAssetReferenceSaving StringAssetReferenceSaving;
 };
 
 /** Allows release builds to override not verifying GC assumptions. Useful for profiling as it's hitchy. */
diff --git a/Engine/Source/Runtime/Engine/Private/AI/Navigation/NavigationSystem.cpp b/Engine/Source/Runtime/Engine/Private/AI/Navigation/NavigationSystem.cpp
--- a/Engine/Source/Runtime/Engine/Private/AI/Navigation/NavigationSystem.cpp
+++ b/Engine/Source/Runtime/Engine/Private/AI/Navigation/NavigationSystem.cpp
@@ -319,7 +319,7 @@ void UNavigationSystem::PostInitProperties()
 			GEngine->OnActorMoved().AddUObject(this, &UNavigationSystem::OnActorMoved);
 		}
 #endif
-		FCoreDelegates::PostLoadMap.AddUObject(this, &UNavigationSystem::OnPostLoadMap);
+		FCoreUObjectDelegates::PostLoadMap.AddUObject(this, &UNavigationSystem::OnPostLoadMap);
 #if WITH_NAVIGATION_GENERATOR
 		UNavigationSystem::NavigationDirtyEvent.AddUObject(this, &UNavigationSystem::OnNavigationDirtied);
 #endif // WITH_NAVIGATION_GENERATOR
@@ -2785,7 +2785,7 @@ void UNavigationSystem::CleanUp(ECleanupMode Mode)
 	}
 #endif // WITH_EDITOR
 
-	FCoreDelegates::PostLoadMap.RemoveAll(this);
+	FCoreUObjectDelegates::PostLoadMap.RemoveAll(this);
 #if WITH_NAVIGATION_GENERATOR
 	UNavigationSystem::NavigationDirtyEvent.RemoveAll(this);
 	FWorldDelegates::LevelAddedToWorld.RemoveAll(this);
diff --git a/Engine/Source/Runtime/Engine/Private/LevelStreaming.cpp b/Engine/Source/Runtime/Engine/Private/LevelStreaming.cpp
--- a/Engine/Source/Runtime/Engine/Private/LevelStreaming.cpp
+++ b/Engine/Source/Runtime/Engine/Private/LevelStreaming.cpp
@@ -200,7 +200,7 @@ void ULevelStreaming::PostLoad()
 			// Convert the FName reference to a TAssetPtr, then broadcast that we loaded a reference
 			// so this reference is gathered by the cooker without having to resave the package.
 			SetWorldAssetByPackageName(PackageName_DEPRECATED);
-			FCoreDelegates::StringAssetReferenceLoaded.ExecuteIfBound(WorldAsset.ToStringReference().ToString());
+			FCoreUObjectDelegates::StringAssetReferenceLoaded.ExecuteIfBound(WorldAsset.ToStringReference().ToString());
 		}
 		else
 		{
diff --git a/Engine/Source/Runtime/Engine/Private/LevelTick.cpp b/Engine/Source/Runtime/Engine/Private/LevelTick.cpp
--- a/Engine/Source/Runtime/Engine/Private/LevelTick.cpp
+++ b/Engine/Source/Runtime/Engine/Private/LevelTick.cpp
@@ -114,7 +114,7 @@ FDetailedTickStats::FDetailedTickStats( int32 InNumObjectsToReport, float InTime
 FDetailedTickStats::~FDetailedTickStats()
 {
 	// remove callback as we are dead
-	FCoreDelegates::PreGarbageCollect.RemoveRaw(this, &FDetailedTickStats::OnPreGarbageCollect);
+	FCoreUObjectDelegates::PreGarbageCollect.RemoveRaw(this, &FDetailedTickStats::OnPreGarbageCollect);
 }
 
 /**
@@ -174,7 +174,7 @@ void FDetailedTickStats::EndObject( UObject* Object, float DeltaTime, bool bForS
 		{
 			GCCallBackRegistered = true;
 			// register callback so that we can avoid finding the wrong stats for new objects reusing memory that used to be associated with a different object
-			FCoreDelegates::PreGarbageCollect.AddRaw(this, &FDetailedTickStats::OnPreGarbageCollect);
+			FCoreUObjectDelegates::PreGarbageCollect.AddRaw(this, &FDetailedTickStats::OnPreGarbageCollect);
 		}
 
 		FTickStats NewTickStats;
diff --git a/Engine/Source/Runtime/Engine/Private/PhysicsEngine/PhysLevel.cpp b/Engine/Source/Runtime/Engine/Private/PhysicsEngine/PhysLevel.cpp
--- a/Engine/Source/Runtime/Engine/Private/PhysicsEngine/PhysLevel.cpp
+++ b/Engine/Source/Runtime/Engine/Private/PhysicsEngine/PhysLevel.cpp
@@ -332,7 +332,7 @@ void InitGamePhys()
 
 	GPhysCommandHandler = new FPhysCommandHandler();
 
-	FCoreDelegates::PreGarbageCollect.AddRaw(GPhysCommandHandler, &FPhysCommandHandler::Flush);
+	FCoreUObjectDelegates::PreGarbageCollect.AddRaw(GPhysCommandHandler, &FPhysCommandHandler::Flush);
 
 	// Init Extensions
 	PxInitExtensions(*GPhysXSDK);
diff --git a/Engine/Source/Runtime/Engine/Private/StreamableManager.cpp b/Engine/Source/Runtime/Engine/Private/StreamableManager.cpp
--- a/Engine/Source/Runtime/Engine/Private/StreamableManager.cpp
+++ b/Engine/Source/Runtime/Engine/Private/StreamableManager.cpp
@@ -59,14 +59,14 @@ struct FStreamable
 
 FStreamableManager::FStreamableManager()
 {
-	FCoreDelegates::PreGarbageCollect.AddRaw(this, &FStreamableManager::OnPreGarbageCollect);
-	FCoreDelegates::PostGarbageCollect.AddRaw(this, &FStreamableManager::OnPostGarbageCollect);
+	FCoreUObjectDelegates::PreGarbageCollect.AddRaw(this, &FStreamableManager::OnPreGarbageCollect);
+	FCoreUObjectDelegates::PostGarbageCollect.AddRaw(this, &FStreamableManager::OnPostGarbageCollect);
 }
 
 FStreamableManager::~FStreamableManager()
 {
-	FCoreDelegates::PreGarbageCollect.RemoveAll(this);
-	FCoreDelegates::PostGarbageCollect.RemoveAll(this);
+	FCoreUObjectDelegates::PreGarbageCollect.RemoveAll(this);
+	FCoreUObjectDelegates::PostGarbageCollect.RemoveAll(this);
 	for (TStreamableMap::TIterator It(StreamableItems); It; ++It)
 	{
 		delete It.Value();
diff --git a/Engine/Source/Runtime/Engine/Private/UnrealEngine.cpp b/Engine/Source/Runtime/Engine/Private/UnrealEngine.cpp
--- a/Engine/Source/Runtime/Engine/Private/UnrealEngine.cpp
+++ b/Engine/Source/Runtime/Engine/Private/UnrealEngine.cpp
@@ -8475,7 +8475,7 @@ bool UEngine::LoadMap( FWorldContext& WorldContext, FURL URL, class UPendingNetG
 	}
 
 	// send a callback message
-	FCoreDelegates::PreLoadMap.Broadcast();
+	FCoreUObjectDelegates::PreLoadMap.Broadcast();
 
 	// Cancel any pending texture streaming requests.  This avoids a significant delay on consoles 
 	// when loading a map and there are a lot of outstanding texture streaming requests from the previous map.
@@ -8946,7 +8946,7 @@ bool UEngine::LoadMap( FWorldContext& WorldContext, FURL URL, class UPendingNetG
 	}
 
 	// send a callback message
-	FCoreDelegates::PostLoadMap.Broadcast();
+	FCoreUObjectDelegates::PostLoadMap.Broadcast();
 	
 	WorldContext.World()->bWorldWasLoadedThisTick = true;
 
diff --git a/Engine/Source/Runtime/Engine/Private/World.cpp b/Engine/Source/Runtime/Engine/Private/World.cpp
--- a/Engine/Source/Runtime/Engine/Private/World.cpp
+++ b/Engine/Source/Runtime/Engine/Private/World.cpp
@@ -1930,8 +1930,8 @@ void FLevelStreamingGCHelper::AddGarbageCollectorCallback()
 	static bool GarbageCollectAdded = false;
 	if ( GarbageCollectAdded == false )
 	{
-		FCoreDelegates::PreGarbageCollect.AddStatic( FLevelStreamingGCHelper::PrepareStreamedOutLevelsForGC );
-		FCoreDelegates::PostGarbageCollect.AddStatic( FLevelStreamingGCHelper::VerifyLevelsGotRemovedByGC );
+		FCoreUObjectDelegates::PreGarbageCollect.AddStatic( FLevelStreamingGCHelper::PrepareStreamedOutLevelsForGC );
+		FCoreUObjectDelegates::PostGarbageCollect.AddStatic( FLevelStreamingGCHelper::VerifyLevelsGotRemovedByGC );
 		GarbageCollectAdded = true;
 	}
 }
@@ -4748,7 +4748,7 @@ UWorld* FSeamlessTravelHandler::Tick()
 				// Called after post seamless travel to make sure players are setup correctly first
 				LoadedWorld->BeginPlay();
 
-				FCoreDelegates::PostLoadMap.Broadcast();
+				FCoreUObjectDelegates::PostLoadMap.Broadcast();
 			}
 			else
 			{
diff --git a/Engine/Source/Runtime/MoviePlayer/Private/DefaultGameMoviePlayer.cpp b/Engine/Source/Runtime/MoviePlayer/Private/DefaultGameMoviePlayer.cpp
--- a/Engine/Source/Runtime/MoviePlayer/Private/DefaultGameMoviePlayer.cpp
+++ b/Engine/Source/Runtime/MoviePlayer/Private/DefaultGameMoviePlayer.cpp
@@ -76,7 +76,7 @@ void FDefaultGameMoviePlayer::Initialize()
 	}
 
 	// Add a delegate to start playing movies when we start loading a map
-	FCoreDelegates::PreLoadMap.AddSP( this, &FDefaultGameMoviePlayer::OnPreLoadMap );
+	FCoreUObjectDelegates::PreLoadMap.AddSP( this, &FDefaultGameMoviePlayer::OnPreLoadMap );
 
 	FPlatformSplash::Hide();
 
@@ -134,7 +134,7 @@ void FDefaultGameMoviePlayer::Shutdown()
 
 	bInitialized = false;
 
-	FCoreDelegates::PreLoadMap.RemoveAll( this );
+	FCoreUObjectDelegates::PreLoadMap.RemoveAll( this );
 
 	LoadingScreenContents.Reset();
 	LoadingScreenWidgetHolder.Reset();
@@ -436,11 +436,11 @@ FReply FDefaultGameMoviePlayer::OnLoadingScreenMouseButtonDown(const FGeometry&
 
 void FDefaultGameMoviePlayer::OnPreLoadMap()
 {
-	FCoreDelegates::PostLoadMap.RemoveAll(this);
+	FCoreUObjectDelegates::PostLoadMap.RemoveAll(this);
 
 	if( PlayMovie() )
 	{
-		FCoreDelegates::PostLoadMap.AddSP(this, &FDefaultGameMoviePlayer::OnPostLoadMap );
+		FCoreUObjectDelegates::PostLoadMap.AddSP(this, &FDefaultGameMoviePlayer::OnPostLoadMap );
 	}
 }
 
