Fix HTML5 Key repeat flag #ttp 345320
SDL Key event repeat flags isn't working in HTML5, replaced with our own method.

diff --git a/Engine/Source/Runtime/Core/Private/HTML5/HTML5InputInterface.cpp b/Engine/Source/Runtime/Core/Private/HTML5/HTML5InputInterface.cpp
--- a/Engine/Source/Runtime/Core/Private/HTML5/HTML5InputInterface.cpp
+++ b/Engine/Source/Runtime/Core/Private/HTML5/HTML5InputInterface.cpp
@@ -23,6 +23,7 @@ FHTML5InputInterface::FHTML5InputInterface( const TSharedRef< FGenericApplicatio
 #else
 	SDL_EnableKeyRepeat(100, 50);
 #endif
+	KeyStates.Init(0, SDLK_LAST);
 }
 
 void FHTML5InputInterface::SetMessageHandler( const TSharedRef< FGenericApplicationMessageHandler >& InMessageHandler )
@@ -39,13 +40,14 @@ void FHTML5InputInterface::Tick(float DeltaTime, const SDL_Event& Event,TSharedR
 #define EVENT_KEY_REPEAT isRepeat
 #define KMOD_GUI KMOD_META
 #else
-#define EVENT_KEY_REPEAT Event.key.repeat
+#define EVENT_KEY_REPEAT (KeyStates[Event.key.keysym.sym]) // Event.key.repeat // <- this value always returns 0 
 #endif
 
 	{
 		if (Event.type == SDL_KEYDOWN)
 		{
 			MessageHandler->OnKeyDown(Event.key.keysym.sym, Event.key.keysym.sym, EVENT_KEY_REPEAT);
+			KeyStates[Event.key.keysym.sym] = true;
 			if (!(Event.key.keysym.mod & KMOD_CTRL) && !(Event.key.keysym.mod & KMOD_ALT) && !(Event.key.keysym.mod & KMOD_GUI))
 			{
 				MessageHandler->OnKeyChar(Event.key.keysym.sym, EVENT_KEY_REPEAT);
@@ -53,6 +55,7 @@ void FHTML5InputInterface::Tick(float DeltaTime, const SDL_Event& Event,TSharedR
 		}
 		else if (Event.type == SDL_KEYUP)
 		{
+			KeyStates[Event.key.keysym.sym] = false;
 			MessageHandler->OnKeyUp(Event.key.keysym.sym, Event.key.keysym.sym, EVENT_KEY_REPEAT);
 #if PLATFORM_HTML5_WIN32 && SDL_MAJOR_VERSION < 2
 			lastKey = SDLK_UNKNOWN;
diff --git a/Engine/Source/Runtime/Core/Public/HTML5/HTML5InputInterface.h b/Engine/Source/Runtime/Core/Public/HTML5/HTML5InputInterface.h
--- a/Engine/Source/Runtime/Core/Public/HTML5/HTML5InputInterface.h
+++ b/Engine/Source/Runtime/Core/Public/HTML5/HTML5InputInterface.h
@@ -36,4 +36,6 @@ private:
 
 	TSharedRef< FGenericApplicationMessageHandler > MessageHandler;
 	const TSharedPtr< ICursor > Cursor;
+
+	TBitArray<FDefaultBitArrayAllocator> KeyStates;
 };
\ No newline at end of file
