Lambda support in delegates and slate attributes
Not supported in multicast delegates for many reasons
Not supported in slate events (may reconsider in future)

Also some cleanup:
Removed comparisons between uobject delegates and sharedptr delegates, as it's impossible for them to match
Removed "const" static delegates, as they didn't really exist
Fixed some comments

diff --git a/Engine/Source/Runtime/Core/Public/Delegates/DelegateInstancesImpl.inl b/Engine/Source/Runtime/Core/Public/Delegates/DelegateInstancesImpl.inl
--- a/Engine/Source/Runtime/Core/Public/Delegates/DelegateInstancesImpl.inl
+++ b/Engine/Source/Runtime/Core/Public/Delegates/DelegateInstancesImpl.inl
@@ -23,6 +23,7 @@
 #define RAW_METHOD_DELEGATE_INSTANCE_CLASS FUNC_COMBINE(TBaseRawMethodDelegateInstance_, FUNC_COMBINE(FUNC_PAYLOAD_SUFFIX, FUNC_CONST_SUFFIX))
 #define UOBJECT_METHOD_DELEGATE_INSTANCE_CLASS FUNC_COMBINE(TBaseUObjectMethodDelegateInstance_, FUNC_COMBINE(FUNC_PAYLOAD_SUFFIX, FUNC_CONST_SUFFIX))
 #define STATIC_DELEGATE_INSTANCE_CLASS FUNC_COMBINE(TBaseStaticDelegateInstance_, FUNC_COMBINE(FUNC_PAYLOAD_SUFFIX, FUNC_CONST_SUFFIX))
+#define FUNCTOR_DELEGATE_INSTANCE_CLASS FUNC_COMBINE(TBaseFunctorDelegateInstance_, FUNC_COMBINE(FUNC_PAYLOAD_SUFFIX, FUNC_CONST_SUFFIX))
 #define UFUNCTION_DELEGATE_INSTANCE_CLASS FUNC_COMBINE(TBaseUFunctionDelegateInstance_, FUNC_COMBINE(FUNC_PAYLOAD_SUFFIX, FUNC_CONST_SUFFIX))
 
 #if FUNC_IS_CONST
@@ -181,10 +182,9 @@ public:
 	{
 		// NOTE: Payload data is not currently considered when comparing delegate instances.
 		// See the comment in multi-cast delegate's Remove() method for more information.
-		if ((InOtherDelegate.GetType() == EDelegateInstanceType::SharedPointerMethod) || 
+		if ((InOtherDelegate.GetType() == EDelegateInstanceType::SharedPointerMethod) ||
 			(InOtherDelegate.GetType() == EDelegateInstanceType::ThreadSafeSharedPointerMethod) ||
-			(InOtherDelegate.GetType() == EDelegateInstanceType::RawMethod) ||
-			(InOtherDelegate.GetType() == EDelegateInstanceType::UObjectMethod))
+			(InOtherDelegate.GetType() == EDelegateInstanceType::RawMethod))
 		{
 			return (GetRawMethodPtrInternal() == InOtherDelegate.GetRawMethodPtr() && UserObject.HasSameObject(InOtherDelegate.GetRawUserObject()));
 		}
@@ -540,9 +540,7 @@ public:
 		// NOTE: Payload data is not currently considered when comparing delegate instances.
 		// See the comment in multi-cast delegate's Remove() method for more information.
 		if ((InOtherDelegate.GetType() == EDelegateInstanceType::UObjectMethod) || 
-			(InOtherDelegate.GetType() == EDelegateInstanceType::RawMethod) ||
-			(InOtherDelegate.GetType() == EDelegateInstanceType::SharedPointerMethod) ||
-			(InOtherDelegate.GetType() == EDelegateInstanceType::ThreadSafeSharedPointerMethod))
+			(InOtherDelegate.GetType() == EDelegateInstanceType::RawMethod))
 		{
 			return (GetRawMethodPtrInternal() == InOtherDelegate.GetRawMethodPtr()) && (UserObject.Get() == InOtherDelegate.GetRawUserObject());
 		}
@@ -598,6 +596,7 @@ private:
 	FUNC_PAYLOAD_MEMBERS
 };
 
+#if !FUNC_IS_CONST
 
 /**
  * Implements a delegate binding for regular C++ functions.
@@ -675,8 +674,7 @@ public:
 
 	virtual RetValType Execute( FUNC_PARAM_LIST ) const override
 	{
-		// Call the member function on the user's object.  And yes, this is the correct C++ syntax for calling a
-		// pointer-to-member function.
+		// Call the static function
 		checkSlow(StaticFuncPtr != nullptr);
 
 		return StaticFuncPtr(DELEGATE_PARAM_PASSTHRU_COMMA_PAYLOAD_PASSIN);
@@ -730,7 +728,126 @@ private:
 };
 
 
-#if !FUNC_IS_CONST
+/**
+ * Implements a delegate binding for C++ functors, e.g. lambdas.
+ */
+template<typename FunctorType, FUNC_PAYLOAD_TEMPLATE_DECL_TYPENAME>
+class FUNCTOR_DELEGATE_INSTANCE_CLASS
+	: public DELEGATE_INSTANCE_INTERFACE_CLASS<FUNC_TEMPLATE_ARGS>
+{
+public:
+	static_assert(TAreTypesEqual<FunctorType, typename TRemoveReference<FunctorType>::Type>::Value, "FunctorType cannot be a reference");
+
+	/**
+	 * Creates and initializes a new instance
+	 *
+	 * @param InFunctor C++ functor
+	 */
+	FUNCTOR_DELEGATE_INSTANCE_CLASS(const FunctorType& InFunctor DELEGATE_COMMA_PAYLOAD_LIST)
+		: Functor(InFunctor)
+		DELEGATE_COMMA_PAYLOAD_INITIALIZER_LIST
+	{
+	}
+
+	FUNCTOR_DELEGATE_INSTANCE_CLASS(FunctorType&& InFunctor DELEGATE_COMMA_PAYLOAD_LIST)
+	: Functor(MoveTemp(InFunctor))
+	DELEGATE_COMMA_PAYLOAD_INITIALIZER_LIST
+	{
+	}
+
+public:
+	// IDelegateInstance interface
+
+	virtual FName GetFunctionName() const override
+	{
+		return NAME_None;
+	}
+
+	virtual const void* GetRawMethodPtr() const override
+	{
+		// casting operator() to void* is not legal C++ if it's a member function
+		// and wouldn't necessarily be a useful thing to return anyway
+		check(0);
+		return nullptr;
+	}
+
+	virtual const void* GetRawUserObject() const override
+	{
+		// returning &Functor wouldn't be particularly useful to the comparison code
+		// as it would always be unique because we store a copy of the functor
+		check(0);
+		return nullptr;
+	}
+
+	virtual EDelegateInstanceType::Type GetType() const override
+	{
+		return EDelegateInstanceType::Functor;
+	}
+
+	virtual bool HasSameObject(const void* UserObject) const override
+	{
+		// Functor Delegates aren't bound to a user object so they can never match
+		return false;
+	}
+
+	virtual bool IsSafeToExecute() const override
+	{
+		// Functors are always considered safe to execute!
+		return true;
+	}
+
+public:
+
+	// DELEGATE_INSTANCE_INTERFACE_CLASS interface
+
+	virtual DELEGATE_INSTANCE_INTERFACE_CLASS<FUNC_TEMPLATE_ARGS>* CreateCopy() override
+	{
+		return Create(Functor DELEGATE_COMMA_PAYLOAD_PASSIN);
+	}
+
+	virtual RetValType Execute(FUNC_PARAM_LIST) const override
+	{
+		return Functor(DELEGATE_PARAM_PASSTHRU_COMMA_PAYLOAD_PASSIN);
+	}
+
+#if FUNC_IS_VOID
+	virtual bool ExecuteIfSafe(FUNC_PARAM_LIST) const override
+	{
+		Execute(FUNC_PARAM_PASSTHRU);
+
+		return true;
+	}
+#endif
+
+	virtual bool IsSameFunction(const DELEGATE_INSTANCE_INTERFACE_CLASS<FUNC_TEMPLATE_ARGS>& InOtherDelegate) const override
+	{
+		// There's no nice way to implement this (we don't have the type info necessary to compare against OtherDelegate's Functor)
+		return false;
+	}
+
+public:
+
+	/**
+	* Creates a new static function delegate binding for the given function pointer.
+	*
+	* @param InFunctor C++ functor
+	* @return The new delegate.
+	*/
+	FORCEINLINE static DELEGATE_INSTANCE_INTERFACE_CLASS<FUNC_TEMPLATE_ARGS>* Create(const FunctorType& InFunctor DELEGATE_COMMA_PAYLOAD_LIST)
+	{
+		return new FUNCTOR_DELEGATE_INSTANCE_CLASS<FunctorType, FUNC_PAYLOAD_TEMPLATE_ARGS>(InFunctor DELEGATE_COMMA_PAYLOAD_PASSTHRU);
+	}
+	FORCEINLINE static DELEGATE_INSTANCE_INTERFACE_CLASS<FUNC_TEMPLATE_ARGS>* Create(FunctorType&& InFunctor DELEGATE_COMMA_PAYLOAD_LIST)
+	{
+		return new FUNCTOR_DELEGATE_INSTANCE_CLASS<FunctorType, FUNC_PAYLOAD_TEMPLATE_ARGS>(MoveTemp(InFunctor) DELEGATE_COMMA_PAYLOAD_PASSTHRU);
+	}
+
+private:
+	FunctorType Functor;
+
+	// Payload member variables, if any.
+	FUNC_PAYLOAD_MEMBERS
+};
 
 /**
  * Implements a delegate binding for UFunctions.
diff --git a/Engine/Source/Runtime/Core/Public/Delegates/DelegateInstancesImpl_Variadics.inl b/Engine/Source/Runtime/Core/Public/Delegates/DelegateInstancesImpl_Variadics.inl
--- a/Engine/Source/Runtime/Core/Public/Delegates/DelegateInstancesImpl_Variadics.inl
+++ b/Engine/Source/Runtime/Core/Public/Delegates/DelegateInstancesImpl_Variadics.inl
@@ -286,8 +286,7 @@ public:
 		// See the comment in multi-cast delegate's Remove() method for more information.
 		if ((InOtherDelegate.GetType() == EDelegateInstanceType::SharedPointerMethod) || 
 			(InOtherDelegate.GetType() == EDelegateInstanceType::ThreadSafeSharedPointerMethod) ||
-			(InOtherDelegate.GetType() == EDelegateInstanceType::RawMethod) ||
-			(InOtherDelegate.GetType() == EDelegateInstanceType::UObjectMethod))
+			(InOtherDelegate.GetType() == EDelegateInstanceType::RawMethod))
 		{
 			return (GetRawMethodPtrInternal() == InOtherDelegate.GetRawMethodPtr() && UserObject.HasSameObject(InOtherDelegate.GetRawUserObject()));
 		}
@@ -659,9 +658,7 @@ public:
 		// NOTE: Payload data is not currently considered when comparing delegate instances.
 		// See the comment in multi-cast delegate's Remove() method for more information.
 		if ((InOtherDelegate.GetType() == EDelegateInstanceType::UObjectMethod) || 
-			(InOtherDelegate.GetType() == EDelegateInstanceType::RawMethod) ||
-			(InOtherDelegate.GetType() == EDelegateInstanceType::SharedPointerMethod) ||
-			(InOtherDelegate.GetType() == EDelegateInstanceType::ThreadSafeSharedPointerMethod))
+			(InOtherDelegate.GetType() == EDelegateInstanceType::RawMethod))
 		{
 			return (GetRawMethodPtrInternal() == InOtherDelegate.GetRawMethodPtr()) && (UserObject.Get() == InOtherDelegate.GetRawUserObject());
 		}
@@ -810,8 +807,7 @@ public:
 
 	virtual RetValType Execute(ParamTypes... Params) const override
 	{
-		// Call the member function on the user's object.  And yes, this is the correct C++ syntax for calling a
-		// pointer-to-member function.
+		// Call the static function
 		checkSlow(StaticFuncPtr != nullptr);
 
 		return Payload.template ApplyAfter_ExplicitReturnType<RetValType>(StaticFuncPtr, Params...);
@@ -877,3 +873,141 @@ public:
 		return true;
 	}
 };
+
+/**
+ * Implements a delegate binding for C++ functors, e.g. lambdas.
+ */
+template <typename FuncType, typename FunctorType, typename... VarTypes>
+class TBaseFunctorDelegateInstance;
+
+template <typename WrappedRetValType, typename... ParamTypes, typename FunctorType, typename... VarTypes>
+class TBaseFunctorDelegateInstance<WrappedRetValType(ParamTypes...), FunctorType, VarTypes...> : public IBaseDelegateInstance<typename TUnwrapType<WrappedRetValType>::Type(ParamTypes...)>
+{
+	static_assert(TAreTypesEqual<FunctorType, typename TRemoveReference<FunctorType>::Type>::Value, "FunctorType cannot be a reference");
+
+	typedef IBaseDelegateInstance<typename TUnwrapType<WrappedRetValType>::Type(ParamTypes...)> Super;
+
+public:
+	typedef typename TUnwrapType<WrappedRetValType>::Type RetValType;
+
+	TBaseFunctorDelegateInstance(const FunctorType& InFunctor, VarTypes... Vars)
+		: Functor(InFunctor)
+		, Payload(Vars...)
+	{
+	}
+	TBaseFunctorDelegateInstance(FunctorType&& InFunctor, VarTypes... Vars)
+		: Functor(MoveTemp(InFunctor))
+		, Payload(Vars...)
+	{
+	}
+
+	// IDelegateInstance interface
+
+	virtual FName GetFunctionName() const override
+	{
+		return NAME_None;
+	}
+
+	virtual const void* GetRawMethodPtr() const override
+	{
+		// casting operator() to void* is not legal C++ if it's a member function
+		// and wouldn't necessarily be a useful thing to return anyway
+		check(0);
+		return nullptr;
+	}
+
+	virtual const void* GetRawUserObject() const override
+	{
+		// returning &Functor wouldn't be particularly useful to the comparison code
+		// as it would always be unique because we store a copy of the functor
+		check(0);
+		return nullptr;
+	}
+
+	virtual EDelegateInstanceType::Type GetType() const override
+	{
+		return EDelegateInstanceType::Functor;
+	}
+
+	virtual bool HasSameObject(const void* UserObject) const override
+	{
+		// Functor Delegates aren't bound to a user object so they can never match
+		return false;
+	}
+
+	virtual bool IsSafeToExecute() const override
+	{
+		// Functors are always considered safe to execute!
+		return true;
+	}
+
+public:
+	// IBaseDelegateInstance interface
+	virtual Super* CreateCopy() override
+	{
+		return Payload.template ApplyAfter_ExplicitReturnType<Super*>(static_cast<Super* (*)(const FunctorType&)>(&Create), Functor);
+	}
+
+	virtual RetValType Execute(ParamTypes... Params) const override
+	{
+		return Payload.template ApplyAfter_ExplicitReturnType<RetValType>(Functor, Params...);
+	}
+
+	virtual bool IsSameFunction(const Super& InOtherDelegate) const override
+	{
+		// There's no nice way to implement this (we don't have the type info necessary to compare against OtherDelegate's Functor)
+		return false;
+	}
+
+public:
+	/**
+	 * Creates a new static function delegate binding for the given function pointer.
+	 *
+	 * @param InFunctor C++ functor
+	 * @return The new delegate.
+	 */
+	FORCEINLINE static Super* Create(const FunctorType& InFunctor, VarTypes... Vars)
+	{
+		return new TBaseFunctorDelegateInstance<RetValType(ParamTypes...), FunctorType, VarTypes...>(InFunctor, Vars...);
+	}
+	FORCEINLINE static Super* Create(FunctorType&& InFunctor, VarTypes... Vars)
+	{
+		return new TBaseFunctorDelegateInstance<RetValType(ParamTypes...), FunctorType, VarTypes...>(MoveTemp(InFunctor), Vars...);
+	}
+
+private:
+	// C++ functor
+	FunctorType Functor;
+
+	// Payload member variables, if any.
+	TTuple<VarTypes...> Payload;
+};
+
+template <typename FunctorType, typename... ParamTypes, typename... VarTypes>
+class TBaseFunctorDelegateInstance<void(ParamTypes...), FunctorType, VarTypes...> : public TBaseFunctorDelegateInstance<TTypeWrapper<void>(ParamTypes...), FunctorType, VarTypes...>
+{
+	typedef TBaseFunctorDelegateInstance<TTypeWrapper<void>(ParamTypes...), FunctorType, VarTypes...> Super;
+
+public:
+	/**
+	 * Creates and initializes a new instance.
+	 *
+	 * @param InFunctor C++ functor
+	 */
+	TBaseFunctorDelegateInstance(const FunctorType& InFunctor, VarTypes... Vars)
+		: Super(InFunctor, Vars...)
+	{
+	}
+	TBaseFunctorDelegateInstance(FunctorType&& InFunctor, VarTypes... Vars)
+		: Super(MoveTemp(InFunctor), Vars...)
+	{
+	}
+
+	virtual bool ExecuteIfSafe(ParamTypes... Params) const override
+	{
+		// Functors are always considered safe to execute!
+		Super::Execute(Params...);
+
+		return true;
+	}
+};
diff --git a/Engine/Source/Runtime/Core/Public/Delegates/DelegateSignatureImpl.inl b/Engine/Source/Runtime/Core/Public/Delegates/DelegateSignatureImpl.inl
--- a/Engine/Source/Runtime/Core/Public/Delegates/DelegateSignatureImpl.inl
+++ b/Engine/Source/Runtime/Core/Public/Delegates/DelegateSignatureImpl.inl
@@ -52,15 +52,16 @@
 #define UOBJECT_METHOD_DELEGATE_INSTANCE_CLASS_FourVars_Const FUNC_COMBINE( UOBJECT_METHOD_DELEGATE_INSTANCE_CLASS_FourVars, _Const )
 
 #define STATIC_DELEGATE_INSTANCE_CLASS_NoVars FUNC_COMBINE( TBaseStaticDelegateInstance_, FUNC_SUFFIX )
-#define STATIC_DELEGATE_INSTANCE_CLASS_NoVars_Const FUNC_COMBINE( STATIC_DELEGATE_INSTANCE_CLASS_NoVars, _Const )
 #define STATIC_DELEGATE_INSTANCE_CLASS_OneVar FUNC_COMBINE( STATIC_DELEGATE_INSTANCE_CLASS_NoVars, _OneVar )
-#define STATIC_DELEGATE_INSTANCE_CLASS_OneVar_Const FUNC_COMBINE( STATIC_DELEGATE_INSTANCE_CLASS_OneVar, _Const )
 #define STATIC_DELEGATE_INSTANCE_CLASS_TwoVars FUNC_COMBINE( STATIC_DELEGATE_INSTANCE_CLASS_NoVars, _TwoVars )
-#define STATIC_DELEGATE_INSTANCE_CLASS_TwoVars_Const FUNC_COMBINE( STATIC_DELEGATE_INSTANCE_CLASS_TwoVars, _Const )
 #define STATIC_DELEGATE_INSTANCE_CLASS_ThreeVars FUNC_COMBINE( STATIC_DELEGATE_INSTANCE_CLASS_NoVars, _ThreeVars )
-#define STATIC_DELEGATE_INSTANCE_CLASS_ThreeVars_Const FUNC_COMBINE( STATIC_DELEGATE_INSTANCE_CLASS_ThreeVars, _Const )
 #define STATIC_DELEGATE_INSTANCE_CLASS_FourVars FUNC_COMBINE( STATIC_DELEGATE_INSTANCE_CLASS_NoVars, _FourVars )
-#define STATIC_DELEGATE_INSTANCE_CLASS_FourVars_Const FUNC_COMBINE( STATIC_DELEGATE_INSTANCE_CLASS_FourVars, _Const )
+
+#define FUNCTOR_DELEGATE_INSTANCE_CLASS_NoVars FUNC_COMBINE( TBaseFunctorDelegateInstance_, FUNC_SUFFIX )
+#define FUNCTOR_DELEGATE_INSTANCE_CLASS_OneVar FUNC_COMBINE( FUNCTOR_DELEGATE_INSTANCE_CLASS_NoVars, _OneVar )
+#define FUNCTOR_DELEGATE_INSTANCE_CLASS_TwoVars FUNC_COMBINE( FUNCTOR_DELEGATE_INSTANCE_CLASS_NoVars, _TwoVars )
+#define FUNCTOR_DELEGATE_INSTANCE_CLASS_ThreeVars FUNC_COMBINE( FUNCTOR_DELEGATE_INSTANCE_CLASS_NoVars, _ThreeVars )
+#define FUNCTOR_DELEGATE_INSTANCE_CLASS_FourVars FUNC_COMBINE( FUNCTOR_DELEGATE_INSTANCE_CLASS_NoVars, _FourVars )
 
 #define DELEGATE_CLASS FUNC_COMBINE( TBaseDelegate_, FUNC_SUFFIX )
 #define BASE_MULTICAST_DELEGATE_CLASS FUNC_COMBINE( TBaseMulticastDelegate_, FUNC_SUFFIX )
@@ -168,6 +169,13 @@ public:
 	template< PAYLOAD_TEMPLATE_DECL_ThreeVars > class TStaticDelegate_ThreeVars : public STATIC_DELEGATE_INSTANCE_CLASS_ThreeVars< FUNC_TEMPLATE_ARGS, PAYLOAD_TEMPLATE_LIST_ThreeVars > { public: TStaticDelegate_ThreeVars( typename STATIC_DELEGATE_INSTANCE_CLASS_ThreeVars< FUNC_TEMPLATE_ARGS, PAYLOAD_TEMPLATE_LIST_ThreeVars >::FFuncPtr InFuncPtr, PAYLOAD_TEMPLATE_ARGS_ThreeVars ) : STATIC_DELEGATE_INSTANCE_CLASS_ThreeVars< FUNC_TEMPLATE_ARGS, PAYLOAD_TEMPLATE_LIST_ThreeVars >( InFuncPtr, PAYLOAD_TEMPLATE_PASSIN_ThreeVars ) {} };
 	template< PAYLOAD_TEMPLATE_DECL_FourVars  > class TStaticDelegate_FourVars  : public STATIC_DELEGATE_INSTANCE_CLASS_FourVars < FUNC_TEMPLATE_ARGS, PAYLOAD_TEMPLATE_LIST_FourVars  > { public: TStaticDelegate_FourVars ( typename STATIC_DELEGATE_INSTANCE_CLASS_FourVars < FUNC_TEMPLATE_ARGS, PAYLOAD_TEMPLATE_LIST_FourVars  >::FFuncPtr InFuncPtr, PAYLOAD_TEMPLATE_ARGS_FourVars  ) : STATIC_DELEGATE_INSTANCE_CLASS_FourVars < FUNC_TEMPLATE_ARGS, PAYLOAD_TEMPLATE_LIST_FourVars  >( InFuncPtr, PAYLOAD_TEMPLATE_PASSIN_FourVars  ) {} };
 
+	/** Declare the user's functor delegate instance types. */
+	template<typename FunctorType                                 > class TFunctorDelegate           : public FUNCTOR_DELEGATE_INSTANCE_CLASS_NoVars   <FunctorType, FUNC_TEMPLATE_ARGS                                 > { public: TFunctorDelegate          (const FunctorType& InFunctor                                 ) : FUNCTOR_DELEGATE_INSTANCE_CLASS_NoVars   <FunctorType, FUNC_TEMPLATE_ARGS                                  >(InFunctor                                   ) {} TFunctorDelegate          (FunctorType&& InFunctor                                 ) : FUNCTOR_DELEGATE_INSTANCE_CLASS_NoVars   <FunctorType, FUNC_TEMPLATE_ARGS                                  >(MoveTemp(InFunctor)                                   ) {} };
+	template<typename FunctorType, PAYLOAD_TEMPLATE_DECL_OneVar   > class TFunctorDelegate_OneVar    : public FUNCTOR_DELEGATE_INSTANCE_CLASS_OneVar   <FunctorType, FUNC_TEMPLATE_ARGS, PAYLOAD_TEMPLATE_LIST_OneVar   > { public: TFunctorDelegate_OneVar   (const FunctorType& InFunctor, PAYLOAD_TEMPLATE_ARGS_OneVar   ) : FUNCTOR_DELEGATE_INSTANCE_CLASS_OneVar   <FunctorType, FUNC_TEMPLATE_ARGS, PAYLOAD_TEMPLATE_LIST_OneVar    >(InFunctor, PAYLOAD_TEMPLATE_PASSIN_OneVar   ) {} TFunctorDelegate_OneVar   (FunctorType&& InFunctor, PAYLOAD_TEMPLATE_ARGS_OneVar   ) : FUNCTOR_DELEGATE_INSTANCE_CLASS_OneVar   <FunctorType, FUNC_TEMPLATE_ARGS, PAYLOAD_TEMPLATE_LIST_OneVar    >(MoveTemp(InFunctor), PAYLOAD_TEMPLATE_PASSIN_OneVar   ) {} };
+	template<typename FunctorType, PAYLOAD_TEMPLATE_DECL_TwoVars  > class TFunctorDelegate_TwoVars   : public FUNCTOR_DELEGATE_INSTANCE_CLASS_TwoVars  <FunctorType, FUNC_TEMPLATE_ARGS, PAYLOAD_TEMPLATE_LIST_TwoVars  > { public: TFunctorDelegate_TwoVars  (const FunctorType& InFunctor, PAYLOAD_TEMPLATE_ARGS_TwoVars  ) : FUNCTOR_DELEGATE_INSTANCE_CLASS_TwoVars  <FunctorType, FUNC_TEMPLATE_ARGS, PAYLOAD_TEMPLATE_LIST_TwoVars   >(InFunctor, PAYLOAD_TEMPLATE_PASSIN_TwoVars  ) {} TFunctorDelegate_TwoVars  (FunctorType&& InFunctor, PAYLOAD_TEMPLATE_ARGS_TwoVars  ) : FUNCTOR_DELEGATE_INSTANCE_CLASS_TwoVars  <FunctorType, FUNC_TEMPLATE_ARGS, PAYLOAD_TEMPLATE_LIST_TwoVars   >(MoveTemp(InFunctor), PAYLOAD_TEMPLATE_PASSIN_TwoVars  ) {} };
+	template<typename FunctorType, PAYLOAD_TEMPLATE_DECL_ThreeVars> class TFunctorDelegate_ThreeVars : public FUNCTOR_DELEGATE_INSTANCE_CLASS_ThreeVars<FunctorType, FUNC_TEMPLATE_ARGS, PAYLOAD_TEMPLATE_LIST_ThreeVars> { public: TFunctorDelegate_ThreeVars(const FunctorType& InFunctor, PAYLOAD_TEMPLATE_ARGS_ThreeVars) : FUNCTOR_DELEGATE_INSTANCE_CLASS_ThreeVars<FunctorType, FUNC_TEMPLATE_ARGS, PAYLOAD_TEMPLATE_LIST_ThreeVars >(InFunctor, PAYLOAD_TEMPLATE_PASSIN_ThreeVars) {} TFunctorDelegate_ThreeVars(FunctorType&& InFunctor, PAYLOAD_TEMPLATE_ARGS_ThreeVars) : FUNCTOR_DELEGATE_INSTANCE_CLASS_ThreeVars<FunctorType, FUNC_TEMPLATE_ARGS, PAYLOAD_TEMPLATE_LIST_ThreeVars >(MoveTemp(InFunctor), PAYLOAD_TEMPLATE_PASSIN_ThreeVars) {} };
+	template<typename FunctorType, PAYLOAD_TEMPLATE_DECL_FourVars > class TFunctorDelegate_FourVars  : public FUNCTOR_DELEGATE_INSTANCE_CLASS_FourVars <FunctorType, FUNC_TEMPLATE_ARGS, PAYLOAD_TEMPLATE_LIST_FourVars > { public: TFunctorDelegate_FourVars (const FunctorType& InFunctor, PAYLOAD_TEMPLATE_ARGS_FourVars ) : FUNCTOR_DELEGATE_INSTANCE_CLASS_FourVars <FunctorType, FUNC_TEMPLATE_ARGS, PAYLOAD_TEMPLATE_LIST_FourVars  >(InFunctor, PAYLOAD_TEMPLATE_PASSIN_FourVars ) {} TFunctorDelegate_FourVars (FunctorType&& InFunctor, PAYLOAD_TEMPLATE_ARGS_FourVars ) : FUNCTOR_DELEGATE_INSTANCE_CLASS_FourVars <FunctorType, FUNC_TEMPLATE_ARGS, PAYLOAD_TEMPLATE_LIST_FourVars  >(MoveTemp(InFunctor), PAYLOAD_TEMPLATE_PASSIN_FourVars ) {} };
+
 public:
 
 	/**
@@ -198,6 +206,36 @@ public:
 		return DELEGATE_CLASS( TStaticDelegate_FourVars< PAYLOAD_TEMPLATE_LIST_FourVars >::Create( InFunc, PAYLOAD_TEMPLATE_PASSIN_FourVars ) );
 	}
 
+	/**
+	 * Static: Creates a C++ lambda delegate
+	 * technically this works for any functor types, but lambdas are the primary use case
+	 */
+	template<typename FunctorType>
+	inline static DELEGATE_CLASS<FUNC_TEMPLATE_ARGS> CreateLambda(FunctorType&& InFunctor)
+	{
+		return DELEGATE_CLASS<FUNC_TEMPLATE_ARGS>(TFunctorDelegate<typename TRemoveReference<FunctorType>::Type>::Create(Forward<FunctorType>(InFunctor)));
+	}
+	template<typename FunctorType, PAYLOAD_TEMPLATE_DECL_OneVar>
+	inline static DELEGATE_CLASS<FUNC_TEMPLATE_ARGS> CreateLambda(FunctorType&& InFunctor, PAYLOAD_TEMPLATE_ARGS_OneVar)
+	{
+		return DELEGATE_CLASS<FUNC_TEMPLATE_ARGS>(TFunctorDelegate_OneVar<typename TRemoveReference<FunctorType>::Type, PAYLOAD_TEMPLATE_LIST_OneVar>::Create(Forward<FunctorType>(InFunctor), PAYLOAD_TEMPLATE_PASSIN_OneVar));
+	}
+	template<typename FunctorType, PAYLOAD_TEMPLATE_DECL_TwoVars>
+	inline static DELEGATE_CLASS<FUNC_TEMPLATE_ARGS> CreateLambda(FunctorType&& InFunctor, PAYLOAD_TEMPLATE_ARGS_TwoVars)
+	{
+		return DELEGATE_CLASS<FUNC_TEMPLATE_ARGS>(TFunctorDelegate_TwoVars<typename TRemoveReference<FunctorType>::Type, PAYLOAD_TEMPLATE_LIST_TwoVars>::Create(Forward<FunctorType>(InFunctor), PAYLOAD_TEMPLATE_PASSIN_TwoVars));
+	}
+	template<typename FunctorType, PAYLOAD_TEMPLATE_DECL_ThreeVars>
+	inline static DELEGATE_CLASS<FUNC_TEMPLATE_ARGS> CreateLambda(FunctorType&& InFunctor, PAYLOAD_TEMPLATE_ARGS_ThreeVars)
+	{
+		return DELEGATE_CLASS<FUNC_TEMPLATE_ARGS>(TFunctorDelegate_ThreeVars<typename TRemoveReference<FunctorType>::Type, PAYLOAD_TEMPLATE_LIST_ThreeVars>::Create(Forward<FunctorType>(InFunctor), PAYLOAD_TEMPLATE_PASSIN_ThreeVars));
+	}
+	template<typename FunctorType, PAYLOAD_TEMPLATE_DECL_FourVars>
+	inline static DELEGATE_CLASS<FUNC_TEMPLATE_ARGS> CreateLambda(FunctorType&& InFunctor, PAYLOAD_TEMPLATE_ARGS_FourVars)
+	{
+		return DELEGATE_CLASS<FUNC_TEMPLATE_ARGS>(TFunctorDelegate_FourVars<typename TRemoveReference<FunctorType>::Type, PAYLOAD_TEMPLATE_LIST_FourVars>::Create(Forward<FunctorType>(InFunctor), PAYLOAD_TEMPLATE_PASSIN_FourVars));
+	}
+
 	/**
 	 * Static: Creates a raw C++ pointer member function delegate.
 	 *
@@ -640,7 +678,7 @@ public:
 public:
 
 	/**
-	 * Creates a raw C++ pointer global function delegate
+	 * Binds a raw C++ pointer global function delegate
 	 */
 	inline void BindStatic( typename FStaticDelegate::FFuncPtr InFunc )
 	{
@@ -667,6 +705,36 @@ public:
 		*this = CreateStatic( InFunc, PAYLOAD_TEMPLATE_PASSIN_FourVars );
 	}
 
+	/**
+	 * Static: Binds a C++ lambda delegate
+	 * technically this works for any functor types, but lambdas are the primary use case
+	*/
+	template<typename FunctorType>
+	inline void BindLambda(FunctorType&& InFunctor)
+	{
+		*this = CreateLambda(Forward<FunctorType>(InFunctor));
+	}
+	template<typename FunctorType, PAYLOAD_TEMPLATE_DECL_OneVar>
+	inline void BindLambda(FunctorType&& InFunctor, PAYLOAD_TEMPLATE_ARGS_OneVar)
+	{
+		*this = CreateLambda(Forward<FunctorType>(InFunctor), PAYLOAD_TEMPLATE_PASSIN_OneVar);
+	}
+	template<typename FunctorType, PAYLOAD_TEMPLATE_DECL_TwoVars>
+	inline void BindLambda(FunctorType&& InFunctor, PAYLOAD_TEMPLATE_ARGS_TwoVars)
+	{
+		*this = CreateLambda(Forward<FunctorType>(InFunctor), PAYLOAD_TEMPLATE_PASSIN_TwoVars);
+	}
+	template<typename FunctorType, PAYLOAD_TEMPLATE_DECL_ThreeVars>
+	inline void BindLambda(FunctorType&& InFunctor, PAYLOAD_TEMPLATE_ARGS_ThreeVars)
+	{
+		*this = CreateLambda(Forward<FunctorType>(InFunctor), PAYLOAD_TEMPLATE_PASSIN_ThreeVars);
+	}
+	template<typename FunctorType, PAYLOAD_TEMPLATE_DECL_FourVars>
+	inline void BindLambda(FunctorType&& InFunctor, PAYLOAD_TEMPLATE_ARGS_FourVars)
+	{
+		*this = CreateLambda(Forward<FunctorType>(InFunctor), PAYLOAD_TEMPLATE_PASSIN_FourVars);
+	}
+
 	/**
 	 * Binds a raw C++ pointer delegate.
 	 *
@@ -779,7 +847,7 @@ public:
 	}
 
 	/**
-	 * Creates a shared pointer-based (fast, not thread-safe) member function delegate.
+	 * Binds a shared pointer-based (fast, not thread-safe) member function delegate.
 	 *
 	 * Shared pointer delegates keep a weak reference to your object.
 	 * You can use ExecuteIfBound() to call them.
@@ -893,7 +961,7 @@ public:
 	}
 
 	/**
-	 * Creates a shared pointer-based (slower, conditionally thread-safe) member function delegate.
+	 * Binds a shared pointer-based (slower, conditionally thread-safe) member function delegate.
 	 *
 	 * Shared pointer delegates keep a weak reference to your object.
 	 * You can use ExecuteIfBound() to call them.
@@ -950,7 +1018,7 @@ public:
 	}
 
 	/**
-	 * Creates a UFunction-based member function delegate.
+	 * Binds a UFunction-based member function delegate.
 	 *
 	 * UFunction delegates keep a weak reference to your object.
 	 * You can use ExecuteIfBound() to call them.
@@ -1092,16 +1160,16 @@ public:
 			return DelegateInstance->IsSameFunction(*OtherInstance);
 		}
 
-			// If neither delegate is initialized to anything yet, then we treat them as equal
+		// If neither delegate is initialized to anything yet, then we treat them as equal
 		if ((DelegateInstance == nullptr) && (OtherInstance == nullptr))
-			{
-				return true;
-			}
-
-			// No match!
-			return false;
+		{
+			return true;
 		}
 
+		// No match!
+		return false;
+	}
+
 private:
 
 	// Declare ourselves as a friend to our multi-cast counterpart, so that it can access our interface directly
diff --git a/Engine/Source/Runtime/Core/Public/Delegates/DelegateSignatureImpl_Variadics.inl b/Engine/Source/Runtime/Core/Public/Delegates/DelegateSignatureImpl_Variadics.inl
--- a/Engine/Source/Runtime/Core/Public/Delegates/DelegateSignatureImpl_Variadics.inl
+++ b/Engine/Source/Runtime/Core/Public/Delegates/DelegateSignatureImpl_Variadics.inl
@@ -97,6 +97,16 @@ public:
 		return TBaseStaticDelegateInstance<TFuncType, VarTypes...>::Create(InFunc, Vars...);
 	}
 
+	/**
+	 * Static: Creates a C++ lambda delegate
+	 * technically this works for any functor types, but lambdas are the primary use case
+	 */
+	template<typename FunctorType, typename... VarTypes>
+	inline static TBaseDelegate<RetValType, ParamTypes...> CreateLambda(FunctorType&& InFunctor, VarTypes... Vars)
+	{
+		return TBaseFunctorDelegateInstance<TFuncType, typename TRemoveReference<FunctorType>::Type>::Create(Forward<FunctorType>(InFunctor), Vars...);
+	}
+
 	/**
 	 * Static: Creates a raw C++ pointer member function delegate.
 	 *
@@ -279,13 +289,23 @@ public:
 public:
 
 	/**
-	 * Creates a raw C++ pointer global function delegate
+	 * Binds a raw C++ pointer global function delegate
 	 */
 	template <typename... VarTypes>
 	inline void BindStatic(typename TBaseStaticDelegateInstance<TFuncType, VarTypes...>::FFuncPtr InFunc, VarTypes... Vars)
 	{
 		*this = CreateStatic(InFunc, Vars...);
 	}
+	
+	/**
+	 * Static: Binds a C++ lambda delegate
+	 * technically this works for any functor types, but lambdas are the primary use case
+	 */
+	template<typename FunctorType, typename... VarTypes>
+	inline void BindLambda(FunctorType&& InFunctor, VarTypes... Vars)
+	{
+		*this = CreateLambda(Forward<FunctorType>(InFunctor), Vars...);
+	}
 
 	/**
 	 * Binds a raw C++ pointer delegate.
@@ -319,7 +339,7 @@ public:
 	}
 
 	/**
-	 * Creates a shared pointer-based (fast, not thread-safe) member function delegate.
+	 * Binds a shared pointer-based (fast, not thread-safe) member function delegate.
 	 *
 	 * Shared pointer delegates keep a weak reference to your object.
 	 * You can use ExecuteIfBound() to call them.
@@ -353,7 +373,7 @@ public:
 	}
 
 	/**
-	 * Creates a shared pointer-based (slower, conditionally thread-safe) member function delegate.
+	 * Binds a shared pointer-based (slower, conditionally thread-safe) member function delegate.
 	 *
 	 * Shared pointer delegates keep a weak reference to your object.
 	 * You can use ExecuteIfBound() to call them.
@@ -370,7 +390,7 @@ public:
 	}
 
 	/**
-	 * Creates a UFunction-based member function delegate.
+	 * Binds a UFunction-based member function delegate.
 	 *
 	 * UFunction delegates keep a weak reference to your object.
 	 * You can use ExecuteIfBound() to call them.
@@ -382,7 +402,7 @@ public:
 	}
 
 	/**
-	 * Creates a UObject-based member function delegate.
+	 * Binds a UObject-based member function delegate.
 	 *
 	 * UObject delegates keep a weak reference to your object.
 	 * You can use ExecuteIfBound() to call them.
diff --git a/Engine/Source/Runtime/Core/Public/Delegates/IDelegateInstance.h b/Engine/Source/Runtime/Core/Public/Delegates/IDelegateInstance.h
--- a/Engine/Source/Runtime/Core/Public/Delegates/IDelegateInstance.h
+++ b/Engine/Source/Runtime/Core/Public/Delegates/IDelegateInstance.h
@@ -26,7 +26,10 @@ namespace EDelegateInstanceType
 		UObjectMethod,
 
 		/** Raw C++ static function pointer */
-		Raw
+		Raw,
+
+		/** C++ functor, e.g. Lambda */
+		Functor,
 	};
 }
 
diff --git a/Engine/Source/Runtime/SlateCore/Public/Widgets/DeclarativeSyntaxSupport.h b/Engine/Source/Runtime/SlateCore/Public/Widgets/DeclarativeSyntaxSupport.h
--- a/Engine/Source/Runtime/SlateCore/Public/Widgets/DeclarativeSyntaxSupport.h
+++ b/Engine/Source/Runtime/SlateCore/Public/Widgets/DeclarativeSyntaxSupport.h
@@ -144,6 +144,15 @@
 			Var = AttrType::Create( AttrType::FGetter::CreateStatic( InFunc, Var1, Var2, Var3, Var4 ) ); \
 			return this->Me(); \
 		} \
+	\
+		/* Bind attribute with delegate to a lambda
+		 * technically this works for any functor types, but lambdas are the primary use case */ \
+		template<typename FunctorType> \
+		WidgetArgsType& AttrName##_Lambda(FunctorType&& InFunctor) \
+		{ \
+			Var = AttrType::Create(AttrType::FGetter::CreateLambda(InFunctor)); \
+			return this->Me(); \
+		} \
 	\
 		/* Bind attribute with delegate to a raw C++ class method */ \
 		template< class UserClass >	\
@@ -312,6 +321,10 @@
 			Var = TargetAttrType::Create( TargetAttrType::FGetter::CreateStatic( &AttrName##_Local::PassThroughDelegate, SourceAttrType::FGetter::CreateStatic( InFunc, Var1, Var2, Var3, Var4 ) ) ); \
 			return this->Me(); \
 		} \
+	\
+		/* We don't support binding FString lambdas to FText attributes because we currently have no way to disambiguate them */ \
+		/*template<typename FunctorType> \
+		WidgetArgsType& AttrName##_Lambda(FunctorType&& InFunctor) \
 	\
 		/* Bind attribute with delegate to a raw C++ class method */ \
 		template< class UserClass >	\
