Fix reporting of crashed application name on OS X when launching the editor without a project and ensure's causing hangs in crashed applications.
Plus provide a UUID for each crash so that all reports get a unique folder name rather than hoping that PID will be enough.

diff --git a/Engine/Source/Runtime/Core/Private/Apple/ApplePlatformSymbolication.cpp b/Engine/Source/Runtime/Core/Private/Apple/ApplePlatformSymbolication.cpp
--- a/Engine/Source/Runtime/Core/Private/Apple/ApplePlatformSymbolication.cpp
+++ b/Engine/Source/Runtime/Core/Private/Apple/ApplePlatformSymbolication.cpp
@@ -66,10 +66,17 @@ struct FApplePlatformSymbolCache
 #endif
 };
 
+static bool GAllowApplePlatformSymbolication = PLATFORM_MAC;
+
+void FApplePlatformSymbolication::SetSymbolicationAllowed(bool const bAllow)
+{
+	GAllowApplePlatformSymbolication = bAllow;
+}
+
 FApplePlatformSymbolCache* FApplePlatformSymbolication::CreateSymbolCache(void)
 {
 #if PLATFORM_MAC
-	return new FApplePlatformSymbolCache;
+	return GAllowApplePlatformSymbolication ? new FApplePlatformSymbolCache : nullptr;
 #else 
 	return nullptr;
 #endif
@@ -78,7 +85,7 @@ FApplePlatformSymbolCache* FApplePlatformSymbolication::CreateSymbolCache(void)
 void FApplePlatformSymbolication::DestroySymbolCache(FApplePlatformSymbolCache* Cache)
 {
 #if PLATFORM_MAC
-	if(Cache)
+	if(GAllowApplePlatformSymbolication && Cache)
 	{
 		for(auto It : Cache->Symbolicators)
 		{
@@ -93,31 +100,33 @@ bool FApplePlatformSymbolication::SymbolInfoForAddress(uint64 ProgramCounter, FP
 {
 #if PLATFORM_MAC
 	bool bOK = false;
-	
-	CSSymbolicatorRef Symbolicator = CSSymbolicatorCreateWithPid(FPlatformProcess::GetCurrentProcessId());
-	if(!CSIsNull(Symbolicator))
+	if (GAllowApplePlatformSymbolication)
 	{
-		CSSourceInfoRef Symbol = CSSymbolicatorGetSourceInfoWithAddressAtTime(Symbolicator, (vm_address_t)ProgramCounter, kCSNow);
-		
-		if(!CSIsNull(Symbol))
+		CSSymbolicatorRef Symbolicator = CSSymbolicatorCreateWithPid(FPlatformProcess::GetCurrentProcessId());
+		if(!CSIsNull(Symbolicator))
 		{
-			out_SymbolInfo.LineNumber = CSSourceInfoGetLineNumber(Symbol);
-			FCStringAnsi::Sprintf(out_SymbolInfo.Filename, CSSourceInfoGetPath(Symbol));
-			FCStringAnsi::Sprintf(out_SymbolInfo.FunctionName, CSSymbolGetName(CSSourceInfoGetSymbol(Symbol)));
-			CSRange CodeRange = CSSourceInfoGetRange(Symbol);
-			out_SymbolInfo.SymbolDisplacement = (ProgramCounter - CodeRange.Location);
+			CSSourceInfoRef Symbol = CSSymbolicatorGetSourceInfoWithAddressAtTime(Symbolicator, (vm_address_t)ProgramCounter, kCSNow);
 			
-			CSSymbolOwnerRef Owner = CSSourceInfoGetSymbolOwner(Symbol);
-			if(!CSIsNull(Owner))
+			if(!CSIsNull(Symbol))
 			{
-				ANSICHAR const* DylibName = CSSymbolOwnerGetName(Owner);
-				FCStringAnsi::Strcpy(out_SymbolInfo.ModuleName, DylibName);
+				out_SymbolInfo.LineNumber = CSSourceInfoGetLineNumber(Symbol);
+				FCStringAnsi::Sprintf(out_SymbolInfo.Filename, CSSourceInfoGetPath(Symbol));
+				FCStringAnsi::Sprintf(out_SymbolInfo.FunctionName, CSSymbolGetName(CSSourceInfoGetSymbol(Symbol)));
+				CSRange CodeRange = CSSourceInfoGetRange(Symbol);
+				out_SymbolInfo.SymbolDisplacement = (ProgramCounter - CodeRange.Location);
 				
-				bOK = out_SymbolInfo.LineNumber != 0;
+				CSSymbolOwnerRef Owner = CSSourceInfoGetSymbolOwner(Symbol);
+				if(!CSIsNull(Owner))
+				{
+					ANSICHAR const* DylibName = CSSymbolOwnerGetName(Owner);
+					FCStringAnsi::Strcpy(out_SymbolInfo.ModuleName, DylibName);
+					
+					bOK = out_SymbolInfo.LineNumber != 0;
+				}
 			}
+			
+			CSRelease(Symbolicator);
 		}
-		
-		CSRelease(Symbolicator);
 	}
 	
 	return bOK;
@@ -131,19 +140,22 @@ bool FApplePlatformSymbolication::SymbolInfoForFunctionFromModule(ANSICHAR const
 #if PLATFORM_MAC
 	bool bOK = false;
 	
-	CSSymbolicatorRef Symbolicator = CSSymbolicatorCreateWithPid(FPlatformProcess::GetCurrentProcessId());
-	if(!CSIsNull(Symbolicator))
+	if (GAllowApplePlatformSymbolication)
 	{
-		CSSymbolRef Symbol = CSSymbolicatorGetSymbolWithMangledNameFromSymbolOwnerWithNameAtTime(Symbolicator, MangledName, ModuleName, kCSNow);
-		
-		if(!CSIsNull(Symbol))
+		CSSymbolicatorRef Symbolicator = CSSymbolicatorCreateWithPid(FPlatformProcess::GetCurrentProcessId());
+		if(!CSIsNull(Symbolicator))
 		{
-			CSRange CodeRange = CSSymbolGetRange(Symbol);
-		
-			bOK = SymbolInfoForAddress(CodeRange.Location, Info);
+			CSSymbolRef Symbol = CSSymbolicatorGetSymbolWithMangledNameFromSymbolOwnerWithNameAtTime(Symbolicator, MangledName, ModuleName, kCSNow);
+			
+			if(!CSIsNull(Symbol))
+			{
+				CSRange CodeRange = CSSymbolGetRange(Symbol);
+			
+				bOK = SymbolInfoForAddress(CodeRange.Location, Info);
+			}
+			
+			CSRelease(Symbolicator);
 		}
-		
-		CSRelease(Symbolicator);
 	}
 	
 	return bOK;
@@ -157,7 +169,7 @@ bool FApplePlatformSymbolication::SymbolInfoForStrippedSymbol(FApplePlatformSymb
 #if PLATFORM_MAC
 	bool bOK = false;
 	
-	if(IFileManager::Get().FileSize(UTF8_TO_TCHAR(ModulePath)) > 0)
+	if(GAllowApplePlatformSymbolication && IFileManager::Get().FileSize(UTF8_TO_TCHAR(ModulePath)) > 0)
 	{
 		CSSymbolicatorRef Symbolicator = { nullptr, nullptr };
 		if ( Cache )
diff --git a/Engine/Source/Runtime/Core/Private/Mac/MacPlatformMisc.cpp b/Engine/Source/Runtime/Core/Private/Mac/MacPlatformMisc.cpp
--- a/Engine/Source/Runtime/Core/Private/Mac/MacPlatformMisc.cpp
+++ b/Engine/Source/Runtime/Core/Private/Mac/MacPlatformMisc.cpp
@@ -15,6 +15,7 @@
 #include "Runtime/Launch/Resources/Version.h"
 #include "EngineVersion.h"
 #include "MacMallocZone.h"
+#include "ApplePlatformSymbolication.h"
 
 #include <dlfcn.h>
 #include <IOKit/IOKitLib.h>
@@ -136,6 +137,8 @@ struct MacApplicationInfo
 		
 		AppPath = FPlatformProcess::GenerateApplicationPath(FApp::GetName(), FApp::GetBuildConfiguration());
 		
+		AppBundleID = FString([[NSBundle mainBundle] bundleIdentifier]);
+		
 		LCID = FString::Printf(TEXT("%d"), FInternationalization::Get().GetCurrentCulture()->GetLCID());
 		
 		PrimaryGPU = FPlatformMisc::GetPrimaryGPUBrand();
@@ -169,6 +172,8 @@ struct MacApplicationInfo
 		check(Status == NOTIFY_STATUS_OK);
 		
 		NumCores = FPlatformMisc::NumberOfCores();
+		
+		FPlatformMisc::CreateGuid(RunUUID);
 	}
 	
 	bool RunningOnBattery;
@@ -185,6 +190,7 @@ struct MacApplicationInfo
 	char MachineCPUString[PATH_MAX+1];
 	FString AppPath;
 	FString AppName;
+	FString AppBundleID;
 	FString OSVersion;
 	FString OSBuild;
 	FString MachineUUID;
@@ -197,6 +203,7 @@ struct MacApplicationInfo
 	FString CommandLine;
 	FString BranchBaseDir;
 	FString PrimaryGPU;
+	FGuid RunUUID;
 };
 static MacApplicationInfo GMacAppInfo;
 
@@ -1086,6 +1093,8 @@ static void DefaultCrashHandler(FMacCrashContext const& Context)
 /** True system-specific crash handler that gets called first */
 static void PlatformCrashHandler(int32 Signal, siginfo_t* Info, void* Context)
 {
+	FApplePlatformSymbolication::SetSymbolicationAllowed( false );
+	
 	FMacCrashContext CrashContext;
 	CrashContext.InitFromSignal(Signal, Info, Context);
 	
@@ -1372,7 +1381,7 @@ void FMacCrashContext::GenerateMinidump(char const* Path) const
 		static uint16 ByteOrderMarker = 0xFEFF;
 		write(ReportFile, &ByteOrderMarker, sizeof(ByteOrderMarker));
 		
-		WriteUTF16String(ReportFile, TEXT("Process:\t"));
+		WriteUTF16String(ReportFile, TEXT("Process:\tUE4-"));
 		WriteUTF16String(ReportFile, *GMacAppInfo.AppName);
 		WriteUTF16String(ReportFile, TEXT(" ["));
 		WriteUTF16String(ReportFile, ItoTCHAR(FPlatformProcess::GetCurrentProcessId(), 10));
@@ -1382,7 +1391,7 @@ void FMacCrashContext::GenerateMinidump(char const* Path) const
 		WriteLine(ReportFile, *GMacAppInfo.AppPath);
 		
 		WriteUTF16String(ReportFile, TEXT("Identifier:\t"));
-		WriteLine(ReportFile, *GMacAppInfo.AppName);
+		WriteLine(ReportFile, *GMacAppInfo.AppBundleID);
 		
 		WriteUTF16String(ReportFile, TEXT("Version:\t"));
 		WriteUTF16String(ReportFile, ENGINE_VERSION_STRINGIFY(ENGINE_MAJOR_VERSION) ENGINE_VERSION_TEXT(".") ENGINE_VERSION_STRINGIFY(ENGINE_MINOR_VERSION) ENGINE_VERSION_TEXT(".") ENGINE_VERSION_STRINGIFY(ENGINE_PATCH_VERSION));
@@ -1630,10 +1639,15 @@ void FMacCrashContext::GenerateCrashInfoAndLaunchReporter() const
 	// create a crash-specific directory
 	char CrashInfoFolder[PATH_MAX] = {};
 	FCStringAnsi::Strncpy(CrashInfoFolder, GMacAppInfo.CrashReportPath, PATH_MAX);
-	FCStringAnsi::Strcat(CrashInfoFolder, PATH_MAX, "/CrashReport-");
+	FCStringAnsi::Strcat(CrashInfoFolder, PATH_MAX, "/CrashReport-UE4-");
 	FCStringAnsi::Strcat(CrashInfoFolder, PATH_MAX, GMacAppInfo.AppNameUTF8);
 	FCStringAnsi::Strcat(CrashInfoFolder, PATH_MAX, "-pid-");
 	FCStringAnsi::Strcat(CrashInfoFolder, PATH_MAX, ItoANSI(getpid(), 10));
+	FCStringAnsi::Strcat(CrashInfoFolder, PATH_MAX, "-");
+	FCStringAnsi::Strcat(CrashInfoFolder, PATH_MAX, ItoANSI(GMacAppInfo.RunUUID.A, 16));
+	FCStringAnsi::Strcat(CrashInfoFolder, PATH_MAX, ItoANSI(GMacAppInfo.RunUUID.B, 16));
+	FCStringAnsi::Strcat(CrashInfoFolder, PATH_MAX, ItoANSI(GMacAppInfo.RunUUID.C, 16));
+	FCStringAnsi::Strcat(CrashInfoFolder, PATH_MAX, ItoANSI(GMacAppInfo.RunUUID.D, 16));
 
 	// Prevent CrashReportClient from spawning another CrashReportClient.
 	const TCHAR* ExecutableName = FPlatformProcess::ExecutableName();
@@ -1697,7 +1711,7 @@ void FMacCrashContext::GenerateCrashInfoAndLaunchReporter() const
 		// copy log
 		FCStringAnsi::Strncpy(FilePath, CrashInfoFolder, PATH_MAX);
 		FCStringAnsi::Strcat(FilePath, PATH_MAX, "/");
-		FCStringAnsi::Strcat(FilePath, PATH_MAX, GMacAppInfo.AppNameUTF8);
+		FCStringAnsi::Strcat(FilePath, PATH_MAX, (!GMacAppInfo.AppName.IsEmpty() ? GMacAppInfo.AppNameUTF8 : "UE4"));
 		FCStringAnsi::Strcat(FilePath, PATH_MAX, ".log");
 		int LogSrc = open(GMacAppInfo.AppLogPath, O_RDONLY);
 		int LogDst = open(FilePath, O_CREAT|O_WRONLY, 0766);
diff --git a/Engine/Source/Runtime/Core/Public/Apple/ApplePlatformSymbolication.h b/Engine/Source/Runtime/Core/Public/Apple/ApplePlatformSymbolication.h
--- a/Engine/Source/Runtime/Core/Public/Apple/ApplePlatformSymbolication.h
+++ b/Engine/Source/Runtime/Core/Public/Apple/ApplePlatformSymbolication.h
@@ -18,6 +18,7 @@ struct FApplePlatformSymbolCache;
  */
 struct CORE_API FApplePlatformSymbolication
 {
+	static void SetSymbolicationAllowed(bool const bAllow);
 	static bool SymbolInfoForAddress(uint64 ProgramCounter, FProgramCounterSymbolInfo& Info);
 	static bool SymbolInfoForFunctionFromModule(ANSICHAR const* MangledName, ANSICHAR const* ModuleName, FProgramCounterSymbolInfo& Info);
 	
