On OS X when using a separate game thread reset the game thread ID to the main thread after stopping the game thread so that static destructors don't assert.

diff --git a/Engine/Source/Runtime/Core/Private/Mac/CocoaThread.cpp b/Engine/Source/Runtime/Core/Private/Mac/CocoaThread.cpp
--- a/Engine/Source/Runtime/Core/Private/Mac/CocoaThread.cpp
+++ b/Engine/Source/Runtime/Core/Private/Mac/CocoaThread.cpp
@@ -360,6 +360,7 @@ FCocoaRunLoopSource* FCocoaRunLoopSource::GameRunLoopSource = nullptr;
 	if (GIsRequestingExit)
 	{
 		MainThreadCall(^{
+			GGameThreadId = FPlatformTLS::GetCurrentThreadId();
 			[NSApp replyToApplicationShouldTerminate:YES];
 		}, NSDefaultRunLoopMode, false);
 	}
