Add binaries directories to the PATH environment variable on Windows, so the loader can always find implicitly linked DLLs. Prevents having to mess around with module load order to link directly to plugin libraries.

diff --git a/Engine/Source/Runtime/Core/Private/Modules/ModuleManager.cpp b/Engine/Source/Runtime/Core/Private/Modules/ModuleManager.cpp
--- a/Engine/Source/Runtime/Core/Private/Modules/ModuleManager.cpp
+++ b/Engine/Source/Runtime/Core/Private/Modules/ModuleManager.cpp
@@ -950,6 +950,8 @@ void FModuleManager::AddBinariesDirectory(const TCHAR *InDirectory, bool bIsGame
 	{
 		EngineBinariesDirectories.Add(InDirectory);
 	}
+
+	FPlatformProcess::AddDllDirectory(InDirectory);
 }
 
 
diff --git a/Engine/Source/Runtime/Core/Private/Windows/WindowsPlatformProcess.cpp b/Engine/Source/Runtime/Core/Private/Windows/WindowsPlatformProcess.cpp
--- a/Engine/Source/Runtime/Core/Private/Windows/WindowsPlatformProcess.cpp
+++ b/Engine/Source/Runtime/Core/Private/Windows/WindowsPlatformProcess.cpp
@@ -31,6 +31,39 @@
 // static variables
 TArray<FString> FWindowsPlatformProcess::DllDirectoryStack;
 
+void FWindowsPlatformProcess::AddDllDirectory(const TCHAR* Directory)
+{
+	// Normalize the input directory
+	FString NormalizedDirectory = Directory;
+	FPaths::NormalizeDirectoryName(NormalizedDirectory);
+	FPaths::MakePlatformFilename(NormalizedDirectory);
+
+	// Get the current value of the PATH variable
+	TArray<TCHAR> PathVariable;
+	PathVariable.AddUninitialized(GetEnvironmentVariable(TEXT("PATH"), NULL, 0));
+	verify(::GetEnvironmentVariable(TEXT("PATH"), PathVariable.GetData(), PathVariable.Num()) == PathVariable.Num() - 1);
+
+	// Set the new path variable with the input directory at the start. Skip over any existing instances of the input directory.
+	FString NewPathVariable = NormalizedDirectory;
+	for(const TCHAR* PathPos = PathVariable.GetData(); PathPos < PathVariable.GetData() + PathVariable.Num(); )
+	{
+		// Scan to the end of this directory
+		const TCHAR* PathEnd = PathPos;
+		while(*PathEnd != ';' && *PathEnd != 0) PathEnd++;
+
+		// Add it to the new path variable if it doesn't match the input directory
+		if(PathEnd - PathPos != NormalizedDirectory.Len() || FCString::Strnicmp(*NormalizedDirectory, PathPos, PathEnd - PathPos) != 0)
+		{
+			NewPathVariable.AppendChar(TEXT(';'));
+			NewPathVariable.AppendChars(PathPos, PathEnd - PathPos);
+		}
+
+		// Move to the next string
+		PathPos = PathEnd + 1;
+	}
+	SetEnvironmentVariable(TEXT("PATH"), *NewPathVariable);
+}
+
 void* FWindowsPlatformProcess::GetDllHandle( const TCHAR* Filename )
 {
 	check(Filename);
diff --git a/Engine/Source/Runtime/Core/Public/GenericPlatform/GenericPlatformProcess.h b/Engine/Source/Runtime/Core/Public/GenericPlatform/GenericPlatformProcess.h
--- a/Engine/Source/Runtime/Core/Public/GenericPlatform/GenericPlatformProcess.h
+++ b/Engine/Source/Runtime/Core/Public/GenericPlatform/GenericPlatformProcess.h
@@ -156,6 +156,12 @@ struct CORE_API FGenericPlatformProcess
 	/** Gets the API version from the specified DLL **/
 	static int32 GetDllApiVersion( const TCHAR* Filename );
 
+	/** Adds a directory to search in when resolving implicitly loaded or filename-only DLLs. **/
+	FORCEINLINE static void AddDllDirectory(const TCHAR* Directory)
+	{
+
+	}
+
 	/** Set a directory to look for DLL files. NEEDS to have a Pop call when complete */
 	FORCEINLINE static void PushDllDirectory(const TCHAR* Directory)
 	{
diff --git a/Engine/Source/Runtime/Core/Public/Windows/WindowsPlatformProcess.h b/Engine/Source/Runtime/Core/Public/Windows/WindowsPlatformProcess.h
--- a/Engine/Source/Runtime/Core/Public/Windows/WindowsPlatformProcess.h
+++ b/Engine/Source/Runtime/Core/Public/Windows/WindowsPlatformProcess.h
@@ -75,6 +75,7 @@ struct CORE_API FWindowsPlatformProcess : public FGenericPlatformProcess
 	static void FreeDllHandle( void* DllHandle );
 	static void* GetDllExport( void* DllHandle, const TCHAR* ProcName );
 	static int32 GetDllApiVersion( const TCHAR* Filename );
+	static void AddDllDirectory(const TCHAR* Directory);
 	static void PushDllDirectory(const TCHAR* Directory);
 	static void PopDllDirectory(const TCHAR* Directory);
 	static void CleanFileCache();
