Forgot file

diff --git a/Engine/Source/Runtime/Core/Public/UObject/ObjectVersion.h b/Engine/Source/Runtime/Core/Public/UObject/ObjectVersion.h
--- a/Engine/Source/Runtime/Core/Public/UObject/ObjectVersion.h
+++ b/Engine/Source/Runtime/Core/Public/UObject/ObjectVersion.h
@@ -622,6 +622,8 @@ enum EUnrealEngineObjectUE4Version
 	VER_UE4_FIX_REFRACTION_INPUT_MASKING,
 	// A global spawn rate for emitters.
 	VER_UE4_GLOBAL_EMITTER_SPAWN_RATE_SCALE,
+	// Cleanup destructible mesh settings
+	VER_UE4_CLEAN_DESTRUCTIBLE_SETTINGS,
 
 	// -----<new versions can be added before this line>-------------------------------------------------
 	// - this needs to be the last line (see note below)
