Send Mac windowWillClose notification to game thread synchronously. Hopefully this will solve a crash that occurs when the game thread is trying to do something with FCocoaWindow whilst the main thread is in a process of releasing it.

diff --git a/Engine/Source/Runtime/Core/Private/Mac/CocoaWindow.cpp b/Engine/Source/Runtime/Core/Private/Mac/CocoaWindow.cpp
--- a/Engine/Source/Runtime/Core/Private/Mac/CocoaWindow.cpp
+++ b/Engine/Source/Runtime/Core/Private/Mac/CocoaWindow.cpp
@@ -515,7 +515,7 @@ NSString* NSPerformDragOperation = @"NSPerformDragOperation";
 	SCOPED_AUTORELEASE_POOL;
 	if(self.bForwardEvents && MacApplication)
 	{
-		FMacEvent::SendToGameRunLoop(notification, self, EMacEventSendMethod::Async, @[ NSDefaultRunLoopMode, UE4CloseEventMode ]);
+		FMacEvent::SendToGameRunLoop(notification, self, EMacEventSendMethod::Sync, @[ NSDefaultRunLoopMode, UE4CloseEventMode ]);
 	}
 	self.bForwardEvents = false;
 	[self setDelegate:nil];
