Core: Added LoadModuleChecked overload to module manager

diff --git a/Engine/Source/Runtime/Core/Private/Modules/ModuleManager.cpp b/Engine/Source/Runtime/Core/Private/Modules/ModuleManager.cpp
--- a/Engine/Source/Runtime/Core/Private/Modules/ModuleManager.cpp
+++ b/Engine/Source/Runtime/Core/Private/Modules/ModuleManager.cpp
@@ -6,6 +6,7 @@
 #include "EngineBuildSettings.h"
 #include "UProjectInfo.h"
 
+
 DEFINE_LOG_CATEGORY_STATIC(LogModuleManager, Log, All);
 
 #if WITH_HOT_RELOAD
@@ -73,13 +74,7 @@ void FModuleManager::FindModules(const TCHAR* WildcardWithoutExtension, TArray<F
 #endif
 }
 
-/**
- * Checks to see if the specified module is currently loaded.  This is an O(1) operation.
- *
- * @param	InModuleName		The base name of the module file.  Should not include path, extension or platform/configuration info.  This is just the "module name" part of the module file name.  Names should be globally unique.
- *
- * @return	True if module is currently loaded, otherwise false
- */
+
 bool FModuleManager::IsModuleLoaded( const FName InModuleName ) const
 {
 	// Do we even know about this module?
@@ -100,6 +95,7 @@ bool FModuleManager::IsModuleLoaded( const FName InModuleName ) const
 	return false;
 }
 
+
 bool FModuleManager::IsModuleUpToDate( const FName InModuleName ) const
 {
 	TMap<FName, FString> ModulePathMap;
@@ -113,6 +109,7 @@ bool FModuleManager::IsModuleUpToDate( const FName InModuleName ) const
 	return CheckModuleCompatibility(*TMap<FName, FString>::TConstIterator(ModulePathMap).Value());
 }
 
+
 void FModuleManager::AddModule( const FName InModuleName )
 {
 	// Do we already know about this module?  If not, we'll create information for this module now.
@@ -234,17 +231,23 @@ void FModuleManager::AddModule( const FName InModuleName )
 	}
 }
 
-/**
- * Adds a module to our list of modules and tries to load it immediately
- *
- * @param	InModuleName	The base name of the module file.  Should not include path, extension or platform/configuration info.  This is just the "name" part of the module file name.  Names should be globally unique.
- */
-TSharedPtr<IModuleInterface> FModuleManager::LoadModule( const FName InModuleName, bool bWasReloaded /*=false*/ )
+
+TSharedPtr<IModuleInterface> FModuleManager::LoadModule( const FName InModuleName, bool bWasReloaded )
 {
 	EModuleLoadResult FailureReason;
 	return LoadModuleWithFailureReason(InModuleName, FailureReason, bWasReloaded );
 }
 
+
+TSharedPtr<IModuleInterface> FModuleManager::LoadModuleChecked( const FName InModuleName, const bool bWasReloaded )
+{
+	TSharedPtr<IModuleInterface> Module = LoadModule(InModuleName, bWasReloaded);
+	check(Module.IsValid());
+
+	return Module;
+}
+
+
 TSharedPtr<IModuleInterface> FModuleManager::LoadModuleWithFailureReason( const FName InModuleName, EModuleLoadResult& OutFailureReason, bool bWasReloaded /*=false*/ )
 {
 	DECLARE_SCOPE_CYCLE_COUNTER(TEXT("Module Load"), STAT_ModuleLoad, STATGROUP_LoadTime);
@@ -475,11 +478,6 @@ bool FModuleManager::UnloadModule( const FName InModuleName, bool bIsShutdown )
 }
 
 
-/**
- * Abandons a loaded module, leaving it loaded in memory but no longer tracking it in the module manager
- *
- * @param	InModuleName	The name of the module to abandon.  Should not include path, extension or platform/configuration info.  This is just the "module name" part of the module file name.
- */
 void FModuleManager::AbandonModule( const FName InModuleName )
 {
 	// Do we even know about this module?
@@ -739,16 +737,19 @@ void FModuleManager::QueryModules( TArray< FModuleStatus >& OutModuleStatuses )
 	}
 }
 
+
 FString FModuleManager::GetModuleFilename(FName ModuleName) const
 {
 	return Modules.FindChecked(ModuleName)->Filename;
 }
 
+
 void FModuleManager::SetModuleFilename(FName ModuleName, const FString& Filename)
 {
 	Modules.FindChecked(ModuleName)->Filename = Filename;
 }
 
+
 FString FModuleManager::GetCleanModuleFilename(FName ModuleName, bool bGameModule)
 {
 	FString Prefix, Suffix;
@@ -756,6 +757,7 @@ FString FModuleManager::GetCleanModuleFilename(FName ModuleName, bool bGameModul
 	return Prefix + ModuleName.ToString() + Suffix;
 }
 
+
 void FModuleManager::GetModuleFilenameFormat(bool bGameModule, FString& OutPrefix, FString& OutSuffix)
 {
 	// Get the module configuration for this directory type
@@ -805,6 +807,7 @@ void FModuleManager::GetModuleFilenameFormat(bool bGameModule, FString& OutPrefi
 	OutSuffix += FPlatformProcess::GetModuleExtension();
 }
 
+
 void FModuleManager::FindModulePaths(const TCHAR* NamePattern, TMap<FName, FString> &OutModulePaths) const
 {
 	// Search through the engine directory
@@ -823,6 +826,7 @@ void FModuleManager::FindModulePaths(const TCHAR* NamePattern, TMap<FName, FStri
 	}
 }
 
+
 void FModuleManager::FindModulePathsInDirectory(const FString& InDirectoryName, bool bIsGameDirectory, const TCHAR* NamePattern, TMap<FName, FString> &OutModulePaths) const
 {
 	// Get the prefix and suffix for module filenames
@@ -850,6 +854,7 @@ void FModuleManager::FindModulePathsInDirectory(const FString& InDirectoryName,
 	}
 }
 
+
 void FModuleManager::UnloadOrAbandonModuleWithCallback( const FName InModuleName, FOutputDevice &Ar )
 {
 	TSharedRef< FModuleInfo > Module = Modules.FindChecked( InModuleName );
@@ -892,15 +897,17 @@ bool FModuleManager::LoadModuleWithCallback( const FName InModuleName, FOutputDe
 
 void FModuleManager::MakeUniqueModuleFilename( const FName InModuleName, FString& UniqueSuffix, FString& UniqueModuleFileName )
 {
-	TSharedRef< FModuleInfo > Module = Modules.FindChecked( InModuleName );
+	TSharedRef<FModuleInfo> Module = Modules.FindChecked(InModuleName);
+
 	do
 	{
 		// Use a random number as the unique file suffix, but mod it to keep it of reasonable length
 		UniqueSuffix = FString::FromInt( FMath::Rand() % 10000 );
 
 		const FString ModuleName = *InModuleName.ToString();
-		const int32 MatchPos = Module->OriginalFilename.Find( ModuleName, ESearchCase::IgnoreCase, ESearchDir::FromEnd);
-		if( ensure( MatchPos != INDEX_NONE ) )
+		const int32 MatchPos = Module->OriginalFilename.Find(ModuleName, ESearchCase::IgnoreCase, ESearchDir::FromEnd);
+
+		if (ensure(MatchPos != INDEX_NONE))
 		{
 			const int32 SuffixPos = MatchPos + ModuleName.Len();
 			UniqueModuleFileName = FString::Printf( TEXT( "%s-%s%s" ),
@@ -922,16 +929,17 @@ const TCHAR *FModuleManager::GetUBTConfiguration()
 bool FModuleManager::CheckModuleCompatibility(const TCHAR* Filename)
 {
 	int32 ModuleApiVersion = FPlatformProcess::GetDllApiVersion(Filename);
-	if(ModuleApiVersion != MODULE_API_VERSION)
+
+	if (ModuleApiVersion != MODULE_API_VERSION)
 	{
 		UE_LOG(LogModuleManager, Warning, TEXT("Found module file %s (API version %d), but it was incompatible with the current engine API version (%d). This is likely a stale module that must be recompiled."), Filename, ModuleApiVersion, MODULE_API_VERSION);
 		return false;
 	}
+
 	return true;
 }
 
 
-
 void FModuleManager::StartProcessingNewlyLoadedObjects()
 {
 	// Only supposed to be called once
@@ -971,9 +979,9 @@ void FModuleManager::SetGameBinariesDirectory(const TCHAR* InDirectory)
 
 bool FModuleManager::DoesLoadedModuleHaveUObjects( const FName ModuleName )
 {
-	if( IsModuleLoaded( ModuleName ) && IsPackageLoaded.IsBound() )
+	if (IsModuleLoaded(ModuleName) && IsPackageLoaded.IsBound())
 	{
-		return IsPackageLoaded.Execute( *FString( FString( TEXT("/Script/") ) + ModuleName.ToString() ) );
+		return IsPackageLoaded.Execute(*FString(FString(TEXT("/Script/")) + ModuleName.ToString()));
 	}
 
 	return false;
diff --git a/Engine/Source/Runtime/Core/Public/Modules/ModuleManager.h b/Engine/Source/Runtime/Core/Public/Modules/ModuleManager.h
--- a/Engine/Source/Runtime/Core/Public/Modules/ModuleManager.h
+++ b/Engine/Source/Runtime/Core/Public/Modules/ModuleManager.h
@@ -150,6 +150,16 @@ public:
 	 */
 	TSharedPtr<IModuleInterface> LoadModule( const FName InModuleName, const bool bWasReloaded = false );
 
+	/**
+	 * Loads the specified module, checking to ensure it exists.
+	 *
+	 * @param InModuleName The base name of the module file.  Should not include path, extension or platform/configuration info.  This is just the "module name" part of the module file name.  Names should be globally unique.
+	 * @param bWasReloaded Indicates that the module has been reloaded (default = false).
+	 * @return The loaded module, or nullptr if the load operation failed.
+	 * @see AbandonModule, IsModuleLoaded, LoadModuleChecked, LoadModulePtr, LoadModuleWithFailureReason, UnloadModule
+	 */
+	TSharedPtr<IModuleInterface> LoadModuleChecked( const FName InModuleName, const bool bWasReloaded = false );
+
 	/**
 	 * Loads a module in memory then calls PostLoad.
 	 *
