Renamed some variables and clarified some comments in text localization live update logic. Less confusing.

diff --git a/Engine/Source/Runtime/Core/Private/Internationalization/TextLocalizationManager.cpp b/Engine/Source/Runtime/Core/Private/Internationalization/TextLocalizationManager.cpp
--- a/Engine/Source/Runtime/Core/Private/Internationalization/TextLocalizationManager.cpp
+++ b/Engine/Source/Runtime/Core/Private/Internationalization/TextLocalizationManager.cpp
@@ -401,7 +401,7 @@ void FTextLocalizationManager::LoadResources(const bool ShouldLoadEditor, const
 
 void FTextLocalizationManager::UpdateLiveTable(const TArray<FLocalizationEntryTracker>& LocalizationEntryTrackers, const bool FilterUpdatesByTableName)
 {
-	// Update existing localized entries/flag existing unlocalized entries.
+	// Update existing localized entries/flag existing newly-unlocalized entries.
 	for(auto NamespaceIterator = LiveTable.NamespaceTable.CreateIterator(); NamespaceIterator; ++NamespaceIterator)
 	{
 		const FString& NamespaceName = NamespaceIterator.Key();
@@ -411,26 +411,26 @@ void FTextLocalizationManager::UpdateLiveTable(const TArray<FLocalizationEntryTr
 			const FString& KeyName = KeyIterator.Key();
 			FStringEntry& LiveStringEntry = KeyIterator.Value();
 
-			const FLocalizationEntryTracker::FEntry* NewEntry = NULL;
+			const FLocalizationEntryTracker::FEntry* UpdateEntry = NULL;
 
 			// Attempt to use resources in prioritized order until we find an entry.
-			for(int32 i = 0; i < LocalizationEntryTrackers.Num() && !NewEntry; ++i)
+			for(int32 i = 0; i < LocalizationEntryTrackers.Num() && !UpdateEntry; ++i)
 			{
 				const FLocalizationEntryTracker& Tracker = LocalizationEntryTrackers[i];
-				const FLocalizationEntryTracker::FKeyTable* const NewKeyTable = Tracker.Namespaces.Find(NamespaceName);
-				const FLocalizationEntryTracker::FEntryArray* const NewEntryArray = NewKeyTable ? NewKeyTable->Find(KeyName) : NULL;
-				const FLocalizationEntryTracker::FEntry* Entry = NewEntryArray && NewEntryArray->Num() ? &((*NewEntryArray)[0]) : NULL;
-				NewEntry = (Entry && (!FilterUpdatesByTableName || LiveStringEntry.TableName == Entry->TableName)) ? Entry : NULL;
+				const FLocalizationEntryTracker::FKeyTable* const UpdateKeyTable = Tracker.Namespaces.Find(NamespaceName);
+				const FLocalizationEntryTracker::FEntryArray* const UpdateEntryArray = UpdateKeyTable ? UpdateKeyTable->Find(KeyName) : NULL;
+				const FLocalizationEntryTracker::FEntry* Entry = UpdateEntryArray && UpdateEntryArray->Num() ? &((*UpdateEntryArray)[0]) : NULL;
+				UpdateEntry = (Entry && (!FilterUpdatesByTableName || LiveStringEntry.TableName == Entry->TableName)) ? Entry : NULL;
 			}
 
-			if( NewEntry )
+			if( UpdateEntry )
 			{
 				// If an entry is unlocalized and the source hash differs, it suggests that the source hash changed - do not replace the display string.
-				if(LiveStringEntry.bIsLocalized || LiveStringEntry.SourceStringHash == NewEntry->SourceStringHash)
+				if(LiveStringEntry.bIsLocalized || LiveStringEntry.SourceStringHash == UpdateEntry->SourceStringHash)
 				{
 					LiveStringEntry.bIsLocalized = true;
-					*(LiveStringEntry.String) = NewEntry->LocalizedString;
-					LiveStringEntry.SourceStringHash = NewEntry->SourceStringHash;
+					*(LiveStringEntry.String) = UpdateEntry->LocalizedString;
+					LiveStringEntry.SourceStringHash = UpdateEntry->SourceStringHash;
 				}
 			}
 			else if(!FilterUpdatesByTableName)
