FString comparison with an ANSI string literal constructs a temporary FString

#ttp 322488

diff --git a/Engine/Source/Runtime/Core/Private/UObject/UnrealNames.cpp b/Engine/Source/Runtime/Core/Private/UObject/UnrealNames.cpp
--- a/Engine/Source/Runtime/Core/Private/UObject/UnrealNames.cpp
+++ b/Engine/Source/Runtime/Core/Private/UObject/UnrealNames.cpp
@@ -406,42 +406,6 @@ FName::FName( EName HardcodedIndex, const TCHAR* Name )
 	Init(Name, NAME_NO_NUMBER_INTERNAL, FNAME_Add, false, HardcodedIndex);
 }
 
-/**
- * Comparision operator.
- *
- * @param	Other	String to compare this name to
- * @return true if name matches the string, false otherwise
- */
-bool FName::operator==( const TCHAR* Other ) const
-{
-	// Find name entry associated with this FName.
-	check( Other );
-	const FNameEntry* const Entry = GetComparisonNameEntry();
-
-	// Temporary buffer to hold split name in case passed in name is of Name_Number format.
-	WIDECHAR TempBuffer[NAME_SIZE];
-	int32 InNumber	= NAME_NO_NUMBER_INTERNAL;
-	int32 TempNumber	= NAME_NO_NUMBER_INTERNAL;
-
-	// Check whether we need to split the passed in string into name and number portion.
-	auto           WideOther    = StringCast<WIDECHAR>(Other);
-	const WIDECHAR* WideOtherPtr = WideOther.Get();
-	if( SplitNameWithCheck( WideOtherPtr, TempBuffer, ARRAY_COUNT(TempBuffer), TempNumber) )
-	{
-		WideOtherPtr	= TempBuffer;
-		InNumber	= NAME_EXTERNAL_TO_INTERNAL(TempNumber);
-	}
-
-	// Report a match if both the number and string portion match.
-	bool bAreNamesMatching = false;
-	if( InNumber == GetNumber() && !FCStringWide::Stricmp( WideOtherPtr, Entry->IsWide() ? Entry->GetWideName() : StringCast<WIDECHAR>(Entry->GetAnsiName()).Get() ) )
-	{
-		bAreNamesMatching = true;
-	}
-
-	return bAreNamesMatching;
-}
-
 /**
  * Compares name to passed in one. Sort is alphabetical ascending.
  *
diff --git a/Engine/Source/Runtime/Core/Public/Android/AndroidString.h b/Engine/Source/Runtime/Core/Public/Android/AndroidString.h
--- a/Engine/Source/Runtime/Core/Public/Android/AndroidString.h
+++ b/Engine/Source/Runtime/Core/Public/Android/AndroidString.h
@@ -12,7 +12,7 @@
 #pragma once
 
 #include "Char.h"
-
+#include "GenericPlatformStricmp.h"
 
 /**
  * Android string implementation
@@ -29,19 +29,20 @@ struct FAndroidPlatformString : public FGenericPlatformString
 		return Dest;
 	}
 
-	template <typename CharType>
-	static inline int32 Stricmp( const CharType* String1, const CharType* String2 )
-	{
-		// walk the strings, comparing them case insensitively
-		for (; *String1 || *String2; String1++, String2++)
-		{
-			CharType Char1 = TChar<CharType>::ToUpper(*String1), Char2 = TChar<CharType>::ToUpper(*String2);
-			if (Char1 != Char2)
-			{
-				return Char1 - Char2;
-			}
-		}
-		return 0;
+	/**
+	 * Compares two strings case-insensitive.
+	 *
+	 * @param String1 First string to compare.
+	 * @param String2 Second string to compare.
+	 *
+	 * @returns Zero if both strings are equal. Greater than zero if first
+	 *          string is greater than the second one. Less than zero
+	 *          otherwise.
+	 */
+	template <typename CharType1, typename CharType2>
+	static inline int32 Stricmp(const CharType1* String1, const CharType2* String2)
+	{
+		return FGenericPlatformStricmp::Stricmp(String1, String2);
 	}
 
 	template <typename CharType>
diff --git a/Engine/Source/Runtime/Core/Public/Containers/UnrealString.h b/Engine/Source/Runtime/Core/Public/Containers/UnrealString.h
--- a/Engine/Source/Runtime/Core/Public/Containers/UnrealString.h
+++ b/Engine/Source/Runtime/Core/Public/Containers/UnrealString.h
@@ -770,25 +770,27 @@ public:
 	/**
 	 * Lexicographically test whether this string is <= the Other given string
 	 * 
-	 * @param Other array of TCHAR to compare against
+	 * @param Other array of CharType to compare against
 	 * @return true if length of this string is lexicographically <= the other, otherwise false
 	 * @note case insensitive
 	 */
-	FORCEINLINE bool operator<=( const TCHAR* Other ) const
+	template <typename CharType>
+	FORCEINLINE bool operator<=(const CharType* Other) const
 	{
-		return !(FCString::Stricmp( **this, Other ) > 0);
+		return !(FPlatformString::Stricmp(**this, Other) > 0);
 	}
 
 	/**
 	 * Lexicographically test whether this string is < the Other given string
 	 * 
-	 * @param Other array of TCHAR to compare against
+	 * @param Other array of CharType to compare against
 	 * @return true if length of this string is lexicographically < the other, otherwise false
 	 * @note case insensitive
 	 */
-	FORCEINLINE bool operator<( const TCHAR* Other ) const
+	template <typename CharType>
+	FORCEINLINE bool operator<(const CharType* Other) const
 	{
-		return FCString::Stricmp( **this, Other ) < 0;
+		return FPlatformString::Stricmp(**this, Other) < 0;
 	}
 
 	/**
@@ -806,25 +808,27 @@ public:
 	/**
 	 * Lexicographically test whether this string is >= the Other given string
 	 * 
-	 * @param Other array of TCHAR to compare against
+	 * @param Other array of CharType to compare against
 	 * @return true if length of this string is lexicographically >= the other, otherwise false
 	 * @note case insensitive
 	 */
-	FORCEINLINE bool operator>=( const TCHAR* Other ) const
+	template <typename CharType>
+	FORCEINLINE bool operator>=(const CharType* Other) const
 	{
-		return !(FCString::Stricmp( **this, Other ) < 0);
+		return !(FPlatformString::Stricmp(**this, Other) < 0);
 	}
 
 	/**
 	 * Lexicographically test whether this string is > the Other given string
 	 * 
-	 * @param Other array of TCHAR to compare against
+	 * @param Other array of CharType to compare against
 	 * @return true if length of this string is lexicographically > the other, otherwise false
 	 * @note case insensitive
 	 */
-	FORCEINLINE bool operator>( const TCHAR* Other ) const
+	template <typename CharType>
+	FORCEINLINE bool operator>(const CharType* Other) const
 	{
-		return FCString::Stricmp( **this, Other ) > 0;
+		return FPlatformString::Stricmp(**this, Other) > 0;
 	}
 
 	/**
@@ -842,13 +846,14 @@ public:
 	/**
 	 * Lexicographically test whether this string is equivalent to the Other given string
 	 * 
-	 * @param Other array of TCHAR to compare against
+	 * @param Other array of CharType to compare against
 	 * @return true if length of this string is lexicographically equivalent to the other, otherwise false
 	 * @note case insensitive
 	 */
-	FORCEINLINE bool operator==( const TCHAR* Other ) const
+	template <typename CharType>
+	FORCEINLINE bool operator==(const CharType* Other) const
 	{
-		return FCString::Stricmp( **this, Other )==0;
+		return FPlatformString::Stricmp(**this, Other) == 0;
 	}
 
 	/**
@@ -866,13 +871,14 @@ public:
 	/**
 	 * Lexicographically test whether this string is not equivalent to the Other given string
 	 * 
-	 * @param Other array of TCHAR to compare against
+	 * @param Other array of CharType to compare against
 	 * @return true if length of this string is lexicographically not equivalent to the other, otherwise false
 	 * @note case insensitive
 	 */
-	FORCEINLINE bool operator!=( const TCHAR* Other ) const
+	template <typename CharType>
+	FORCEINLINE bool operator!=(const CharType* Other) const
 	{
-		return FCString::Stricmp( **this, Other )!=0;
+		return FPlatformString::Stricmp(**this, Other) != 0;
 	}
 
 	/**
diff --git a/Engine/Source/Runtime/Core/Public/GenericPlatform/GenericPlatformStricmp.h b/Engine/Source/Runtime/Core/Public/GenericPlatform/GenericPlatformStricmp.h
new file mode 100644
--- /dev/null
+++ b/Engine/Source/Runtime/Core/Public/GenericPlatform/GenericPlatformStricmp.h
@@ -0,0 +1,102 @@
+// Copyright 1998-2014 Epic Games, Inc. All Rights Reserved.
+
+#pragma once
+
+#include "Char.h"
+#include "UnrealTypeTraits.h"
+
+/**
+ * This trait tells if given CharTypeA is comparison compatible with CharTypeB, i.e.
+ * if CharTypeA contains whole CharTypeB character set on the same positions.
+ *
+ * This is true e.g. for WIDECHAR and ANSICHAR because first 256 characters of WIDECHAR is
+ * basically ANSICHAR table.
+ */
+template<typename A, typename B>	struct TIsComparisonCompatibleChar							{ enum { Value = false }; };
+template<>							struct TIsComparisonCompatibleChar<WIDECHAR, ANSICHAR>		{ enum { Value = true }; };
+template<>							struct TIsComparisonCompatibleChar<UTF8CHAR, ANSICHAR>		{ enum { Value = true }; };
+template<>							struct TIsComparisonCompatibleChar<UTF16CHAR, ANSICHAR>		{ enum { Value = true }; };
+template<>							struct TIsComparisonCompatibleChar<UTF32CHAR, ANSICHAR>		{ enum { Value = true }; };
+
+/**
+ * Static struct that implements generic stricmp functionality.
+ */
+struct FGenericPlatformStricmp
+{
+private:
+	/**
+	 * Compares two strings case-insensitive assuming that both char types are
+	 * compatible.
+	 *
+	 * @param String1 First string to compare.
+	 * @param String2 Second string to compare.
+	 *
+	 * @returns Zero if both strings are equal. Greater than zero if first
+	 *          string is greater than the second one. Less than zero
+	 *          otherwise.
+	 */
+	template <typename CompatibleCharType1, typename CompatibleCharType2>
+	static inline int32 CompatibleCharTypesStricmp(const CompatibleCharType1* String1, const CompatibleCharType2* String2)
+	{
+		// Walk the strings, comparing them case insensitively.
+		for (; *String1 || *String2; String1++, String2++)
+		{
+			CompatibleCharType1 Char1 = TChar<CompatibleCharType1>::ToLower(*String1);
+			CompatibleCharType2 Char2 = TChar<CompatibleCharType2>::ToLower(*String2);
+
+			if (Char1 != Char2)
+			{
+				return Char1 - Char2;
+			}
+		}
+		return 0;
+	}
+
+	/**
+	 * Helper function wrapper to verify compatibility. This overload checks if CharType1 == CharType2.
+	 */
+	template <typename CharType1, typename CharType2>
+	static inline int32 StricmpImpl(const CharType1* String1, const CharType2* String2, typename TEnableIf<TIsSame<CharType1, CharType2>::Value>::Type* Dummy1 = nullptr)
+	{
+		return CompatibleCharTypesStricmp(String1, String2);
+	}
+
+	/**
+	 * Helper function wrapper to verify compatibility. This overload checks if
+	 * CharType1 is compatible with CharType2.
+	 */
+	template <typename CharType1, typename CharType2>
+	static inline int32 StricmpImpl(const CharType1* String1, const CharType2* String2, typename TEnableIf<TIsComparisonCompatibleChar<CharType1, CharType2>::Value>::Type* Dummy = nullptr)
+	{
+		return CompatibleCharTypesStricmp(String1, String2);
+	}
+
+	/**
+	 * Helper function wrapper to verify compatibility. This overload checks if
+	 * CharType2 is compatible with CharType1. This is the mirrored version of
+	 * above function, so we don't have to declare compatibility trait for both
+	 * ways.
+	 */
+	template <typename CharType1, typename CharType2>
+	static inline int32 StricmpImpl(const CharType1* String1, const CharType2* String2, typename TEnableIf<TIsComparisonCompatibleChar<CharType2, CharType1>::Value>::Type* Dummy = nullptr)
+	{
+		return CompatibleCharTypesStricmp(String1, String2);
+	}
+
+public:
+	/**
+	 * Compares two strings case-insensitive.
+	 *
+	 * @param String1 First string to compare.
+	 * @param String2 Second string to compare.
+	 *
+	 * @returns Zero if both strings are equal. Greater than zero if first
+	 *          string is greater than the second one. Less than zero
+	 *          otherwise.
+	 */
+	template <typename CharType1, typename CharType2>
+	static inline int32 Stricmp(const CharType1* String1, const CharType2* String2)
+	{
+		return FGenericPlatformStricmp::StricmpImpl(String1, String2);
+	}
+};
\ No newline at end of file
diff --git a/Engine/Source/Runtime/Core/Public/GenericPlatform/MicrosoftPlatformString.h b/Engine/Source/Runtime/Core/Public/GenericPlatform/MicrosoftPlatformString.h
--- a/Engine/Source/Runtime/Core/Public/GenericPlatform/MicrosoftPlatformString.h
+++ b/Engine/Source/Runtime/Core/Public/GenericPlatform/MicrosoftPlatformString.h
@@ -3,6 +3,7 @@
 #pragma once
 
 #include "Char.h"
+#include "GenericPlatformStricmp.h"
 
 /**
 * Microsoft specific implementation 
@@ -68,20 +69,6 @@ struct FMicrosoftPlatformString : public FGenericPlatformString
 		return _tcsncmp( String1, String2, Count );
 	}
 
-	static FORCEINLINE int32 Stricmp( const WIDECHAR* String1, const WIDECHAR* String2 )
-	{
-		// walk the strings, comparing them case insensitively
-		for (; *String1 || *String2; String1++, String2++)
-		{
-			WIDECHAR Char1 = TChar<WIDECHAR>::ToLower(*String1), Char2 = TChar<WIDECHAR>::ToLower(*String2);
-			if (Char1 != Char2)
-			{
-				return Char1 - Char2;
-			}
-		}
-		return 0;
-	}
-
 	static FORCEINLINE int32 Strnicmp( const WIDECHAR* String1, const WIDECHAR* String2, SIZE_T Count )
 	{
 		return _tcsnicmp( String1, String2, Count );
@@ -206,9 +193,37 @@ struct FMicrosoftPlatformString : public FGenericPlatformString
 		return strncmp( String1, String2, Count );
 	}
 
-	static FORCEINLINE int32 Stricmp( const ANSICHAR* String1, const ANSICHAR* String2 )
+	/**
+	 * Compares two strings case-insensitive.
+	 *
+	 * Specialized version for ANSICHAR types.
+	 *
+	 * @param String1 First string to compare.
+	 * @param String2 Second string to compare.
+	 *
+	 * @returns Zero if both strings are equal. Greater than zero if first
+	 *          string is greater than the second one. Less than zero
+	 *          otherwise.
+	 */
+	static FORCEINLINE int32 Stricmp(const ANSICHAR* String1, const ANSICHAR* String2)
+	{
+		return _stricmp(String1, String2);
+	}
+
+	/**
+	 * Compares two strings case-insensitive.
+	 *
+	 * @param String1 First string to compare.
+	 * @param String2 Second string to compare.
+	 *
+	 * @returns Zero if both strings are equal. Greater than zero if first
+	 *          string is greater than the second one. Less than zero
+	 *          otherwise.
+	 */
+	template <typename CharType1, typename CharType2>
+	static FORCEINLINE int32 Stricmp(const CharType1* String1, const CharType2* String2)
 	{
-		return _stricmp( String1, String2 );
+		return FGenericPlatformStricmp::Stricmp(String1, String2);
 	}
 
 	static FORCEINLINE int32 Strnicmp( const ANSICHAR* String1, const ANSICHAR* String2, SIZE_T Count )
diff --git a/Engine/Source/Runtime/Core/Public/GenericPlatform/StandardPlatformString.h b/Engine/Source/Runtime/Core/Public/GenericPlatform/StandardPlatformString.h
--- a/Engine/Source/Runtime/Core/Public/GenericPlatform/StandardPlatformString.h
+++ b/Engine/Source/Runtime/Core/Public/GenericPlatform/StandardPlatformString.h
@@ -3,6 +3,7 @@
 #pragma once
 
 #include "Char.h"
+#include "GenericPlatformStricmp.h"
 
 /**
 * Standard implementation
@@ -19,19 +20,21 @@ struct FStandardPlatformString : public FGenericPlatformString
 		return Dest;
 	}
 
-	template <typename CharType>
-	static inline int32 Stricmp( const CharType* String1, const CharType* String2 )
-	{
-		// walk the strings, comparing them case insensitively
-		for (; *String1 || *String2; String1++, String2++)
-		{
-			CharType Char1 = TChar<CharType>::ToLower(*String1), Char2 = TChar<CharType>::ToLower(*String2);
-			if (Char1 != Char2)
-			{
-				return Char1 - Char2;
-			}
-		}
-		return 0;
+public:
+	/**
+	 * Compares two strings case-insensitive.
+	 *
+	 * @param String1 First string to compare.
+	 * @param String2 Second string to compare.
+	 *
+	 * @returns Zero if both strings are equal. Greater than zero if first
+	 *          string is greater than the second one. Less than zero
+	 *          otherwise.
+	 */
+	template <typename CharType1, typename CharType2>
+	static inline int32 Stricmp(const CharType1* String1, const CharType2* String2)
+	{
+		return FGenericPlatformStricmp::Stricmp(String1, String2);
 	}
 
 	template <typename CharType>
diff --git a/Engine/Source/Runtime/Core/Public/Templates/UnrealTypeTraits.h b/Engine/Source/Runtime/Core/Public/Templates/UnrealTypeTraits.h
--- a/Engine/Source/Runtime/Core/Public/Templates/UnrealTypeTraits.h
+++ b/Engine/Source/Runtime/Core/Public/Templates/UnrealTypeTraits.h
@@ -8,6 +8,9 @@
 
 #pragma once
 
+#include "AssertionMacros.h"
+#include "Misc/OutputDevice.h"
+
 /*-----------------------------------------------------------------------------
  * Macros to abstract the presence of certain compiler intrinsic type traits 
  -----------------------------------------------------------------------------*/
@@ -94,6 +97,14 @@ template<typename T> struct TIsArithmeticType
 	enum { Value = TIsIntegralType<T>::Value || TIsFloatType<T>::Value } ;
 };
 
+/**
+ * TIsSame
+ *
+ * Unreal implementation of std::is_same trait.
+ */
+template<typename A, typename B>	struct TIsSame			{ enum { Value = false	}; };
+template<typename T>				struct TIsSame<T, T>	{ enum { Value = true	}; };
+
 /**
  * TIsCharType
  */
diff --git a/Engine/Source/Runtime/Core/Public/UObject/NameTypes.h b/Engine/Source/Runtime/Core/Public/UObject/NameTypes.h
--- a/Engine/Source/Runtime/Core/Public/UObject/NameTypes.h
+++ b/Engine/Source/Runtime/Core/Public/UObject/NameTypes.h
@@ -763,12 +763,60 @@ public:
 	explicit FName( EName HardcodedIndex, const TCHAR* Name );
 
 	/**
-	 * Comparison operator.
+	 * Equality operator.
 	 *
 	 * @param	Other	String to compare this name to
 	 * @return true if name matches the string, false otherwise
 	 */
-	bool operator==( const TCHAR* Other ) const;
+	template <typename CharType>
+	bool operator==(const CharType* Other) const
+	{
+		// Find name entry associated with this FName.
+		check(Other);
+		const FNameEntry* const Entry = GetComparisonNameEntry();
+
+		// Temporary buffer to hold split name in case passed in name is of Name_Number format.
+		WIDECHAR TempBuffer[NAME_SIZE];
+		int32 InNumber = NAME_NO_NUMBER_INTERNAL;
+		int32 TempNumber = NAME_NO_NUMBER_INTERNAL;
+
+		// Check whether we need to split the passed in string into name and number portion.
+		auto WideOther = StringCast<WIDECHAR>(Other);
+		const WIDECHAR* WideOtherPtr = WideOther.Get();
+		if (SplitNameWithCheck(WideOtherPtr, TempBuffer, ARRAY_COUNT(TempBuffer), TempNumber))
+		{
+			WideOtherPtr = TempBuffer;
+			InNumber = NAME_EXTERNAL_TO_INTERNAL(TempNumber);
+		}
+
+		// Report a match if both the number and string portion match.
+		bool bAreNamesMatching = false;
+		if (InNumber == GetNumber())
+		{
+			if (Entry->IsWide())
+			{
+				bAreNamesMatching = !FPlatformString::Stricmp(WideOtherPtr, Entry->GetWideName());
+			}
+			else
+			{
+				bAreNamesMatching = !FPlatformString::Stricmp(WideOtherPtr, Entry->GetAnsiName());
+			}
+		}
+
+		return bAreNamesMatching;
+	}
+
+	/**
+	 * Inequality operator.
+	 *
+	 * @param	Other	String to compare this name to
+	 * @return true if name matches the string, false otherwise
+	 */
+	template <typename CharType>
+	bool operator!=(const CharType* Other) const
+	{
+		return !operator==(Other);
+	}
 
 	static void StaticInit();
 	static void DisplayHash( class FOutputDevice& Ar );
@@ -967,17 +1015,31 @@ FORCEINLINE FName ScriptNameToName(const FScriptName& InName)
 
 
 /**
- * Comparison operator with TCHAR* on left hand side and FName on right hand side
+ * Equality operator with CharType* on left hand side and FName on right hand side
  * 
- * @param	LHS		TCHAR to compare to FName
- * @param	RHS		FName to compare to TCHAR
- * @return true if strings match, false otherwise
+ * @param	LHS		CharType to compare to FName
+ * @param	RHS		FName to compare to CharType
+ * @return True if strings match, false otherwise.
  */
-inline bool operator==( const TCHAR *LHS, const FName &RHS )
+template <typename CharType>
+inline bool operator==(const CharType *LHS, const FName &RHS)
 {
 	return RHS == LHS;
 }
 
+/**
+ * Inequality operator with CharType* on left hand side and FName on right hand side
+ *
+ * @param	LHS		CharType to compare to FName
+ * @param	RHS		FName to compare to CharType
+ * @return True if strings don't match, false otherwise.
+ */
+template <typename CharType>
+inline bool operator!=(const CharType *LHS, const FName &RHS)
+{
+	return RHS != LHS;
+}
+
 /** FNames act like PODs. */
 template <> struct TIsPODType<FName> { enum { Value = true }; };
 
