Core - Added some basic documentation to the stats system

#docs

diff --git a/Engine/Source/Runtime/Core/Public/Stats/Stats.h b/Engine/Source/Runtime/Core/Public/Stats/Stats.h
--- a/Engine/Source/Runtime/Core/Public/Stats/Stats.h
+++ b/Engine/Source/Runtime/Core/Public/Stats/Stats.h
@@ -2,6 +2,286 @@
 
 #pragma once
 
+/** 
+ *	Unreal Engine Stats system
+ *	
+ *	This is a preliminary version of the documentation, any comments are welcome :)
+ *	
+ *	This system allows you to collect various performance data and then the data can be used to optimize your game.
+ *	There are a few methods how to achieve this. This quick tutorial will describe all of them.
+ *	For stats commands check out method PrintStatsHelpToOutputDevice();
+ *	
+ *	Stats system in the UE4 supports following stats types:
+ *		Cycle Counter - a generic cycle counter used to counting the number of cycles during the lifetime of the object
+ *		Float/Dword Counter - a counter that is cleared every frame
+ *		Float/Dword Accumulator - a counter that is not cleared every frame, persistent stat, but it can be reset
+ *		Memory - a special type of counter that is optimized for memory tracking
+ *	
+ *	Each stat needs to be grouped, this usually corresponds with displaying the specified stat group i.e. 'stat statsystem' which displays stats' related data.
+ *	
+ *	To define a stat group you need to use one of the following methods:
+ *				DECLARE_STATS_GROUP(GroupDesc,GroupId,GroupCat) - declares a stats group which is enabled by default 
+ *				DECLARE_STATS_GROUP_VERBOSE(GroupDesc,GroupId,GroupCat) - declares a stats group which is disabled by default
+ *				DECLARE_STATS_GROUP_MAYBE_COMPILED_OUT(GroupDesc,GroupId,GroupCat) - declares a stats group which is disabled by default and may be stripped by the compiler
+ *	
+ *	where
+ *		GroupDesc is a text description of the group
+ *		GroupId is an UNIQUE id of the group
+ *		GroupCat is reserved for future use
+ *		CompileIn if set to true, the compiler may strip it out
+ *		
+ *	It can be done in the source or header file depending the usage scope.
+ *	
+ *	Examples:
+ *				DECLARE_STATS_GROUP(TEXT("Threading"), STATGROUP_Threading, STATCAT_Advanced);
+ *				DECLARE_STATS_GROUP_VERBOSE(TEXT("Linker Load"), STATGROUP_LinkerLoad, STATCAT_Advanced);
+ *	
+ *	Now, you can declare/define a stat.
+ *	A stat can be used only in one cpp file, in the function scope, in the module scope or can be used in the whole project.
+ *	
+ *	For one file scope you need to use one of the following methods depending on the stat type.
+ *				DECLARE_CYCLE_STAT(CounterName,StatId,GroupId) - declares a cycle counter stat
+ *		
+ *				DECLARE_SCOPE_CYCLE_COUNTER(CounterName,StatId,GroupId) - declares a cycle counter stat and uses it at the same time, it is limited to one function scope
+ *				QUICK_SCOPE_CYCLE_COUNTER(StatId) - declares a cycle counter stat that will belong to stat group called 'Quick'
+ *				RETURN_QUICK_DECLARE_CYCLE_STAT(StatId,GroupId) - returns a cycle counter, used by a few specialized classes, more information later
+ *		
+ *				DECLARE_FLOAT_COUNTER_STAT(CounterName,StatId,GroupId) - declares a float counter, technically speaking it's based on the double type, 8 bytes
+ *				DECLARE_DWORD_COUNTER_STAT(CounterName,StatId,GroupId) - declared a dword counter, technically speaking it's based on the qword type, 8 bytes
+ *				DECLARE_FLOAT_ACCUMULATOR_STAT(CounterName,StatId,GroupId) - declares a float accumulator
+ *				DECLARE_DWORD_ACCUMULATOR_STAT(CounterName,StatId,GroupId) - declares a dword accumulator
+ *				DECLARE_MEMORY_STAT(CounterName,StatId,GroupId) - declares a memory counter, same as the dword accumulator, but will be displayed with memory specific units
+ *				DECLARE_MEMORY_STAT_POOL(CounterName,StatId,GroupId,Pool) - declares a memory counter with a pool
+ *	
+ *		If you want to have these stats accessible in the whole project/or wider range of files you need to use extern version.
+ *		These methods are the same as the previously mentioned but with _EXTERN and the end of the name, here is the list:
+ *				DECLARE_CYCLE_STAT_EXTERN(CounterName,StatId,GroupId, API)
+ *				DECLARE_FLOAT_COUNTER_STAT_EXTERN(CounterName,StatId,GroupId, API)
+ *				DECLARE_DWORD_COUNTER_STAT_EXTERN(CounterName,StatId,GroupId, API)
+ *				DECLARE_FLOAT_ACCUMULATOR_STAT_EXTERN(CounterName,StatId,GroupId, API)
+ *				DECLARE_DWORD_ACCUMULATOR_STAT_EXTERN(CounterName,StatId,GroupId, API)
+ *				DECLARE_MEMORY_STAT_EXTERN(CounterName,StatId,GroupId, API)
+ *				DECLARE_MEMORY_STAT_POOL_EXTERN(CounterName,StatId,GroupId,Pool, API)
+ *				
+ *		Then in the source file you need to define those stats.
+ *				DEFINE_STAT(CounterName) - defines stats declared with _EXTERN
+ *	
+ *	where
+ *		CounterName is a text description of the stat
+ *		StatId is an UNIQUE id of the stat
+ *		GroupId is an id of the group that the stat will belong to, the GroupId from DECLARE_STATS_GROUP*
+ *		Pool is a platform specific memory pool, more details later
+ *		API is the *_API of module, can be empty if the stat will be used only in that module
+ *
+ *	Examples:
+ *		Custom memory stats with pools
+ *			First you need to add a new pool to enum EMemoryCounterRegion, it can be global or platform specific.
+ *			
+ *				enum EMemoryCounterRegion
+ *				{
+ *					MCR_Invalid,	// not memory
+ *					MCR_Physical,	// main system memory
+ *					MCR_GPU,		// memory directly a GPU (graphics card, etc)
+ *					MCR_GPUSystem,	// system memory directly accessible by a GPU
+ *					MCR_TexturePool,// presized texture pools
+ *					MCR_MAX
+ *				};
+ *			
+ *			This is an example that will allow using the pools every where, see CORE_API. 
+ *			THE NAME OF THE POOL MUST START WITH MCR_
+ *			Header file.	
+ *				DECLARE_MEMORY_STAT_POOL_EXTERN(TEXT("Physical Memory Pool [Physical]"),	MCR_Physical,		STATGROUP_Memory,  FPlatformMemory::MCR_Physical,	CORE_API);
+ *				DECLARE_MEMORY_STAT_POOL_EXTERN(TEXT("GPU Memory Pool [GPU]"),				MCR_GPU,			STATGROUP_Memory,  FPlatformMemory::MCR_GPU,		CORE_API);
+ *				DECLARE_MEMORY_STAT_POOL_EXTERN(TEXT("Texture Memory Pool [Texture]"),		MCR_TexturePool,	STATGROUP_Memory,  FPlatformMemory::MCR_TexturePool,CORE_API);
+ *				
+ *			Source file.
+ *				DEFINE_STAT(MCR_Physical);
+ *				DEFINE_STAT(MCR_GPU);
+ *				DEFINE_STAT(MCR_TexturePool);
+ *			
+ *			This is a pool, so it needs to be initialized. Usually in the F*PlatformMemory::Init()
+ *				SET_MEMORY_STAT(MCR_Physical, PhysicalPoolLimit);
+ *				SET_MEMORY_STAT(MCR_GPU, GPUPoolLimit);
+ *				SET_MEMORY_STAT(MCR_TexturePool, TexturePoolLimit);
+ *			
+ *			Now we have pools, so we can setup memory stats for those pools.
+ *			Accessible everywhere.
+ *				DECLARE_MEMORY_STAT_POOL_EXTERN(TEXT("Index buffer memory"),		STAT_IndexBufferMemory,		STATGROUP_RHI, FPlatformMemory::MCR_GPU, RHI_API);
+ *				DECLARE_MEMORY_STAT_POOL_EXTERN(TEXT("Vertex buffer memory"),		STAT_VertexBufferMemory,	STATGROUP_RHI, FPlatformMemory::MCR_GPU, RHI_API);
+ *				DECLARE_MEMORY_STAT_POOL_EXTERN(TEXT("Structured buffer memory"),	STAT_StructuredBufferMemory,STATGROUP_RHI, FPlatformMemory::MCR_GPU, RHI_API);
+ *				DECLARE_MEMORY_STAT_POOL_EXTERN(TEXT("Pixel buffer memory"),		STAT_PixelBufferMemory,		STATGROUP_RHI, FPlatformMemory::MCR_GPU, RHI_API);
+ *				
+ *			Accessible only in the module where defined.	
+ *				DECLARE_MEMORY_STAT_POOL_EXTERN(TEXT("Pool Memory Size"), STAT_TexturePoolSize,				STATGROUP_Streaming, FPlatformMemory::MCR_TexturePool, );
+ *				DECLARE_MEMORY_STAT_POOL_EXTERN(TEXT("Pool Memory Used"), STAT_TexturePoolAllocatedSize,	STATGROUP_Streaming, FPlatformMemory::MCR_TexturePool, );
+ *				
+ *			And the last thing, updating the memory stats.
+ *				INC_MEMORY_STAT_BY(STAT_PixelBufferMemory,NumBytes) - increases a memory stat by the specified value
+ *				DEC_MEMORY_STAT_BY(STAT_PixelBufferMemory,NumBytes) - decreases a memory stat by the specified value
+ *				SET_MEMORY_STAT(STAT_PixelBufferMemory,NumBytes) - sets a memory stat to the specified value
+ *			
+ *		Regular memory stats, without pools	
+ *				DECLARE_MEMORY_STAT(TEXT("Total Physical"),		STAT_TotalPhysical,		STATGROUP_MemoryPlatform);
+ *				DECLARE_MEMORY_STAT(TEXT("Total Virtual"),		STAT_TotalVirtual,		STATGROUP_MemoryPlatform);
+ *				DECLARE_MEMORY_STAT(TEXT("Page Size"),			STAT_PageSize,			STATGROUP_MemoryPlatform);
+ *				DECLARE_MEMORY_STAT(TEXT("Total Physical GB"),	STAT_TotalPhysicalGB,	STATGROUP_MemoryPlatform);
+ *		
+ *		Or DECLARE_MEMORY_STAT_EXTERN in the header file and then DEFINE_STAT in the source file.
+ *		Updating the memory stats is done the same way as in the version with pools.
+ *		
+ *		
+ *		Performance data using the cycle counters.
+ *			First you need to add cycle counters.
+ *				DECLARE_CYCLE_STAT(TEXT("Broadcast"),	STAT_StatsBroadcast,STATGROUP_StatSystem);
+ *				DECLARE_CYCLE_STAT(TEXT("Condense"),	STAT_StatsCondense,	STATGROUP_StatSystem);
+ *				
+ *			Or DECLARE_CYCLE_STAT_EXTERN in the header file and then DEFINE_STAT in the source file.
+ *			
+ *			Now you can grab the performance data.
+ *			
+ *				Stats::Broadcast()
+ *				{
+ *					SCOPE_CYCLE_COUNTER(STAT_StatsBroadcast);
+ *					...
+ *					// a piece of code
+ *					...
+ *				}
+ *
+ *			and that's all.
+ *			Sometimes you don't want to grab the stats every time the function is called, so you can use conditional cycle counter.
+ *			It's not very common, but may be useful.
+ *			
+ *				Stats::Broadcast(bool bSomeCondition)
+ *				{
+ *					CONDITIONAL_SCOPE_CYCLE_COUNTER(STAT_StatsBroadcast,bSomeCondition);
+ *					...
+ *					// a piece of code
+ *					...
+ *				}
+ *
+ *			If you want to grab the performance data from one function you can use following construction.
+ *			
+ *				Stats::Broadcast(bool bSomeCondition)
+ *				{
+ *					DECLARE_SCOPE_CYCLE_COUNTER(TEXT("Broadcast"), STAT_StatsBroadcast, STATGROUP_StatSystem);
+ *					...
+ *					// a piece of code
+ *					...
+ *				}
+ *				
+ *			or
+ *			
+ *				Stats::Broadcast(bool bSomeCondition)
+ *				{
+ *					QUICK_SCOPE_CYCLE_COUNTER(TEXT("Stats::Broadcast"));
+ *					...
+ *					// a piece of code
+ *					...
+ *				}
+ *				
+ *			Mostly used for temporary stats.
+ *			
+ *			Those all cycle counters are used to generate the hierarchy. So you can get more detailed information about performance data.
+ *			There is also an option to set flat cycle counter.
+ *			
+ *				Stats::Broadcast(bool bSomeCondition)
+ *				{
+ *					const uint32 BroadcastBeginTime = FPlatformTime::Cycles();
+ *					...
+ *					// a piece of code
+ *					...
+ *					const uint32 BroadcastEndTime = FPlatformTime::Cycles();
+ *					SET_CYCLE_COUNTER(STAT_StatsBroadcast, BroadcastEndTime-BroadcastBeginTime);
+ *				}
+ *				
+ *		
+ *		A few tasks implemented in the UE4 use a different approach in terms of getting the performance data.
+ *		They implement method GetStatId(). If there is no GetStatId(), the code will not compile.
+ *		Here is an example.
+ *		
+ *				class FParallelAnimationCompletionTask
+ *				{
+ *					// ...
+ *					// a piece of code
+ *					FORCEINLINE TStatId GetStatId() const
+ *					{
+ *						RETURN_QUICK_DECLARE_CYCLE_STAT(FParallelAnimationCompletionTask, STATGROUP_TaskGraphTasks);
+ *					}
+ *					// a piece of code
+ *					// ...
+ *				};
+ *				
+ *				
+ *		Generic data using the float or dword counters.		
+ *			First you need to add a few counters.
+ *				DECLARE_FLOAT_COUNTER_STAT_EXTERN(STAT_FloatCounter,StatId,STATGROUP_TestGroup, CORE_API)
+ *				DECLARE_DWORD_COUNTER_STAT_EXTERN(STAT_DwordCounter,StatId,STATGROUP_TestGroup, CORE_API)
+ *				DECLARE_FLOAT_ACCUMULATOR_STAT_EXTERN(STAT_FloatAccumulator,StatId,STATGROUP_TestGroup, CORE_API)
+ *				DECLARE_DWORD_ACCUMULATOR_STAT_EXTERN(STAT_DwordAccumulator,StatId,STATGROUP_TestGroup, CORE_API)
+ *				 
+ *			Updating counters.
+ *				INC_DWORD_STAT(StatId) - increases a dword stat by 1
+ *				DEC_DWORD_STAT(StatId) - decreases a dword stat by 1
+ *				INC_DWORD_STAT_BY(StatId,Amount) - increases a dword stat by the specified value
+ *				DEC_DWORD_STAT_BY(StatId,Amount) - decreases a dword stat by the specified value
+ *				SET_DWORD_STAT(StatId,Value) - sets a dword stat to the specified value
+
+ *				INC_FLOAT_STAT_BY(StatId,Amount) - increases a float stat by the specified value
+ *				DEC_FLOAT_STAT_BY(StatId,Amount) - decreases a float stat by the specified value
+ *				SET_FLOAT_STAT(StatId,Value) - sets a float stat to the specified value
+ *			
+ *			
+ *		A few helper methods
+ *				 GET_STATID(StatId) - returns an instance of the TStatId of the stat, ADVANCED
+ *				 GET_STATDESCRIPTION(StatId) - returns a description of the stat
+ *
+ *		
+ *		If you don't want to use the stats system and just log some performance data, there is functionality for this.
+ *		
+ *				SCOPE_SECONDS_COUNTER(double&Seconds) - captures time passed in seconds, adding delta time to passed in variable
+ *		
+ *				Stats::Broadcast()
+ *				{
+ *					double ThisTime = 0;
+ *					{
+ *						 SCOPE_SECONDS_COUNTER(ThisTime);
+*						...
+ *						// a piece of code
+ *						...
+ *					}
+ *					UE_LOG(LogTemp, Log, TEXT("Stats::Broadcast %.2f"), ThisTime );	
+ *				}
+ *		
+ *				FScopeLogTime - utility class to log time passed in seconds, adding cumulative stats to passed in variable, print the performance data to the log in the destructor
+ *			
+ *				SCOPE_LOG_TIME(Name,CumulativePtr) - using the given name prints the performance data and gathers cumulative stats
+ *				SCOPE_LOG_TIME_IN_SECONDS(Name,CumulativePtr) - the same as above, but prints in seconds
+ *			
+ *				SCOPE_LOG_TIME_FUNC() - using the funcion name prints the performance data, cannot be nested
+ *				SCOPE_LOG_TIME_FUNC_WITH_GLOBAL(CumulativePtr), same as above, but gather cumulative stats
+ *			
+ *			A few examples.
+ *			
+ *				double GMyBroadcastTime = 0.0;
+ *				Stats::Broadcast()
+ *				{
+ *					SCOPE_LOG_TIME("Stats::Broadcast", &GMyBroadcastTime );
+ *					SCOPE_LOG_TIME_IN_SECONDS("Stats::Broadcast (sec)", &GMyBroadcastTime );
+ *					...
+ *					// a piece of code
+ *					...
+ *				}
+ *				
+ *				Stats::Condense()
+ *				{
+ *					SCOPE_LOG_TIME_FUNC(); // The name should be "Stats::Condense()", may differ across compilers
+ *					SCOPE_LOG_TIME_FUNC_WITH_GLOBAL(&GMyBroadcastTime);
+ *					...
+ *					// a piece of code
+ *					...
+ *				}
+ */
+
 #define FORCEINLINE_STATS FORCEINLINE
 #define checkStats(x)
 
@@ -117,7 +397,7 @@ FORCEINLINE void StatsMasterEnableSubtract(int32 Value = 1)
 
 #define DECLARE_STATS_GROUP(GroupDesc,GroupId,GroupCat)
 #define DECLARE_STATS_GROUP_VERBOSE(GroupDesc,GroupId,GroupCat)
-#define DECLARE_STATS_GROUP_COMPILED_OUT(GroupDesc,GroupId,GroupCat)
+#define DECLARE_STATS_GROUP_MAYBE_COMPILED_OUT(GroupDesc,GroupId,GroupCat,CompileIn)
 
 #define SET_CYCLE_COUNTER(Stat,Cycles)
 #define INC_DWORD_STAT(StatId)
diff --git a/Engine/Source/Runtime/Core/Public/Stats/Stats2.h b/Engine/Source/Runtime/Core/Public/Stats/Stats2.h
--- a/Engine/Source/Runtime/Core/Public/Stats/Stats2.h
+++ b/Engine/Source/Runtime/Core/Public/Stats/Stats2.h
@@ -1503,11 +1503,11 @@ struct FStat_##StatName\
 	#define SCOPE_CYCLE_COUNTER_GUARD
 #endif // UE_BUILD_DEBUG
 
-#define DECLARE_SCOPE_CYCLE_COUNTER(CounterName,StatId,GroupId) \
+#define DECLARE_SCOPE_CYCLE_COUNTER(CounterName,Stat,GroupId) \
 	SCOPE_CYCLE_COUNTER_GUARD \
-	DECLARE_STAT(CounterName,StatId,GroupId,EStatDataType::ST_int64, true, true, FPlatformMemory::MCR_Invalid); \
-	static DEFINE_STAT(StatId) \
-	FScopeCycleCounter CycleCount_##StatId(GET_STATID(StatId));
+	DECLARE_STAT(CounterName,Stat,GroupId,EStatDataType::ST_int64, true, true, FPlatformMemory::MCR_Invalid); \
+	static DEFINE_STAT(Stat) \
+	FScopeCycleCounter CycleCount_##Stat(GET_STATID(Stat));
 
 #define QUICK_SCOPE_CYCLE_COUNTER(Stat) \
 	SCOPE_CYCLE_COUNTER_GUARD \
