Removing PLATFORM_HAS_vsnprintf (no longer used)

diff --git a/Engine/Source/Runtime/Core/Public/Android/AndroidPlatform.h b/Engine/Source/Runtime/Core/Public/Android/AndroidPlatform.h
--- a/Engine/Source/Runtime/Core/Public/Android/AndroidPlatform.h
+++ b/Engine/Source/Runtime/Core/Public/Android/AndroidPlatform.h
@@ -34,7 +34,6 @@ typedef FAndroidTypes FPlatformTypes;
 #define PLATFORM_LITTLE_ENDIAN						1
 #define PLATFORM_SUPPORTS_PRAGMA_PACK				1
 #define PLATFORM_USE_LS_SPEC_FOR_WIDECHAR			1
-#define PLATFORM_HAS_vsnprintf						0
 #define PLATFORM_HAS_BSD_TIME						1
 #define PLATFORM_USE_PTHREADS						1
 #define PLATFORM_MAX_FILEPATH_LENGTH				MAX_PATH
diff --git a/Engine/Source/Runtime/Core/Public/HAL/Platform.h b/Engine/Source/Runtime/Core/Public/HAL/Platform.h
--- a/Engine/Source/Runtime/Core/Public/HAL/Platform.h
+++ b/Engine/Source/Runtime/Core/Public/HAL/Platform.h
@@ -164,9 +164,6 @@
 #ifndef PLATFORM_TCHAR_IS_4_BYTES
 	#define PLATFORM_TCHAR_IS_4_BYTES			0
 #endif
-#ifndef PLATFORM_HAS_vsnprintf
-	#define PLATFORM_HAS_vsnprintf				1
-#endif
 #ifndef PLATFORM_HAS_BSD_TIME
 	#define PLATFORM_HAS_BSD_TIME				1
 #endif
diff --git a/Engine/Source/Runtime/Core/Public/HTML5/HTML5Platform.h b/Engine/Source/Runtime/Core/Public/HTML5/HTML5Platform.h
--- a/Engine/Source/Runtime/Core/Public/HTML5/HTML5Platform.h
+++ b/Engine/Source/Runtime/Core/Public/HTML5/HTML5Platform.h
@@ -28,7 +28,6 @@ typedef FHTML5Types FPlatformTypes;
 #define PLATFORM_LITTLE_ENDIAN						1
 #define PLATFORM_SUPPORTS_PRAGMA_PACK				1
 #define PLATFORM_USE_LS_SPEC_FOR_WIDECHAR			1
-#define PLATFORM_HAS_vsnprintf						1
 #if PLATFORM_HTML5_WIN32
 #define PLATFORM_HAS_BSD_TIME						0
 #define PLATFORM_COMPILER_HAS_DEFAULTED_FUNCTIONS	0
diff --git a/Engine/Source/Runtime/Core/Public/IOS/IOSPlatform.h b/Engine/Source/Runtime/Core/Public/IOS/IOSPlatform.h
--- a/Engine/Source/Runtime/Core/Public/IOS/IOSPlatform.h
+++ b/Engine/Source/Runtime/Core/Public/IOS/IOSPlatform.h
@@ -33,7 +33,6 @@ typedef FIOSPlatformTypes FPlatformTypes;
 #define PLATFORM_SUPPORTS_PRAGMA_PACK					1
 #define PLATFORM_COMPILER_DISTINGUISHES_INT_AND_LONG	1
 #define PLATFORM_TCHAR_IS_4_BYTES						1
-#define PLATFORM_HAS_vsnprintf							0
 #define PLATFORM_USE_SYSTEM_VSWPRINTF					0
 #define PLATFORM_HAS_BSD_TIME							1
 #define PLATFORM_MAX_FILEPATH_LENGTH					MAX_PATH
diff --git a/Engine/Source/Runtime/Core/Public/Linux/LinuxPlatform.h b/Engine/Source/Runtime/Core/Public/Linux/LinuxPlatform.h
--- a/Engine/Source/Runtime/Core/Public/Linux/LinuxPlatform.h
+++ b/Engine/Source/Runtime/Core/Public/Linux/LinuxPlatform.h
@@ -36,7 +36,6 @@ typedef FLinuxPlatformTypes FPlatformTypes;
 #define PLATFORM_SUPPORTS_PRAGMA_PACK				1
 #define PLATFORM_USE_LS_SPEC_FOR_WIDECHAR			1
 #define PLATFORM_TCHAR_IS_4_BYTES					1
-#define PLATFORM_HAS_vsnprintf						0
 #define PLATFORM_HAS_BSD_TIME						1
 #define PLATFORM_USE_PTHREADS						1
 #define PLATFORM_MAX_FILEPATH_LENGTH				MAX_PATH /* @todo linux: avoid using PATH_MAX as it is known to be broken */
diff --git a/Engine/Source/Runtime/Core/Public/Mac/MacPlatform.h b/Engine/Source/Runtime/Core/Public/Mac/MacPlatform.h
--- a/Engine/Source/Runtime/Core/Public/Mac/MacPlatform.h
+++ b/Engine/Source/Runtime/Core/Public/Mac/MacPlatform.h
@@ -34,7 +34,6 @@ typedef FMacPlatformTypes FPlatformTypes;
 #define PLATFORM_USE_SYSTEM_VSWPRINTF				0
 #define PLATFORM_COMPILER_DISTINGUISHES_INT_AND_LONG			1
 #define PLATFORM_TCHAR_IS_4_BYTES					1
-#define PLATFORM_HAS_vsnprintf						0
 #define PLATFORM_HAS_BSD_TIME						1
 //#define PLATFORM_USE_PTHREADS						1
 #define PLATFORM_MAX_FILEPATH_LENGTH				MAX_PATH
diff --git a/Engine/Source/Runtime/Core/Public/WinRT/WinRTARMPlatform.h b/Engine/Source/Runtime/Core/Public/WinRT/WinRTARMPlatform.h
--- a/Engine/Source/Runtime/Core/Public/WinRT/WinRTARMPlatform.h
+++ b/Engine/Source/Runtime/Core/Public/WinRT/WinRTARMPlatform.h
@@ -35,7 +35,6 @@ typedef FWinRTTypes FPlatformTypes;
 #define PLATFORM_USE_LS_SPEC_FOR_WIDECHAR					0
 #define PLATFORM_COMPILER_DISTINGUISHES_DWORD_AND_UINT		1
 #define PLATFORM_COMPILER_HAS_GENERIC_KEYWORD				1
-#define PLATFORM_HAS_vsnprintf								0
 #define PLATFORM_HAS_BSD_TIME								0
 #define PLATFORM_USE_PTHREADS								0
 #define PLATFORM_MAX_FILEPATH_LENGTH						MAX_PATH
diff --git a/Engine/Source/Runtime/Core/Public/WinRT/WinRTPlatform.h b/Engine/Source/Runtime/Core/Public/WinRT/WinRTPlatform.h
--- a/Engine/Source/Runtime/Core/Public/WinRT/WinRTPlatform.h
+++ b/Engine/Source/Runtime/Core/Public/WinRT/WinRTPlatform.h
@@ -39,7 +39,6 @@ typedef FWinRTTypes FPlatformTypes;
 #define PLATFORM_USE_LS_SPEC_FOR_WIDECHAR					0
 #define PLATFORM_COMPILER_DISTINGUISHES_DWORD_AND_UINT		1
 #define PLATFORM_COMPILER_HAS_GENERIC_KEYWORD				1
-#define PLATFORM_HAS_vsnprintf								0
 #define PLATFORM_HAS_BSD_TIME								0
 #define PLATFORM_HAS_BSD_SOCKETS							0
 #define PLATFORM_USE_PTHREADS								0
diff --git a/Engine/Source/Runtime/Core/Public/Windows/WIndowsPlatform.h b/Engine/Source/Runtime/Core/Public/Windows/WIndowsPlatform.h
--- a/Engine/Source/Runtime/Core/Public/Windows/WIndowsPlatform.h
+++ b/Engine/Source/Runtime/Core/Public/Windows/WIndowsPlatform.h
@@ -52,7 +52,6 @@ typedef FWindowsPlatformTypes FPlatformTypes;
 //#define PLATFORM_USE_LS_SPEC_FOR_WIDECHAR					1
 //#define PLATFORM_USE_SYSTEM_VSWPRINTF						1
 //#define PLATFORM_TCHAR_IS_4_BYTES							0
-#define PLATFORM_HAS_vsnprintf								0
 #define PLATFORM_HAS_BSD_TIME								0
 #define PLATFORM_USE_PTHREADS								0
 #define PLATFORM_MAX_FILEPATH_LENGTH						MAX_PATH
