CrashReport - Added FMallocCrash, still not enabled

diff --git a/Engine/Source/Runtime/Core/Private/HAL/MallocCrash.cpp b/Engine/Source/Runtime/Core/Private/HAL/MallocCrash.cpp
new file mode 100644
--- /dev/null
+++ b/Engine/Source/Runtime/Core/Private/HAL/MallocCrash.cpp
@@ -0,0 +1,443 @@
+// Copyright 1998-2014 Epic Games, Inc. All Rights Reserved.
+
+#include "Core.h"
+#include "MallocCrash.h"
+
+/** Describes a pool. */
+struct FPoolDesc
+{
+	FPoolDesc( const uint32 InSize, const uint32 InNumAllocs ):
+		Size(InSize),
+		NumAllocs(InNumAllocs)
+	{}
+
+	/** Size of the pool. */
+	const uint32 Size;
+
+	/** Number of allocations in the pool. */
+	const uint32 NumAllocs;
+};
+
+/** Generated by the FMallocCrash::PrintPoolsUsage. */
+const FPoolDesc& FMallocCrash::GetPoolDesc( uint32 Index ) const
+{
+	static const FPoolDesc AllPoolDesc[NUM_POOLS]=
+	{
+		FPoolDesc(   64, 224),
+		FPoolDesc(   96, 144),
+		FPoolDesc(  128,  80),
+		FPoolDesc(  192, 560),
+		FPoolDesc(  256, 384),
+		FPoolDesc(  384, 208),
+		FPoolDesc(  512,  48),
+		FPoolDesc(  768,  32),
+		FPoolDesc( 1024,  32),
+		FPoolDesc( 2048,  32),
+		FPoolDesc( 4096,  32),
+		FPoolDesc( 8192,  32),
+		FPoolDesc(16384,  16),
+		FPoolDesc(32768,  16),
+	};
+	return AllPoolDesc[Index];
+}
+
+struct FMallocCrashPool
+{
+	uint32 NumUsed;
+	uint32 MaxUsedIndex;
+	uint32 MaxNumUsed;
+	uint32 TotalNumUsed;
+
+	/** Allocation size for this pool. */
+	const uint32 AllocationSize;
+
+	/** Fixed list of allocations for the specified size for this pool. */
+	FPtrInfo* Allocations[FMallocCrash::MAX_NUM_ALLOCS_IN_POOL];
+
+	/** Maximum number of allocations that can be made for this pool. */
+	const uint32 MaxNumAllocations;
+
+	/** Memory allocated in the pool and memory used by a fixed array. */
+	uint32 AllocatedMemory;
+
+	FMallocCrashPool( const FPoolDesc& PoolDesc, FMallocCrash& Outer ): 
+		NumUsed(0),
+		MaxUsedIndex(0),
+		MaxNumUsed(0),
+		TotalNumUsed(0),
+		AllocationSize( PoolDesc.Size+FMallocCrash::PER_ALLOC_OVERHEAD ),
+		MaxNumAllocations( PoolDesc.NumAllocs )
+	{
+		for( uint32 Index = 0; Index < MaxNumAllocations; ++Index )
+		{
+			uint8* NewPtr = (uint8*)Outer.AllocateFromSmallPool( AllocationSize );
+			Allocations[Index] = new(NewPtr) FPtrInfo( NewPtr+FMallocCrash::PER_ALLOC_OVERHEAD );
+		}
+
+		// Zero the rest.
+		for( uint32 Index = MaxNumAllocations; Index < FMallocCrash::MAX_NUM_ALLOCS_IN_POOL; ++Index )
+		{
+			Allocations[Index] = nullptr;
+		}
+
+		AllocatedMemory = MaxNumAllocations*AllocationSize + sizeof(Allocations);
+
+		FPlatformMisc::LowLevelOutputDebugStringf( TEXT( "FMallocCrashPool Size=%5u Num=%4i Mem=%8u\n" ), PoolDesc.Size, MaxNumAllocations, AllocatedMemory );
+	}
+
+	uint8* AllocateFromPool( uint32 InAllocationSize )
+	{
+		DebugVerify();
+
+		// Find an unused allocation.
+		FPtrInfo* PtrInfo = nullptr;
+		for( uint32 Index = 0; Index < MaxNumAllocations; ++Index )
+		{
+			FPtrInfo* PtrIt = Allocations[Index];
+			if( PtrIt->Size == 0 )
+			{
+				PtrInfo = PtrIt;
+				MaxUsedIndex = FMath::Max(MaxUsedIndex,Index);
+				break;
+			}
+		}
+
+		if( PtrInfo )
+		{
+			NumUsed++;
+			TotalNumUsed++;
+			MaxNumUsed = FMath::Max(MaxNumUsed,NumUsed);
+			PtrInfo->Size = InAllocationSize;
+
+			FMemory::Memset( (void*)PtrInfo->Ptr, FMallocCrash::MEM_TAG, PtrInfo->Size );
+			DebugVerify();
+
+			//FPlatformMisc::LowLevelOutputDebugStringf( TEXT( "Malloc Size=%u, PooledPtr=0x%016llx \n" ), AllocationSize, (uint64)PtrInfo->Ptr );
+
+			return PtrInfo->Ptr;
+		}
+		else
+		{
+			FPlatformMisc::DebugBreak();
+			FPlatformMisc::LowLevelOutputDebugStringf( TEXT( "AllocateFromPool run out of memory allocating %u bytes for %u allocations\n" ), InAllocationSize, MaxNumAllocations );
+			FPlatformMisc::LowLevelOutputDebugString( TEXT( "Please increase MaxNumAllocations for that pool\n" ) );
+			FPlatformMisc::RequestExit( true );
+		}
+		return nullptr;
+	}
+
+	/** Tries to free a pointer. */
+	void TryFreeFromPool( uint8* Ptr )
+	{
+		//const uint32 PtrSize = FMallocCrash::GetAllocationSize(Ptr);
+		//FPlatformMisc::LowLevelOutputDebugStringf( TEXT( "Free SizeWithOverhead=%u, PooledPtr=0x%016llx \n" ), PtrSize, (uint64)Ptr );
+
+		bool bRemoved = false;
+		for( uint32 Index = 0; Index < MaxNumAllocations; ++Index )
+		{
+			FPtrInfo* PtrIt = Allocations[Index];
+			if( PtrIt->Ptr == Ptr )
+			{
+				FMemory::Memset( (void*)PtrIt->Ptr, FMallocCrash::MEM_WIPETAG, PtrIt->Size );
+
+				PtrIt->Size = 0;
+				bRemoved = true;
+				NumUsed--;
+
+				break;
+			}
+		}
+
+		if( !bRemoved )
+		{
+			FPlatformMisc::DebugBreak();
+		}
+
+		DebugVerify();
+	}
+
+private:
+	void DebugVerify()
+	{
+#if	_DEBUG
+		for( uint32 Index = 0; Index < MaxNumAllocations; ++Index )
+		{
+			FPtrInfo* PtrIt = Allocations[Index];
+			if( PtrIt->Size > 32768 )
+			{
+				FPlatformMisc::DebugBreak();
+			}
+		}
+#endif // _DEBUG
+	}
+};
+
+FMallocCrash::FMallocCrash( FMalloc* MainMalloc ) :
+	CrashedThreadId( 0 ),
+	LargeMemoryPoolOffset( 0 ),
+	SmallMemoryPoolOffset( 0 ),
+	PreviousMalloc( MainMalloc )
+{
+	LargeMemoryPool = Align( (uint8*)FPlatformMemory::BinnedAllocFromOS( LARGE_MEMORYPOOL_SIZE + REQUIRED_ALIGNMENT ), REQUIRED_ALIGNMENT );
+	SmallMemoryPool = Align( (uint8*)FPlatformMemory::BinnedAllocFromOS( GetSmallPoolTotalSize() + REQUIRED_ALIGNMENT ), REQUIRED_ALIGNMENT );
+	if( !SmallMemoryPool || !LargeMemoryPool )
+	{
+		FPlatformMisc::RequestExit(true);
+	}
+	InitializeSmallPools();
+	FPlatformMisc::LowLevelOutputDebugStringf( TEXT( "FMallocCrash overhead is %u bytes\n" ), LARGE_MEMORYPOOL_SIZE+GetSmallPoolTotalSize() );
+}
+
+FMallocCrash::~FMallocCrash()
+{
+}
+
+FMallocCrash& FMallocCrash::Get( FMalloc* MainMalloc /*= nullptr*/ )
+{
+	static FMallocCrash CrashMalloc( MainMalloc );
+	return CrashMalloc;
+}
+
+void FMallocCrash::SetAsGMalloc()
+{
+	InternalLock.Lock();
+	GMalloc = this;
+	CrashedThreadId = FPlatformTLS::GetCurrentThreadId();
+}
+
+void* FMallocCrash::Malloc( SIZE_T Size, uint32 Alignment )
+{
+	const uint32 Size32 = (uint32)Size;
+	if( Alignment > 16 )
+	{
+		FPlatformMisc::DebugBreak();
+		FPlatformMisc::LowLevelOutputDebugString( TEXT( "Alignment > 16 is not supported\n" ) );
+	}
+
+	if( IsOnCrashedThread() )
+	{
+		FMallocCrashPool* Pool = FindPoolFromSize( Size32 );
+		if( Pool )
+		{
+			const uint8* PooledPtr = Pool->AllocateFromPool( Size32 );	
+			return (void*)PooledPtr;
+		}
+		else
+		{
+			const uint32 SizeWithOverhead = Size32 + PER_ALLOC_OVERHEAD;
+			LargeMemoryPoolOffset = Align( LargeMemoryPoolOffset, REQUIRED_ALIGNMENT );
+
+			if( LargeMemoryPoolOffset + SizeWithOverhead <= LARGE_MEMORYPOOL_SIZE )
+			{
+				const uint32 ReturnMemoryPoolOffset = LargeMemoryPoolOffset;
+				LargeMemoryPoolOffset += SizeWithOverhead;
+
+				FPtrInfo* PtrInfo = (FPtrInfo*)(LargeMemoryPool+ReturnMemoryPoolOffset);
+				PtrInfo->Size = Size32;
+				PtrInfo->Ptr = LargeMemoryPool+ReturnMemoryPoolOffset+PER_ALLOC_OVERHEAD;
+
+				FPlatformMisc::LowLevelOutputDebugStringf( TEXT( "Malloc Size=%u LargeMemoryPoolOffset=%u \n" ), Size32, LargeMemoryPoolOffset );
+				return (void*)PtrInfo->Ptr;
+			}
+			else
+			{
+				FPlatformMisc::DebugBreak();
+				FPlatformMisc::LowLevelOutputDebugStringf( TEXT( "MallocCrash run out of memory allocating %u bytes, free %u bytes\n" ), Size32, LARGE_MEMORYPOOL_SIZE-LargeMemoryPoolOffset );
+				FPlatformMisc::LowLevelOutputDebugString( TEXT( "Please increase LARGE_MEMORYPOOL_SIZE \n" ) );
+				FPlatformMisc::RequestExit( true );			
+			}
+		}
+	}
+	return nullptr;
+}
+
+void* FMallocCrash::Realloc( void* Ptr, SIZE_T NewSize, uint32 Alignment )
+{
+	if( IsOnCrashedThread() )
+	{
+		void* Result = nullptr;
+		if( Ptr && NewSize )
+		{
+			SIZE_T PtrSize = 0;
+			const bool bPreviousMalloc = NewSize > 0 && Ptr && !(IsPtrInLargePool(Ptr)||IsPtrInSmallPool(Ptr));
+
+			if( bPreviousMalloc )
+			{
+				// Realloc from the previous allocator.
+				PreviousMalloc->GetAllocationSize(Ptr,PtrSize);
+				if( PtrSize == 0 )
+				{
+					PtrSize = 0;
+				}
+			}
+			else
+			{
+				PtrSize = GetAllocationSize(Ptr);
+			}
+			
+			Result = Malloc( NewSize, REQUIRED_ALIGNMENT );
+			FMemory::Memcpy( Result, Ptr, FMath::Min(NewSize,PtrSize) );
+			
+			if( PtrSize > 32768 )
+			{
+				FPlatformMisc::LowLevelOutputDebugStringf( TEXT( "Realloc PtrSize=%u NewSize=%u PooledPtr=0x%016llx\n" ), (uint32)PtrSize, (uint32)NewSize, (uint64)Ptr );
+			}
+
+			Free( Ptr );
+		}
+		else if( Ptr == nullptr )
+		{
+			Result = Malloc( NewSize, REQUIRED_ALIGNMENT );
+		}
+		else
+		{
+			Free( Ptr );
+			Result = nullptr;
+		}
+		return Result;
+	}
+	return nullptr;
+}
+
+void FMallocCrash::Free( void* Ptr )
+{
+	if( IsOnCrashedThread() )
+	{
+		if( IsPtrInSmallPool(Ptr) )
+		{
+			FMallocCrashPool* Pool = FindPoolFromSize( GetAllocationSize(Ptr) );
+			if( Pool )
+			{
+				Pool->TryFreeFromPool( (uint8*)Ptr );
+			}
+			else
+			{
+				FPlatformMisc::DebugBreak();
+			}
+		}
+		else if( IsPtrInLargePool(Ptr) )
+		{
+			// Not implemented yet.
+		}
+		else
+		{
+			// From the previous allocator.
+		}
+	}
+}
+
+void FMallocCrash::PrintPoolsUsage()
+{
+	FPlatformMisc::LowLevelOutputDebugString( TEXT( "FPoolDesc used\n" ) );
+	for( uint32 Index = 0; Index < FMallocCrash::NUM_POOLS; ++Index )
+	{
+		const FMallocCrashPool& CrashPool = *Pools[Index];
+		FPlatformMisc::LowLevelOutputDebugStringf( TEXT( "FPoolDesc(%5u,%4u),\n" ), CrashPool.AllocationSize-FMallocCrash::PER_ALLOC_OVERHEAD, CrashPool.MaxUsedIndex );
+	}
+
+	FPlatformMisc::LowLevelOutputDebugString( TEXT( "FPoolDesc tweaked\n" ) );
+	for( uint32 Index = 0; Index < FMallocCrash::NUM_POOLS; ++Index )
+	{
+		const FMallocCrashPool& CrashPool = *Pools[Index];
+		FPlatformMisc::LowLevelOutputDebugStringf( TEXT( "FPoolDesc(%5u,%4u),\n" ), CrashPool.AllocationSize-FMallocCrash::PER_ALLOC_OVERHEAD, Align(CrashPool.MaxUsedIndex*2+16,16) );
+	}
+	FPlatformMisc::LowLevelOutputDebugStringf( TEXT( "LargeMemoryPoolOffset=%u\n" ), LargeMemoryPoolOffset );
+}
+
+bool FMallocCrash::IsOnCrashedThread() const
+{
+	// Suspend threads other than the crashed one to prevent serious memory errors.
+	// Only the crashed thread can do anything meaningful from here anyway.
+	if( CrashedThreadId == FPlatformTLS::GetCurrentThreadId() )
+	{
+		return true;
+	}
+	else
+	{
+		FPlatformProcess::SleepInfinite();
+		return false;
+	}
+}
+
+bool FMallocCrash::IsPtrInLargePool( void* Ptr ) const
+{
+	const bool bResult = (Ptr >= &LargeMemoryPool[0] && Ptr < &LargeMemoryPool[LARGE_MEMORYPOOL_SIZE]);
+	return bResult;
+}
+
+bool FMallocCrash::IsPtrInSmallPool( void* Ptr ) const
+{
+	const bool bResult = (Ptr >= &SmallMemoryPool[0] && Ptr < &SmallMemoryPool[GetSmallPoolTotalSize()]);
+	return bResult;
+}
+
+bool FMallocCrash::GetAllocationSize( void *Original, SIZE_T &SizeOut )
+{
+	SizeOut = FMallocCrash::GetAllocationSize(Original);
+	return true;
+}
+
+uint32 FMallocCrash::GetSmallPoolTotalSize() const
+{
+	static uint32 TotalSize = 0;
+	if( TotalSize == 0 )
+	{
+		for( uint32 Index = 0; Index < NUM_POOLS; ++Index )
+		{
+			const FPoolDesc& PoolDesc = GetPoolDesc(Index);
+			check(PoolDesc.NumAllocs%16==0);
+			check(PoolDesc.Size%16==0);
+			TotalSize += PoolDesc.NumAllocs*(PoolDesc.Size+PER_ALLOC_OVERHEAD);
+		}
+	}
+	return TotalSize;
+}
+
+void FMallocCrash::InitializeSmallPools()
+{
+	uint32 MallocCrashOverhead = 0;
+
+	for( uint32 Index = 0; Index < NUM_POOLS; ++Index )
+	{
+		Pools[Index] = new FMallocCrashPool( GetPoolDesc(Index), *this );
+		MallocCrashOverhead += Pools[Index]->AllocatedMemory;
+	}
+
+	check(SmallMemoryPoolOffset==GetSmallPoolTotalSize());
+}
+
+FMallocCrashPool* FMallocCrash::FindPoolFromSize( uint32 AllocationSize ) const
+{
+	for( uint32 Index = 0; Index < NUM_POOLS; ++Index )
+	{
+		FMallocCrashPool* Pool = Pools[Index];
+		if( AllocationSize <= Pool->AllocationSize-PER_ALLOC_OVERHEAD )
+		{
+			return Pool;
+		}
+	}
+
+	// Use large allocation pool.
+	return nullptr;
+}
+
+uint8* FMallocCrash::AllocateFromSmallPool( uint32 AllocationSize )
+{
+	if( SmallMemoryPoolOffset + AllocationSize <= GetSmallPoolTotalSize() )
+	{
+		const uint32 ReturnMemoryPoolOffset = SmallMemoryPoolOffset;
+		SmallMemoryPoolOffset += AllocationSize;
+		return (uint8*)SmallMemoryPool+ReturnMemoryPoolOffset;
+	}
+	else
+	{
+		check(0);
+	}
+	return nullptr;
+}
+
+uint32 FMallocCrash::GetAllocationSize( void *Original )
+{
+	FPtrInfo* PtrInfo = (FPtrInfo*)((uint8*)Original-PER_ALLOC_OVERHEAD);
+	return PtrInfo->Size;
+}
\ No newline at end of file
diff --git a/Engine/Source/Runtime/Core/Private/HAL/UnrealMemory.cpp b/Engine/Source/Runtime/Core/Private/HAL/UnrealMemory.cpp
--- a/Engine/Source/Runtime/Core/Private/HAL/UnrealMemory.cpp
+++ b/Engine/Source/Runtime/Core/Private/HAL/UnrealMemory.cpp
@@ -13,12 +13,16 @@
 #include "MallocDebug.h"
 #include "MallocProfiler.h"
 #include "MallocThreadSafeProxy.h"
+#include "MallocCrash.h"
 
 
 /** Helper function called on first allocation to create and initialize GMalloc */
 void GCreateMalloc()
 {
 	GMalloc = FPlatformMemory::BaseAllocator();
+	// Setup malloc crash as soon as possible.
+	FMallocCrash::Get( GMalloc );
+
 // so now check to see if we are using a Mem Profiler which wraps the GMalloc
 #if USE_MALLOC_PROFILER
 	#if WITH_ENGINE && IS_MONOLITHIC
diff --git a/Engine/Source/Runtime/Core/Public/HAL/MallocCrash.h b/Engine/Source/Runtime/Core/Public/HAL/MallocCrash.h
new file mode 100644
--- /dev/null
+++ b/Engine/Source/Runtime/Core/Public/HAL/MallocCrash.h
@@ -0,0 +1,141 @@
+// Copyright 1998-2014 Epic Games, Inc. All Rights Reserved.
+
+#pragma once
+
+/** Describes a pointer. */
+struct FPtrInfo
+{
+	/** Size of the allocation. */
+	uint64 Size;
+	/** Address of the allocation. */
+	uint8* Ptr;
+
+#if PLATFORM_32BITS
+	/** Explicit padding for 32 bit builds */
+	uint8 Padding[4];
+#endif
+
+	FPtrInfo():
+		Size(0),
+		Ptr(0)
+	{}
+
+	FPtrInfo( void* NewPtr ):
+		Size(0),
+		Ptr((uint8*)NewPtr)
+	{}
+};
+
+struct FPoolDesc;
+struct FMallocCrashPool;
+
+/**
+ * Simple pooled memory allocator that uses preallocated memory.
+ * Instance of this class replaces GMalloc after a crash, so we can use dynamic memory allocation even if the app crashes due to OOM.
+ */
+struct FMallocCrash final : public FMalloc
+{
+	friend struct FPoolDesc;
+	friend struct FMallocCrashPool;
+
+private:
+	enum
+	{
+		LARGE_MEMORYPOOL_SIZE = 2 * 1024 * 1024,
+		REQUIRED_ALIGNMENT = 16,
+		PER_ALLOC_OVERHEAD = sizeof(FPtrInfo),
+
+		NUM_POOLS = 14,
+		MAX_NUM_ALLOCS_IN_POOL = 2048,
+
+		MEM_TAG     = 0xfe,
+		MEM_WIPETAG = 0xcd,
+	};
+
+public:
+	FMallocCrash( FMalloc* MainMalloc );
+	virtual ~FMallocCrash();
+
+	/** Creates a new instance. */
+	static CORE_API FMallocCrash& Get( FMalloc* MainMalloc = nullptr );
+
+	/**
+	 * Sets as GMalloc.
+	 * This method locks to the thread that crashed.
+	 * Any next calls to this method or GMalloc from other threads will dead-lock those threads,
+	 * but this is ok, because we are shutting down.
+	 * This also fixes a lot of potential issues of using dynamic allocation during crash dumping.
+	 * 
+	 * @warning
+	 * This is not super safe, may interfere with other allocations, ie. replacing vtable during running 
+	 * the code from the previous malloc will probably crash other threads.
+	 */
+	void SetAsGMalloc();
+
+	// FMalloc interface.
+	virtual void* Malloc( SIZE_T Size, uint32 Alignment ) override;
+
+	virtual void* Realloc( void* Ptr, SIZE_T NewSize, uint32 Alignment ) override;
+
+	virtual void Free( void* /*Ptr*/ ) override;
+
+	virtual bool GetAllocationSize( void *Original, SIZE_T &SizeOut ) override;
+
+	virtual bool IsInternallyThreadSafe() const override
+	{
+		return true;
+	}
+
+	virtual bool ValidateHeap() override
+	{
+		// Nothing to do here.
+		return true;
+	}
+
+	virtual const TCHAR * GetDescriptiveName() override
+	{
+		return TEXT( "MallocCrash" );
+	}
+
+	CORE_API void PrintPoolsUsage();
+
+protected:
+	/** Whether it is safe to call crash malloc's methods. */
+	bool IsOnCrashedThread() const;
+
+	bool IsPtrInLargePool(void* Ptr) const;
+	bool IsPtrInSmallPool(void* Ptr) const;
+
+	const FPoolDesc& GetPoolDesc( uint32 Index ) const;
+	uint32 GetSmallPoolTotalSize() const;
+
+	void InitializeSmallPools();
+	FMallocCrashPool* FindPoolFromSize( uint32 AllocationSize ) const;
+	uint8* AllocateFromSmallPool( uint32 AllocationSize );
+
+	static uint32 GetAllocationSize(void *Original);
+
+protected:
+	/** Used to lock crash malloc to one thread. */
+	FCriticalSection InternalLock;
+
+	/** The id of the thread that crashed. */
+	int32 CrashedThreadId;
+
+	/** Preallocated memory pool for large allocations. */
+	uint8* LargeMemoryPool;
+
+	/** Current position in the large memory pool. */
+	uint32 LargeMemoryPoolOffset;
+
+	/** Preallocated memory pool for small allocations. */
+	uint8* SmallMemoryPool;
+
+	/** Current position in the small memory pool. */
+	uint32 SmallMemoryPoolOffset;
+
+	/** Previously used malloc. */
+	FMalloc* PreviousMalloc;
+
+	FMallocCrashPool* Pools[NUM_POOLS];
+};
