CrashReport - Minor tweak

diff --git a/Engine/Source/Runtime/Core/Private/GenericPlatform/GenericPlatformStackWalk.cpp b/Engine/Source/Runtime/Core/Private/GenericPlatform/GenericPlatformStackWalk.cpp
--- a/Engine/Source/Runtime/Core/Private/GenericPlatform/GenericPlatformStackWalk.cpp
+++ b/Engine/Source/Runtime/Core/Private/GenericPlatform/GenericPlatformStackWalk.cpp
@@ -15,7 +15,6 @@ FProgramCounterSymbolInfo::FProgramCounterSymbolInfo() :
 	FCStringAnsi::Strncpy( Filename, "UnknownFile", MAX_NAME_LENGHT );
 }
 
-
 bool FGenericPlatformStackWalk::ProgramCounterToHumanReadableString( int32 CurrentCallDepth, uint64 ProgramCounter, ANSICHAR* HumanReadableString, SIZE_T HumanReadableStringSize, FGenericCrashContext* Context )
 {
 	if (HumanReadableString && HumanReadableStringSize > 0)
@@ -32,11 +31,7 @@ bool FGenericPlatformStackWalk::ProgramCounterToHumanReadableString( int32 Curre
 		const UPTRINT RealPos = FMath::Max( (UPTRINT)Pos0, (UPTRINT)Pos1 );
 		const ANSICHAR* StrippedModuleName = RealPos > 0 ? (const ANSICHAR*)(RealPos+1) : SymbolInfo.ModuleName;
 	
-#if	PLATFORM_64BITS
 		FCStringAnsi::Sprintf( StackLine, "%s!%s (0x%016llx) + %i bytes [%s:%i]", StrippedModuleName, (const ANSICHAR*)SymbolInfo.FunctionName, SymbolInfo.ProgramCounter, SymbolInfo.SymbolDisplacement, (const ANSICHAR*)SymbolInfo.Filename, SymbolInfo.LineNumber );
-#else
-		FCStringAnsi::Sprintf( StackLine, "%s!%s (0x%08x) + %i bytes [%s:%i]", StrippedModuleName, (const ANSICHAR*)SymbolInfo.FunctionName, (uint32)SymbolInfo.ProgramCounter, SymbolInfo.OffsetInModule, (const ANSICHAR*)SymbolInfo.Filename, SymbolInfo.LineNumber );
-#endif
 
 		// Append the stack line.
 		FCStringAnsi::Strcat( HumanReadableString, HumanReadableStringSize, StackLine );
