Merging CL #2318747 from 4.5 branch to Main

diff --git a/Engine/Source/Runtime/Core/Private/Mac/MacApplication.cpp b/Engine/Source/Runtime/Core/Private/Mac/MacApplication.cpp
--- a/Engine/Source/Runtime/Core/Private/Mac/MacApplication.cpp
+++ b/Engine/Source/Runtime/Core/Private/Mac/MacApplication.cpp
@@ -328,17 +328,17 @@ void FMacApplication::ProcessEvent(FMacEvent const* const Event)
 	if(AppKitEvent)
 	{
 		// Process a standard NSEvent as we always did
-		MacApplication->ProcessNSEvent(AppKitEvent, Event->GetWindow(), Event->GetMousePosition());
+		TSharedPtr<FMacWindow> MacWindow = FindWindowByNSWindow(MacApplication->Windows, &MacApplication->WindowsMutex, Event->GetWindow());
+		MacApplication->ProcessNSEvent(AppKitEvent, MacWindow, Event->GetMousePosition());
 	}
 	else if(Notification)
 	{
 		// Notifications need to mapped to the right handler function, that's no longer handled in the
 		// call location.
 		NSString* const NotificationName = [Notification name];
-		TSharedPtr<FMacWindow> Window = Event->GetWindow();
-		if(Window.IsValid())
+		FCocoaWindow* CocoaWindow = Event->GetWindow();
+		if (CocoaWindow)
 		{
-			FCocoaWindow* CocoaWindow = Window->GetWindowHandle();
 			if(NotificationName == NSWindowDidResizeNotification)
 			{
 				MacApplication->OnWindowDidResize( CocoaWindow );
diff --git a/Engine/Source/Runtime/Core/Private/Mac/MacEvent.cpp b/Engine/Source/Runtime/Core/Private/Mac/MacEvent.cpp
--- a/Engine/Source/Runtime/Core/Private/Mac/MacEvent.cpp
+++ b/Engine/Source/Runtime/Core/Private/Mac/MacEvent.cpp
@@ -13,12 +13,8 @@ FMacEvent::FMacEvent(NSEvent* const Event)
 	check(Event);
 	[Event retain];
 	
-	NSWindow* Window = FMacApplication::FindMacEventWindow(Event);
-	if(Window)
-	{
-		MacWindow = FMacApplication::FindMacWindowByNSWindow((FCocoaWindow*)Window);
-	}
-	
+	CocoaWindow = FMacApplication::FindMacEventWindow(Event);
+
 	CGEventRef CGEvent = [Event CGEvent];
 	CGPoint CGMouseLocation = CGEventGetLocation(CGEvent);
 	MousePosition.X = CGMouseLocation.x;
@@ -31,12 +27,9 @@ FMacEvent::FMacEvent(NSNotification* const Notification, NSWindow* const Window)
 	SCOPED_AUTORELEASE_POOL;
 	check(Notification);
 	[Notification retain];
-	
-	if(Window)
-	{
-		MacWindow = FMacApplication::FindMacWindowByNSWindow((FCocoaWindow*)Window);
-	}
-	
+
+	CocoaWindow = [Window isKindOfClass:[FCocoaWindow class]] ? (FCocoaWindow*)Window : nullptr;
+
 	CGEventRef Event = CGEventCreate(NULL);
 	CGPoint CursorPos = CGEventGetLocation(Event);
 	MousePosition.X = CursorPos.x;
@@ -107,9 +100,9 @@ NSNotification* FMacEvent::GetNotification(void) const
 	}
 }
 
-TSharedPtr<FMacWindow> FMacEvent::GetWindow(void) const
+FCocoaWindow* FMacEvent::GetWindow(void) const
 {
-	return MacWindow;
+	return CocoaWindow;
 }
 
 FVector2D FMacEvent::GetMousePosition(void) const
diff --git a/Engine/Source/Runtime/Core/Public/Mac/MacEvent.h b/Engine/Source/Runtime/Core/Public/Mac/MacEvent.h
--- a/Engine/Source/Runtime/Core/Public/Mac/MacEvent.h
+++ b/Engine/Source/Runtime/Core/Public/Mac/MacEvent.h
@@ -45,14 +45,14 @@ public:
 	// Get the NSNotification for this FMacEvent, nil if not an NSNotification.
 	NSNotification* GetNotification(void) const;
 	
-	// Get the FMacWindow for this FMacEvent, may be nullptr if not a window event, or for a non-FMacWindow.
-	TSharedPtr<FMacWindow> GetWindow(void) const;
+	// Get the FCocoaWindow for this FMacEvent, may be nullptr if not a window event, or for a non-FCocoaWindow.
+	FCocoaWindow* GetWindow(void) const;
 	
 	// Get the mouse position for this FMacEvent at the time of posting on the AppKit Main Thread.
 	FVector2D GetMousePosition(void) const;
 	
 private:
-	TSharedPtr<FMacWindow> MacWindow; // The retained UE4 window for this event, or nullptr
+	FCocoaWindow* CocoaWindow; // The retained window for this event, or nullptr
 	NSObject* EventData; // The retained NSEvent or NSNotification for this event or nil.
 	FVector2D MousePosition; // The mouse position for the event in UE4 coordinates.
 };
