Code analysis macro rename, conflicts with CoreAnimation Framework

diff --git a/Engine/Source/Runtime/Core/Public/GenericPlatform/GenericPlatformMemory.h b/Engine/Source/Runtime/Core/Public/GenericPlatform/GenericPlatformMemory.h
--- a/Engine/Source/Runtime/Core/Public/GenericPlatform/GenericPlatformMemory.h
+++ b/Engine/Source/Runtime/Core/Public/GenericPlatform/GenericPlatformMemory.h
@@ -164,7 +164,7 @@ struct CORE_API FGenericPlatformMemory
 	/** Initializes platform memory specific constants. */
 	static void Init();
 	
-	static CA_NORETURN void OnOutOfMemory(uint64 Size, uint32 Alignment);
+	static CA_NO_RETURN void OnOutOfMemory(uint64 Size, uint32 Alignment);
 
 	/** Initializes the memory pools, should be called by the init function. */
 	static void SetupMemoryPools();
diff --git a/Engine/Source/Runtime/Core/Public/HAL/MallocBinned.h b/Engine/Source/Runtime/Core/Public/HAL/MallocBinned.h
--- a/Engine/Source/Runtime/Core/Public/HAL/MallocBinned.h
+++ b/Engine/Source/Runtime/Core/Public/HAL/MallocBinned.h
@@ -315,7 +315,7 @@ private:
 #endif
 
 	// Implementation. 
-	CA_NORETURN void OutOfMemory(uint64 Size, uint32 Alignment=0)
+	CA_NO_RETURN void OutOfMemory(uint64 Size, uint32 Alignment=0)
 	{
 		// this is expected not to return
 		FPlatformMemory::OnOutOfMemory(Size, Alignment);
diff --git a/Engine/Source/Runtime/Core/Public/Misc/CoreMiscDefines.h b/Engine/Source/Runtime/Core/Public/Misc/CoreMiscDefines.h
--- a/Engine/Source/Runtime/Core/Public/Misc/CoreMiscDefines.h
+++ b/Engine/Source/Runtime/Core/Public/Misc/CoreMiscDefines.h
@@ -70,7 +70,7 @@
 #endif
 
 #if USING_CODE_ANALYSIS
-	#if !defined( CA_IN ) || !defined( CA_OUT ) || !defined( CA_READ_ONLY ) || !defined( CA_WRITE_ONLY ) || !defined( CA_VALID_POINTER ) || !defined( CA_CHECK_RETVAL ) || !defined( CA_NORETURN) || !defined( CA_SUPPRESS ) || !defined( CA_ASSUME )
+	#if !defined( CA_IN ) || !defined( CA_OUT ) || !defined( CA_READ_ONLY ) || !defined( CA_WRITE_ONLY ) || !defined( CA_VALID_POINTER ) || !defined( CA_CHECK_RETVAL ) || !defined( CA_NO_RETURN ) || !defined( CA_SUPPRESS ) || !defined( CA_ASSUME )
 		#error Code analysis macros are not configured correctly for this platform
 	#endif
 #else
@@ -81,7 +81,7 @@
 	#define CA_WRITE_ONLY
 	#define CA_VALID_POINTER
 	#define CA_CHECK_RETVAL
-	#define CA_NORETURN
+	#define CA_NO_RETURN
 	#define CA_SUPPRESS( WarningNumber )
 	#define CA_ASSUME( Expr )
 #endif
diff --git a/Engine/Source/Runtime/Core/Public/Windows/WindowsPlatformCodeAnalysis.h b/Engine/Source/Runtime/Core/Public/Windows/WindowsPlatformCodeAnalysis.h
--- a/Engine/Source/Runtime/Core/Public/Windows/WindowsPlatformCodeAnalysis.h
+++ b/Engine/Source/Runtime/Core/Public/Windows/WindowsPlatformCodeAnalysis.h
@@ -62,7 +62,7 @@
 	#define CA_CHECK_RETVAL [returnvalue:Post(MustCheck=Yes)]
 
 	// Function is expected to never return
-	#define CA_NORETURN __declspec(noreturn)
+	#define CA_NO_RETURN __declspec(noreturn)
 
 	// Suppresses a warning for a single occurrence.  Should be used only for code analysis warnings on Windows platform!
 	#define CA_SUPPRESS( WarningNumber ) __pragma( warning( suppress: WarningNumber ) )
