Internal compiler error fix for pre-release Clang 3.4.2, which affected Xcode 5.1.1.

diff --git a/Engine/Source/Runtime/Core/Public/Delegates/DelegateSignatureImpl_Variadics.inl b/Engine/Source/Runtime/Core/Public/Delegates/DelegateSignatureImpl_Variadics.inl
--- a/Engine/Source/Runtime/Core/Public/Delegates/DelegateSignatureImpl_Variadics.inl
+++ b/Engine/Source/Runtime/Core/Public/Delegates/DelegateSignatureImpl_Variadics.inl
@@ -92,7 +92,7 @@ public:
 	 * Static: Creates a raw C++ pointer global function delegate
 	 */
 	template <typename... VarTypes>
-	inline static TBaseDelegate<RetValType, ParamTypes...> CreateStatic(typename TBaseStaticDelegateInstance<TFuncType, VarTypes...>::FFuncPtr InFunc, VarTypes... Vars)
+	inline static TBaseDelegate<RetValType, ParamTypes...> CreateStatic(typename TIdentity<RetValType (*)(ParamTypes..., VarTypes...)>::Type InFunc, VarTypes... Vars)
 	{
 		return TBaseStaticDelegateInstance<TFuncType, VarTypes...>::Create(InFunc, Vars...);
 	}
