Stop titled windows on Mac OS X wandering down the screen when repeatedly entering & exiting fullscreen mode.

diff --git a/Engine/Source/Runtime/Core/Private/Mac/MacWindow.cpp b/Engine/Source/Runtime/Core/Private/Mac/MacWindow.cpp
--- a/Engine/Source/Runtime/Core/Private/Mac/MacWindow.cpp
+++ b/Engine/Source/Runtime/Core/Private/Mac/MacWindow.cpp
@@ -274,7 +274,7 @@ void FMacWindow::MoveWindowTo( int32 X, int32 Y )
 {
 	MainThreadCall(^{
 		SCOPED_AUTORELEASE_POOL;
-		const int32 InvertedY = FPlatformMisc::ConvertSlateYPositionToCocoa(Y) - [WindowHandle frame].size.height + 1;
+		const int32 InvertedY = FPlatformMisc::ConvertSlateYPositionToCocoa(Y) - [WindowHandle openGLFrame].size.height + 1;
 		[WindowHandle setFrameOrigin: NSMakePoint(X, InvertedY)];
 	}, UE4ResizeEventMode, true);
 }
