MallocCrash - Small and large memory pool size are now multiple of PageSize, If PageSize is not set, this will use 64k as PageSize

diff --git a/Engine/Source/Runtime/Core/Private/HAL/MallocCrash.cpp b/Engine/Source/Runtime/Core/Private/HAL/MallocCrash.cpp
--- a/Engine/Source/Runtime/Core/Private/HAL/MallocCrash.cpp
+++ b/Engine/Source/Runtime/Core/Private/HAL/MallocCrash.cpp
@@ -120,7 +120,7 @@ struct FMallocCrashPool
 		{
 			FPlatformMisc::DebugBreak();
 			FPlatformMisc::LowLevelOutputDebugStringf( TEXT( "AllocateFromPool run out of memory allocating %u bytes for %u allocations\n" ), InAllocationSize, MaxNumAllocations );
-			FPlatformMisc::LowLevelOutputDebugString( TEXT( "Please increase MaxNumAllocations for that pool\n" ) );
+			FPlatformMisc::LowLevelOutputDebugString( TEXT( "Please increase MaxNumAllocations for that pool, exiting...\n" ) );
 			FPlatformMisc::RequestExit( true );
 		}
 		return nullptr;
@@ -178,14 +178,23 @@ FMallocCrash::FMallocCrash( FMalloc* MainMalloc ) :
 	SmallMemoryPoolOffset( 0 ),
 	PreviousMalloc( MainMalloc )
 {
-	LargeMemoryPool = Align( (uint8*)FPlatformMemory::BinnedAllocFromOS( LARGE_MEMORYPOOL_SIZE + REQUIRED_ALIGNMENT ), REQUIRED_ALIGNMENT );
-	SmallMemoryPool = Align( (uint8*)FPlatformMemory::BinnedAllocFromOS( GetSmallPoolTotalSize() + REQUIRED_ALIGNMENT ), REQUIRED_ALIGNMENT );
+	const uint32 LargeMemoryPoolSize = Align(LARGE_MEMORYPOOL_SIZE,SafePageSize());
+	LargeMemoryPool = (uint8*)FPlatformMemory::BinnedAllocFromOS( LargeMemoryPoolSize );
+	SmallMemoryPool = (uint8*)FPlatformMemory::BinnedAllocFromOS( (SIZE_T)GetSmallPoolTotalSize() );
 	if( !SmallMemoryPool || !LargeMemoryPool )
 	{
+		FPlatformMisc::LowLevelOutputDebugString( TEXT( "Memory pools allocations failed, exiting...\n" ) );
 		FPlatformMisc::RequestExit(true);
 	}
+
+	if((UPTRINT(LargeMemoryPool) & (REQUIRED_ALIGNMENT - 1)) != 0 || (UPTRINT(SmallMemoryPool) & (REQUIRED_ALIGNMENT - 1)) != 0)
+	{
+		FPlatformMisc::LowLevelOutputDebugString( TEXT( "OS allocations must be aligned to a value multiple of 16, exiting...\n" ) );
+		FPlatformMisc::RequestExit(true);
+	}
+
 	InitializeSmallPools();
-	FPlatformMisc::LowLevelOutputDebugStringf( TEXT( "FMallocCrash overhead is %u bytes\n" ), LARGE_MEMORYPOOL_SIZE+GetSmallPoolTotalSize() );
+	FPlatformMisc::LowLevelOutputDebugStringf( TEXT( "FMallocCrash overhead is %u bytes\n" ), LargeMemoryPoolSize+GetSmallPoolTotalSize() );
 }
 
 FMallocCrash::~FMallocCrash()
@@ -243,7 +252,7 @@ void* FMallocCrash::Malloc( SIZE_T Size, uint32 Alignment )
 			{
 				FPlatformMisc::DebugBreak();
 				FPlatformMisc::LowLevelOutputDebugStringf( TEXT( "MallocCrash run out of memory allocating %u bytes, free %u bytes\n" ), Size32, LARGE_MEMORYPOOL_SIZE-LargeMemoryPoolOffset );
-				FPlatformMisc::LowLevelOutputDebugString( TEXT( "Please increase LARGE_MEMORYPOOL_SIZE \n" ) );
+				FPlatformMisc::LowLevelOutputDebugString( TEXT( "Please increase LARGE_MEMORYPOOL_SIZE, exiting...\n" ) );
 				FPlatformMisc::RequestExit( true );			
 			}
 		}
@@ -263,11 +272,21 @@ void* FMallocCrash::Realloc( void* Ptr, SIZE_T NewSize, uint32 Alignment )
 
 			if( bPreviousMalloc )
 			{
-				// Realloc from the previous allocator.
-				PreviousMalloc->GetAllocationSize(Ptr,PtrSize);
-				if( PtrSize == 0 )
+				// At this moment, we can safely get allocation size only from the binned malloc, this may change in future.
+				if( FCStringWide::Strcmp( PreviousMalloc->GetDescriptiveName(), TEXT("binned") ) == 0 )
+				{
+					// Realloc from the previous allocator.
+					PreviousMalloc->GetAllocationSize(Ptr,PtrSize);
+					if( PtrSize == 0 )
+					{
+						PtrSize = 0;
+					}
+				}
+				// There is nothing we can do about it.
+				else
 				{
-					PtrSize = 0;
+					FPlatformMisc::LowLevelOutputDebugString( TEXT( "Realloc from previous malloc, exiting...\n" ) );
+					FPlatformMisc::RequestExit( true );	
 				}
 			}
 			else
@@ -389,6 +408,8 @@ uint32 FMallocCrash::GetSmallPoolTotalSize() const
 			check(PoolDesc.Size%16==0);
 			TotalSize += PoolDesc.NumAllocs*(PoolDesc.Size+PER_ALLOC_OVERHEAD);
 		}
+
+		TotalSize = Align(TotalSize,SafePageSize());
 	}
 	return TotalSize;
 }
@@ -403,7 +424,7 @@ void FMallocCrash::InitializeSmallPools()
 		MallocCrashOverhead += Pools[Index]->AllocatedMemory;
 	}
 
-	check(SmallMemoryPoolOffset==GetSmallPoolTotalSize());
+	check(SmallMemoryPoolOffset<=GetSmallPoolTotalSize());
 }
 
 FMallocCrashPool* FMallocCrash::FindPoolFromSize( uint32 AllocationSize ) const
@@ -440,4 +461,14 @@ uint32 FMallocCrash::GetAllocationSize( void *Original )
 {
 	FPtrInfo* PtrInfo = (FPtrInfo*)((uint8*)Original-PER_ALLOC_OVERHEAD);
 	return PtrInfo->Size;
-}
\ No newline at end of file
+}
+
+uint32 FMallocCrash::SafePageSize()
+{
+	const uint32 PageSize = (uint32)FPlatformMemory::GetStats().PageSize;
+	if( PageSize == 0 )
+	{
+		return 65536u;
+	}
+	return PageSize;
+}
diff --git a/Engine/Source/Runtime/Core/Public/HAL/MallocCrash.h b/Engine/Source/Runtime/Core/Public/HAL/MallocCrash.h
--- a/Engine/Source/Runtime/Core/Public/HAL/MallocCrash.h
+++ b/Engine/Source/Runtime/Core/Public/HAL/MallocCrash.h
@@ -115,6 +115,11 @@ protected:
 
 	static uint32 GetAllocationSize(void *Original);
 
+	/**
+	 * @return page size, if page size is not initialized, returns 64k.
+	 */
+	static uint32 SafePageSize();
+
 protected:
 	/** Used to lock crash malloc to one thread. */
 	FCriticalSection InternalLock;
