CrashReport - Malloc crash for Windows

diff --git a/Engine/Source/Runtime/Core/Private/Misc/OutputDevice.cpp b/Engine/Source/Runtime/Core/Private/Misc/OutputDevice.cpp
--- a/Engine/Source/Runtime/Core/Private/Misc/OutputDevice.cpp
+++ b/Engine/Source/Runtime/Core/Private/Misc/OutputDevice.cpp
@@ -7,6 +7,7 @@
 #include "Core.h"
 #include "ExceptionHandling.h"
 #include "VarargsHelper.h"
+#include "MallocCrash.h"
 
 DEFINE_LOG_CATEGORY_STATIC(LogOutputDevice, Log, All);
 
@@ -164,6 +165,11 @@ void StaticFailDebug( const TCHAR* Error, const ANSICHAR* File, int32 Line, cons
 	// For ensure log should be flushed in the engine loop.
 	if( !bIsEnsure )
 	{
+#if	PLATFORM_WINDOWS // @TODO yrx 2014-09-30 Remove later.
+		// Switch to malloc crash.
+		FMallocCrash::Get().SetAsGMalloc(); 
+#endif // PLATFORM_WINDOWS
+
 		GLog->PanicFlushThreadedLogs();
 	}
 
