Update the window position to {0,0} when entering EWindowMode::Fullscreen on OS X, since that's what it is.

diff --git a/Engine/Source/Runtime/Core/Private/Mac/MacApplication.cpp b/Engine/Source/Runtime/Core/Private/Mac/MacApplication.cpp
--- a/Engine/Source/Runtime/Core/Private/Mac/MacApplication.cpp
+++ b/Engine/Source/Runtime/Core/Private/Mac/MacApplication.cpp
@@ -1285,8 +1285,8 @@ void FMacApplication::OnWindowDidMove( FCocoaWindow* Window )
 	NSRect OpenGLFrame = [Window openGLFrame];
 	
 	const int32 X = (int32)WindowFrame.origin.x;
-	const int32 Y = FPlatformMisc::ConvertSlateYPositionToCocoa( (int32)WindowFrame.origin.y ) - OpenGLFrame.size.height + 1;
-
+	const int32 Y = ([Window windowMode] != EWindowMode::Fullscreen) ? FPlatformMisc::ConvertSlateYPositionToCocoa( (int32)WindowFrame.origin.y ) - OpenGLFrame.size.height + 1 : 0;
+	
 	TSharedPtr< FMacWindow > EventWindow = FindWindowByNSWindow( Windows, &WindowsMutex, Window );
 	if( EventWindow.IsValid() )
 	{
