Finish implementing EWindowMode::Fullscreen on Mac OS X by lying to Slate about mouse coordinates & window size, just like we did in UE3. When the fullscreen mode changes set the mouse scale to 1 for Windowed & WindowedFullscreen, or ViewSize/WindowSize for Fullscreen and then use this scaling factor to marshal mouse coordinates between Slate & the OS. EWindowMode::Fullscreen should now work correctly on OS X & use the fast path through CGLFlushDrawable rather than having to render at display resolution.

diff --git a/Engine/Source/Runtime/Core/Private/Mac/CocoaWindow.cpp b/Engine/Source/Runtime/Core/Private/Mac/CocoaWindow.cpp
--- a/Engine/Source/Runtime/Core/Private/Mac/CocoaWindow.cpp
+++ b/Engine/Source/Runtime/Core/Private/Mac/CocoaWindow.cpp
@@ -6,6 +6,7 @@
 #include "CocoaTextView.h"
 #include "MacEvent.h"
 #include "CocoaThread.h"
+#include "MacCursor.h"
 
 TArray< FCocoaWindow* > GRunningModalWindows;
 
@@ -59,7 +60,7 @@ NSString* NSPerformDragOperation = @"NSPerformDragOperation";
 	SCOPED_AUTORELEASE_POOL;
 	if(self.TargetWindowMode == EWindowMode::Fullscreen || WindowMode == EWindowMode::Fullscreen)
 	{
-		return self.PreFullScreenRect;
+		return {{0, 0}, self.PreFullScreenRect.size};
 	}
 	else if([self styleMask] & (NSTexturedBackgroundWindowMask))
 	{
@@ -368,6 +369,15 @@ NSString* NSPerformDragOperation = @"NSPerformDragOperation";
 	{
 		FMacEvent::SendToGameRunLoop(notification, self, EMacEventSendMethod::Async, @[ NSDefaultRunLoopMode, UE4FullscreenEventMode ]);
 	}
+	FMacCursor* MacCursor = (FMacCursor*)MacApplication->Cursor.Get();
+	if ( MacCursor )
+	{
+		NSSize WindowSize = [self frame].size;
+		NSSize ViewSize = [self openGLFrame].size;
+		float WidthScale = ViewSize.width / WindowSize.width;
+		float HeightScale = ViewSize.height / WindowSize.height;
+		MacCursor->SetMouseScaling(FVector2D(WidthScale, HeightScale));
+	}
 }
 
 - (void)windowDidExitFullScreen:(NSNotification *)notification
@@ -378,6 +388,11 @@ NSString* NSPerformDragOperation = @"NSPerformDragOperation";
 	{
 		FMacEvent::SendToGameRunLoop(notification, self, EMacEventSendMethod::Async, @[ NSDefaultRunLoopMode, UE4FullscreenEventMode ]);
 	}
+	FMacCursor* MacCursor = (FMacCursor*)MacApplication->Cursor.Get();
+	if ( MacCursor )
+	{
+		MacCursor->SetMouseScaling(FVector2D(1.0f, 1.0f));
+	}
 }
 
 - (void)windowDidBecomeKey:(NSNotification *)Notification
diff --git a/Engine/Source/Runtime/Core/Private/Mac/MacCursor.cpp b/Engine/Source/Runtime/Core/Private/Mac/MacCursor.cpp
--- a/Engine/Source/Runtime/Core/Private/Mac/MacCursor.cpp
+++ b/Engine/Source/Runtime/Core/Private/Mac/MacCursor.cpp
@@ -9,6 +9,7 @@ FMacCursor::FMacCursor()
 :	bIsVisible(true)
 ,	bAssociateMouseCursor(false)
 ,	MouseWarpDelta(FVector2D::ZeroVector)
+,	MouseScale(1.0f, 1.0f)
 {
 	SCOPED_AUTORELEASE_POOL;
 
@@ -153,16 +154,16 @@ FVector2D FMacCursor::GetPosition() const
 	CGPoint CursorPos = CGEventGetLocation(Event);
 	CFRelease(Event);
 	
-	return FVector2D(FMath::TruncToFloat(CursorPos.x), FMath::TruncToFloat(CursorPos.y));
+	return FVector2D(FMath::TruncToFloat(CursorPos.x * MouseScale.X), FMath::TruncToFloat(CursorPos.y * MouseScale.X));
 }
 
 void FMacCursor::SetPosition( const int32 X, const int32 Y )
 {
 	FVector2D CurrentPos = GetPosition();
 	FVector2D NewPos(X, Y);
-	MouseWarpDelta += (NewPos - CurrentPos);
+	MouseWarpDelta += (NewPos - CurrentPos) / MouseScale;
 
-	WarpCursor(X, Y);
+	WarpCursor(X / MouseScale.X, Y / MouseScale.Y);
 }
 
 void FMacCursor::SetType( const EMouseCursor::Type InNewCursor )
@@ -309,3 +310,13 @@ void FMacCursor::AssociateMouseAndCursorPosition( bool const bEnable )
 		CGAssociateMouseAndMouseCursorPosition( bEnable );
 	}
 }
+
+void FMacCursor::SetMouseScaling( FVector2D Scale )
+{
+	MouseScale = Scale;
+}
+
+FVector2D FMacCursor::GetMouseScaling( void )
+{
+	return MouseScale;
+}
diff --git a/Engine/Source/Runtime/Core/Private/Mac/MacWindow.cpp b/Engine/Source/Runtime/Core/Private/Mac/MacWindow.cpp
--- a/Engine/Source/Runtime/Core/Private/Mac/MacWindow.cpp
+++ b/Engine/Source/Runtime/Core/Private/Mac/MacWindow.cpp
@@ -3,6 +3,7 @@
 #include "Core.h"
 #include "MacWindow.h"
 #include "MacApplication.h"
+#include "MacCursor.h"
 #include "CocoaTextView.h"
 #include "CocoaThread.h"
 
@@ -556,7 +557,9 @@ bool FMacWindow::IsPointInWindow( int32 X, int32 Y ) const
 		
 		if([WindowHandle isOnActiveSpace])
 		{
-			PointInWindow = (NSPointInRect(NSMakePoint(X, Y), VisibleFrame) == YES);
+			FMacCursor* MacCursor = (FMacCursor*)MacApplication->Cursor.Get();
+			FVector2D MouseScale = MacCursor ? MacCursor->GetMouseScaling() : FVector2D(1.0f, 1.0f);
+			PointInWindow = (NSPointInRect(NSMakePoint(X / MouseScale.X, Y / MouseScale.Y), VisibleFrame) == YES);
 		}
 	}
 	return PointInWindow;
diff --git a/Engine/Source/Runtime/Core/Public/Mac/MacCursor.h b/Engine/Source/Runtime/Core/Public/Mac/MacCursor.h
--- a/Engine/Source/Runtime/Core/Public/Mac/MacCursor.h
+++ b/Engine/Source/Runtime/Core/Public/Mac/MacCursor.h
@@ -38,6 +38,10 @@ public:
 	FVector2D GetMouseWarpDelta( bool const bClearAccumulatedDelta );
 
 	void AssociateMouseAndCursorPosition( bool const bEnable );
+	
+	void SetMouseScaling( FVector2D Scale );
+	
+	FVector2D GetMouseScaling( void );
 
 private:
 
@@ -53,4 +57,5 @@ private:
 	NSCursor* CurrentCursor;
 	
 	FVector2D MouseWarpDelta;
+	FVector2D MouseScale;
 };
