Linux: clipboard support (PR #489).

- contributed by [REDACTED_PERSON].

diff --git a/Engine/Source/Runtime/Core/Private/Linux/LinuxPlatformMisc.cpp b/Engine/Source/Runtime/Core/Private/Linux/LinuxPlatformMisc.cpp
--- a/Engine/Source/Runtime/Core/Private/Linux/LinuxPlatformMisc.cpp
+++ b/Engine/Source/Runtime/Core/Private/Linux/LinuxPlatformMisc.cpp
@@ -313,6 +313,34 @@ uint32 FLinuxPlatformMisc::GetKeyMap( uint16* KeyCodes, FString* KeyNames, uint3
 	return NumMappings;
 }
 
+void FLinuxPlatformMisc::ClipboardCopy(const TCHAR* Str)
+{
+	if (SDL_HasClipboardText() == SDL_TRUE)
+	{
+		if (SDL_SetClipboardText(TCHAR_TO_UTF8(Str)))
+		{
+			UE_LOG(LogInit, Fatal, TEXT("Error copying clipboard contents: %s\n"), ANSI_TO_TCHAR(SDL_GetError()));
+		}
+	}
+}
+void FLinuxPlatformMisc::ClipboardPaste(class FString& Result)
+{
+	char* ClipContent;
+	ClipContent = SDL_GetClipboardText();
+
+	if (!ClipContent)
+	{
+		UE_LOG(LogInit, Fatal, TEXT("Error pasting clipboard contents: %s\n"), ANSI_TO_TCHAR(SDL_GetError()));
+		// unreachable
+		Result = TEXT("");
+	}
+	else
+	{
+		Result = FString(UTF8_TO_TCHAR(ClipContent));
+	}
+	SDL_free(ClipContent);
+}
+
 EAppReturnType::Type FLinuxPlatformMisc::MessageBoxExt(EAppMsgType::Type MsgType, const TCHAR* Text, const TCHAR* Caption)
 {
 	int NumberOfButtons = 0;
diff --git a/Engine/Source/Runtime/Core/Public/Linux/LinuxPlatformMisc.h b/Engine/Source/Runtime/Core/Public/Linux/LinuxPlatformMisc.h
--- a/Engine/Source/Runtime/Core/Public/Linux/LinuxPlatformMisc.h
+++ b/Engine/Source/Runtime/Core/Public/Linux/LinuxPlatformMisc.h
@@ -107,6 +107,9 @@ struct CORE_API FLinuxPlatformMisc : public FGenericPlatformMisc
 	static void LowLevelOutputDebugString(const TCHAR *Message);
 	static bool ControlScreensaver(EScreenSaverAction Action);
 
+	static void ClipboardCopy(const TCHAR* Str);
+	static void ClipboardPaste(class FString& Dest);
+
 	static const TCHAR* RootDir();
 	static void NormalizePath(FString& InPath);
 
