CrashReport - Added GenericPlatformContext

diff --git a/Engine/Source/Programs/CrashReporter/CrashReportClient/Private/Mac/CrashReportClientMainMac.cpp b/Engine/Source/Programs/CrashReporter/CrashReportClient/Private/Mac/CrashReportClientMainMac.cpp
--- a/Engine/Source/Programs/CrashReporter/CrashReportClient/Private/Mac/CrashReportClientMainMac.cpp
+++ b/Engine/Source/Programs/CrashReporter/CrashReportClient/Private/Mac/CrashReportClientMainMac.cpp
@@ -2,6 +2,7 @@
 
 #include "CrashReportClientApp.h"
 #include "ExceptionHandling.h"
+#include "MacPlatformCrashContext.h"
 #include "CocoaThread.h"
 
 /**
diff --git a/Engine/Source/Runtime/Core/Private/Android/AndroidMisc.cpp b/Engine/Source/Runtime/Core/Private/Android/AndroidMisc.cpp
--- a/Engine/Source/Runtime/Core/Private/Android/AndroidMisc.cpp
+++ b/Engine/Source/Runtime/Core/Private/Android/AndroidMisc.cpp
@@ -6,7 +6,8 @@
 #include <cpu-features.h>
 #include "ModuleManager.h"
 #include <android/keycodes.h>
-
+#include "AndroidPlatformCrashContext.h"
+#include "MallocCrash.h"
 
 DEFINE_LOG_CATEGORY_STATIC(LogEngine, Log, All);
 
@@ -250,6 +251,9 @@ void (* GCrashHandlerPointer)(const FGenericCrashContext& Context) = NULL;
 /** True system-specific crash handler that gets called first */
 void PlatformCrashHandler(int32 Signal, siginfo* Info, void* Context)
 {
+	// Switch to malloc crash.
+	//FMallocCrash::Get().SetAsGMalloc(); @todo uncomment after verification
+
 	fprintf(stderr, "Signal %d caught.\n", Signal);
 
 	FAndroidCrashContext CrashContext;
diff --git a/Engine/Source/Runtime/Core/Private/GenericPlatform/GenericPlatformContext.cpp b/Engine/Source/Runtime/Core/Private/GenericPlatform/GenericPlatformContext.cpp
new file mode 100644
--- /dev/null
+++ b/Engine/Source/Runtime/Core/Private/GenericPlatform/GenericPlatformContext.cpp
@@ -0,0 +1,136 @@
+// Copyright 1998-2014 Epic Games, Inc. All Rights Reserved.
+
+#include "Core.h"
+#include "GenericPlatform/GenericPlatformContext.h"
+#include "Runtime/Launch/Resources/Version.h"
+
+const ANSICHAR* FGenericCrashContext::CrashContextRuntimeXMLNameA = "CrashContext.runtime-xml";
+const TCHAR* FGenericCrashContext::CrashContextRuntimeXMLNameW = TEXT( "CrashContext.runtime-xml" );
+
+namespace
+{
+	static FString CrashDescriptionCommonBuffer;
+}
+
+FGenericCrashContext::FGenericCrashContext() :
+CommonBuffer( CrashDescriptionCommonBuffer )
+{
+	CommonBuffer.Reserve( 32768 );
+}
+
+void FGenericCrashContext::SerializeContentToBuffer()
+{
+	AddHeader();
+
+	// Add common crash properties.
+	AddCrashProperty( TEXT( "GameName" ), *FString::Printf( TEXT( "UE4-%s" ), FApp::GetGameName() ) );
+	AddCrashProperty( TEXT( "PlatformName" ), *FString( FPlatformProperties::PlatformName() ) );
+	AddCrashProperty( TEXT( "EngineMode" ), FPlatformMisc::GetEngineMode() );
+	AddCrashProperty( TEXT( "EngineVersion" ), ENGINE_VERSION_STRING );
+	AddCrashProperty( TEXT( "CommandLine" ), FCommandLine::Get() );
+	AddCrashProperty( TEXT( "BaseDir" ), FPlatformProcess::BaseDir() );
+	AddCrashProperty( TEXT( "LanguageLCID" ), FInternationalization::Get().GetCurrentCulture()->GetLCID() );
+
+	// Get the user name only for non-UE4 releases.
+	if( !FRocketSupport::IsRocket() )
+	{
+		// Remove periods from internal user names to match AutoReporter user names
+		// The name prefix is read by CrashRepository.AddNewCrash in the website code
+
+		AddCrashProperty( TEXT( "UserName" ), *FString( FPlatformProcess::UserName() ).Replace( TEXT( "." ), TEXT( "" ) ) );
+	}
+
+	AddCrashProperty( TEXT( "MachineId" ), *FPlatformMisc::GetMachineId().ToString( EGuidFormats::Digits ) );
+	AddCrashProperty( TEXT( "EpicAccountId" ), *FPlatformMisc::GetEpicAccountId() );
+
+	AddCrashProperty( TEXT( "CallStack" ), TEXT( "" ) );
+	AddCrashProperty( TEXT( "SourceContext" ), TEXT( "" ) );
+	AddCrashProperty( TEXT( "UserDescription" ), TEXT( "" ) );
+	AddCrashProperty( TEXT( "ErrorMessage" ), (const TCHAR*)GErrorMessage );
+
+	AddCrashProperty( TEXT( "TimeOfCrash" ), FDateTime::UtcNow().GetTicks() );
+
+	// Add misc stats.
+	AddCrashProperty( TEXT( "Misc.NumberOfCores" ), FPlatformMisc::NumberOfCores() );
+	AddCrashProperty( TEXT( "Misc.NumberOfCoresIncludingHyperthreads" ), FPlatformMisc::NumberOfCoresIncludingHyperthreads() );
+	AddCrashProperty( TEXT( "Misc.Is64bitOperatingSystem" ), (int32)FPlatformMisc::Is64bitOperatingSystem() );
+
+	AddCrashProperty( TEXT( "Misc.CPUVendor" ), *FPlatformMisc::GetCPUVendor() );
+	AddCrashProperty( TEXT( "Misc.CPUBrand" ), *FPlatformMisc::GetCPUBrand() );
+	AddCrashProperty( TEXT( "Misc.PrimaryGPUBrand" ), *FPlatformMisc::GetPrimaryGPUBrand() );
+
+	FString OSVersion, OSSubVersion;
+	FPlatformMisc::GetOSVersions( OSVersion, OSSubVersion );
+	AddCrashProperty( TEXT( "Misc.OSVersion" ), *FString::Printf( TEXT( "%s, %s" ), *OSVersion, *OSSubVersion ) );
+
+	uint64 AppDiskTotalNumberOfBytes = 0;
+	uint64 AppDiskNumberOfFreeBytes = 0;
+	FPlatformMisc::GetDiskTotalAndFreeSpace( FPlatformProcess::BaseDir(), AppDiskTotalNumberOfBytes, AppDiskNumberOfFreeBytes );
+	AddCrashProperty( TEXT( "Misc.AppDiskTotalNumberOfBytes" ), AppDiskTotalNumberOfBytes );
+	AddCrashProperty( TEXT( "Misc.AppDiskNumberOfFreeBytes" ), AppDiskNumberOfFreeBytes );
+
+	// Add memory stats.
+	const FPlatformMemoryStats MemStats = FPlatformMemory::GetStats();
+
+	AddCrashProperty( TEXT( "MemoryStats.TotalPhysical" ), (uint64)MemStats.TotalPhysical );
+	AddCrashProperty( TEXT( "MemoryStats.TotalVirtual" ), (uint64)MemStats.TotalVirtual );
+	AddCrashProperty( TEXT( "MemoryStats.PageSize" ), (uint64)MemStats.PageSize );
+	AddCrashProperty( TEXT( "MemoryStats.TotalPhysicalGB" ), MemStats.TotalPhysicalGB );
+	AddCrashProperty( TEXT( "MemoryStats.AvailablePhysical" ), (uint64)MemStats.AvailablePhysical );
+	AddCrashProperty( TEXT( "MemoryStats.AvailableVirtual" ), (uint64)MemStats.AvailableVirtual );
+	AddCrashProperty( TEXT( "MemoryStats.UsedPhysical" ), (uint64)MemStats.UsedPhysical );
+	AddCrashProperty( TEXT( "MemoryStats.PeakUsedPhysical" ), (uint64)MemStats.PeakUsedPhysical );
+	AddCrashProperty( TEXT( "MemoryStats.UsedVirtual" ), (uint64)MemStats.UsedVirtual );
+	AddCrashProperty( TEXT( "MemoryStats.PeakUsedVirtual" ), (uint64)MemStats.PeakUsedVirtual );
+	AddCrashProperty( TEXT( "MemoryStats.bIsOOM" ), (int32)FPlatformMemory::bIsOOM );
+
+	// @TODO yrx 2014-09-10 
+	//Architecture
+	//CrashedModuleName
+	//LoadedModules
+	AddFooter();
+
+	// Add platform specific properties.
+	CommonBuffer += TEXT( "<PlatformSpecificProperties>" ) LINE_TERMINATOR;
+	AddPlatformSpecificProperties();
+	CommonBuffer += TEXT( "</PlatformSpecificProperties>" ) LINE_TERMINATOR;
+}
+
+void FGenericCrashContext::SerializeAsXML( const TCHAR* Filename )
+{
+	SerializeContentToBuffer();
+	FFileHelper::SaveStringToFile( CommonBuffer, Filename, FFileHelper::EEncodingOptions::ForceUTF8 );
+}
+
+void FGenericCrashContext::AddPlatformSpecificProperties()
+{
+	// Nothing really to do here. Can be overridden by the platform code.
+	// @see FWindowsPlatformCrashContext::AddPlatformSpecificProperties
+}
+
+void FGenericCrashContext::AddHeader()
+{
+	CommonBuffer += TEXT( "<?xml version=\"1.0\" encoding=\"UTF-8\"?>" ) LINE_TERMINATOR;
+	CommonBuffer += TEXT( "<RuntimeCrashDescription>" ) LINE_TERMINATOR;
+
+	AddCrashProperty( TEXT( "Version" ), (int32)ECrashDescVersions::VER_2_AddedNewProperties );
+}
+
+void FGenericCrashContext::AddFooter()
+{
+	CommonBuffer += TEXT( "</RuntimeCrashDescription>" ) LINE_TERMINATOR;
+}
+
+void FGenericCrashContext::AddCrashProperty( const TCHAR* PropertyName, const TCHAR* PropertyValue )
+{
+	CommonBuffer += TEXT( "<" );
+	CommonBuffer += PropertyName;
+	CommonBuffer += TEXT( ">" );
+
+	CommonBuffer += PropertyValue;
+
+	CommonBuffer += TEXT( "</" );
+	CommonBuffer += PropertyName;
+	CommonBuffer += TEXT( ">" );
+	CommonBuffer += LINE_TERMINATOR;
+}
diff --git a/Engine/Source/Runtime/Core/Private/IOS/IOSPlatformMisc.cpp b/Engine/Source/Runtime/Core/Private/IOS/IOSPlatformMisc.cpp
--- a/Engine/Source/Runtime/Core/Private/IOS/IOSPlatformMisc.cpp
+++ b/Engine/Source/Runtime/Core/Private/IOS/IOSPlatformMisc.cpp
@@ -12,6 +12,8 @@
 #include "IOSView.h"
 #include "GenericPlatformChunkInstall.h"
 
+#include "Apple/ApplePlatformCrashContext.h"
+
 /** Amount of free memory in MB reported by the system at starup */
 CORE_API int32 GStartupFreeMemoryMB;
 
diff --git a/Engine/Source/Runtime/Core/Private/Linux/LinuxPlatformCrashContext.cpp b/Engine/Source/Runtime/Core/Private/Linux/LinuxPlatformCrashContext.cpp
new file mode 100644
--- /dev/null
+++ b/Engine/Source/Runtime/Core/Private/Linux/LinuxPlatformCrashContext.cpp
@@ -0,0 +1,821 @@
+// Copyright 1998-2014 Epic Games, Inc. All Rights Reserved.
+
+#include "Core.h"
+#include "Linux/LinuxPlatformCrashContext.h"
+#include "Runtime/Launch/Resources/Version.h"
+
+#include <sys/utsname.h>	// for uname()
+#include <signal.h>
+
+// these are not actually system headers, but a TPS library (see ThirdParty/elftoolchain)
+#include <libelf.h>
+#include <_libelf.h>
+#include <libdwarf.h>
+#include <dwarf.h>
+
+
+FString DescribeSignal(int32 Signal, siginfo_t* Info)
+{
+	FString ErrorString;
+
+#define HANDLE_CASE(a,b) case a: ErrorString += TEXT(#a ": " b); break;
+
+	switch (Signal)
+	{
+	case SIGSEGV:
+		ErrorString += TEXT("SIGSEGV: invalid attempt to access memory at address ");
+		ErrorString += FString::Printf(TEXT("0x%08x"), (uint32*)Info->si_addr);
+		break;
+	case SIGBUS:
+		ErrorString += TEXT("SIGBUS: invalid attempt to access memory at address ");
+		ErrorString += FString::Printf(TEXT("0x%08x"), (uint32*)Info->si_addr);
+		break;
+
+		HANDLE_CASE(SIGINT, "program interrupted")
+		HANDLE_CASE(SIGQUIT, "user-requested crash")
+		HANDLE_CASE(SIGILL, "illegal instruction")
+		HANDLE_CASE(SIGTRAP, "trace trap")
+		HANDLE_CASE(SIGABRT, "abort() called")
+		HANDLE_CASE(SIGFPE, "floating-point exception")
+		HANDLE_CASE(SIGKILL, "program killed")
+		HANDLE_CASE(SIGSYS, "non-existent system call invoked")
+		HANDLE_CASE(SIGPIPE, "write on a pipe with no reader")
+		HANDLE_CASE(SIGTERM, "software termination signal")
+		HANDLE_CASE(SIGSTOP, "stop")
+
+	default:
+		ErrorString += FString::Printf(TEXT("Signal %d (unknown)"), Signal);
+	}
+
+	return ErrorString;
+#undef HANDLE_CASE
+}
+
+FLinuxCrashContext::~FLinuxCrashContext()
+{
+	if (BacktraceSymbols)
+	{
+		// glibc uses malloc() to allocate this, and we only need to free one pointer, see http://www.gnu.org/software/libc/manual/html_node/Backtraces.html
+		free(BacktraceSymbols);
+		BacktraceSymbols = NULL;
+	}
+
+	if (DebugInfo)
+	{
+		Dwarf_Error ErrorInfo;
+		dwarf_finish(DebugInfo, &ErrorInfo);
+		DebugInfo = NULL;
+	}
+
+	if (ElfHdr)
+	{
+		elf_end(ElfHdr);
+		ElfHdr = NULL;
+	}
+
+	if (ExeFd >= 0)
+	{
+		close(ExeFd);
+		ExeFd = -1;
+	}
+}
+
+void FLinuxCrashContext::InitFromSignal(int32 InSignal, siginfo_t* InInfo, void* InContext)
+{
+	Signal = InSignal;
+	Info = InInfo;
+	Context = reinterpret_cast< ucontext_t* >( InContext );
+
+	// open ourselves for examination
+	if (!FParse::Param( FCommandLine::Get(), TEXT(CMDARG_SUPPRESS_DWARF_PARSING)))
+	{
+		ExeFd = open("/proc/self/exe", O_RDONLY);
+		if (ExeFd >= 0)
+		{
+			Dwarf_Error ErrorInfo;
+			// allocate DWARF debug descriptor
+			if (dwarf_init(ExeFd, DW_DLC_READ, NULL, NULL, &DebugInfo, &ErrorInfo) == DW_DLV_OK)
+			{
+				// get ELF descritor
+				if (dwarf_get_elf(DebugInfo, &ElfHdr, &ErrorInfo) != DW_DLV_OK)
+				{
+					dwarf_finish(DebugInfo, &ErrorInfo);
+					DebugInfo = NULL;
+
+					close(ExeFd);
+					ExeFd = -1;
+				}
+			}
+			else
+			{
+				DebugInfo = NULL;
+				close(ExeFd);
+				ExeFd = -1;
+			}
+		}
+	}
+
+	FCString::Strcat(SignalDescription, ARRAY_COUNT( SignalDescription ) - 1, *DescribeSignal(Signal, Info));
+}
+
+/**
+ * Finds a function name in DWARF DIE (Debug Information Entry).
+ * For more info on DWARF format, see http://www.dwarfstd.org/Download.php , http://www.ibm.com/developerworks/library/os-debugging/
+ *
+ * @return true if we need to stop search (i.e. either found it or some error happened)
+ */
+bool FindFunctionNameInDIE(Dwarf_Debug DebugInfo, Dwarf_Die Die, Dwarf_Addr Addr, const char **OutFuncName)
+{
+	Dwarf_Error ErrorInfo;
+	Dwarf_Half Tag;
+	Dwarf_Unsigned LowerPC, HigherPC;
+	char *TempFuncName;
+	int ReturnCode;
+
+	if (dwarf_tag(Die, &Tag, &ErrorInfo) != DW_DLV_OK || Tag != DW_TAG_subprogram ||
+		dwarf_attrval_unsigned(Die, DW_AT_low_pc, &LowerPC, &ErrorInfo) != DW_DLV_OK ||
+		dwarf_attrval_unsigned(Die, DW_AT_high_pc, &HigherPC, &ErrorInfo) != DW_DLV_OK ||
+		Addr < LowerPC || HigherPC <= Addr
+		) 
+	{
+		return false;
+	}
+	
+	// found it
+	*OutFuncName = NULL;
+	Dwarf_Attribute SubAt;
+	ReturnCode = dwarf_attr(Die, DW_AT_name, &SubAt, &ErrorInfo);
+	if (ReturnCode == DW_DLV_ERROR)
+	{
+		return true;	// error, but stop the search
+	}
+	else if (ReturnCode == DW_DLV_OK) 
+	{
+		if (dwarf_formstring(SubAt, &TempFuncName, &ErrorInfo))
+		{
+			*OutFuncName = NULL;
+		}
+		else
+		{
+			*OutFuncName = TempFuncName;
+		}
+		return true;
+	}
+
+	// DW_AT_Name is not present, look in DW_AT_specification
+	Dwarf_Attribute SpecAt;
+	if (dwarf_attr(Die, DW_AT_specification, &SpecAt, &ErrorInfo))
+	{
+		// not found, tough luck
+		return false;
+	}
+
+	Dwarf_Off Offset;
+	if (dwarf_global_formref(SpecAt, &Offset, &ErrorInfo))
+	{
+		return false;
+	}
+
+	Dwarf_Die SpecDie;
+	if (dwarf_offdie(DebugInfo, Offset, &SpecDie, &ErrorInfo))
+	{
+		return false;
+	}
+
+	if (dwarf_attrval_string(SpecDie, DW_AT_name, OutFuncName, &ErrorInfo))
+	{
+		*OutFuncName = NULL;
+	}
+
+	return true;
+}
+
+/**
+ * Finds a function name in DWARF DIE (Debug Information Entry) and its children.
+ * For more info on DWARF format, see http://www.dwarfstd.org/Download.php , http://www.ibm.com/developerworks/library/os-debugging/
+ * Note: that function is not exactly traversing the tree, but this "seems to work"(tm). Not sure if we need to descend properly (taking child of every sibling), this
+ * takes too much time (and callstacks seem to be fine without it).
+ */
+void FindFunctionNameInDIEAndChildren(Dwarf_Debug DebugInfo, Dwarf_Die Die, Dwarf_Addr Addr, const char **OutFuncName)
+{
+	if (OutFuncName == NULL || *OutFuncName != NULL)
+	{
+		return;
+	}
+
+	// search for this Die
+	if (FindFunctionNameInDIE(DebugInfo, Die, Addr, OutFuncName))
+	{
+		return;
+	}
+
+	Dwarf_Die PrevChild = Die, Current = NULL;
+	Dwarf_Error ErrorInfo;
+
+	int32 MaxChildrenAllowed = 32 * 1024 * 1024;	// safeguard to make sure we never get into an infinite loop
+	for(;;)
+	{
+		if (--MaxChildrenAllowed <= 0)
+		{
+			fprintf(stderr, "Breaking out from what seems to be an infinite loop during DWARF parsing (too many children).\n");
+			return;
+		}
+
+		// Get the child
+		if (dwarf_child(PrevChild, &Current, &ErrorInfo) != DW_DLV_OK)
+		{
+			return;	// bail out
+		}
+
+		PrevChild = Current;
+
+		// look for in the child
+		if (FindFunctionNameInDIE(DebugInfo, Current, Addr, OutFuncName))
+		{
+			return;	// got the function name!
+		}
+
+		// search among child's siblings
+		int32 MaxSiblingsAllowed = 64 * 1024 * 1024;	// safeguard to make sure we never get into an infinite loop
+		for(;;)
+		{
+			if (--MaxSiblingsAllowed <= 0)
+			{
+				fprintf(stderr, "Breaking out from what seems to be an infinite loop during DWARF parsing (too many siblings).\n");
+				break;
+			}
+
+			Dwarf_Die Prev = Current;
+			if (dwarf_siblingof(DebugInfo, Prev, &Current, &ErrorInfo) != DW_DLV_OK || Current == NULL)
+			{
+				break;
+			}
+
+			if (FindFunctionNameInDIE(DebugInfo, Current, Addr, OutFuncName))
+			{
+				return;	// got the function name!
+			}
+		}
+	};
+}
+
+bool FLinuxCrashContext::GetInfoForAddress(void * Address, const char **OutFunctionNamePtr, const char **OutSourceFilePtr, int *OutLineNumberPtr)
+{
+	if (DebugInfo == NULL)
+	{
+		return false;
+	}
+
+	Dwarf_Die Die;
+	Dwarf_Unsigned Addr = reinterpret_cast< Dwarf_Unsigned >( Address ), LineNumber = 0;
+	const char * SrcFile = NULL;
+
+	static_assert(sizeof(Dwarf_Unsigned) >= sizeof(Address), "Dwarf_Unsigned type should be long enough to represent pointers. Check libdwarf bitness.");
+
+	int ReturnCode = DW_DLV_OK;
+	Dwarf_Error ErrorInfo;
+	bool bExitHeaderLoop = false;
+	int32 MaxCompileUnitsAllowed = 16 * 1024 * 1024;	// safeguard to make sure we never get into an infinite loop
+	const int32 kMaxBufferLinesAllowed = 16 * 1024 * 1024;	// safeguard to prevent too long line loop
+	for(;;)
+	{
+		if (--MaxCompileUnitsAllowed <= 0)
+		{
+			fprintf(stderr, "Breaking out from what seems to be an infinite loop during DWARF parsing (too many compile units).\n");
+			ReturnCode = DW_DLE_DIE_NO_CU_CONTEXT;	// invalidate
+			break;
+		}
+
+		if (bExitHeaderLoop)
+			break;
+
+		ReturnCode = dwarf_next_cu_header(DebugInfo, NULL, NULL, NULL, NULL, NULL, &ErrorInfo);
+		if (ReturnCode != DW_DLV_OK)
+			break;
+
+		Die = NULL;
+
+		while(dwarf_siblingof(DebugInfo, Die, &Die, &ErrorInfo) == DW_DLV_OK)
+		{
+			Dwarf_Half Tag;
+			if (dwarf_tag(Die, &Tag, &ErrorInfo) != DW_DLV_OK)
+			{
+				bExitHeaderLoop = true;
+				break;
+			}
+
+			if (Tag == DW_TAG_compile_unit)
+			{
+				break;
+			}
+		}
+
+		if (Die == NULL)
+		{
+			break;
+		}
+
+		// check if address is inside this CU
+		Dwarf_Unsigned LowerPC, HigherPC;
+		if (!dwarf_attrval_unsigned(Die, DW_AT_low_pc, &LowerPC, &ErrorInfo) && !dwarf_attrval_unsigned(Die, DW_AT_high_pc, &HigherPC, &ErrorInfo))
+		{
+			if (Addr < LowerPC || Addr >= HigherPC)
+			{
+				continue;
+			}
+		}
+
+		Dwarf_Line * LineBuf;
+		Dwarf_Signed NumLines = kMaxBufferLinesAllowed;
+		if (dwarf_srclines(Die, &LineBuf, &NumLines, &ErrorInfo) != DW_DLV_OK)
+		{
+			// could not get line info for some reason
+			break;
+		}
+
+		if (NumLines >= kMaxBufferLinesAllowed)
+		{
+			fprintf(stderr, "Number of lines associated with a DIE looks unreasonable (%ld), early quitting.\n", NumLines);
+			ReturnCode = DW_DLE_DIE_NO_CU_CONTEXT;	// invalidate
+			break;
+		}
+
+		// look which line is that
+		Dwarf_Addr LineAddress, PrevLineAddress = ~0ULL;
+		Dwarf_Unsigned PrevLineNumber = 0;
+		const char * PrevSrcFile = NULL;
+		char * SrcFileTemp = NULL;
+		for (int Idx = 0; Idx < NumLines; ++Idx)
+		{
+			if (dwarf_lineaddr(LineBuf[Idx], &LineAddress, &ErrorInfo) != 0 ||
+				dwarf_lineno(LineBuf[Idx], &LineNumber, &ErrorInfo) != 0)
+			{
+				bExitHeaderLoop = true;
+				break;
+			}
+
+			if (!dwarf_linesrc(LineBuf[Idx], &SrcFileTemp, &ErrorInfo))
+			{
+				SrcFile = SrcFileTemp;
+			}
+
+			// check if we hit the exact line
+			if (Addr == LineAddress)
+			{
+				bExitHeaderLoop = true;
+				break;
+			}
+			else if (PrevLineAddress < Addr && Addr < LineAddress)
+			{
+				LineNumber = PrevLineNumber;
+				SrcFile = PrevSrcFile;
+				bExitHeaderLoop = true;
+				break;
+			}
+
+			PrevLineAddress = LineAddress;
+			PrevLineNumber = LineNumber;
+			PrevSrcFile = SrcFile;
+		}
+
+	}
+
+	const char * FunctionName = NULL;
+	if (ReturnCode == DW_DLV_OK)
+	{
+		FindFunctionNameInDIEAndChildren(DebugInfo, Die, Addr, &FunctionName);
+	}
+
+	if (OutFunctionNamePtr != NULL && FunctionName != NULL)
+	{
+		*OutFunctionNamePtr = FunctionName;
+	}
+
+	if (OutSourceFilePtr != NULL && SrcFile != NULL)
+	{
+		*OutSourceFilePtr = SrcFile;
+		
+		if (OutLineNumberPtr != NULL)
+		{
+			*OutLineNumberPtr = LineNumber;
+		}
+	}
+
+	// Resets internal CU pointer, so next time we get here it begins from the start
+	while (ReturnCode != DW_DLV_NO_ENTRY) 
+	{
+		if (ReturnCode == DW_DLV_ERROR)
+			break;
+		ReturnCode = dwarf_next_cu_header(DebugInfo, NULL, NULL, NULL, NULL, NULL, &ErrorInfo);
+	}
+
+	// if we weren't able to find a function name, don't trust the source file either
+	return FunctionName != NULL;
+}
+
+/**
+ * Handles graceful termination. Gives time to exit gracefully, but second signal will quit immediately.
+ */
+void GracefulTerminationHandler(int32 Signal, siginfo_t* Info, void* Context)
+{
+	printf("CtrlCHandler: Signal=%d\n", Signal);
+
+	// make sure as much data is written to disk as possible
+	if (GLog)
+	{
+		GLog->Flush();
+	}
+	if (GWarn)
+	{
+		GWarn->Flush();
+	}
+	if (GError)
+	{
+		GError->Flush();
+	}
+
+	if( !GIsRequestingExit )
+	{
+		GIsRequestingExit = 1;
+	}
+	else
+	{
+		FPlatformMisc::RequestExit(true);
+	}
+}
+
+void CreateExceptionInfoString(int32 Signal, siginfo_t* Info)
+{
+	FString ErrorString = TEXT("Unhandled Exception: ");
+	ErrorString += DescribeSignal(Signal, Info);
+	FCString::Strncpy(GErrorExceptionDescription, *ErrorString, FMath::Min(ErrorString.Len() + 1, (int32)ARRAY_COUNT(GErrorExceptionDescription)));
+}
+
+namespace
+{
+	/** 
+	 * Write a line of UTF-8 to a file
+	 */
+	void WriteLine(FArchive* ReportFile, const ANSICHAR* Line = NULL)
+	{
+		if( Line != NULL )
+		{
+			int64 StringBytes = FCStringAnsi::Strlen(Line);
+			ReportFile->Serialize(( void* )Line, StringBytes);
+		}
+
+		// use Windows line terminator
+		static ANSICHAR WindowsTerminator[] = "\r\n";
+		ReportFile->Serialize(WindowsTerminator, 2);
+	}
+
+	/**
+	 * Serializes UTF string to UTF-16
+	 */
+	void WriteUTF16String(FArchive* ReportFile, const TCHAR * UTFString4BytesChar, uint32 NumChars)
+	{
+		check(UTFString4BytesChar != NULL || NumChars == 0);
+		static_assert(sizeof(TCHAR) == 4, "Platform TCHAR is not 4 bytes. Revisit this function.");
+
+		for (uint32 Idx = 0; Idx < NumChars; ++Idx)
+		{
+			ReportFile->Serialize(const_cast< TCHAR* >( &UTFString4BytesChar[Idx] ), 2);
+		}
+	}
+
+	/** 
+	 * Writes UTF-16 line to a file
+	 */
+	void WriteLine(FArchive* ReportFile, const TCHAR* Line)
+	{
+		if( Line != NULL )
+		{
+			int64 NumChars = FCString::Strlen(Line);
+			WriteUTF16String(ReportFile, Line, NumChars);
+		}
+
+		// use Windows line terminator
+		static TCHAR WindowsTerminator[] = TEXT("\r\n");
+		WriteUTF16String(ReportFile, WindowsTerminator, 2);
+	}
+}
+
+/** 
+ * Write all the data mined from the minidump to a text file
+ */
+void FLinuxCrashContext::GenerateReport(const FString & DiagnosticsPath) const
+{
+	FArchive* ReportFile = IFileManager::Get().CreateFileWriter(*DiagnosticsPath);
+	if (ReportFile != NULL)
+	{
+		FString Line;
+
+		WriteLine(ReportFile, "Generating report for minidump");
+		WriteLine(ReportFile);
+
+		Line = FString::Printf(TEXT("Application version %d.%d.%d.0" ), ENGINE_MAJOR_VERSION, ENGINE_MINOR_VERSION, ENGINE_PATCH_VERSION);
+		WriteLine(ReportFile, TCHAR_TO_UTF8(*Line));
+
+		Line = FString::Printf(TEXT(" ... built from changelist %d"), ENGINE_VERSION);
+		WriteLine(ReportFile, TCHAR_TO_UTF8(*Line));
+		WriteLine(ReportFile);
+
+		utsname UnixName;
+		if (uname(&UnixName) == 0)
+		{
+			Line = FString::Printf(TEXT( "OS version %s %s (network name: %s)" ), ANSI_TO_TCHAR(UnixName.sysname), ANSI_TO_TCHAR(UnixName.release), ANSI_TO_TCHAR(UnixName.nodename));
+			WriteLine(ReportFile, TCHAR_TO_UTF8(*Line));	
+			Line = FString::Printf( TEXT( "Running %d %s processors (%d logical cores)" ), FPlatformMisc::NumberOfCores(), ANSI_TO_TCHAR(UnixName.machine), FPlatformMisc::NumberOfCoresIncludingHyperthreads());
+			WriteLine(ReportFile, TCHAR_TO_UTF8(*Line));
+		}
+		else
+		{
+			Line = FString::Printf(TEXT("OS version could not be determined (%d, %s)"), errno, ANSI_TO_TCHAR(strerror(errno)));
+			WriteLine(ReportFile, TCHAR_TO_UTF8(*Line));	
+			Line = FString::Printf( TEXT( "Running %d unknown processors" ), FPlatformMisc::NumberOfCores());
+			WriteLine(ReportFile, TCHAR_TO_UTF8(*Line));
+		}
+		Line = FString::Printf(TEXT("Exception was \"%s\""), SignalDescription);
+		WriteLine(ReportFile, TCHAR_TO_UTF8(*Line));
+		WriteLine(ReportFile);
+
+		WriteLine(ReportFile, "<SOURCE START>");
+		WriteLine(ReportFile, "<SOURCE END>");
+		WriteLine(ReportFile);
+
+		WriteLine(ReportFile, "<CALLSTACK START>");
+		WriteLine(ReportFile, MinidumpCallstackInfo);
+		WriteLine(ReportFile, "<CALLSTACK END>");
+		WriteLine(ReportFile);
+
+		WriteLine(ReportFile, "0 loaded modules");
+
+		WriteLine(ReportFile);
+
+		Line = FString::Printf(TEXT("Report end!"));
+		WriteLine(ReportFile, TCHAR_TO_UTF8(*Line));
+
+		ReportFile->Close();
+		delete ReportFile;
+	}
+}
+
+/** 
+ * Mimics Windows WER format
+ */
+void GenerateWindowsErrorReport(const FString & WERPath)
+{
+	FArchive* ReportFile = IFileManager::Get().CreateFileWriter(*WERPath);
+	if (ReportFile != NULL)
+	{
+		// write BOM
+		static uint16 ByteOrderMarker = 0xFEFF;
+		ReportFile->Serialize(&ByteOrderMarker, sizeof(ByteOrderMarker));
+
+		WriteLine(ReportFile, TEXT("<?xml version=\"1.0\" encoding=\"UTF-16\"?>"));
+		WriteLine(ReportFile, TEXT("<WERReportMetadata>"));
+		
+		WriteLine(ReportFile, TEXT("\t<OSVersionInformation>"));
+		WriteLine(ReportFile, TEXT("\t\t<WindowsNTVersion>0.0</WindowsNTVersion>"));
+		WriteLine(ReportFile, TEXT("\t\t<Build>No Build</Build>"));
+		WriteLine(ReportFile, TEXT("\t\t<Product>Linux</Product>"));
+		WriteLine(ReportFile, TEXT("\t\t<Edition>No Edition</Edition>"));
+		WriteLine(ReportFile, TEXT("\t\t<BuildString>No BuildString</BuildString>"));
+		WriteLine(ReportFile, TEXT("\t\t<Revision>0</Revision>"));
+		WriteLine(ReportFile, TEXT("\t\t<Flavor>No Flavor</Flavor>"));
+		WriteLine(ReportFile, TEXT("\t\t<Architecture>Unknown Architecture</Architecture>"));
+		WriteLine(ReportFile, TEXT("\t\t<LCID>0</LCID>"));
+		WriteLine(ReportFile, TEXT("\t</OSVersionInformation>"));
+		
+		WriteLine(ReportFile, TEXT("\t<ParentProcessInformation>"));
+		WriteLine(ReportFile, *FString::Printf(TEXT("\t\t<ParentProcessId>%d</ParentProcessId>"), getppid()));
+		WriteLine(ReportFile, TEXT("\t\t<ParentProcessPath>C:\\Windows\\explorer.exe</ParentProcessPath>"));			// FIXME: supply valid?
+		WriteLine(ReportFile, TEXT("\t\t<ParentProcessCmdLine>C:\\Windows\\Explorer.EXE</ParentProcessCmdLine>"));	// FIXME: supply valid?
+		WriteLine(ReportFile, TEXT("\t</ParentProcessInformation>"));
+
+		WriteLine(ReportFile, TEXT("\t<ProblemSignatures>"));
+		WriteLine(ReportFile, TEXT("\t\t<EventType>APPCRASH</EventType>"));
+		WriteLine(ReportFile, *FString::Printf(TEXT("\t\t<Parameter0>UE4-%s</Parameter0>"), FApp::GetGameName()));
+		WriteLine(ReportFile, *FString::Printf(TEXT("\t\t<Parameter1>%d.%d.%d.0</Parameter1>"), ENGINE_MAJOR_VERSION, ENGINE_MINOR_VERSION, ENGINE_PATCH_VERSION));
+		WriteLine(ReportFile, TEXT("\t\t<Parameter2>0</Parameter2>"));													// FIXME: supply valid?
+		WriteLine(ReportFile, TEXT("\t\t<Parameter3>Unknown Fault Module</Parameter3>"));										// FIXME: supply valid?
+		WriteLine(ReportFile, TEXT("\t\t<Parameter4>0.0.0.0</Parameter4>"));													// FIXME: supply valid?
+		WriteLine(ReportFile, TEXT("\t\t<Parameter5>00000000</Parameter5>"));													// FIXME: supply valid?
+		WriteLine(ReportFile, TEXT("\t\t<Parameter6>00000000</Parameter6>"));													// FIXME: supply valid?
+		WriteLine(ReportFile, TEXT("\t\t<Parameter7>0000000000000000</Parameter7>"));											// FIXME: supply valid?
+		WriteLine(ReportFile, *FString::Printf(TEXT("\t\t<Parameter8>!%s!</Parameter8>"), FCommandLine::Get()));				// FIXME: supply valid? Only partially valid
+		WriteLine(ReportFile, *FString::Printf(TEXT("\t\t<Parameter9>%s!%s!%s!%d</Parameter9>"), TEXT( BRANCH_NAME ), FPlatformProcess::BaseDir(), FPlatformMisc::GetEngineMode(), BUILT_FROM_CHANGELIST));
+		WriteLine(ReportFile, TEXT("\t</ProblemSignatures>"));
+
+		WriteLine(ReportFile, TEXT("\t<DynamicSignatures>"));
+		WriteLine(ReportFile, TEXT("\t\t<Parameter1>6.1.7601.2.1.0.256.48</Parameter1>"));
+		WriteLine(ReportFile, TEXT("\t\t<Parameter2>1033</Parameter2>"));
+		WriteLine(ReportFile, TEXT("\t</DynamicSignatures>"));
+
+		WriteLine(ReportFile, TEXT("\t<SystemInformation>"));
+		WriteLine(ReportFile, TEXT("\t\t<MID>11111111-2222-3333-4444-555555555555</MID>"));							// FIXME: supply valid?
+		
+		WriteLine(ReportFile, TEXT("\t\t<SystemManufacturer>Unknown.</SystemManufacturer>"));						// FIXME: supply valid?
+		WriteLine(ReportFile, TEXT("\t\t<SystemProductName>Linux machine</SystemProductName>"));					// FIXME: supply valid?
+		WriteLine(ReportFile, TEXT("\t\t<BIOSVersion>A02</BIOSVersion>"));											// FIXME: supply valid?
+		WriteLine(ReportFile, TEXT("\t</SystemInformation>"));
+
+		WriteLine(ReportFile, TEXT("</WERReportMetadata>"));
+
+		ReportFile->Close();
+		delete ReportFile;
+	}
+}
+
+/** 
+ * Creates (fake so far) minidump
+ */
+void GenerateMinidump(const FString & Path)
+{
+	FArchive* ReportFile = IFileManager::Get().CreateFileWriter(*Path);
+	if (ReportFile != NULL)
+	{
+		// write BOM
+		static uint32 Garbage = 0xDEADBEEF;
+		ReportFile->Serialize(&Garbage, sizeof(Garbage));
+
+		ReportFile->Close();
+		delete ReportFile;
+	}
+}
+
+
+int32 ReportCrash(const FLinuxCrashContext & Context)
+{
+	static bool GAlreadyCreatedMinidump = false;
+	// Only create a minidump the first time this function is called.
+	// (Can be called the first time from the RenderThread, then a second time from the MainThread.)
+	if ( GAlreadyCreatedMinidump == false )
+	{
+		GAlreadyCreatedMinidump = true;
+
+		const SIZE_T StackTraceSize = 65535;
+		ANSICHAR* StackTrace = (ANSICHAR*) FMemory::Malloc( StackTraceSize );
+		StackTrace[0] = 0;
+		// Walk the stack and dump it to the allocated memory (ignore first 2 callstack lines as those are in stack walking code)
+		FPlatformStackWalk::StackWalkAndDump( StackTrace, StackTraceSize, 2, const_cast< FLinuxCrashContext* >( &Context ) );
+
+		FCString::Strncat( GErrorHist, ANSI_TO_TCHAR(StackTrace), ARRAY_COUNT(GErrorHist) - 1 );
+		CreateExceptionInfoString(Context.Signal, Context.Info);
+
+		FMemory::Free( StackTrace );
+	}
+
+	return 0;
+}
+
+/**
+ * Generates information for crash reporter
+ */
+void GenerateCrashInfoAndLaunchReporter(const FLinuxCrashContext & Context)
+{
+	// do not report crashes for tools (particularly for crash reporter itself)
+#if !IS_PROGRAM
+
+	// create a crash-specific directory
+	FString CrashInfoFolder = FString::Printf(TEXT("crashinfo-%s-pid-%d-%s-%s"), FApp::GetGameName(), getpid(), 
+		*FDateTime::Now().ToString(), *FGuid::NewGuid().ToString());
+	FString CrashInfoAbsolute = FPaths::ConvertRelativePathToFull(CrashInfoFolder);
+	if (IFileManager::Get().MakeDirectory(*CrashInfoFolder))
+	{
+		// generate "minidump"
+		Context.GenerateReport(FPaths::Combine(*CrashInfoFolder, TEXT("diagnostics.txt")));
+
+		// generate "WER"
+		GenerateWindowsErrorReport(FPaths::Combine(*CrashInfoFolder, TEXT("wermeta.xml")));
+
+		// generate "minidump" (just >1 byte)
+		GenerateMinidump(FPaths::Combine(*CrashInfoFolder, TEXT("minidump.dmp")));
+
+		// Introduces a new runtime crash context. Will replace all Windows related crash reporting.
+		//FCStringAnsi::Strncpy(FilePath, CrashInfoFolder, PATH_MAX);
+		//FCStringAnsi::Strcat(FilePath, PATH_MAX, "/" );
+		//FCStringAnsi::Strcat(FilePath, PATH_MAX, FGenericCrashContext::CrashContextRuntimeXMLNameA );
+		//SerializeAsXML( FilePath ); @todo uncomment after verification
+
+		// copy log
+		FString LogSrcAbsolute = FPlatformOutputDevices::GetAbsoluteLogFilename();
+		FString LogDstAbsolute = FPaths::Combine(*CrashInfoAbsolute, *FPaths::GetCleanFilename(LogSrcAbsolute));
+		FPaths::NormalizeDirectoryName(LogDstAbsolute);
+		static_cast<void>(IFileManager::Get().Copy(*LogDstAbsolute, *LogSrcAbsolute));	// best effort, so don't care about result: couldn't copy -> tough, no log
+
+		// try launching the tool and wait for its exit, if at all
+		const TCHAR * RelativePathToCrashReporter = TEXT("../../../engine/binaries/linux/crashreportclient");	// FIXME: painfully hard-coded
+		if (!FPaths::FileExists(RelativePathToCrashReporter))
+		{
+			RelativePathToCrashReporter = TEXT("../../../Engine/Binaries/Linux/CrashReportClient");	// FIXME: even more painfully hard-coded
+		}
+
+		// show on the console
+		printf("Starting %s\n", StringCast<ANSICHAR>(RelativePathToCrashReporter).Get());
+		FProcHandle RunningProc = FPlatformProcess::CreateProc(RelativePathToCrashReporter, *(CrashInfoAbsolute + TEXT("/")), true, false, false, NULL, 0, NULL, NULL);
+		if (FPlatformProcess::IsProcRunning(RunningProc))
+		{
+			// do not wait indefinitely
+			double kTimeOut = 3 * 60.0;
+			double StartSeconds = FPlatformTime::Seconds();
+			for(;;)
+			{
+				if (!FPlatformProcess::IsProcRunning(RunningProc))
+				{
+					break;
+				}
+
+				if (FPlatformTime::Seconds() - StartSeconds > kTimeOut)
+				{
+					break;
+				}
+
+				FPlatformProcess::Sleep(1.0f);
+			};
+		}
+	}
+
+#endif
+
+	FPlatformMisc::RequestExit(true);
+}
+
+/**
+ * Good enough default crash reporter.
+ */
+void DefaultCrashHandler(const FLinuxCrashContext & Context)
+{
+	// Switch to malloc crash.
+	//FMallocCrash::Get().SetAsGMalloc(); 
+
+	printf("DefaultCrashHandler: Signal=%d\n", Context.Signal);
+
+	ReportCrash(Context);
+	if (GLog)
+	{
+		GLog->Flush();
+	}
+	if (GWarn)
+	{
+		GWarn->Flush();
+	}
+	if (GError)
+	{
+		GError->Flush();
+		GError->HandleError();
+	}
+
+	return GenerateCrashInfoAndLaunchReporter(Context);
+}
+
+/** Global pointer to crash handler */
+void (* GCrashHandlerPointer)(const FGenericCrashContext & Context) = NULL;
+
+/** True system-specific crash handler that gets called first */
+void PlatformCrashHandler(int32 Signal, siginfo_t* Info, void* Context)
+{
+	fprintf(stderr, "Signal %d caught.\n", Signal);
+
+	FLinuxCrashContext CrashContext;
+	CrashContext.InitFromSignal(Signal, Info, Context);
+
+	if (GCrashHandlerPointer)
+	{
+		GCrashHandlerPointer(CrashContext);
+	}
+	else
+	{
+		// call default one
+		DefaultCrashHandler(CrashContext);
+	}
+}
+
+void FLinuxPlatformMisc::SetGracefulTerminationHandler()
+{
+	struct sigaction Action;
+	FMemory::Memzero(&Action, sizeof(struct sigaction));
+	Action.sa_sigaction = GracefulTerminationHandler;
+	sigemptyset(&Action.sa_mask);
+	Action.sa_flags = SA_SIGINFO | SA_RESTART | SA_ONSTACK;
+	sigaction(SIGINT, &Action, NULL);
+	sigaction(SIGTERM, &Action, NULL);
+	sigaction(SIGHUP, &Action, NULL);	//  this should actually cause the server to just re-read configs (restart?)
+}
+
+void FLinuxPlatformMisc::SetCrashHandler(void (* CrashHandler)(const FGenericCrashContext & Context))
+{
+	GCrashHandlerPointer = CrashHandler;
+
+	struct sigaction Action;
+	FMemory::Memzero(&Action, sizeof(struct sigaction));
+	Action.sa_sigaction = PlatformCrashHandler;
+	sigemptyset(&Action.sa_mask);
+	Action.sa_flags = SA_SIGINFO | SA_RESTART | SA_ONSTACK;
+	sigaction(SIGQUIT, &Action, NULL);	// SIGQUIT is a user-initiated "crash".
+	sigaction(SIGILL, &Action, NULL);
+	sigaction(SIGFPE, &Action, NULL);
+	sigaction(SIGBUS, &Action, NULL);
+	sigaction(SIGSEGV, &Action, NULL);
+	sigaction(SIGSYS, &Action, NULL);
+}
\ No newline at end of file
diff --git a/Engine/Source/Runtime/Core/Private/Linux/LinuxPlatformMisc.cpp b/Engine/Source/Runtime/Core/Private/Linux/LinuxPlatformMisc.cpp
--- a/Engine/Source/Runtime/Core/Private/Linux/LinuxPlatformMisc.cpp
+++ b/Engine/Source/Runtime/Core/Private/Linux/LinuxPlatformMisc.cpp
@@ -13,17 +13,8 @@
 #include <fcntl.h>
 #include <signal.h>
 
-// these are not actually system headers, but a TPS library (see ThirdParty/elftoolchain)
-#include <libelf.h>
-#include <_libelf.h>
-#include <libdwarf.h>
-#include <dwarf.h>
-
 #include "ModuleManager.h"
 
-#include "../../Launch/Resources/Version.h"
-#include <sys/utsname.h>	// for uname()
-
 // define for glibc 2.12.2 and lower (which is shipped with CentOS 6.x and which we target by default)
 #define __secure_getenv getenv
 
@@ -35,9 +26,6 @@ namespace PlatformMiscLimits
 	};
 };
 
-// commandline parameter to suppress DWARF parsing (greatly speeds up callstack generation)
-#define CMDARG_SUPPRESS_DWARF_PARSING			"nodwarf"
-
 namespace
 {
 	/**
@@ -618,803 +606,6 @@ uint32 FLinuxPlatformMisc::GetCPUInfo()
 }
 #endif // PLATFORM_HAS_CPUID
 
-FString DescribeSignal(int32 Signal, siginfo_t* Info)
-{
-	FString ErrorString;
-
-#define HANDLE_CASE(a,b) case a: ErrorString += TEXT(#a ": " b); break;
-
-	switch (Signal)
-	{
-	case SIGSEGV:
-		ErrorString += TEXT("SIGSEGV: invalid attempt to access memory at address ");
-		ErrorString += FString::Printf(TEXT("0x%08x"), (uint32*)Info->si_addr);
-		break;
-	case SIGBUS:
-		ErrorString += TEXT("SIGBUS: invalid attempt to access memory at address ");
-		ErrorString += FString::Printf(TEXT("0x%08x"), (uint32*)Info->si_addr);
-		break;
-
-		HANDLE_CASE(SIGINT, "program interrupted")
-		HANDLE_CASE(SIGQUIT, "user-requested crash")
-		HANDLE_CASE(SIGILL, "illegal instruction")
-		HANDLE_CASE(SIGTRAP, "trace trap")
-		HANDLE_CASE(SIGABRT, "abort() called")
-		HANDLE_CASE(SIGFPE, "floating-point exception")
-		HANDLE_CASE(SIGKILL, "program killed")
-		HANDLE_CASE(SIGSYS, "non-existent system call invoked")
-		HANDLE_CASE(SIGPIPE, "write on a pipe with no reader")
-		HANDLE_CASE(SIGTERM, "software termination signal")
-		HANDLE_CASE(SIGSTOP, "stop")
-
-	default:
-		ErrorString += FString::Printf(TEXT("Signal %d (unknown)"), Signal);
-	}
-
-	return ErrorString;
-#undef HANDLE_CASE
-}
-
-FLinuxCrashContext::~FLinuxCrashContext()
-{
-	if (BacktraceSymbols)
-	{
-		// glibc uses malloc() to allocate this, and we only need to free one pointer, see http://www.gnu.org/software/libc/manual/html_node/Backtraces.html
-		free(BacktraceSymbols);
-		BacktraceSymbols = NULL;
-	}
-
-	if (DebugInfo)
-	{
-		Dwarf_Error ErrorInfo;
-		dwarf_finish(DebugInfo, &ErrorInfo);
-		DebugInfo = NULL;
-	}
-
-	if (ElfHdr)
-	{
-		elf_end(ElfHdr);
-		ElfHdr = NULL;
-	}
-
-	if (ExeFd >= 0)
-	{
-		close(ExeFd);
-		ExeFd = -1;
-	}
-}
-
-void FLinuxCrashContext::InitFromSignal(int32 InSignal, siginfo_t* InInfo, void* InContext)
-{
-	Signal = InSignal;
-	Info = InInfo;
-	Context = reinterpret_cast< ucontext_t* >( InContext );
-
-	// open ourselves for examination
-	if (!FParse::Param( FCommandLine::Get(), TEXT(CMDARG_SUPPRESS_DWARF_PARSING)))
-	{
-		ExeFd = open("/proc/self/exe", O_RDONLY);
-		if (ExeFd >= 0)
-		{
-			Dwarf_Error ErrorInfo;
-			// allocate DWARF debug descriptor
-			if (dwarf_init(ExeFd, DW_DLC_READ, NULL, NULL, &DebugInfo, &ErrorInfo) == DW_DLV_OK)
-			{
-				// get ELF descritor
-				if (dwarf_get_elf(DebugInfo, &ElfHdr, &ErrorInfo) != DW_DLV_OK)
-				{
-					dwarf_finish(DebugInfo, &ErrorInfo);
-					DebugInfo = NULL;
-
-					close(ExeFd);
-					ExeFd = -1;
-				}
-			}
-			else
-			{
-				DebugInfo = NULL;
-				close(ExeFd);
-				ExeFd = -1;
-			}
-		}
-	}
-
-	FCString::Strcat(SignalDescription, ARRAY_COUNT( SignalDescription ) - 1, *DescribeSignal(Signal, Info));
-}
-
-/**
- * Finds a function name in DWARF DIE (Debug Information Entry).
- * For more info on DWARF format, see http://www.dwarfstd.org/Download.php , http://www.ibm.com/developerworks/library/os-debugging/
- *
- * @return true if we need to stop search (i.e. either found it or some error happened)
- */
-bool FindFunctionNameInDIE(Dwarf_Debug DebugInfo, Dwarf_Die Die, Dwarf_Addr Addr, const char **OutFuncName)
-{
-	Dwarf_Error ErrorInfo;
-	Dwarf_Half Tag;
-	Dwarf_Unsigned LowerPC, HigherPC;
-	char *TempFuncName;
-	int ReturnCode;
-
-	if (dwarf_tag(Die, &Tag, &ErrorInfo) != DW_DLV_OK || Tag != DW_TAG_subprogram ||
-		dwarf_attrval_unsigned(Die, DW_AT_low_pc, &LowerPC, &ErrorInfo) != DW_DLV_OK ||
-		dwarf_attrval_unsigned(Die, DW_AT_high_pc, &HigherPC, &ErrorInfo) != DW_DLV_OK ||
-		Addr < LowerPC || HigherPC <= Addr
-		) 
-	{
-		return false;
-	}
-	
-	// found it
-	*OutFuncName = NULL;
-	Dwarf_Attribute SubAt;
-	ReturnCode = dwarf_attr(Die, DW_AT_name, &SubAt, &ErrorInfo);
-	if (ReturnCode == DW_DLV_ERROR)
-	{
-		return true;	// error, but stop the search
-	}
-	else if (ReturnCode == DW_DLV_OK) 
-	{
-		if (dwarf_formstring(SubAt, &TempFuncName, &ErrorInfo))
-		{
-			*OutFuncName = NULL;
-		}
-		else
-		{
-			*OutFuncName = TempFuncName;
-		}
-		return true;
-	}
-
-	// DW_AT_Name is not present, look in DW_AT_specification
-	Dwarf_Attribute SpecAt;
-	if (dwarf_attr(Die, DW_AT_specification, &SpecAt, &ErrorInfo))
-	{
-		// not found, tough luck
-		return false;
-	}
-
-	Dwarf_Off Offset;
-	if (dwarf_global_formref(SpecAt, &Offset, &ErrorInfo))
-	{
-		return false;
-	}
-
-	Dwarf_Die SpecDie;
-	if (dwarf_offdie(DebugInfo, Offset, &SpecDie, &ErrorInfo))
-	{
-		return false;
-	}
-
-	if (dwarf_attrval_string(SpecDie, DW_AT_name, OutFuncName, &ErrorInfo))
-	{
-		*OutFuncName = NULL;
-	}
-
-	return true;
-}
-
-/**
- * Finds a function name in DWARF DIE (Debug Information Entry) and its children.
- * For more info on DWARF format, see http://www.dwarfstd.org/Download.php , http://www.ibm.com/developerworks/library/os-debugging/
- * Note: that function is not exactly traversing the tree, but this "seems to work"(tm). Not sure if we need to descend properly (taking child of every sibling), this
- * takes too much time (and callstacks seem to be fine without it).
- */
-void FindFunctionNameInDIEAndChildren(Dwarf_Debug DebugInfo, Dwarf_Die Die, Dwarf_Addr Addr, const char **OutFuncName)
-{
-	if (OutFuncName == NULL || *OutFuncName != NULL)
-	{
-		return;
-	}
-
-	// search for this Die
-	if (FindFunctionNameInDIE(DebugInfo, Die, Addr, OutFuncName))
-	{
-		return;
-	}
-
-	Dwarf_Die PrevChild = Die, Current = NULL;
-	Dwarf_Error ErrorInfo;
-
-	int32 MaxChildrenAllowed = 32 * 1024 * 1024;	// safeguard to make sure we never get into an infinite loop
-	for(;;)
-	{
-		if (--MaxChildrenAllowed <= 0)
-		{
-			fprintf(stderr, "Breaking out from what seems to be an infinite loop during DWARF parsing (too many children).\n");
-			return;
-		}
-
-		// Get the child
-		if (dwarf_child(PrevChild, &Current, &ErrorInfo) != DW_DLV_OK)
-		{
-			return;	// bail out
-		}
-
-		PrevChild = Current;
-
-		// look for in the child
-		if (FindFunctionNameInDIE(DebugInfo, Current, Addr, OutFuncName))
-		{
-			return;	// got the function name!
-		}
-
-		// search among child's siblings
-		int32 MaxSiblingsAllowed = 64 * 1024 * 1024;	// safeguard to make sure we never get into an infinite loop
-		for(;;)
-		{
-			if (--MaxSiblingsAllowed <= 0)
-			{
-				fprintf(stderr, "Breaking out from what seems to be an infinite loop during DWARF parsing (too many siblings).\n");
-				break;
-			}
-
-			Dwarf_Die Prev = Current;
-			if (dwarf_siblingof(DebugInfo, Prev, &Current, &ErrorInfo) != DW_DLV_OK || Current == NULL)
-			{
-				break;
-			}
-
-			if (FindFunctionNameInDIE(DebugInfo, Current, Addr, OutFuncName))
-			{
-				return;	// got the function name!
-			}
-		}
-	};
-}
-
-bool FLinuxCrashContext::GetInfoForAddress(void* Address, const char **OutFunctionNamePtr, const char **OutSourceFilePtr, int *OutLineNumberPtr)
-{
-	if (DebugInfo == NULL)
-	{
-		return false;
-	}
-
-	Dwarf_Die Die;
-	Dwarf_Unsigned Addr = reinterpret_cast< Dwarf_Unsigned >( Address ), LineNumber = 0;
-	const char * SrcFile = NULL;
-
-	static_assert(sizeof(Dwarf_Unsigned) >= sizeof(Address), "Dwarf_Unsigned type should be long enough to represent pointers. Check libdwarf bitness.");
-
-	int ReturnCode = DW_DLV_OK;
-	Dwarf_Error ErrorInfo;
-	bool bExitHeaderLoop = false;
-	int32 MaxCompileUnitsAllowed = 16 * 1024 * 1024;	// safeguard to make sure we never get into an infinite loop
-	const int32 kMaxBufferLinesAllowed = 16 * 1024 * 1024;	// safeguard to prevent too long line loop
-	for(;;)
-	{
-		if (--MaxCompileUnitsAllowed <= 0)
-		{
-			fprintf(stderr, "Breaking out from what seems to be an infinite loop during DWARF parsing (too many compile units).\n");
-			ReturnCode = DW_DLE_DIE_NO_CU_CONTEXT;	// invalidate
-			break;
-		}
-
-		if (bExitHeaderLoop)
-			break;
-
-		ReturnCode = dwarf_next_cu_header(DebugInfo, NULL, NULL, NULL, NULL, NULL, &ErrorInfo);
-		if (ReturnCode != DW_DLV_OK)
-			break;
-
-		Die = NULL;
-
-		while(dwarf_siblingof(DebugInfo, Die, &Die, &ErrorInfo) == DW_DLV_OK)
-		{
-			Dwarf_Half Tag;
-			if (dwarf_tag(Die, &Tag, &ErrorInfo) != DW_DLV_OK)
-			{
-				bExitHeaderLoop = true;
-				break;
-			}
-
-			if (Tag == DW_TAG_compile_unit)
-			{
-				break;
-			}
-		}
-
-		if (Die == NULL)
-		{
-			break;
-		}
-
-		// check if address is inside this CU
-		Dwarf_Unsigned LowerPC, HigherPC;
-		if (!dwarf_attrval_unsigned(Die, DW_AT_low_pc, &LowerPC, &ErrorInfo) && !dwarf_attrval_unsigned(Die, DW_AT_high_pc, &HigherPC, &ErrorInfo))
-		{
-			if (Addr < LowerPC || Addr >= HigherPC)
-			{
-				continue;
-			}
-		}
-
-		Dwarf_Line * LineBuf;
-		Dwarf_Signed NumLines = kMaxBufferLinesAllowed;
-		if (dwarf_srclines(Die, &LineBuf, &NumLines, &ErrorInfo) != DW_DLV_OK)
-		{
-			// could not get line info for some reason
-			break;
-		}
-
-		if (NumLines >= kMaxBufferLinesAllowed)
-		{
-			fprintf(stderr, "Number of lines associated with a DIE looks unreasonable (%ld), early quitting.\n", NumLines);
-			ReturnCode = DW_DLE_DIE_NO_CU_CONTEXT;	// invalidate
-			break;
-		}
-
-		// look which line is that
-		Dwarf_Addr LineAddress, PrevLineAddress = ~0ULL;
-		Dwarf_Unsigned PrevLineNumber = 0;
-		const char * PrevSrcFile = NULL;
-		char * SrcFileTemp = NULL;
-		for (int Idx = 0; Idx < NumLines; ++Idx)
-		{
-			if (dwarf_lineaddr(LineBuf[Idx], &LineAddress, &ErrorInfo) != 0 ||
-				dwarf_lineno(LineBuf[Idx], &LineNumber, &ErrorInfo) != 0)
-			{
-				bExitHeaderLoop = true;
-				break;
-			}
-
-			if (!dwarf_linesrc(LineBuf[Idx], &SrcFileTemp, &ErrorInfo))
-			{
-				SrcFile = SrcFileTemp;
-			}
-
-			// check if we hit the exact line
-			if (Addr == LineAddress)
-			{
-				bExitHeaderLoop = true;
-				break;
-			}
-			else if (PrevLineAddress < Addr && Addr < LineAddress)
-			{
-				LineNumber = PrevLineNumber;
-				SrcFile = PrevSrcFile;
-				bExitHeaderLoop = true;
-				break;
-			}
-
-			PrevLineAddress = LineAddress;
-			PrevLineNumber = LineNumber;
-			PrevSrcFile = SrcFile;
-		}
-
-	}
-
-	const char * FunctionName = NULL;
-	if (ReturnCode == DW_DLV_OK)
-	{
-		FindFunctionNameInDIEAndChildren(DebugInfo, Die, Addr, &FunctionName);
-	}
-
-	if (OutFunctionNamePtr != NULL && FunctionName != NULL)
-	{
-		*OutFunctionNamePtr = FunctionName;
-	}
-
-	if (OutSourceFilePtr != NULL && SrcFile != NULL)
-	{
-		*OutSourceFilePtr = SrcFile;
-		
-		if (OutLineNumberPtr != NULL)
-		{
-			*OutLineNumberPtr = LineNumber;
-		}
-	}
-
-	// Resets internal CU pointer, so next time we get here it begins from the start
-	while (ReturnCode != DW_DLV_NO_ENTRY) 
-	{
-		if (ReturnCode == DW_DLV_ERROR)
-			break;
-		ReturnCode = dwarf_next_cu_header(DebugInfo, NULL, NULL, NULL, NULL, NULL, &ErrorInfo);
-	}
-
-	// if we weren't able to find a function name, don't trust the source file either
-	return FunctionName != NULL;
-}
-
-/**
- * Handles graceful termination. Gives time to exit gracefully, but second signal will quit immediately.
- */
-void GracefulTerminationHandler(int32 Signal, siginfo_t* Info, void* Context)
-{
-	printf("CtrlCHandler: Signal=%d\n", Signal);
-
-	// make sure as much data is written to disk as possible
-	if (GLog)
-	{
-		GLog->Flush();
-	}
-	if (GWarn)
-	{
-		GWarn->Flush();
-	}
-	if (GError)
-	{
-		GError->Flush();
-	}
-
-	if( !GIsRequestingExit )
-	{
-		GIsRequestingExit = 1;
-	}
-	else
-	{
-		FPlatformMisc::RequestExit(true);
-	}
-}
-
-void CreateExceptionInfoString(int32 Signal, siginfo_t* Info)
-{
-	FString ErrorString = TEXT("Unhandled Exception: ");
-	ErrorString += DescribeSignal(Signal, Info);
-	FCString::Strncpy(GErrorExceptionDescription, *ErrorString, FMath::Min(ErrorString.Len() + 1, (int32)ARRAY_COUNT(GErrorExceptionDescription)));
-}
-
-namespace
-{
-	/** 
-	 * Write a line of UTF-8 to a file
-	 */
-	void WriteLine(FArchive* ReportFile, const ANSICHAR* Line = NULL)
-	{
-		if( Line != NULL )
-		{
-			int64 StringBytes = FCStringAnsi::Strlen(Line);
-			ReportFile->Serialize(( void* )Line, StringBytes);
-		}
-
-		// use Windows line terminator
-		static ANSICHAR WindowsTerminator[] = "\r\n";
-		ReportFile->Serialize(WindowsTerminator, 2);
-	}
-
-	/**
-	 * Serializes UTF string to UTF-16
-	 */
-	void WriteUTF16String(FArchive* ReportFile, const TCHAR* UTFString4BytesChar, uint32 NumChars)
-	{
-		check(UTFString4BytesChar != NULL || NumChars == 0);
-		static_assert(sizeof(TCHAR) == 4, "Platform TCHAR is not 4 bytes. Revisit this function.");
-
-		for (uint32 Idx = 0; Idx < NumChars; ++Idx)
-		{
-			ReportFile->Serialize(const_cast< TCHAR* >( &UTFString4BytesChar[Idx] ), 2);
-		}
-	}
-
-	/** 
-	 * Writes UTF-16 line to a file
-	 */
-	void WriteLine(FArchive* ReportFile, const TCHAR* Line)
-	{
-		if( Line != NULL )
-		{
-			int64 NumChars = FCString::Strlen(Line);
-			WriteUTF16String(ReportFile, Line, NumChars);
-		}
-
-		// use Windows line terminator
-		static TCHAR WindowsTerminator[] = TEXT("\r\n");
-		WriteUTF16String(ReportFile, WindowsTerminator, 2);
-	}
-}
-
-/** 
- * Write all the data mined from the minidump to a text file
- */
-void FLinuxCrashContext::GenerateReport(const FString& DiagnosticsPath) const
-{
-	FArchive* ReportFile = IFileManager::Get().CreateFileWriter(*DiagnosticsPath);
-	if (ReportFile != NULL)
-	{
-		FString Line;
-
-		WriteLine(ReportFile, "Generating report for minidump");
-		WriteLine(ReportFile);
-
-		Line = FString::Printf(TEXT("Application version %d.%d.%d.%d" ), 1, 0, ENGINE_VERSION_HIWORD, ENGINE_VERSION_LOWORD);
-		WriteLine(ReportFile, TCHAR_TO_UTF8(*Line));
-
-		Line = FString::Printf(TEXT(" ... built from changelist %d"), ENGINE_VERSION);
-		WriteLine(ReportFile, TCHAR_TO_UTF8(*Line));
-		WriteLine(ReportFile);
-
-		utsname UnixName;
-		if (uname(&UnixName) == 0)
-		{
-			Line = FString::Printf(TEXT( "OS version %s %s (network name: %s)" ), ANSI_TO_TCHAR(UnixName.sysname), ANSI_TO_TCHAR(UnixName.release), ANSI_TO_TCHAR(UnixName.nodename));
-			WriteLine(ReportFile, TCHAR_TO_UTF8(*Line));	
-			Line = FString::Printf( TEXT( "Running %d %s processors (%d logical cores)" ), FPlatformMisc::NumberOfCores(), ANSI_TO_TCHAR(UnixName.machine), FPlatformMisc::NumberOfCoresIncludingHyperthreads());
-			WriteLine(ReportFile, TCHAR_TO_UTF8(*Line));
-		}
-		else
-		{
-			Line = FString::Printf(TEXT("OS version could not be determined (%d, %s)"), errno, ANSI_TO_TCHAR(strerror(errno)));
-			WriteLine(ReportFile, TCHAR_TO_UTF8(*Line));	
-			Line = FString::Printf( TEXT( "Running %d unknown processors" ), FPlatformMisc::NumberOfCores());
-			WriteLine(ReportFile, TCHAR_TO_UTF8(*Line));
-		}
-		Line = FString::Printf(TEXT("Exception was \"%s\""), SignalDescription);
-		WriteLine(ReportFile, TCHAR_TO_UTF8(*Line));
-		WriteLine(ReportFile);
-
-		WriteLine(ReportFile, "<SOURCE START>");
-		WriteLine(ReportFile, "<SOURCE END>");
-		WriteLine(ReportFile);
-
-		WriteLine(ReportFile, "<CALLSTACK START>");
-		WriteLine(ReportFile, MinidumpCallstackInfo);
-		WriteLine(ReportFile, "<CALLSTACK END>");
-		WriteLine(ReportFile);
-
-		WriteLine(ReportFile, "0 loaded modules");
-
-		WriteLine(ReportFile);
-
-		Line = FString::Printf(TEXT("Report end!"));
-		WriteLine(ReportFile, TCHAR_TO_UTF8(*Line));
-
-		ReportFile->Close();
-		delete ReportFile;
-	}
-}
-
-/** 
- * Mimics Windows WER format
- */
-void GenerateWindowsErrorReport(const FString& WERPath)
-{
-	FArchive* ReportFile = IFileManager::Get().CreateFileWriter(*WERPath);
-	if (ReportFile != NULL)
-	{
-		// write BOM
-		static uint16 ByteOrderMarker = 0xFEFF;
-		ReportFile->Serialize(&ByteOrderMarker, sizeof(ByteOrderMarker));
-
-		WriteLine(ReportFile, TEXT("<?xml version=\"1.0\" encoding=\"UTF-16\"?>"));
-		WriteLine(ReportFile, TEXT("<WERReportMetadata>"));
-		
-		WriteLine(ReportFile, TEXT("\t<OSVersionInformation>"));
-		WriteLine(ReportFile, TEXT("\t\t<WindowsNTVersion>0.0</WindowsNTVersion>"));
-		WriteLine(ReportFile, TEXT("\t\t<Build>No Build</Build>"));
-		WriteLine(ReportFile, TEXT("\t\t<Product>Linux</Product>"));
-		WriteLine(ReportFile, TEXT("\t\t<Edition>No Edition</Edition>"));
-		WriteLine(ReportFile, TEXT("\t\t<BuildString>No BuildString</BuildString>"));
-		WriteLine(ReportFile, TEXT("\t\t<Revision>0</Revision>"));
-		WriteLine(ReportFile, TEXT("\t\t<Flavor>No Flavor</Flavor>"));
-		WriteLine(ReportFile, TEXT("\t\t<Architecture>Unknown Architecture</Architecture>"));
-		WriteLine(ReportFile, TEXT("\t\t<LCID>0</LCID>"));
-		WriteLine(ReportFile, TEXT("\t</OSVersionInformation>"));
-		
-		WriteLine(ReportFile, TEXT("\t<ParentProcessInformation>"));
-		WriteLine(ReportFile, *FString::Printf(TEXT("\t\t<ParentProcessId>%d</ParentProcessId>"), getppid()));
-		WriteLine(ReportFile, TEXT("\t\t<ParentProcessPath>C:\\Windows\\explorer.exe</ParentProcessPath>"));			// FIXME: supply valid?
-		WriteLine(ReportFile, TEXT("\t\t<ParentProcessCmdLine>C:\\Windows\\Explorer.EXE</ParentProcessCmdLine>"));	// FIXME: supply valid?
-		WriteLine(ReportFile, TEXT("\t</ParentProcessInformation>"));
-
-		WriteLine(ReportFile, TEXT("\t<ProblemSignatures>"));
-		WriteLine(ReportFile, TEXT("\t\t<EventType>APPCRASH</EventType>"));
-		WriteLine(ReportFile, *FString::Printf(TEXT("\t\t<Parameter0>UE4-%s</Parameter0>"), FApp::GetGameName()));
-		WriteLine(ReportFile, *FString::Printf(TEXT("\t\t<Parameter1>%d.%d.%d.0</Parameter1>"), ENGINE_MAJOR_VERSION, ENGINE_MINOR_VERSION, ENGINE_PATCH_VERSION));
-		WriteLine(ReportFile, TEXT("\t\t<Parameter2>0</Parameter2>"));													// FIXME: supply valid?
-		WriteLine(ReportFile, TEXT("\t\t<Parameter3>Unknown Fault Module</Parameter3>"));										// FIXME: supply valid?
-		WriteLine(ReportFile, TEXT("\t\t<Parameter4>0.0.0.0</Parameter4>"));													// FIXME: supply valid?
-		WriteLine(ReportFile, TEXT("\t\t<Parameter5>00000000</Parameter5>"));													// FIXME: supply valid?
-		WriteLine(ReportFile, TEXT("\t\t<Parameter6>00000000</Parameter6>"));													// FIXME: supply valid?
-		WriteLine(ReportFile, TEXT("\t\t<Parameter7>0000000000000000</Parameter7>"));											// FIXME: supply valid?
-		WriteLine(ReportFile, *FString::Printf(TEXT("\t\t<Parameter8>!%s!</Parameter8>"), FCommandLine::Get()));				// FIXME: supply valid? Only partially valid
-		WriteLine(ReportFile, *FString::Printf(TEXT("\t\t<Parameter9>%s!%s!%s!%d</Parameter9>"), TEXT( BRANCH_NAME ), FPlatformProcess::BaseDir(), FPlatformMisc::GetEngineMode(), BUILT_FROM_CHANGELIST));
-		WriteLine(ReportFile, TEXT("\t</ProblemSignatures>"));
-
-		WriteLine(ReportFile, TEXT("\t<DynamicSignatures>"));
-		WriteLine(ReportFile, TEXT("\t\t<Parameter1>6.1.7601.2.1.0.256.48</Parameter1>"));
-		WriteLine(ReportFile, TEXT("\t\t<Parameter2>1033</Parameter2>"));
-		WriteLine(ReportFile, TEXT("\t</DynamicSignatures>"));
-
-		WriteLine(ReportFile, TEXT("\t<SystemInformation>"));
-		WriteLine(ReportFile, TEXT("\t\t<MID>11111111-2222-3333-4444-555555555555</MID>"));							// FIXME: supply valid?
-		
-		WriteLine(ReportFile, TEXT("\t\t<SystemManufacturer>Unknown.</SystemManufacturer>"));						// FIXME: supply valid?
-		WriteLine(ReportFile, TEXT("\t\t<SystemProductName>Linux machine</SystemProductName>"));					// FIXME: supply valid?
-		WriteLine(ReportFile, TEXT("\t\t<BIOSVersion>A02</BIOSVersion>"));											// FIXME: supply valid?
-		WriteLine(ReportFile, TEXT("\t</SystemInformation>"));
-
-		WriteLine(ReportFile, TEXT("</WERReportMetadata>"));
-
-		ReportFile->Close();
-		delete ReportFile;
-	}
-}
-
-/** 
- * Creates (fake so far) minidump
- */
-void GenerateMinidump(const FString& Path)
-{
-	FArchive* ReportFile = IFileManager::Get().CreateFileWriter(*Path);
-	if (ReportFile != NULL)
-	{
-		// write BOM
-		static uint32 Garbage = 0xDEADBEEF;
-		ReportFile->Serialize(&Garbage, sizeof(Garbage));
-
-		ReportFile->Close();
-		delete ReportFile;
-	}
-}
-
-
-int32 ReportCrash(const FLinuxCrashContext& Context)
-{
-	static bool GAlreadyCreatedMinidump = false;
-	// Only create a minidump the first time this function is called.
-	// (Can be called the first time from the RenderThread, then a second time from the MainThread.)
-	if ( GAlreadyCreatedMinidump == false )
-	{
-		GAlreadyCreatedMinidump = true;
-
-		const SIZE_T StackTraceSize = 65535;
-		ANSICHAR* StackTrace = (ANSICHAR*) FMemory::Malloc( StackTraceSize );
-		StackTrace[0] = 0;
-		// Walk the stack and dump it to the allocated memory (ignore first 2 callstack lines as those are in stack walking code)
-		FPlatformStackWalk::StackWalkAndDump( StackTrace, StackTraceSize, 2, const_cast< FLinuxCrashContext* >( &Context ) );
-
-		FCString::Strncat( GErrorHist, ANSI_TO_TCHAR(StackTrace), ARRAY_COUNT(GErrorHist) - 1 );
-		CreateExceptionInfoString(Context.Signal, Context.Info);
-
-		FMemory::Free( StackTrace );
-	}
-
-	return 0;
-}
-
-/**
- * Generates information for crash reporter
- */
-void GenerateCrashInfoAndLaunchReporter(const FLinuxCrashContext& Context)
-{
-	// do not report crashes for tools (particularly for crash reporter itself)
-#if !IS_PROGRAM
-
-	// create a crash-specific directory
-	FString CrashInfoFolder = FString::Printf(TEXT("crashinfo-%s-pid-%d-%s-%s"), FApp::GetGameName(), getpid(), 
-		*FDateTime::Now().ToString(), *FGuid::NewGuid().ToString());
-	FString CrashInfoAbsolute = FPaths::ConvertRelativePathToFull(CrashInfoFolder);
-	if (IFileManager::Get().MakeDirectory(*CrashInfoFolder))
-	{
-		// generate "minidump"
-		Context.GenerateReport(FPaths::Combine(*CrashInfoFolder, TEXT("diagnostics.txt")));
-
-		// generate "WER"
-		GenerateWindowsErrorReport(FPaths::Combine(*CrashInfoFolder, TEXT("wermeta.xml")));
-
-		// generate "minidump" (just >1 byte)
-		GenerateMinidump(FPaths::Combine(*CrashInfoFolder, TEXT("minidump.dmp")));
-
-		// copy log
-		FString LogSrcAbsolute = FPlatformOutputDevices::GetAbsoluteLogFilename();
-		FString LogDstAbsolute = FPaths::Combine(*CrashInfoAbsolute, *FPaths::GetCleanFilename(LogSrcAbsolute));
-		FPaths::NormalizeDirectoryName(LogDstAbsolute);
-		static_cast<void>(IFileManager::Get().Copy(*LogDstAbsolute, *LogSrcAbsolute));	// best effort, so don't care about result: couldn't copy -> tough, no log
-
-		// try launching the tool and wait for its exit, if at all
-		const TCHAR* RelativePathToCrashReporter = TEXT("../../../engine/binaries/linux/crashreportclient");	// FIXME: painfully hard-coded
-		if (!FPaths::FileExists(RelativePathToCrashReporter))
-		{
-			RelativePathToCrashReporter = TEXT("../../../Engine/Binaries/Linux/CrashReportClient");	// FIXME: even more painfully hard-coded
-		}
-
-		// show on the console
-		printf("Starting %s\n", StringCast<ANSICHAR>(RelativePathToCrashReporter).Get());
-		FProcHandle RunningProc = FPlatformProcess::CreateProc(RelativePathToCrashReporter, *(CrashInfoAbsolute + TEXT("/")), true, false, false, NULL, 0, NULL, NULL);
-		if (FPlatformProcess::IsProcRunning(RunningProc))
-		{
-			// do not wait indefinitely
-			double kTimeOut = 3 * 60.0;
-			double StartSeconds = FPlatformTime::Seconds();
-			for(;;)
-			{
-				if (!FPlatformProcess::IsProcRunning(RunningProc))
-				{
-					break;
-				}
-
-				if (FPlatformTime::Seconds() - StartSeconds > kTimeOut)
-				{
-					break;
-				}
-
-				FPlatformProcess::Sleep(1.0f);
-			};
-		}
-	}
-
-#endif
-
-	FPlatformMisc::RequestExit(true);
-}
-
-/**
- * Good enough default crash reporter.
- */
-void DefaultCrashHandler(const FLinuxCrashContext& Context)
-{
-	printf("DefaultCrashHandler: Signal=%d\n", Context.Signal);
-
-	ReportCrash(Context);
-	if (GLog)
-	{
-		GLog->Flush();
-	}
-	if (GWarn)
-	{
-		GWarn->Flush();
-	}
-	if (GError)
-	{
-		GError->Flush();
-		GError->HandleError();
-	}
-
-	return GenerateCrashInfoAndLaunchReporter(Context);
-}
-
-/** Global pointer to crash handler */
-void (* GCrashHandlerPointer)(const FGenericCrashContext& Context) = NULL;
-
-/** True system-specific crash handler that gets called first */
-void PlatformCrashHandler(int32 Signal, siginfo_t* Info, void* Context)
-{
-	fprintf(stderr, "Signal %d caught.\n", Signal);
-
-	FLinuxCrashContext CrashContext;
-	CrashContext.InitFromSignal(Signal, Info, Context);
-
-	if (GCrashHandlerPointer)
-	{
-		GCrashHandlerPointer(CrashContext);
-	}
-	else
-	{
-		// call default one
-		DefaultCrashHandler(CrashContext);
-	}
-}
-
-void FLinuxPlatformMisc::SetGracefulTerminationHandler()
-{
-	struct sigaction Action;
-	FMemory::Memzero(&Action, sizeof(struct sigaction));
-	Action.sa_sigaction = GracefulTerminationHandler;
-	sigemptyset(&Action.sa_mask);
-	Action.sa_flags = SA_SIGINFO | SA_RESTART | SA_ONSTACK;
-	sigaction(SIGINT, &Action, NULL);
-	sigaction(SIGTERM, &Action, NULL);
-	sigaction(SIGHUP, &Action, NULL);	//  this should actually cause the server to just re-read configs (restart?)
-}
-
-void FLinuxPlatformMisc::SetCrashHandler(void (* CrashHandler)(const FGenericCrashContext& Context))
-{
-	GCrashHandlerPointer = CrashHandler;
-
-	struct sigaction Action;
-	FMemory::Memzero(&Action, sizeof(struct sigaction));
-	Action.sa_sigaction = PlatformCrashHandler;
-	sigemptyset(&Action.sa_mask);
-	Action.sa_flags = SA_SIGINFO | SA_RESTART | SA_ONSTACK;
-	sigaction(SIGQUIT, &Action, NULL);	// SIGQUIT is a user-initiated "crash".
-	sigaction(SIGILL, &Action, NULL);
-	sigaction(SIGFPE, &Action, NULL);
-	sigaction(SIGBUS, &Action, NULL);
-	sigaction(SIGSEGV, &Action, NULL);
-	sigaction(SIGSYS, &Action, NULL);
-}
-
 #if !UE_BUILD_SHIPPING
 bool FLinuxPlatformMisc::IsDebuggerPresent()
 {
diff --git a/Engine/Source/Runtime/Core/Private/Mac/MacMallocZone.cpp b/Engine/Source/Runtime/Core/Private/Mac/MacMallocZone.cpp
--- a/Engine/Source/Runtime/Core/Private/Mac/MacMallocZone.cpp
+++ b/Engine/Source/Runtime/Core/Private/Mac/MacMallocZone.cpp
@@ -3,6 +3,7 @@
 #include "Core.h"
 #include "MacMallocZone.h"
 #include "MacPlatformMisc.h"
+#include "MacPlatformCrashContext.h"
 #include <malloc/malloc.h>
 
 FMacMallocZone::FMacMallocZone( uint64 const InitialSize )
diff --git a/Engine/Source/Runtime/Core/Private/Mac/MacPlatformMisc.cpp b/Engine/Source/Runtime/Core/Private/Mac/MacPlatformMisc.cpp
--- a/Engine/Source/Runtime/Core/Private/Mac/MacPlatformMisc.cpp
+++ b/Engine/Source/Runtime/Core/Private/Mac/MacPlatformMisc.cpp
@@ -16,6 +16,7 @@
 #include "EngineVersion.h"
 #include "MacMallocZone.h"
 #include "ApplePlatformSymbolication.h"
+#include "MacPlatformCrashContext.h"
 
 #include <dlfcn.h>
 #include <IOKit/IOKitLib.h>
@@ -1708,6 +1709,12 @@ void FMacCrashContext::GenerateCrashInfoAndLaunchReporter() const
 			close(ReportFile);
 		}
 		
+		// Introduces a new runtime crash context. Will replace all Windows related crash reporting.
+		FCStringAnsi::Strncpy(FilePath, CrashInfoFolder, PATH_MAX);
+		FCStringAnsi::Strcat(FilePath, PATH_MAX, "/" );
+		FCStringAnsi::Strcat(FilePath, PATH_MAX, FGenericCrashContext::CrashContextRuntimeXMLNameA );
+		//SerializeAsXML( FilePath ); @todo uncomment after verification
+		
 		// copy log
 		FCStringAnsi::Strncpy(FilePath, CrashInfoFolder, PATH_MAX);
 		FCStringAnsi::Strcat(FilePath, PATH_MAX, "/");
diff --git a/Engine/Source/Runtime/Core/Private/Windows/WindowsPlatformExceptionHandling.cpp b/Engine/Source/Runtime/Core/Private/Windows/WindowsPlatformExceptionHandling.cpp
--- a/Engine/Source/Runtime/Core/Private/Windows/WindowsPlatformExceptionHandling.cpp
+++ b/Engine/Source/Runtime/Core/Private/Windows/WindowsPlatformExceptionHandling.cpp
@@ -1,7 +1,9 @@
 // Copyright 1998-2014 Epic Games, Inc. All Rights Reserved.
 
 #include "Core.h"
+#include "MallocCrash.h"
 #include "ExceptionHandling.h"
+#include "GenericPlatform/GenericPlatformContext.h"
 #include "../../Launch/Resources/Version.h"
 
 #include "AllowWindowsPlatformTypes.h"
@@ -190,6 +192,14 @@ void AddMiniDump(HREPORT ReportHandle, EXCEPTION_POINTERS* ExceptionInfo)
 	}
 }
 
+struct FWindowsPlatformCrashContext : public FGenericCrashContext
+{
+	void AddPlatformSpecificProperties()
+	{
+		AddCrashProperty( TEXT( "Platform.IsRunningWindows" ), 1 );
+	}
+};
+
 /** 
  * Add miscellaneous files to the report. Currently the log and the video file
  */
@@ -199,7 +209,14 @@ void AddMiscFiles( HREPORT ReportHandle )
 	WerReportAddFile( ReportHandle, *LogFileName, WerFileTypeOther, WER_FILE_ANONYMOUS_DATA ); 
 
 	FString CrashVideoPath = FPaths::GameLogDir() / TEXT( "CrashVideo.avi" );
-	WerReportAddFile( ReportHandle, *CrashVideoPath, WerFileTypeOther, WER_FILE_ANONYMOUS_DATA ); 
+	WerReportAddFile( ReportHandle, *CrashVideoPath, WerFileTypeOther, WER_FILE_ANONYMOUS_DATA );
+
+	// Introduces a new runtime crash context. Will replace all Windows related crash reporting.
+	const FString CrashContextXMLPath = FPaths::Combine( *FPaths::GameLogDir(), FWindowsPlatformCrashContext::CrashContextRuntimeXMLNameW );
+	FWindowsPlatformCrashContext CrashContext;
+	CrashContext.SerializeAsXML( *CrashContextXMLPath );
+
+	WerReportAddFile( ReportHandle, *CrashContextXMLPath, WerFileTypeOther, WER_FILE_ANONYMOUS_DATA );
 }
 
 /**
@@ -446,6 +463,9 @@ int32 NullReportCrash( LPEXCEPTION_POINTERS ExceptionInfo )
 
 int32 ReportCrash( LPEXCEPTION_POINTERS ExceptionInfo )
 {
+	// Switch to malloc crash.
+	FMallocCrash::Get().SetAsGMalloc();
+
 	// Only create a minidump the first time this function is called.
 	// (Can be called the first time from the RenderThread, then a second time from the MainThread.)
 	if( FPlatformAtomics::InterlockedIncrement( &ReportCrashCallCount ) != 1 )
@@ -457,7 +477,7 @@ int32 ReportCrash( LPEXCEPTION_POINTERS ExceptionInfo )
 
 	// Note: ReportCrashUsingCrashReportClient reads the callstack from GErrorHist
 	const SIZE_T StackTraceSize = 65535;
-	ANSICHAR* StackTrace = (ANSICHAR*) FMemory::SystemMalloc( StackTraceSize );
+	ANSICHAR* StackTrace = (ANSICHAR*) GMalloc->Malloc( StackTraceSize );
 	StackTrace[0] = 0;
 	// Walk the stack and dump it to the allocated memory.
 	FPlatformStackWalk::StackWalkAndDump( StackTrace, StackTraceSize, 0, ExceptionInfo->ContextRecord );
@@ -472,7 +492,7 @@ int32 ReportCrash( LPEXCEPTION_POINTERS ExceptionInfo )
 
 	FCString::Strncat( GErrorHist, ANSI_TO_TCHAR(StackTrace), ARRAY_COUNT(GErrorHist) );
 
-	FMemory::SystemFree( StackTrace );
+	GMalloc->Free( StackTrace );
 
 #if WINVER > 0x502	// Windows Error Reporting is not supported on Windows XP
 	if (GUseCrashReportClient)
diff --git a/Engine/Source/Runtime/Core/Public/Android/AndroidMisc.h b/Engine/Source/Runtime/Core/Public/Android/AndroidMisc.h
--- a/Engine/Source/Runtime/Core/Public/Android/AndroidMisc.h
+++ b/Engine/Source/Runtime/Core/Public/Android/AndroidMisc.h
@@ -9,43 +9,6 @@
 
 //@todo android: this entire file
 
-struct CORE_API FAndroidCrashContext : public FGenericCrashContext
-{
-	/** Signal number */
-	int32 Signal;
-	
-	/** Additional signal info */
-	siginfo* Info;
-	
-	/** Thread context */
-	void* Context;
-
-	FAndroidCrashContext()
-		:	Signal(0)
-		,	Info(NULL)
-		,	Context(NULL)
-	{
-	}
-
-	~FAndroidCrashContext()
-	{
-	}
-
-	/**
-	 * Inits the crash context from data provided by a signal handler.
-	 *
-	 * @param InSignal number (SIGSEGV, etc)
-	 * @param InInfo additional info (e.g. address we tried to read, etc)
-	 * @param InContext thread context
-	 */
-	void InitFromSignal(int32 InSignal, siginfo* InInfo, void* InContext)
-	{
-		Signal = InSignal;
-		Info = InInfo;
-		Context = InContext;
-	}
-};
-
 /**
  * Android implementation of the misc OS functions
  */
diff --git a/Engine/Source/Runtime/Core/Public/Android/AndroidPlatformCrashContext.h b/Engine/Source/Runtime/Core/Public/Android/AndroidPlatformCrashContext.h
new file mode 100644
--- /dev/null
+++ b/Engine/Source/Runtime/Core/Public/Android/AndroidPlatformCrashContext.h
@@ -0,0 +1,42 @@
+// Copyright 1998-2014 Epic Games, Inc. All Rights Reserved.
+
+#pragma once
+
+#include "GenericPlatform/GenericPlatformContext.h"
+
+struct CORE_API FAndroidCrashContext : public FGenericCrashContext
+{
+	/** Signal number */
+	int32 Signal;
+	
+	/** Additional signal info */
+	siginfo* Info;
+	
+	/** Thread context */
+	void* Context;
+
+	FAndroidCrashContext()
+		:	Signal(0)
+		,	Info(NULL)
+		,	Context(NULL)
+	{
+	}
+
+	~FAndroidCrashContext()
+	{
+	}
+
+	/**
+	 * Inits the crash context from data provided by a signal handler.
+	 *
+	 * @param InSignal number (SIGSEGV, etc)
+	 * @param InInfo additional info (e.g. address we tried to read, etc)
+	 * @param InContext thread context
+	 */
+	void InitFromSignal(int32 InSignal, siginfo* InInfo, void* InContext)
+	{
+		Signal = InSignal;
+		Info = InInfo;
+		Context = InContext;
+	}
+};
\ No newline at end of file
diff --git a/Engine/Source/Runtime/Core/Public/Apple/ApplePlatformCrashContext.h b/Engine/Source/Runtime/Core/Public/Apple/ApplePlatformCrashContext.h
--- a/Engine/Source/Runtime/Core/Public/Apple/ApplePlatformCrashContext.h
+++ b/Engine/Source/Runtime/Core/Public/Apple/ApplePlatformCrashContext.h
@@ -6,6 +6,8 @@
 
 #pragma once
 
+#include "GenericPlatform/GenericPlatformContext.h"
+
 /**
  * Declaration for common crash context implementation for Apple platforms.
  */
diff --git a/Engine/Source/Runtime/Core/Public/GenericPlatform/GenericPlatformContext.h b/Engine/Source/Runtime/Core/Public/GenericPlatform/GenericPlatformContext.h
new file mode 100644
--- /dev/null
+++ b/Engine/Source/Runtime/Core/Public/GenericPlatform/GenericPlatformContext.h
@@ -0,0 +1,70 @@
+// Copyright 1998-2014 Epic Games, Inc. All Rights Reserved.
+
+#pragma once
+
+class FString;
+
+/** Enumerates crash description versions. */
+enum class ECrashDescVersions
+{
+	/** Introduces a new crash description format. */
+	VER_1_NewCrashFormat,
+
+	/** Added misc properties (CPU,GPU,OS,etc), memory related stats and platform specific properties as generic payload. */
+	VER_2_AddedNewProperties,
+};
+
+/**
+ *	Contains a runtime crash's properties that are common for all platforms.
+ *	This may change in the future.
+ */
+struct CORE_API FGenericCrashContext
+{
+public:
+	static const ANSICHAR* CrashContextRuntimeXMLNameA;
+	static const TCHAR* CrashContextRuntimeXMLNameW;
+
+	/** Default constructor. */
+	FGenericCrashContext();
+
+	/** Serializes all data to the buffer. */
+	void SerializeContentToBuffer();
+
+	const FString& GetBuffer() const
+	{
+		return CommonBuffer;
+	}
+
+	/** Serializes crash's informations to the specified filename. */
+	void SerializeAsXML( const TCHAR* Filename );
+
+	/** Writes a common property to the buffer. */
+	void AddCrashProperty( const TCHAR* PropertyName, const TCHAR* PropertyValue );
+
+	/** Writes a common property to the buffer. */
+	template <typename Type>
+	void AddCrashProperty( const TCHAR* PropertyName, const Type& Value )
+	{
+		AddCrashProperty( PropertyName, *TTypeToString<Type>::ToString( Value ) );
+	}
+
+private:
+	/** Serializes platform specific properties to the buffer. */
+	virtual void AddPlatformSpecificProperties();
+
+	/** Writes header information to the buffer. */
+	void AddHeader();
+
+	/** Writes footer to the buffer. */
+	void AddFooter();
+
+	/** The buffer used to store the crash's properties. */
+	FString& CommonBuffer;
+
+	// FNoncopyable
+	FGenericCrashContext( const FGenericCrashContext& );
+	FGenericCrashContext& operator=(const FGenericCrashContext&);
+};
+
+struct CORE_API FGenericMemoryWarningContext
+{};
diff --git a/Engine/Source/Runtime/Core/Public/GenericPlatform/GenericPlatformMisc.h b/Engine/Source/Runtime/Core/Public/GenericPlatform/GenericPlatformMisc.h
--- a/Engine/Source/Runtime/Core/Public/GenericPlatform/GenericPlatformMisc.h
+++ b/Engine/Source/Runtime/Core/Public/GenericPlatform/GenericPlatformMisc.h
@@ -7,6 +7,9 @@
 
 #pragma once
 
+struct FGenericCrashContext;
+struct FGenericMemoryWarningContext;
+
 class FText;
 
 template<typename KeyType, typename ValueType> class TPair;
@@ -182,126 +185,6 @@ namespace EAppReturnType
 	};
 }
 
-/**
- *	Contains crash's description that is common for all platform and can be obtained in run-time.
- *	Since we can't used dynamic memory allocation everything is defined as simple as possible.
- *	This may change in the future.
- */
-struct CORE_API FGenericCrashDescription
-{
-	/**
-	 * The name of the game that crashed. (AppID)
-	 */
-	TCHAR GameName[64];
-
-	/**
-	 * The mode the game was in e.g. editor.
-	 */
-	TCHAR EngineMode[64];
-
-	/**
-	 * The platform that crashed e.g. Win64.
-	 * Main platform name [detailed information like OS version]
-	 */
-	TCHAR Platform[64];
-
-	/**
-	* Encoded engine version. (AppVersion)
-	* E.g. 4.3.0.0-2215663+UE4-Releases+4.3
-	* BuildVersion-BuiltFromCL-BranchName
-	* @EngineVersion	varchar(64)
-	*/
-	//FEngineVersion EngineVersion;
-
-	//Architecture;
-
-	//NumberOfCores;
-
-	//CrashedModuleName;
-
-	//LoadedModules ? ;
-
-	//static const FPlatformMemoryConstants& GetConstants();
-	//static FPlatformMemoryStats GetStats();
-
-	/**
-	 * The three component version of the app e.g. 4.4.1
-	 */
-	TCHAR BuildVersion[32];
-
-	/**
-	 * Built from changelist.
-	 */
-	uint32 BuiltFromCL;
-
-	/**
-	 * The name of the branch this game was built out of.
-	 */
-	TCHAR BranchName[32];
-
-	/**
-	 * The command line of the application that crashed.
-	 */
-	TCHAR CommandLine[512];
-
-	/**
-	 * The base directory where the app was running.
-	 */
-	TCHAR BaseDir[512];
-
-	/**
-	 * The language ID the application that crashed.
-	 */
-	TCHAR LanguageLCID[64];
-
-	/**
-	 * The name of the user that caused this crash.
-	 */
-	TCHAR UserName[64];
-
-	/**
-	 * The unique ID used to identify the machine the crash occurred on.
-	 */
-	TCHAR MachineId[64];
-
-	/**
-	 * The Epic account ID for the user who last used the Launcher.
-	 */
-	TCHAR EpicAccountId[64];
-
-	/**
-	 * A string representing the callstack of the crash.
-	 */
-	TCHAR CallStack[16384];
-
-	/**
-	 * A string representing the user description of the crash.
-	 */
-	TCHAR UserDescription[512];
-
-	/**
-	 * The error message, can be assertion message, ensure message or message from the fatal error.
-	 */
-	TCHAR ErrorMessage[512];
-
-	/**
-	 * The UTC time the crash occurred.
-	 */
-	//	FDateTime TimeOfCrash;
-
-	/** Lenght of the crash's description, in bytes. */
-	int32 Lenght;
-};
-
-struct CORE_API FGenericCrashContext
-{
-};
-
-struct CORE_API FGenericMemoryWarningContext
-{
-};
-
-
 /**
 * Generic implementation for most platforms
 **/
@@ -743,7 +626,7 @@ public:
 	/** Return the name of the platform features module. Can be NULL if there are no extra features for this platform */
 	static const TCHAR* GetPlatformFeaturesModuleName()
 	{
-		// by deafult, no module
+		// by default, no module
 		return NULL;
 	}
 
diff --git a/Engine/Source/Runtime/Core/Public/Linux/LinuxPlatformCrashContext.h b/Engine/Source/Runtime/Core/Public/Linux/LinuxPlatformCrashContext.h
new file mode 100644
--- /dev/null
+++ b/Engine/Source/Runtime/Core/Public/Linux/LinuxPlatformCrashContext.h
@@ -0,0 +1,81 @@
+// Copyright 1998-2014 Epic Games, Inc. All Rights Reserved.
+
+#pragma once
+
+#include "GenericPlatform/GenericPlatformContext.h"
+
+// commandline parameter to suppress DWARF parsing (greatly speeds up callstack generation)
+#define CMDARG_SUPPRESS_DWARF_PARSING			"nodwarf"
+
+struct CORE_API FLinuxCrashContext : public FGenericCrashContext
+{
+	/** Signal number */
+	int32 Signal;
+	
+	/** Additional signal info */
+	siginfo_t* Info;
+	
+	/** Thread context */
+	ucontext_t*	Context;
+
+	/** Symbols received via backtrace_symbols(), if any (note that we will need to clean it up) */
+	char ** BacktraceSymbols;
+
+	/** File descriptor needed for libelf to open (our own) binary */
+	int ExeFd;
+
+	/** Elf header as used by libelf (forward-declared in the same way as libelf does it, it's normally of Elf type) */
+	struct _Elf* ElfHdr;
+
+	/** DWARF handle used by libdwarf (forward-declared in the same way as libdwarf does it, it's normally of Dwarf_Debug type) */
+	struct _Dwarf_Debug	* DebugInfo;
+
+	/** Memory reserved for "exception" (signal) info */
+	TCHAR SignalDescription[128];
+
+	/** Memory reserved for minidump-style callstack info */
+	char MinidumpCallstackInfo[16384];
+
+	FLinuxCrashContext()
+		:	Signal(0)
+		,	Info(NULL)
+		,	Context(NULL)
+		,	BacktraceSymbols(NULL)
+		,	ExeFd(-1)
+		,	ElfHdr(NULL)
+		,	DebugInfo(NULL)
+	{
+		SignalDescription[ 0 ] = 0;
+		MinidumpCallstackInfo[ 0 ] = 0;
+	}
+
+	~FLinuxCrashContext();
+
+	/**
+	 * Inits the crash context from data provided by a signal handler.
+	 *
+	 * @param InSignal number (SIGSEGV, etc)
+	 * @param InInfo additional info (e.g. address we tried to read, etc)
+	 * @param InContext thread context
+	 */
+	void InitFromSignal(int32 InSignal, siginfo_t* InInfo, void* InContext);
+
+	/**
+	 * Gets information for the crash.
+	 *
+	 * @param Address the address to look up info for
+	 * @param OutFunctionNamePtr pointer to function name (may be NULL). Caller doesn't have to free it, but need to consider it temporary (i.e. next GetInfoForAddress() call on any thread may change it).
+	 * @param OutSourceFilePtr pointer to source filename (may be NULL). Caller doesn't have to free it, but need to consider it temporary (i.e. next GetInfoForAddress() call on any thread may change it).
+	 * @param OutLineNumberPtr pointer to line in a source file (may be NULL). Caller doesn't have to free it, but need to consider it temporary (i.e. next GetInfoForAddress() call on any thread may change it).
+	 *
+	 * @return true if succeeded in getting the info. If false is returned, none of above parameters should be trusted to contain valid data!
+	 */
+	bool GetInfoForAddress(void * Address, const char **OutFunctionNamePtr, const char **OutSourceFilePtr, int *OutLineNumberPtr);
+
+	/**
+	 * Dumps all the data from crash context to the "minidump" report.
+	 *
+	 * @param DiagnosticsPath Path to put the file to
+	 */
+	void GenerateReport(const FString & DiagnosticsPath) const;
+};
diff --git a/Engine/Source/Runtime/Core/Public/Linux/LinuxPlatformMisc.h b/Engine/Source/Runtime/Core/Public/Linux/LinuxPlatformMisc.h
--- a/Engine/Source/Runtime/Core/Public/Linux/LinuxPlatformMisc.h
+++ b/Engine/Source/Runtime/Core/Public/Linux/LinuxPlatformMisc.h
@@ -7,78 +7,7 @@
 
 #pragma once
 
-struct CORE_API FLinuxCrashContext : public FGenericCrashContext
-{
-	/** Signal number */
-	int32 Signal;
-	
-	/** Additional signal info */
-	siginfo_t* Info;
-	
-	/** Thread context */
-	ucontext_t*	Context;
-
-	/** Symbols received via backtrace_symbols(), if any (note that we will need to clean it up) */
-	char ** BacktraceSymbols;
-
-	/** File descriptor needed for libelf to open (our own) binary */
-	int ExeFd;
-
-	/** Elf header as used by libelf (forward-declared in the same way as libelf does it, it's normally of Elf type) */
-	struct _Elf* ElfHdr;
-
-	/** DWARF handle used by libdwarf (forward-declared in the same way as libdwarf does it, it's normally of Dwarf_Debug type) */
-	struct _Dwarf_Debug	* DebugInfo;
-
-	/** Memory reserved for "exception" (signal) info */
-	TCHAR SignalDescription[128];
-
-	/** Memory reserved for minidump-style callstack info */
-	char MinidumpCallstackInfo[16384];
-
-	FLinuxCrashContext()
-		:	Signal(0)
-		,	Info(NULL)
-		,	Context(NULL)
-		,	BacktraceSymbols(NULL)
-		,	ExeFd(-1)
-		,	ElfHdr(NULL)
-		,	DebugInfo(NULL)
-	{
-		SignalDescription[ 0 ] = 0;
-		MinidumpCallstackInfo[ 0 ] = 0;
-	}
-
-	~FLinuxCrashContext();
-
-	/**
-	 * Inits the crash context from data provided by a signal handler.
-	 *
-	 * @param InSignal number (SIGSEGV, etc)
-	 * @param InInfo additional info (e.g. address we tried to read, etc)
-	 * @param InContext thread context
-	 */
-	void InitFromSignal(int32 InSignal, siginfo_t* InInfo, void* InContext);
-
-	/**
-	 * Gets information for the crash.
-	 *
-	 * @param Address the address to look up info for
-	 * @param OutFunctionNamePtr pointer to function name (may be NULL). Caller doesn't have to free it, but need to consider it temporary (i.e. next GetInfoForAddress() call on any thread may change it).
-	 * @param OutSourceFilePtr pointer to source filename (may be NULL). Caller doesn't have to free it, but need to consider it temporary (i.e. next GetInfoForAddress() call on any thread may change it).
-	 * @param OutLineNumberPtr pointer to line in a source file (may be NULL). Caller doesn't have to free it, but need to consider it temporary (i.e. next GetInfoForAddress() call on any thread may change it).
-	 *
-	 * @return true if succeeded in getting the info. If false is returned, none of above parameters should be trusted to contain valid data!
-	 */
-	bool GetInfoForAddress(void* Address, const char **OutFunctionNamePtr, const char **OutSourceFilePtr, int *OutLineNumberPtr);
-
-	/**
-	 * Dumps all the data from crash context to the "minidump" report.
-	 *
-	 * @param DiagnosticsPath Path to put the file to
-	 */
-	void GenerateReport(const FString& DiagnosticsPath) const;
-};
+struct FGenericCrashContext;
 
 /**
  * Linux implementation of the misc OS functions
diff --git a/Engine/Source/Runtime/Core/Public/Mac/MacPlatformCrashContext.h b/Engine/Source/Runtime/Core/Public/Mac/MacPlatformCrashContext.h
new file mode 100644
--- /dev/null
+++ b/Engine/Source/Runtime/Core/Public/Mac/MacPlatformCrashContext.h
@@ -0,0 +1,18 @@
+// Copyright 1998-2014 Epic Games, Inc. All Rights Reserved.
+
+#pragma once
+
+#include "Apple/ApplePlatformCrashContext.h"
+
+struct CORE_API FMacCrashContext : public FApplePlatformCrashContext
+{
+	/** Mimics Windows WER format */
+	void GenerateWindowsErrorReport(char const* WERPath) const;
+
+	/** Creates (fake so far) minidump */
+	void GenerateMinidump(char const* Path) const;
+
+	/** Generates information for crash reporter */
+	void GenerateCrashInfoAndLaunchReporter() const;
+};
+
diff --git a/Engine/Source/Runtime/Core/Public/Mac/MacPlatformIncludes.h b/Engine/Source/Runtime/Core/Public/Mac/MacPlatformIncludes.h
--- a/Engine/Source/Runtime/Core/Public/Mac/MacPlatformIncludes.h
+++ b/Engine/Source/Runtime/Core/Public/Mac/MacPlatformIncludes.h
@@ -80,7 +80,6 @@ struct RECT
 // include platform implementations
 #include "Mac/MacPlatformMemory.h"
 #include "Apple/ApplePlatformString.h"
-#include "Apple/ApplePlatformCrashContext.h"
 #include "Mac/MacPlatformMisc.h"
 #include "Apple/ApplePlatformStackWalk.h"
 #include "Mac/MacPlatformMath.h"
diff --git a/Engine/Source/Runtime/Core/Public/Mac/MacPlatformMisc.h b/Engine/Source/Runtime/Core/Public/Mac/MacPlatformMisc.h
--- a/Engine/Source/Runtime/Core/Public/Mac/MacPlatformMisc.h
+++ b/Engine/Source/Runtime/Core/Public/Mac/MacPlatformMisc.h
@@ -193,16 +193,3 @@ enum EMacModifierKeys
 	MMK_RightAlt		= 0xF761,
 	MMK_RightControl	= 0xF762
 };
-
-struct CORE_API FMacCrashContext : public FApplePlatformCrashContext
-{
-	/** Mimics Windows WER format */
-	void GenerateWindowsErrorReport(char const* WERPath) const;
-	
-	/** Creates (fake so far) minidump */
-	void GenerateMinidump(char const* Path) const;
-	
-	/** Generates information for crash reporter */
-	void GenerateCrashInfoAndLaunchReporter() const;
-};
-
diff --git a/Engine/Source/Runtime/Launch/Private/Android/AndroidJNI.cpp b/Engine/Source/Runtime/Launch/Private/Android/AndroidJNI.cpp
--- a/Engine/Source/Runtime/Launch/Private/Android/AndroidJNI.cpp
+++ b/Engine/Source/Runtime/Launch/Private/Android/AndroidJNI.cpp
@@ -2,6 +2,7 @@
 
 #include "LaunchPrivatePCH.h"
 #include "ExceptionHandling.h"
+#include "AndroidPlatformCrashContext.h"
 
 #include "AndroidJNI.h"
 #include <Android/asset_manager.h>
diff --git a/Engine/Source/Runtime/Launch/Private/Linux/LaunchLinux.cpp b/Engine/Source/Runtime/Launch/Private/Linux/LaunchLinux.cpp
--- a/Engine/Source/Runtime/Launch/Private/Linux/LaunchLinux.cpp
+++ b/Engine/Source/Runtime/Launch/Private/Linux/LaunchLinux.cpp
@@ -2,6 +2,7 @@
 
 #include "LaunchPrivatePCH.h"
 #include "ExceptionHandling.h"
+#include "LinuxPlatformCrashContext.h"
 
 #if WITH_EDITOR
 #include "Editor/MainFrame/Private/Frame/MainFrameActions.h"
diff --git a/Engine/Source/Runtime/Launch/Private/Mac/LaunchMac.cpp b/Engine/Source/Runtime/Launch/Private/Mac/LaunchMac.cpp
--- a/Engine/Source/Runtime/Launch/Private/Mac/LaunchMac.cpp
+++ b/Engine/Source/Runtime/Launch/Private/Mac/LaunchMac.cpp
@@ -2,6 +2,7 @@
 
 #include "LaunchPrivatePCH.h"
 #include "ExceptionHandling.h"
+#include "MacPlatformCrashContext.h"
 
 #if WITH_EDITOR
 	#include "MainFrame.h"
