Fix high-precision mouse mode in EWindowMode::Fullscreen on OS X - needed to use the correct scale when calling the different functions.

diff --git a/Engine/Source/Runtime/Core/Private/Mac/MacApplication.cpp b/Engine/Source/Runtime/Core/Private/Mac/MacApplication.cpp
--- a/Engine/Source/Runtime/Core/Private/Mac/MacApplication.cpp
+++ b/Engine/Source/Runtime/Core/Private/Mac/MacApplication.cpp
@@ -582,14 +582,16 @@ void FMacApplication::ProcessNSEvent(NSEvent* const Event, TSharedPtr< FMacWindo
 					while ((screen = [screenEnumerator nextObject]) && !NSMouseInRect(NSMakePoint(HighPrecisionMousePos.X, HighPrecisionMousePos.Y), screen.frame, NO))
 						;
 					
+					FMacCursor* MacCursor = static_cast<FMacCursor*>( Cursor.Get() );
+					
 					// Clamp to no more than the reported delta - a single event of no mouse movement won't be noticed
 					// but going in the wrong direction will.
 					const FVector2D FullDelta([Event deltaX], [Event deltaY]);
 					const FVector2D WarpDelta(FMath::Abs(AccumDelta.X)<FMath::Abs(FullDelta.X) ? AccumDelta.X : FullDelta.X, FMath::Abs(AccumDelta.Y)<FMath::Abs(FullDelta.Y) ? AccumDelta.Y : FullDelta.Y);
 					
-					FVector2D Delta = (FullDelta - WarpDelta) / 2.f;
+					FVector2D Delta = ((FullDelta - WarpDelta) / 2.f) * MacCursor->GetMouseScaling();
 					
-					HighPrecisionMousePos = static_cast<FMacCursor*>( Cursor.Get() )->GetPosition() + Delta;
+					HighPrecisionMousePos = MacCursor->GetPosition() + Delta;
 					MacCursor->UpdateCursorClipping( HighPrecisionMousePos );
 					
 					// Clamp to the current screen and avoid the menu bar and dock to prevent popups and other
@@ -604,8 +606,8 @@ void FMacApplication::ProcessNSEvent(NSEvent* const Event, TSharedPtr< FMacWindo
 					NSRect FullFrame = [screen frame];
 					VisibleFrame.origin.y = (FullFrame.origin.y+FullFrame.size.height) - (VisibleFrame.origin.y + VisibleFrame.size.height);
 					
-					HighPrecisionMousePos.X = FMath::Clamp(HighPrecisionMousePos.X, (float)VisibleFrame.origin.x, (float)(VisibleFrame.origin.x + VisibleFrame.size.width)-1.f);
-					HighPrecisionMousePos.Y = FMath::Clamp(HighPrecisionMousePos.Y, (float)VisibleFrame.origin.y, (float)(VisibleFrame.origin.y + VisibleFrame.size.height)-1.f);
+					HighPrecisionMousePos.X = FMath::Clamp(HighPrecisionMousePos.X / MacCursor->GetMouseScaling().X, (float)VisibleFrame.origin.x, (float)(VisibleFrame.origin.x + VisibleFrame.size.width)-1.f);
+					HighPrecisionMousePos.Y = FMath::Clamp(HighPrecisionMousePos.Y / MacCursor->GetMouseScaling().Y, (float)VisibleFrame.origin.y, (float)(VisibleFrame.origin.y + VisibleFrame.size.height)-1.f);
 					
 					MacCursor->WarpCursor( HighPrecisionMousePos.X, HighPrecisionMousePos.Y );
 					MessageHandler->OnRawMouseMove( Delta.X, Delta.Y );
