On OS X when reshaping the window in EWindowMode::Fullscreen just update the FCocoaWindow's PreFullScreenRect & send the resize notification so that the viewport is ultimately resized to match. This fixes changing resolution when already in EWindowMode::Fullscreen and when going from a windowed resolution that has been reduced to fit on the desktop, but which can be used unmodified when fullscreen.

diff --git a/Engine/Source/Runtime/Core/Private/Mac/MacWindow.cpp b/Engine/Source/Runtime/Core/Private/Mac/MacWindow.cpp
--- a/Engine/Source/Runtime/Core/Private/Mac/MacWindow.cpp
+++ b/Engine/Source/Runtime/Core/Private/Mac/MacWindow.cpp
@@ -4,6 +4,7 @@
 #include "MacWindow.h"
 #include "MacApplication.h"
 #include "MacCursor.h"
+#include "MacEvent.h"
 #include "CocoaTextView.h"
 #include "CocoaThread.h"
 
@@ -254,6 +255,12 @@ void FMacWindow::ReshapeWindow( int32 X, int32 Y, int32 Width, int32 Height )
 				WindowHandle->bZoomed = [WindowHandle isZoomed];
 			}, UE4ResizeEventMode, true);
 		}
+		else
+		{
+			WindowHandle.PreFullScreenRect = NSMakeRect(WindowHandle.PreFullScreenRect.origin.x, WindowHandle.PreFullScreenRect.origin.y, (CGFloat)Width, (CGFloat)Height);
+			
+			FMacEvent::SendToGameRunLoop([NSNotification notificationWithName:NSWindowDidResizeNotification object:WindowHandle], WindowHandle, EMacEventSendMethod::Sync, @[ NSDefaultRunLoopMode, UE4ResizeEventMode, UE4ShowEventMode, UE4FullscreenEventMode ]);
+		}
 		
 		MessageHandler->FinishedReshapingWindow( SharedThis( this ) );
 	}
