Implement GetDefaultLocale for HTML5

diff --git a/Engine/Source/Runtime/Core/Private/HTML5/HTML5PlatformMisc.cpp b/Engine/Source/Runtime/Core/Private/HTML5/HTML5PlatformMisc.cpp
--- a/Engine/Source/Runtime/Core/Private/HTML5/HTML5PlatformMisc.cpp
+++ b/Engine/Source/Runtime/Core/Private/HTML5/HTML5PlatformMisc.cpp
@@ -18,6 +18,7 @@ void FHTML5Misc::PlatformInit()
 	// Identity.
 	UE_LOG(LogInit, Log, TEXT("Computer: %s"), FPlatformProcess::ComputerName());
 	UE_LOG(LogInit, Log, TEXT("User: %s"), FPlatformProcess::UserName());
+	UE_LOG(LogInit, Log, TEXT("Current Culture: %s"), *FHTML5Misc::GetDefaultLocale());
 
 	// Timer resolution.
 	UE_LOG(LogInit, Log, TEXT("High frequency timer resolution =%f MHz"), 0.000001 / FPlatformTime::GetSecondsPerCycle());
@@ -135,6 +136,22 @@ uint32 FHTML5Misc::GetKeyMap(uint16* KeyCodes, FString* KeyNames, uint32 MaxMapp
 	return NumMappings;
 }
 
+FString FHTML5Misc::GetDefaultLocale()
+{
+#if PLATFORM_HTML5_BROWSER
+	char AsciiCultureName[512];
+	if (UE_GetCurrentCultureName(AsciiCultureName,sizeof(AsciiCultureName)))
+	{
+		return FString(StringCast<TCHAR>(AsciiCultureName).Get());
+	}
+	else
+#endif
+	{
+		icu::Locale ICUDefaultLocale = icu::Locale::getDefault();
+		return FString(ICUDefaultLocale.getName());
+	}
+}
+
 EAppReturnType::Type FHTML5Misc::MessageBoxExt(EAppMsgType::Type MsgType, const TCHAR* Text, const TCHAR* Caption)
 {
 
diff --git a/Engine/Source/Runtime/Core/Public/HTML5/HTML5PlatformMisc.h b/Engine/Source/Runtime/Core/Public/HTML5/HTML5PlatformMisc.h
--- a/Engine/Source/Runtime/Core/Public/HTML5/HTML5PlatformMisc.h
+++ b/Engine/Source/Runtime/Core/Public/HTML5/HTML5PlatformMisc.h
@@ -18,6 +18,7 @@ struct CORE_API FHTML5Misc : public FGenericPlatformMisc
 	static class GenericApplication* CreateApplication();
 	static uint32 GetKeyMap(uint16* KeyCodes, FString* KeyNames, uint32 MaxMappings);
 	static uint32 GetCharKeyMap(uint16* KeyCodes, FString* KeyNames, uint32 MaxMappings);
+	static FString GetDefaultLocale();
 	static EAppReturnType::Type MessageBoxExt( EAppMsgType::Type MsgType, const TCHAR* Text, const TCHAR* Caption );
 	FORCEINLINE static int32 NumberOfCores()
 	{
@@ -53,7 +54,6 @@ struct CORE_API FHTML5Misc : public FGenericPlatformMisc
 	{
 		wprintf(TEXT("%ls"), Str);
 	}
-
 };
 
 typedef FHTML5Misc FPlatformMisc;
diff --git a/Engine/Source/Runtime/HTML5/HTML5JS/Private/HTML5JavaScriptFx.js b/Engine/Source/Runtime/HTML5/HTML5JS/Private/HTML5JavaScriptFx.js
--- a/Engine/Source/Runtime/HTML5/HTML5JS/Private/HTML5JavaScriptFx.js
+++ b/Engine/Source/Runtime/HTML5/HTML5JS/Private/HTML5JavaScriptFx.js
@@ -101,6 +101,15 @@ var UE_JavaScriptLibary =
             alert(text);
         }
         return 1; 
+    },
+
+    UE_GetCurrentCultureName: function (address, outsize) {
+        var culture_name = navigator.language || navigator.browserLanguage;
+        if (culture_name.lenght >= outsize) {
+            return 0;
+        }
+        Module.writeAsciiToMemory(culture_name, address);
+        return 1;
     }
 };
 
diff --git a/Engine/Source/Runtime/HTML5/HTML5JS/Public/HTML5JavaScriptFx.h b/Engine/Source/Runtime/HTML5/HTML5JS/Public/HTML5JavaScriptFx.h
--- a/Engine/Source/Runtime/HTML5/HTML5JS/Public/HTML5JavaScriptFx.h
+++ b/Engine/Source/Runtime/HTML5/HTML5JS/Public/HTML5JavaScriptFx.h
@@ -7,5 +7,6 @@
 		 bool UE_SaveGame(const char* Name, const int UserIndex, const char* indata, int insize);
 		 bool UE_LoadGame(const char* Name, const int UserIndex, char **outdata, int* outsize);
 		 int  UE_MessageBox( int MsgType, const char* Text, const char* Caption);
+		 int UE_GetCurrentCultureName(const char* OutName, int outsize);
 	}
 
