Compile fix for HTML5.

diff --git a/Engine/Source/Runtime/Core/Private/HTML5/HTML5Window.cpp b/Engine/Source/Runtime/Core/Private/HTML5/HTML5Window.cpp
--- a/Engine/Source/Runtime/Core/Private/HTML5/HTML5Window.cpp
+++ b/Engine/Source/Runtime/Core/Private/HTML5/HTML5Window.cpp
@@ -5,9 +5,8 @@
 
 #if !PLATFORM_HTML5_WIN32 
 #include <emscripten.h>
-#else
-#include <SDL/SDL.h>
 #endif
+#include <SDL/SDL.h>
 
 FHTML5Window::~FHTML5Window()
 {
