[INTEGRATE] Merging using UE4-To-UE4-LauncherReleases, all up to cl 2321357

diff --git a/Engine/Source/Runtime/Core/Private/Misc/Parse.cpp b/Engine/Source/Runtime/Core/Private/Misc/Parse.cpp
--- a/Engine/Source/Runtime/Core/Private/Misc/Parse.cpp
+++ b/Engine/Source/Runtime/Core/Private/Misc/Parse.cpp
@@ -277,23 +277,15 @@ bool FParse::Text( const TCHAR* Str, FText& Value, const TCHAR* Namespace )
 					{
 						if( c == '"' )
 						{
-							// Closing quote of parsed string
-							if( !ParsedString.IsEmpty() )
+							bInQuotes = false;
+							++StringCount;
+							if (StringCount == ExpectedStringCount - 2)
 							{
-								bInQuotes = false;
-								++StringCount;
-								if( StringCount == ExpectedStringCount - 2 )
-								{
-									NamespaceString = ParsedString;
-								}
-								else if( StringCount == ExpectedStringCount - 1 )
-								{
-									KeyString = ParsedString;
-								}						
+								NamespaceString = ParsedString;
 							}
-							else
+							else if (StringCount == ExpectedStringCount - 1)
 							{
-								bError = true;
+								KeyString = ParsedString;
 							}
 						}
 						else if( c == '\\' )
@@ -341,10 +333,15 @@ bool FParse::Text( const TCHAR* Str, FText& Value, const TCHAR* Namespace )
 				++Str;
 			}
 
+			if (KeyString.Len() == 0 && (bFoundNSLocText ? NamespaceString : Namespace).Len() > 0)
+			{
+				bError = true;
+			}
+
 			if( *Str == ')' && !bError && StringCount == ExpectedStringCount )
 			{
 				if ( !FText::FindText( bFoundNSLocText ? NamespaceString : Namespace, KeyString, /*OUT*/Value ) )
-				{
+				{ 
 					Value = FText::FromString( ParsedString );
 				}
 
diff --git a/Engine/Source/Runtime/Slate/Private/Framework/Text/TextDecorators.cpp b/Engine/Source/Runtime/Slate/Private/Framework/Text/TextDecorators.cpp
--- a/Engine/Source/Runtime/Slate/Private/Framework/Text/TextDecorators.cpp
+++ b/Engine/Source/Runtime/Slate/Private/Framework/Text/TextDecorators.cpp
@@ -99,8 +99,8 @@ TSharedRef< FHyperlinkDecorator > FHyperlinkDecorator::Create( FString Id, const
 }
 
 FHyperlinkDecorator::FHyperlinkDecorator( FString InId, const FSlateHyperlinkRun::FOnClick& InNavigateDelegate, const FSlateHyperlinkRun::FOnGetTooltipText& InToolTipTextDelegate, const FSlateHyperlinkRun::FOnGenerateTooltip& InToolTipDelegate )
-	: Id( InId )
-	, NavigateDelegate( InNavigateDelegate )
+	: NavigateDelegate( InNavigateDelegate )
+	, Id(InId)
 	, ToolTipTextDelegate( InToolTipTextDelegate )
 	, ToolTipDelegate( InToolTipDelegate )
 {
diff --git a/Engine/Source/Runtime/Slate/Public/Framework/Text/TextDecorators.h b/Engine/Source/Runtime/Slate/Public/Framework/Text/TextDecorators.h
--- a/Engine/Source/Runtime/Slate/Public/Framework/Text/TextDecorators.h
+++ b/Engine/Source/Runtime/Slate/Public/Framework/Text/TextDecorators.h
@@ -68,14 +68,16 @@ public:
 
 	virtual TSharedRef< ISlateRun > Create(const TSharedRef<class FTextLayout>& TextLayout, const FTextRunParseResults& RunParseResult, const FString& OriginalText, const TSharedRef< FString >& InOutModelText, const ISlateStyle* Style) override;
 
-private:
+protected:
 
 	FHyperlinkDecorator( FString InId, const FSlateHyperlinkRun::FOnClick& InNavigateDelegate, const FSlateHyperlinkRun::FOnGetTooltipText& InToolTipTextDelegate, const FSlateHyperlinkRun::FOnGenerateTooltip& InToolTipDelegate );
 
-private:
+protected:
 
-	FString Id;
 	FSlateHyperlinkRun::FOnClick NavigateDelegate;
+
+private:
+	FString Id;
 	FSlateHyperlinkRun::FOnGetTooltipText ToolTipTextDelegate;
 	FSlateHyperlinkRun::FOnGenerateTooltip ToolTipDelegate;
 };
