FMacPlatformProcess::IsThisApplicationForeground() will now return false after fast switching to a second user

diff --git a/Engine/Source/Runtime/Core/Private/Mac/MacApplication.cpp b/Engine/Source/Runtime/Core/Private/Mac/MacApplication.cpp
--- a/Engine/Source/Runtime/Core/Private/Mac/MacApplication.cpp
+++ b/Engine/Source/Runtime/Core/Private/Mac/MacApplication.cpp
@@ -113,6 +113,7 @@ FMacApplication::FMacApplication()
 	, bIsProcessingNSEvent( false )
 	, ModifierKeysFlags( 0 )
 	, CurrentModifierFlags( 0 )
+	, bIsWorkspaceSessionActive( true )
 	, EventMonitor( NULL )
 {
 	CGDisplayRegisterReconfigurationCallback(FMacApplication::OnDisplayReconfiguration, this);
@@ -233,6 +234,16 @@ FMacApplication::FMacApplication()
 									GVolumeMultiplier = 0.0f;
 								}];
 	
+	WorkspaceActivationObserver = [[[NSWorkspace sharedWorkspace] notificationCenter] addObserverForName:NSWorkspaceSessionDidBecomeActiveNotification object:[NSWorkspace sharedWorkspace] queue:[NSOperationQueue mainQueue] usingBlock:^(NSNotification* Notification){
+								   printf("workspace activated\n");
+									   bIsWorkspaceSessionActive = true;
+								   }];
+
+	WorkspaceDeactivationObserver = [[[NSWorkspace sharedWorkspace] notificationCenter] addObserverForName:NSWorkspaceSessionDidResignActiveNotification object:[NSWorkspace sharedWorkspace] queue:[NSOperationQueue mainQueue] usingBlock:^(NSNotification* Notification){
+									 printf("workspace deactivated\n");
+									   bIsWorkspaceSessionActive = false;
+								   }];
+
 #if WITH_EDITOR
 	FMemory::MemZero(GestureUsage);
 	LastGestureUsed = EGestureEvent::None;
@@ -259,6 +270,18 @@ FMacApplication::~FMacApplication()
 		AppDeactivationObserver = nil;
 	}
 	
+	if(WorkspaceActivationObserver)
+	{
+		[[NSNotificationCenter defaultCenter] removeObserver:WorkspaceActivationObserver];
+		WorkspaceActivationObserver = nil;
+	}
+
+	if(WorkspaceDeactivationObserver)
+	{
+		[[NSNotificationCenter defaultCenter] removeObserver:WorkspaceDeactivationObserver];
+		WorkspaceDeactivationObserver = nil;
+	}
+
 	CGDisplayRemoveReconfigurationCallback(FMacApplication::OnDisplayReconfiguration, this);
 	if (MouseCaptureWindow)
 	{
diff --git a/Engine/Source/Runtime/Core/Private/Mac/MacPlatformProcess.cpp b/Engine/Source/Runtime/Core/Private/Mac/MacPlatformProcess.cpp
--- a/Engine/Source/Runtime/Core/Private/Mac/MacPlatformProcess.cpp
+++ b/Engine/Source/Runtime/Core/Private/Mac/MacPlatformProcess.cpp
@@ -6,6 +6,7 @@
 
 #include "Core.h"
 #include "../../Launch/Resources/Version.h"
+#include "MacApplication.h"
 #include <mach-o/dyld.h>
 #include <libproc.h>
 
@@ -581,7 +582,7 @@ FString FMacPlatformProcess::GetApplicationName( uint32 ProcessId )
 bool FMacPlatformProcess::IsThisApplicationForeground()
 {
 	SCOPED_AUTORELEASE_POOL;
-	return [NSApp isActive];
+	return [NSApp isActive] && MacApplication && MacApplication->IsWorkspaceSessionActive();
 }
 
 void FMacPlatformProcess::CleanFileCache()
diff --git a/Engine/Source/Runtime/Core/Public/Mac/MacApplication.h b/Engine/Source/Runtime/Core/Public/Mac/MacApplication.h
--- a/Engine/Source/Runtime/Core/Public/Mac/MacApplication.h
+++ b/Engine/Source/Runtime/Core/Public/Mac/MacApplication.h
@@ -78,6 +78,8 @@ public:
 
 	bool IsProcessingNSEvent() const { return bIsProcessingNSEvent; }
 
+	bool IsWorkspaceSessionActive() const { return bIsWorkspaceSessionActive; }
+
 #if WITH_EDITOR
     virtual void SendAnalytics(IAnalyticsProvider* Provider) override;
 
@@ -181,12 +183,20 @@ private:
 
 	TSharedPtr<FMacTextInputMethodSystem> TextInputMethodSystem;
 
+	bool bIsWorkspaceSessionActive;
+
 	/** Notification center observer for application activation events */
 	id AppActivationObserver;
 
 	/** Notification center observer for application deactivation events */
 	id AppDeactivationObserver;
 
+	/** Notification center observer for workspace activation events */
+	id WorkspaceActivationObserver;
+
+	/** Notification center observer for workspace deactivation events */
+	id WorkspaceDeactivationObserver;
+
 #if WITH_EDITOR
 	/** Holds the last gesture used to try and capture unique uses for gestures. */
 	EGestureEvent::Type LastGestureUsed;
