Improved hierarchical localization data loading logic. Now handles prioritizing by combinations of language, script, and region.

diff --git a/Engine/Source/Developer/WebBrowser/Private/WebBrowserSingleton.cpp b/Engine/Source/Developer/WebBrowser/Private/WebBrowserSingleton.cpp
--- a/Engine/Source/Developer/WebBrowser/Private/WebBrowserSingleton.cpp
+++ b/Engine/Source/Developer/WebBrowser/Private/WebBrowserSingleton.cpp
@@ -125,7 +125,7 @@ FString FWebBrowserSingleton::GetCurrentLocaleCode()
 {
 	FCultureRef Culture = FInternationalization::Get().GetCurrentCulture();
 	FString LocaleCode = Culture->GetTwoLetterISOLanguageName();
-	FString Country = Culture->GetCountry();
+	FString Country = Culture->GetRegion();
 	if (!Country.IsEmpty())
 	{
 		LocaleCode = LocaleCode + TEXT("-") + Country;
diff --git a/Engine/Source/Editor/UnrealEd/Private/Commandlets/InternationalizationExportCommandlet.cpp b/Engine/Source/Editor/UnrealEd/Private/Commandlets/InternationalizationExportCommandlet.cpp
--- a/Engine/Source/Editor/UnrealEd/Private/Commandlets/InternationalizationExportCommandlet.cpp
+++ b/Engine/Source/Editor/UnrealEd/Private/Commandlets/InternationalizationExportCommandlet.cpp
@@ -315,7 +315,7 @@ bool FindDelimitedString( const FString& InStr, const FString& LeftDelim, const
 FPortableObjectCulture::FPortableObjectCulture( const FString& LangCode, const FString& PluralForms )
 	: LanguageCode( LangCode )
 	, LanguagePluralForms( PluralForms )
-	, Culture( MakeShareable( new FCulture( LangCode ) ) )
+	, Culture( FCulture::Create( LangCode ) )
 {
 
 }
@@ -323,7 +323,7 @@ FPortableObjectCulture::FPortableObjectCulture( const FString& LangCode, const F
 FPortableObjectCulture::FPortableObjectCulture( const FPortableObjectCulture& Other )
 	: LanguageCode( Other.LanguageCode )
 	, LanguagePluralForms( Other.LanguagePluralForms )
-	, Culture( MakeShareable( new FCulture( Other.LanguageCode ) ) )
+	, Culture( FCulture::Create( Other.LanguageCode ) )
 {
 
 }
@@ -331,7 +331,7 @@ FPortableObjectCulture::FPortableObjectCulture( const FPortableObjectCulture& Ot
 void FPortableObjectCulture::SetLanguageCode( const FString& LangCode )
 {
 	LanguageCode = LangCode;
-	Culture = MakeShareable( new FCulture( LangCode ) );
+	Culture = FCulture::Create( LangCode );
 }
 
 
@@ -352,7 +352,7 @@ FString FPortableObjectCulture::Country() const
 	FString Result;
 	if( Culture.IsValid() )
 	{
-		Result = Culture->GetCountry();
+		Result = Culture->GetRegion();
 	}
 	return Result;
 }
diff --git a/Engine/Source/Editor/UnrealEd/Private/Settings/ProjectPackagingSettingsCustomization.h b/Engine/Source/Editor/UnrealEd/Private/Settings/ProjectPackagingSettingsCustomization.h
--- a/Engine/Source/Editor/UnrealEd/Private/Settings/ProjectPackagingSettingsCustomization.h
+++ b/Engine/Source/Editor/UnrealEd/Private/Settings/ProjectPackagingSettingsCustomization.h
@@ -238,20 +238,26 @@ protected:
 		switch(FilterCulturesChoice)
 		{
 		case EFilterCulturesChoices::AllAvailableCultures:
-			CultureList.Empty();
-			for(const auto& Culture : FInternationalization::Get().GetAllCultures())
 			{
-				CultureList.Add(Culture);
+				CultureList.Empty();
+				TArray<FString> CultureNames;
+				FInternationalization::Get().GetCultureNames(CultureNames);
+				for(const FString& CultureName : CultureNames)
+				{
+					CultureList.Add(FInternationalization::Get().GetCulture(CultureName));
+				}
 			}
 			break;
 		case EFilterCulturesChoices::OnlyLocalizedCultures:
-			TArray<FCultureRef> LocalizedCultureList;
-			TArray<FString> LocalizationPaths = FPaths::GetGameLocalizationPaths();
-			FInternationalization::Get().GetCulturesWithAvailableLocalization(LocalizationPaths, LocalizedCultureList, true);
-			CultureList.Empty();
-			for(const auto& Culture : LocalizedCultureList)
 			{
-				CultureList.Add(Culture);
+				TArray<FCultureRef> LocalizedCultureList;
+				TArray<FString> LocalizationPaths = FPaths::GetGameLocalizationPaths();
+				FInternationalization::Get().GetCulturesWithAvailableLocalization(LocalizationPaths, LocalizedCultureList, true);
+				CultureList.Empty();
+				for(const auto& Culture : LocalizedCultureList)
+				{
+					CultureList.Add(Culture);
+				}
 			}
 			break;
 		}
diff --git a/Engine/Source/Runtime/Core/Private/Internationalization/Culture.cpp b/Engine/Source/Runtime/Core/Private/Internationalization/Culture.cpp
--- a/Engine/Source/Runtime/Core/Private/Internationalization/Culture.cpp
+++ b/Engine/Source/Runtime/Core/Private/Internationalization/Culture.cpp
@@ -1,4 +1,4 @@
-// Copyright 1998-2014 Epic Games, Inc. All Rights Reserved.
+﻿// Copyright 1998-2014 Epic Games, Inc. All Rights Reserved.
 
 #include "Core.h"
 
@@ -8,6 +8,18 @@
 #include "LegacyCulture.h"
 #endif
 
+#if UE_ENABLE_ICU
+FCulturePtr FCulture::Create(const FString& LocaleName)
+{
+	return MakeShareable(new FCulture(LocaleName));
+}
+#else
+FCulturePtr FCulture::Create(const FText& InDisplayName, const FString& InEnglishName, const int InKeyboardLayoutId, const int InLCID, const FString& InName, const FString& InNativeName, const FString& InUnrealLegacyThreeLetterISOLanguageName, const FString& InThreeLetterISOLanguageName, const FString& InTwoLetterISOLanguageName, const FNumberFormattingRules& InNumberFormattingRule, const FTextFormattingRules& InTextFormattingRule, const FDateTimeFormattingRules& InDateTimeFormattingRule)
+{
+	return MakeShareable(new FCulture(InDisplayName, InEnglishName, InKeyboardLayoutId, InLCID, InName, InNativeName, InUnrealLegacyThreeLetterISOLanguageName, InThreeLetterISOLanguageName, InTwoLetterISOLanguageName, InNumberFormattingRule, InTextFormattingRule, InDateTimeFormattingRule));
+}
+#endif
+
 #if UE_ENABLE_ICU
 FCulture::FCulture(const FString& LocaleName)
 	: Implementation( new FICUCultureImplementation( LocaleName ) )
@@ -44,9 +56,28 @@ int FCulture::GetLCID() const
 	return Implementation->GetLCID();
 }
 
-FString FCulture::GetName() const
+TArray<FString> FCulture::GetPrioritizedParentCultureNames()
 {
-	return Implementation->GetName();
+	const FString LanguageCode = GetTwoLetterISOLanguageName();
+	const FString ScriptCode = GetScript();
+	const FString RegionCode = GetRegion();
+
+	TArray<FString> LocaleTagCombinations;
+	if (!ScriptCode.IsEmpty() && !RegionCode.IsEmpty())
+	{
+		LocaleTagCombinations.Add(LanguageCode + TEXT("_") + ScriptCode + TEXT("_") + RegionCode);
+	}
+	if (!RegionCode.IsEmpty())
+	{
+		LocaleTagCombinations.Add(LanguageCode + TEXT("_") + RegionCode);
+	}
+	if (!ScriptCode.IsEmpty())
+	{
+		LocaleTagCombinations.Add(LanguageCode + TEXT("_") + ScriptCode);
+	}
+	LocaleTagCombinations.Add(LanguageCode);
+
+	return LocaleTagCombinations;
 }
 
 FString FCulture::GetCanonicalName(const FString& Name)
@@ -54,9 +85,9 @@ FString FCulture::GetCanonicalName(const FString& Name)
 	return FImplementation::GetCanonicalName(Name);
 }
 
-FString FCulture::GetParentName(const FString& CultureName)
+FString FCulture::GetName() const
 {
-	return FImplementation::GetParentName(CultureName);
+	return Implementation->GetName();
 }
 
 FString FCulture::GetNativeName() const
@@ -64,16 +95,6 @@ FString FCulture::GetNativeName() const
 	return Implementation->GetNativeName();
 }
 
-FString FCulture::GetNativeLanguage() const
-{
-	return Implementation->GetNativeLanguage();
-}
-
-FString FCulture::GetNativeRegion() const
-{
-	return Implementation->GetNativeRegion();
-}
-
 FString FCulture::GetUnrealLegacyThreeLetterISOLanguageName() const
 {
 	return Implementation->GetUnrealLegacyThreeLetterISOLanguageName();
@@ -89,10 +110,29 @@ FString FCulture::GetTwoLetterISOLanguageName() const
 	return Implementation->GetTwoLetterISOLanguageName();
 }
 
-FString FCulture::GetCountry() const
+FString FCulture::GetNativeLanguage() const
+{
+	return Implementation->GetNativeLanguage();
+}
+
+FString FCulture::GetRegion() const
+{
+#if UE_ENABLE_ICU
+	return Implementation->GetRegion();
+#else
+	return TEXT("");
+#endif
+}
+
+FString FCulture::GetNativeRegion() const
+{
+	return Implementation->GetNativeRegion();
+}
+
+FString FCulture::GetScript() const
 {
 #if UE_ENABLE_ICU
-	return Implementation->GetCountry();
+	return Implementation->GetScript();
 #else
 	return TEXT("");
 #endif
diff --git a/Engine/Source/Runtime/Core/Private/Internationalization/Cultures/AmericanEnglish.h b/Engine/Source/Runtime/Core/Private/Internationalization/Cultures/AmericanEnglish.h
deleted file mode 100644
diff --git a/Engine/Source/Runtime/Core/Private/Internationalization/Cultures/IndiaHindi.h b/Engine/Source/Runtime/Core/Private/Internationalization/Cultures/IndiaHindi.h
deleted file mode 100644
diff --git a/Engine/Source/Runtime/Core/Private/Internationalization/Cultures/InvariantCulture.h b/Engine/Source/Runtime/Core/Private/Internationalization/Cultures/InvariantCulture.h
--- a/Engine/Source/Runtime/Core/Private/Internationalization/Cultures/InvariantCulture.h
+++ b/Engine/Source/Runtime/Core/Private/Internationalization/Cultures/InvariantCulture.h
@@ -1,6 +1,8 @@
 ﻿// Copyright 1998-2014 Epic Games, Inc. All Rights Reserved.
 #pragma once
 
+#if !UE_ENABLE_ICU
+
 #include "Culture.h"
 
 #define LOCTEXT_NAMESPACE "Internationalization"
@@ -10,9 +12,6 @@ class FInvariantCulture
 public:
 	static FCultureRef Create()
 	{
-#if UE_ENABLE_ICU
-		FCultureRef Culture = MakeShareable( new FCulture( FString() ) );
-#else
 		TArray<FString> InNativeDigits;
 
 		for (int i = 0; i<=10; ++i)
@@ -131,7 +130,7 @@ public:
 			);
 
 
-		FCultureRef Culture = MakeShareable( new FCulture(
+		FCultureRef Culture = FCulture::Create(
 			  LOCTEXT("InvariantCultureDisplayName", "Invariant Language (Invariant Country)")		//const FText DisplayName
 			, FString(TEXT("Invariant Language (Invariant Country)"))			//const FString EnglishName
 			, 1033												//const int KeyboardLayoutId
@@ -144,11 +143,12 @@ public:
 			, InvariantCultureNumberFormatting					//const FNumberFormattingRules NumberFormattingRule
 			, InvariantCultureTextFormatting						//const FTextFormattingRules TextFormattingRule
 			, InvariantCultureDateFormatting						//const FDateTimeFormattingRules DateTimeFormattingRule
-			) );
-#endif
+			).ToSharedRef();
 
 		return Culture;
 	}
 };
 
 #undef LOCTEXT_NAMESPACE
+
+#endif
\ No newline at end of file
diff --git a/Engine/Source/Runtime/Core/Private/Internationalization/Cultures/JapaneseCulture.h b/Engine/Source/Runtime/Core/Private/Internationalization/Cultures/JapaneseCulture.h
deleted file mode 100644
diff --git a/Engine/Source/Runtime/Core/Private/Internationalization/Cultures/KoreanCulture.h b/Engine/Source/Runtime/Core/Private/Internationalization/Cultures/KoreanCulture.h
deleted file mode 100644
diff --git a/Engine/Source/Runtime/Core/Private/Internationalization/Cultures/LEETCulture.h b/Engine/Source/Runtime/Core/Private/Internationalization/Cultures/LEETCulture.h
deleted file mode 100644
diff --git a/Engine/Source/Runtime/Core/Private/Internationalization/ICUCulture.cpp b/Engine/Source/Runtime/Core/Private/Internationalization/ICUCulture.cpp
--- a/Engine/Source/Runtime/Core/Private/Internationalization/ICUCulture.cpp
+++ b/Engine/Source/Runtime/Core/Private/Internationalization/ICUCulture.cpp
@@ -102,11 +102,6 @@ int FCulture::FICUCultureImplementation::GetLCID() const
 	return ICULocale.getLCID();
 }
 
-FString FCulture::FICUCultureImplementation::GetName() const
-{
-	return ICULocale.getName();
-}
-
 FString FCulture::FICUCultureImplementation::GetCanonicalName(const FString& Name)
 {
 	static const int32 MaximumNameLength = 64;
@@ -119,12 +114,9 @@ FString FCulture::FICUCultureImplementation::GetCanonicalName(const FString& Nam
 	return CanonicalName;
 }
 
-FString FCulture::FICUCultureImplementation::GetParentName(const FString& CultureName)
+FString FCulture::FICUCultureImplementation::GetName() const
 {
-	char Parent[32];
-	UErrorCode ICUStatus = U_ZERO_ERROR;
-	int32_t ParentLength = uloc_getParent(TCHAR_TO_ANSI(*CultureName), Parent, 32, &ICUStatus);
-	return FString(Parent);
+	return ICULocale.getName();
 }
 
 FString FCulture::FICUCultureImplementation::GetNativeName() const
@@ -134,6 +126,34 @@ FString FCulture::FICUCultureImplementation::GetNativeName() const
 	return ICUUtilities::ConvertString(ICUResult);
 }
 
+FString FCulture::FICUCultureImplementation::GetUnrealLegacyThreeLetterISOLanguageName() const
+{
+	FString Result( ICULocale.getISO3Language() );
+
+	// Legacy Overrides (INT, JPN, KOR), also for new web localization (CHN)
+	// and now for any other languages (FRA, DEU...) for correct redirection of documentation web links
+	if (Result == TEXT("eng"))
+	{
+		Result = TEXT("INT");
+	}
+	else
+	{
+		Result = Result.ToUpper();
+	}
+
+	return Result;
+}
+
+FString FCulture::FICUCultureImplementation::GetThreeLetterISOLanguageName() const
+{
+	return ICULocale.getISO3Language();
+}
+
+FString FCulture::FICUCultureImplementation::GetTwoLetterISOLanguageName() const
+{
+	return ICULocale.getLanguage();
+}
+
 FString FCulture::FICUCultureImplementation::GetNativeLanguage() const
 {
 	icu::UnicodeString ICUNativeLanguage;
@@ -153,6 +173,11 @@ FString FCulture::FICUCultureImplementation::GetNativeLanguage() const
 	return NativeLanguage;
 }
 
+FString FCulture::FICUCultureImplementation::GetRegion() const
+{
+	return ICULocale.getCountry();
+}
+
 FString FCulture::FICUCultureImplementation::GetNativeRegion() const
 {
 	icu::UnicodeString ICUNativeCountry;
@@ -172,37 +197,9 @@ FString FCulture::FICUCultureImplementation::GetNativeRegion() const
 	return NativeCountry;
 }
 
-FString FCulture::FICUCultureImplementation::GetUnrealLegacyThreeLetterISOLanguageName() const
-{
-	FString Result( ICULocale.getISO3Language() );
-
-	// Legacy Overrides (INT, JPN, KOR), also for new web localization (CHN)
-	// and now for any other languages (FRA, DEU...) for correct redirection of documentation web links
-	if (Result == TEXT("eng"))
-	{
-		Result = TEXT("INT");
-	}
-	else
-	{
-		Result = Result.ToUpper();
-	}
-
-	return Result;
-}
-
-FString FCulture::FICUCultureImplementation::GetThreeLetterISOLanguageName() const
+FString FCulture::FICUCultureImplementation::GetScript() const
 {
-	return ICULocale.getISO3Language();
-}
-
-FString FCulture::FICUCultureImplementation::GetTwoLetterISOLanguageName() const
-{
-	return ICULocale.getLanguage();
-}
-
-FString FCulture::FICUCultureImplementation::GetCountry() const
-{
-	return ICULocale.getCountry();
+	return ICULocale.getScript();
 }
 
 FString FCulture::FICUCultureImplementation::GetVariant() const
diff --git a/Engine/Source/Runtime/Core/Private/Internationalization/ICUCulture.h b/Engine/Source/Runtime/Core/Private/Internationalization/ICUCulture.h
--- a/Engine/Source/Runtime/Core/Private/Internationalization/ICUCulture.h
+++ b/Engine/Source/Runtime/Core/Private/Internationalization/ICUCulture.h
@@ -114,25 +114,25 @@ class FCulture::FICUCultureImplementation
 
 	int GetLCID() const;
 
-	FString GetName() const;
-
 	static FString GetCanonicalName(const FString& Name);
-
-	static FString GetParentName(const FString& CultureName);
+	
+	FString GetName() const;
 
 	FString GetNativeName() const;
 
-	FString GetNativeLanguage() const;
-
-	FString GetNativeRegion() const;
-
 	FString GetUnrealLegacyThreeLetterISOLanguageName() const;
 
 	FString GetThreeLetterISOLanguageName() const;
 
 	FString GetTwoLetterISOLanguageName() const;
 
-	FString GetCountry() const;
+	FString GetNativeLanguage() const;
+
+	FString GetRegion() const;
+
+	FString GetNativeRegion() const;
+
+	FString GetScript() const;
 
 	FString GetVariant() const;
 
diff --git a/Engine/Source/Runtime/Core/Private/Internationalization/ICUInternationalization.cpp b/Engine/Source/Runtime/Core/Private/Internationalization/ICUInternationalization.cpp
new file mode 100644
--- /dev/null
+++ b/Engine/Source/Runtime/Core/Private/Internationalization/ICUInternationalization.cpp
@@ -0,0 +1,399 @@
+// Copyright 1998-2014 Epic Games, Inc. All Rights Reserved.
+#include "Core.h"
+#include "ICUInternationalization.h"
+
+#if UE_ENABLE_ICU
+
+#include "ICUUtilities.h"
+#include "ICUBreakIterator.h"
+#include <unicode/locid.h>
+#include <unicode/timezone.h>
+#include <unicode/uclean.h>
+#include <unicode/udata.h>
+#include <unicode/umachine.h>
+
+namespace
+{
+	struct FICUOverrides
+	{
+#if STATS
+		static int64 BytesInUseCount;
+		static int64 CachedBytesInUseCount;
+#endif
+
+		static void* U_CALLCONV Malloc(const void* context, size_t size)
+		{
+			void* Result = FMemory::Malloc(size);
+#if STATS
+			BytesInUseCount += FMemory::GetAllocSize(Result);
+			if(FThreadStats::IsThreadingReady() && CachedBytesInUseCount != BytesInUseCount)
+			{
+				SET_MEMORY_STAT(STAT_MemoryICUTotalAllocationSize, BytesInUseCount);
+				CachedBytesInUseCount = BytesInUseCount;
+			}
+#endif
+			return Result;
+		}
+
+		static void* U_CALLCONV Realloc(const void* context, void* mem, size_t size)
+		{
+			return FMemory::Realloc(mem, size);
+		}
+
+		static void U_CALLCONV Free(const void* context, void* mem)
+		{
+#if STATS
+			BytesInUseCount -= FMemory::GetAllocSize(mem);
+			if(FThreadStats::IsThreadingReady() && CachedBytesInUseCount != BytesInUseCount)
+			{
+				SET_MEMORY_STAT(STAT_MemoryICUTotalAllocationSize, BytesInUseCount);
+				CachedBytesInUseCount = BytesInUseCount;
+			}
+#endif
+			FMemory::Free(mem);
+		}
+	};
+
+#if STATS
+	int64 FICUOverrides::BytesInUseCount = 0;
+	int64 FICUOverrides::CachedBytesInUseCount = 0;
+#endif
+}
+
+FICUInternationalization::FICUInternationalization(FInternationalization* const InI18N)
+	: I18N(InI18N)
+{
+
+}
+
+bool FICUInternationalization::Initialize()
+{
+	UErrorCode ICUStatus = U_ZERO_ERROR;
+
+	// Linux needs to have those compiled statically at least until we settle on .so location for deployed/native builds
+#if (IS_PROGRAM || !IS_MONOLITHIC) && !PLATFORM_LINUX
+	LoadDLLs();
+#endif /*IS_PROGRAM || !IS_MONOLITHIC*/
+
+	u_setMemoryFunctions(NULL, &(FICUOverrides::Malloc), &(FICUOverrides::Realloc), &(FICUOverrides::Free), &(ICUStatus));
+
+	const FString DataDirectoryRelativeToContent = FString(TEXT("Localization")) / FString(TEXT("ICU"));
+	IFileManager& FileManager = IFileManager::Get();
+
+	const FString PotentialDataDirectories[] = 
+	{
+		FPaths::GameContentDir() / DataDirectoryRelativeToContent, // Try game content directory.
+		FPaths::EngineContentDir() / DataDirectoryRelativeToContent, // Try engine content directory.
+	};
+
+	bool HasFoundDataDirectory = false;
+	for(const auto& DataDirectoryString : PotentialDataDirectories)
+	{
+		if( FileManager.DirectoryExists( *DataDirectoryString ) )
+		{
+			u_setDataDirectory( StringCast<char>( *DataDirectoryString ).Get() );
+			HasFoundDataDirectory = true;
+			break;
+		}
+	}
+
+	auto GetPrioritizedDataDirectoriesString = [&]() -> FString
+	{
+		FString Result;
+		for(const auto& DataDirectoryString : PotentialDataDirectories)
+		{
+			if(!Result.IsEmpty())
+			{
+				Result += TEXT("\n");
+			}
+			Result += DataDirectoryString;
+		}
+		return Result;
+	};
+	checkf( HasFoundDataDirectory, TEXT("ICU data directory was not discovered:\n%s"), *(GetPrioritizedDataDirectoriesString()) );
+
+	u_setDataFileFunctions(nullptr, &FICUInternationalization::OpenDataFile, &FICUInternationalization::CloseDataFile, &(ICUStatus));
+	u_init(&(ICUStatus));
+
+	FICUBreakIteratorManager::Create();
+
+	I18N->InvariantCulture = FindOrMakeCulture(TEXT(""));
+	I18N->DefaultCulture = FindOrMakeCulture(icu::Locale::getDefault().getName());
+	SetCurrentCulture( I18N->GetDefaultCulture()->GetName() );
+	return U_SUCCESS(ICUStatus) ? true : false;
+}
+
+void FICUInternationalization::Terminate()
+{
+	FICUBreakIteratorManager::Destroy();
+	CachedCultures.Empty();
+
+	u_cleanup();
+#if IS_PROGRAM || !IS_MONOLITHIC
+	UnloadDLLs();
+#endif //IS_PROGRAM || !IS_MONOLITHIC
+}
+
+#if (IS_PROGRAM || !IS_MONOLITHIC)
+void FICUInternationalization::LoadDLLs()
+{
+	// The base directory for ICU binaries is consistent on all platforms.
+	FString ICUBinariesRoot = FPaths::EngineDir() / TEXT("Binaries") / TEXT("ThirdParty") / TEXT("ICU") / TEXT("icu4c-53_1");
+
+#if PLATFORM_WINDOWS
+#if PLATFORM_64BITS
+	const FString PlatformFolderName = TEXT("Win64");
+#elif PLATFORM_32BITS
+	const FString PlatformFolderName = TEXT("Win32");
+#endif //PLATFORM_*BITS
+
+#if _MSC_VER >= 1800
+	const FString VSVersionFolderName = TEXT("VS2013");
+#else
+	const FString VSVersionFolderName = TEXT("VS2012");
+#endif //_MSVC_VER
+
+	// Windows requires support for 32/64 bit and also VC11 and VC12 runtimes.
+	const FString TargetSpecificPath = ICUBinariesRoot / PlatformFolderName / VSVersionFolderName;
+
+	// Windows libraries use a specific naming convention.
+	FString LibraryNameStems[] =
+	{
+		TEXT("dt"),   // Data
+		TEXT("uc"),   // Unicode Common
+		TEXT("in"),   // Internationalization
+		TEXT("le"),   // Layout Engine
+		TEXT("lx"),   // Layout Extensions
+		TEXT("io")	// Input/Output
+	};
+#else
+	// Non-Windows libraries use a consistent naming convention.
+	FString LibraryNameStems[] =
+	{
+		TEXT("data"),   // Data
+		TEXT("uc"),   // Unicode Common
+		TEXT("i18n"),   // Internationalization
+		TEXT("le"),   // Layout Engine
+		TEXT("lx"),   // Layout Extensions
+		TEXT("io")	// Input/Output
+	};
+#if PLATFORM_LINUX
+	const FString TargetSpecificPath = ICUBinariesRoot / TEXT("Linux") / TEXT("x86_64-unknown-linux-gnu");
+#elif PLATFORM_MAC
+	const FString TargetSpecificPath = ICUBinariesRoot / TEXT("Mac");
+#endif //PLATFORM_*
+#endif //PLATFORM_*
+
+#if UE_BUILD_DEBUG && !defined(NDEBUG)
+	const FString LibraryNamePostfix = TEXT("d");
+#else
+	const FString LibraryNamePostfix = TEXT("");
+#endif //DEBUG
+
+	for(FString Stem : LibraryNameStems)
+	{
+#if PLATFORM_WINDOWS
+		FString LibraryName = "icu" + Stem + LibraryNamePostfix + "53" "." "dll";
+#elif PLATFORM_LINUX
+		FString LibraryName = "lib" "icu" + Stem + LibraryNamePostfix + ".53.1" "." "so";
+#elif PLATFORM_MAC
+		FString LibraryName = "lib" "icu" + Stem + LibraryNamePostfix + ".53.1" "." "dylib";
+#endif //PLATFORM_*
+
+		void* DLLHandle = FPlatformProcess::GetDllHandle(*(TargetSpecificPath / LibraryName));
+		check(DLLHandle != nullptr);
+		DLLHandles.Add(DLLHandle);
+	}
+}
+
+void FICUInternationalization::UnloadDLLs()
+{
+	for( const auto DLLHandle : DLLHandles )
+	{
+		FPlatformProcess::FreeDllHandle(DLLHandle);
+	}
+	DLLHandles.Reset();
+}
+#endif // IS_PROGRAM || !IS_MONOLITHIC
+
+#if STATS
+namespace
+{
+	int64 DataFileBytesInUseCount = 0;
+	int64 CachedDataFileBytesInUseCount = 0;
+}
+#endif
+
+void FICUInternationalization::SetCurrentCulture(const FString& Name)
+{
+	FCulturePtr NewCurrentCulture = FindOrMakeCulture(Name);
+
+	if (NewCurrentCulture.IsValid())
+	{
+		if (NewCurrentCulture.ToSharedRef() != I18N->CurrentCulture)
+		{
+			I18N->CurrentCulture = NewCurrentCulture.ToSharedRef();
+
+			UErrorCode ICUStatus = U_ZERO_ERROR;
+			uloc_setDefault(StringCast<char>(*Name).Get(), &ICUStatus);
+
+			FCoreDelegates::OnCultureChanged.Broadcast();
+		}
+	}
+}
+
+void FICUInternationalization::GetCultureNames(TArray<FString>& CultureNames) const
+{
+	int32_t LocaleCount;
+	const icu::Locale* const AvailableLocales = icu::Locale::getAvailableLocales(LocaleCount);
+
+	CultureNames.Reset(LocaleCount);
+	for (int32 i = 0; i < LocaleCount; ++i)
+	{
+		CultureNames.Add(AvailableLocales[i].getName());
+	}
+}
+
+FCulturePtr FICUInternationalization::GetCulture(const FString& Name)
+{
+	return FindOrMakeCulture(Name);
+}
+
+FCulturePtr FICUInternationalization::FindOrMakeCulture(const FString& Name)
+{
+	FCultureRef* FoundCulture = CachedCultures.Find(Name);
+	if (!FoundCulture)
+	{
+		FCulturePtr NewCulture = FCulture::Create(Name);
+		if (NewCulture.IsValid())
+		{
+			FoundCulture = &(CachedCultures.Add(Name, NewCulture.ToSharedRef()));
+		}
+	}
+
+	return FoundCulture ? FCulturePtr(*FoundCulture) : FCulturePtr(nullptr);
+}
+
+UBool FICUInternationalization::OpenDataFile(const void* context, void** fileContext, void** contents, const char* path)
+{
+	auto& PathToCachedFileDataMap = FInternationalization::Get().Implementation->PathToCachedFileDataMap;
+
+	// Try to find existing buffer
+	FICUInternationalization::FICUCachedFileData* CachedFileData = PathToCachedFileDataMap.Find(path);
+
+	// If there's no file context, we might have to load the file.
+	if (!CachedFileData)
+	{
+		// Attempt to load the file.
+		FArchive* FileAr = IFileManager::Get().CreateFileReader(StringCast<TCHAR>(path).Get());
+		if (FileAr)
+		{
+			const int64 FileSize = FileAr->TotalSize();
+
+			// Create file data.
+			CachedFileData = &(PathToCachedFileDataMap.Emplace(FString(path), FICUInternationalization::FICUCachedFileData(FileSize)));
+
+			// Load file into buffer.
+			FileAr->Serialize(CachedFileData->Buffer, FileSize); 
+			delete FileAr;
+
+			// Stat tracking.
+#if STATS
+			DataFileBytesInUseCount += FMemory::GetAllocSize(CachedFileData->Buffer);
+			if (FThreadStats::IsThreadingReady() && CachedDataFileBytesInUseCount != DataFileBytesInUseCount)
+			{
+				SET_MEMORY_STAT(STAT_MemoryICUDataFileAllocationSize, DataFileBytesInUseCount);
+				CachedDataFileBytesInUseCount = DataFileBytesInUseCount;
+			}
+#endif
+		}
+	}
+
+	// Add a reference, either the initial one or an additional one.
+	if (CachedFileData)
+	{
+		++(CachedFileData->ReferenceCount);
+	}
+
+	// Use the file path as the context, so we can look up the cached file data later and decrement its reference count.
+	*fileContext = CachedFileData ? new FString(path) : nullptr;
+
+	// Use the buffer from the cached file data.
+	*contents = CachedFileData ? CachedFileData->Buffer : nullptr;
+
+	// If we have cached file data, we must have loaded new data or found existing data, so we've successfully "opened" and "read" the file into "contents".
+	return CachedFileData != nullptr;
+}
+
+void FICUInternationalization::CloseDataFile(const void* context, void* const fileContext, void* const contents)
+{
+	// Early out on null context.
+	if (fileContext == nullptr)
+	{
+		return;
+	}
+
+	auto& PathToCachedFileDataMap = FInternationalization::Get().Implementation->PathToCachedFileDataMap;
+
+	// The file context is the path to the file.
+	FString* const Path = reinterpret_cast<FString*>(fileContext);
+	check(Path);
+
+	// Look up the cached file data so we can maintain references.
+	FICUInternationalization::FICUCachedFileData* const CachedFileData = PathToCachedFileDataMap.Find(*Path);
+	check(CachedFileData);
+	check(CachedFileData->Buffer == contents);
+
+	// Remove a reference.
+	--(CachedFileData->ReferenceCount);
+
+	// If the last reference has been removed, the cached file data is not longer needed.
+	if (CachedFileData->ReferenceCount == 0)
+	{
+		// Stat tracking.
+#if STATS
+		DataFileBytesInUseCount -= FMemory::GetAllocSize(CachedFileData->Buffer);
+		if (FThreadStats::IsThreadingReady() && CachedDataFileBytesInUseCount != DataFileBytesInUseCount)
+		{
+			SET_MEMORY_STAT(STAT_MemoryICUDataFileAllocationSize, DataFileBytesInUseCount);
+			CachedDataFileBytesInUseCount = DataFileBytesInUseCount;
+		}
+#endif
+
+		// Delete the cached file data.
+		PathToCachedFileDataMap.Remove(*Path);
+	}
+
+	// The path string we allocated for tracking is no longer necessary.
+	delete Path;
+}
+
+FICUInternationalization::FICUCachedFileData::FICUCachedFileData(const int64 FileSize)
+	: ReferenceCount(0)
+	, Buffer( FICUOverrides::Malloc(nullptr, FileSize) )
+{
+}
+
+FICUInternationalization::FICUCachedFileData::FICUCachedFileData(const FICUCachedFileData& Source)
+{
+	checkf(false, TEXT("Cached file data for ICU may not be copy constructed. Something is trying to copy construct FICUCachedFileData."));
+}
+
+FICUInternationalization::FICUCachedFileData::FICUCachedFileData(FICUCachedFileData&& Source)
+	: ReferenceCount(Source.ReferenceCount)
+	, Buffer( Source.Buffer )
+{
+	// Make sure the moved source object doesn't retain the pointer and free the memory we now point to.
+	Source.Buffer = nullptr;
+}
+
+FICUInternationalization::FICUCachedFileData::~FICUCachedFileData()
+{
+	if (Buffer)
+	{
+		check(ReferenceCount == 0);
+		FICUOverrides::Free(nullptr, Buffer);
+	}
+}
+
+#endif
\ No newline at end of file
diff --git a/Engine/Source/Runtime/Core/Private/Internationalization/ICUInternationalization.h b/Engine/Source/Runtime/Core/Private/Internationalization/ICUInternationalization.h
new file mode 100644
--- /dev/null
+++ b/Engine/Source/Runtime/Core/Private/Internationalization/ICUInternationalization.h
@@ -0,0 +1,52 @@
+// Copyright 1998-2014 Epic Games, Inc. All Rights Reserved.
+#pragma once
+
+#if UE_ENABLE_ICU
+
+class FICUInternationalization
+{
+public:
+	FICUInternationalization(FInternationalization* const I18N);
+
+	bool Initialize();
+	void Terminate();
+
+	void SetCurrentCulture(const FString& Name);
+	void GetCultureNames(TArray<FString>& CultureNames) const;
+	FCulturePtr GetCulture(const FString& Name);
+
+private:
+#if (IS_PROGRAM || !IS_MONOLITHIC)
+	void LoadDLLs();
+	void UnloadDLLs();
+#endif
+
+	FCulturePtr FindOrMakeCulture(const FString& Name);
+
+private:
+	FInternationalization* const I18N;
+
+	TArray< void* > DLLHandles;
+
+	TMap<FString, FCultureRef> CachedCultures;
+
+	static UBool OpenDataFile(const void* context, void** fileContext, void** contents, const char* path);
+	static void CloseDataFile(const void* context, void* const fileContext, void* const contents);
+
+	// Tidy class for storing the count of references for an ICU data file and the file's data itself.
+	struct FICUCachedFileData
+	{
+		FICUCachedFileData(const int64 FileSize);
+		FICUCachedFileData(const FICUCachedFileData& Source);
+		FICUCachedFileData(FICUCachedFileData&& Source);
+		~FICUCachedFileData();
+
+		uint32 ReferenceCount;
+		void* Buffer;
+	};
+
+	// Map for associating ICU data file paths with cached file data, to prevent multiple copies of immutable ICU data files from residing in memory.
+	TMap<FString, FICUCachedFileData> PathToCachedFileDataMap;
+};
+
+#endif
\ No newline at end of file
diff --git a/Engine/Source/Runtime/Core/Private/Internationalization/Internationalization.cpp b/Engine/Source/Runtime/Core/Private/Internationalization/Internationalization.cpp
--- a/Engine/Source/Runtime/Core/Private/Internationalization/Internationalization.cpp
+++ b/Engine/Source/Runtime/Core/Private/Internationalization/Internationalization.cpp
@@ -1,42 +1,17 @@
 ﻿// Copyright 1998-2014 Epic Games, Inc. All Rights Reserved.
 
 #include "Core.h"
-#include "InvariantCulture.h"
-#if ENABLE_LOC_TESTING
-#include "LEETCulture.h"
-#endif
-#include "AmericanEnglish.h"
-// this header has hardcoded 2byte characters that won't work with 1 byte TCHARs
-#if !PLATFORM_TCHAR_IS_1_BYTE
-#include "IndiaHindi.h"
-#include "JapaneseCulture.h"
-#include "KoreanCulture.h"
-#endif
 
 #if UE_ENABLE_ICU
-#include "ICUUtilities.h"
-#include "ICUBreakIterator.h"
-#include <unicode/locid.h>
-#include <unicode/timezone.h>
-#include <unicode/uclean.h>
-#include <unicode/udata.h>
-#include <unicode/umachine.h>
+#include "ICUInternationalization.h"
+#else
+#include "LegacyInternationalization.h"
 #endif
 
 #define LOCTEXT_NAMESPACE "Internationalization"
 
 FInternationalization* FInternationalization::Instance;
 
-#if UE_ENABLE_ICU
-// Simple struct that operates as a namespace for ICU data file callbacks.
-// Struct used instead of namespace in order to avoid having to friend individual methods, which requires forward declaration and thus would create a "public" ICU dependency in the header.
-struct FICUDataCallbacks
-{
-	static UBool OpenDataFile(const void* context, void** fileContext, void** contents, const char* path);
-	static void CloseDataFile(const void* context, void* const fileContext, void* const contents);
-};
-#endif
-
 FInternationalization& FInternationalization::Get()
 {
 	if(!Instance)
@@ -58,126 +33,15 @@ void FInternationalization::TearDown()
 	}
 }
 
-void FInternationalization::GetTimeZonesIDs(TArray<FString>& TimeZonesIDs) const
-{
-	TimeZonesIDs.Empty();
-#if UE_ENABLE_ICU
-	UErrorCode ICUStatus = U_ZERO_ERROR;
-
-	TSharedPtr<icu::StringEnumeration> StringEnumeration( icu::TimeZone::createEnumeration() );
-	TimeZonesIDs.Reserve( StringEnumeration->count(ICUStatus) );
-
-	const icu::UnicodeString* ICUString;
-	do
-	{
-		ICUString = StringEnumeration->snext(ICUStatus);
-		if(ICUString)
-		{
-			FString NativeString;
-			ICUUtilities::ConvertString(*ICUString, NativeString);
-			TimeZonesIDs.Add( NativeString );
-		}
-	} while( ICUString );
-#endif
-}
-
 void FInternationalization::SetCurrentCulture(const FString& Name)
 {
-	const int32 CultureIndex = GetCultureIndex(Name);
-
-	if(CultureIndex != INDEX_NONE && CurrentCultureIndex != CultureIndex)
-	{
-		CurrentCultureIndex = CultureIndex;
-
-#if UE_ENABLE_ICU
-		UErrorCode ICUStatus = U_ZERO_ERROR;
-		icu::Locale::setDefault(icu::Locale( TCHAR_TO_ANSI(*Name) ), ICUStatus);
-#endif
-
-		FCoreDelegates::OnCultureChanged.Broadcast();
-	}
+	Implementation->SetCurrentCulture(Name);
 }
 
-FCultureRef FInternationalization::GetCurrentCulture() const
+FCulturePtr FInternationalization::GetCulture(const FString& Name)
 {
-	return AllCultures[CurrentCultureIndex];
-}
-
-FCulturePtr FInternationalization::GetCulture(const FString& Name) const
-{
-	int32 CultureIndex = GetCultureIndex(Name);
-	return CultureIndex != -1 ? AllCultures[CultureIndex] : FCulturePtr();
-}
-
-int32 FInternationalization::GetCultureIndex(const FString& Name) const
-{
-	const FString CanonicalName = FCulture::GetCanonicalName(Name);
-
-	const int32 CultureCount = AllCultures.Num();
-	int32 i;
-	for (i = 0; i < CultureCount; ++i)
-	{
-		if( AllCultures[i]->GetName() == CanonicalName )
-		{
-			break;
-		}
-	}
-	if(i >= CultureCount)
-	{
-		i = -1;
-	}
-	return i;
-}
-
-#if UE_ENABLE_ICU
-namespace
-{
-	struct FICUOverrides
-	{
-#if STATS
-		static int64 BytesInUseCount;
-		static int64 CachedBytesInUseCount;
-#endif
-
-		static void* U_CALLCONV Malloc(const void* context, size_t size)
-		{
-			void* Result = FMemory::Malloc(size);
-#if STATS
-			BytesInUseCount += FMemory::GetAllocSize(Result);
-			if(FThreadStats::IsThreadingReady() && CachedBytesInUseCount != BytesInUseCount)
-			{
-				SET_MEMORY_STAT(STAT_MemoryICUTotalAllocationSize, BytesInUseCount);
-				CachedBytesInUseCount = BytesInUseCount;
-			}
-#endif
-			return Result;
-		}
-
-		static void* U_CALLCONV Realloc(const void* context, void* mem, size_t size)
-		{
-			return FMemory::Realloc(mem, size);
-		}
-
-		static void U_CALLCONV Free(const void* context, void* mem)
-		{
-#if STATS
-			BytesInUseCount -= FMemory::GetAllocSize(mem);
-			if(FThreadStats::IsThreadingReady() && CachedBytesInUseCount != BytesInUseCount)
-			{
-				SET_MEMORY_STAT(STAT_MemoryICUTotalAllocationSize, BytesInUseCount);
-				CachedBytesInUseCount = BytesInUseCount;
-			}
-#endif
-			FMemory::Free(mem);
-		}
-	};
-
-#if STATS
-	int64 FICUOverrides::BytesInUseCount = 0;
-	int64 FICUOverrides::CachedBytesInUseCount = 0;
-#endif
+	return Implementation->GetCulture(Name);
 }
-#endif
 
 void FInternationalization::Initialize()
 {
@@ -193,173 +57,22 @@ void FInternationalization::Initialize()
 		~FInitializingGuard()	{IsInitializing = false;}
 	} InitializingGuard;
 
-#if UE_ENABLE_ICU
-	UErrorCode ICUStatus = U_ZERO_ERROR;
-
-// Linux needs to have those compiled statically at least until we settle on .so location for deployed/native builds
-#if (IS_PROGRAM || !IS_MONOLITHIC) && !PLATFORM_LINUX
-	LoadDLLs();
-#endif /*IS_PROGRAM || !IS_MONOLITHIC*/
-
-	u_setMemoryFunctions(NULL, &(FICUOverrides::Malloc), &(FICUOverrides::Realloc), &(FICUOverrides::Free), &(ICUStatus));
-
-	const FString DataDirectoryRelativeToContent = FString(TEXT("Localization")) / FString(TEXT("ICU"));
-	IFileManager& FileManager = IFileManager::Get();
-
-	const FString PotentialDataDirectories[] = 
-		{
-			FPaths::GameContentDir() / DataDirectoryRelativeToContent, // Try game content directory.
-			FPaths::EngineContentDir() / DataDirectoryRelativeToContent, // Try engine content directory.
-		};
-
-	bool HasFoundDataDirectory = false;
-	for(const auto& DataDirectoryString : PotentialDataDirectories)
-	{
-		if( FileManager.DirectoryExists( *DataDirectoryString ) )
-		{
-			u_setDataDirectory( StringCast<char>( *DataDirectoryString ).Get() );
-			HasFoundDataDirectory = true;
-			break;
-		}
-	}
-
-	auto GetPrioritizedDataDirectoriesString = [&]() -> FString
-	{
-		FString Result;
-		for(const auto& DataDirectoryString : PotentialDataDirectories)
-		{
-			if(!Result.IsEmpty())
-			{
-				Result += TEXT("\n");
-			}
-			Result += DataDirectoryString;
-		}
-		return Result;
-	};
-	checkf( HasFoundDataDirectory, TEXT("ICU data directory was not discovered:\n%s"), *(GetPrioritizedDataDirectoriesString()) );
-
-	u_setDataFileFunctions(nullptr, &FICUDataCallbacks::OpenDataFile, &FICUDataCallbacks::CloseDataFile, &(ICUStatus));
-	u_init(&(ICUStatus));
-
-	FICUBreakIteratorManager::Create();
-#endif /*UE_ENABLE_ICU*/
-
-	PopulateAllCultures();
-
-#if UE_ENABLE_ICU
-	SetCurrentCulture( GetDefaultCulture()->GetName() );
-#else
-	SetCurrentCulture( TEXT("") );
-#endif
-
-#if UE_ENABLE_ICU
-	bIsInitialized = U_SUCCESS(ICUStatus) ? true : false;
-#else
-	bIsInitialized = true;
-#endif
+	bIsInitialized = Implementation->Initialize();
 }
 
 void FInternationalization::Terminate()
 {
 	DefaultCulture.Reset();
 	InvariantCulture.Reset();
-	AllCultures.Empty();
+	CurrentCulture.Reset();
 
-#if UE_ENABLE_ICU
-	FICUBreakIteratorManager::Destroy();
-	u_cleanup();
-#if IS_PROGRAM || !IS_MONOLITHIC
-	UnloadDLLs();
-#endif //IS_PROGRAM || !IS_MONOLITHIC
-#endif //UE_ENABLE_ICU
+	Implementation->Terminate();
 
 	bIsInitialized = false;
 	delete Instance;
 	Instance = nullptr;
 }
 
-#if UE_ENABLE_ICU && (IS_PROGRAM || !IS_MONOLITHIC)
-void FInternationalization::LoadDLLs()
-{
-	// The base directory for ICU binaries is consistent on all platforms.
-	FString ICUBinariesRoot = FPaths::EngineDir() / TEXT("Binaries") / TEXT("ThirdParty") / TEXT("ICU") / TEXT("icu4c-53_1");
-
-#if PLATFORM_WINDOWS
-#if PLATFORM_64BITS
-	const FString PlatformFolderName = TEXT("Win64");
-#elif PLATFORM_32BITS
-	const FString PlatformFolderName = TEXT("Win32");
-#endif //PLATFORM_*BITS
-
-#if _MSC_VER >= 1800
-	const FString VSVersionFolderName = TEXT("VS2013");
-#else
-	const FString VSVersionFolderName = TEXT("VS2012");
-#endif //_MSVC_VER
-
-	// Windows requires support for 32/64 bit and also VC11 and VC12 runtimes.
-	const FString TargetSpecificPath = ICUBinariesRoot / PlatformFolderName / VSVersionFolderName;
-
-	// Windows libraries use a specific naming convention.
-	FString LibraryNameStems[] =
-	{
-		TEXT("dt"),   // Data
-		TEXT("uc"),   // Unicode Common
-		TEXT("in"),   // Internationalization
-		TEXT("le"),   // Layout Engine
-		TEXT("lx"),   // Layout Extensions
-		TEXT("io")	// Input/Output
-	};
-#else
-	// Non-Windows libraries use a consistent naming convention.
-	FString LibraryNameStems[] =
-	{
-		TEXT("data"),   // Data
-		TEXT("uc"),   // Unicode Common
-		TEXT("i18n"),   // Internationalization
-		TEXT("le"),   // Layout Engine
-		TEXT("lx"),   // Layout Extensions
-		TEXT("io")	// Input/Output
-	};
-#if PLATFORM_LINUX
-	const FString TargetSpecificPath = ICUBinariesRoot / TEXT("Linux") / TEXT("x86_64-unknown-linux-gnu");
-#elif PLATFORM_MAC
-	const FString TargetSpecificPath = ICUBinariesRoot / TEXT("Mac");
-#endif //PLATFORM_*
-#endif //PLATFORM_*
-
-#if UE_BUILD_DEBUG && !defined(NDEBUG)
-	const FString LibraryNamePostfix = TEXT("d");
-#else
-	const FString LibraryNamePostfix = TEXT("");
-#endif //DEBUG
-
-	for(FString Stem : LibraryNameStems)
-	{
-#if PLATFORM_WINDOWS
-		FString LibraryName = "icu" + Stem + LibraryNamePostfix + "53" "." "dll";
-#elif PLATFORM_LINUX
-		FString LibraryName = "lib" "icu" + Stem + LibraryNamePostfix + ".53.1" "." "so";
-#elif PLATFORM_MAC
-		FString LibraryName = "lib" "icu" + Stem + LibraryNamePostfix + ".53.1" "." "dylib";
-#endif //PLATFORM_*
-
-		void* DLLHandle = FPlatformProcess::GetDllHandle(*(TargetSpecificPath / LibraryName));
-		check(DLLHandle != nullptr);
-		DLLHandles.Add(DLLHandle);
-	}
-}
-
-void FInternationalization::UnloadDLLs()
-{
-	for( const auto DLLHandle : DLLHandles )
-	{
-		FPlatformProcess::FreeDllHandle(DLLHandle);
-	}
-	DLLHandles.Reset();
-}
-#endif //UE_ENABLE_ICU && (IS_PROGRAM || !IS_MONOLITHIC)
-
 #if ENABLE_LOC_TESTING
 namespace
 {
@@ -503,14 +216,10 @@ FString& FInternationalization::Leetify(FString& SourceString)
 
 void FInternationalization::GetCultureNames(TArray<FString>& CultureNames) const
 {
-	CultureNames.Reset();
-	for(const auto& Culture : AllCultures)
-	{
-		CultureNames.Add(Culture->GetName());
-	}
+	Implementation->GetCultureNames(CultureNames);
 }
 
-void FInternationalization::GetCulturesWithAvailableLocalization(const TArray<FString>& InLocalizationPaths, TArray< FCultureRef >& OutAvailableCultures, const bool bIncludeDerivedCultures) const
+void FInternationalization::GetCulturesWithAvailableLocalization(const TArray<FString>& InLocalizationPaths, TArray< FCultureRef >& OutAvailableCultures, const bool bIncludeDerivedCultures)
 {
 	OutAvailableCultures.Reset();
 
@@ -551,14 +260,22 @@ void FInternationalization::GetCulturesWithAvailableLocalization(const TArray<FS
 	// Find any cultures that are a partial match for those we have translations for.
 	if(bIncludeDerivedCultures)
 	{
-		for(const auto& Culture : AllCultures)
+		TArray<FString> CultureNames;
+		GetCultureNames(CultureNames);
+		for(const FString& CultureName : CultureNames)
 		{
-			for(FString ParentCultureName = Culture->GetName(); !ParentCultureName.IsEmpty(); ParentCultureName = FCulture::GetParentName(ParentCultureName))
+			FCulturePtr Culture = GetCulture(CultureName);
+			if (Culture.IsValid())
 			{
-				if(AllLocalizationFolders.Contains(ParentCultureName))
+				TArray<FString> PrioritizedParentCultureNames = Culture->GetPrioritizedParentCultureNames();
+
+				for (const FString& CultureName : PrioritizedParentCultureNames)
 				{
-					OutAvailableCultures.AddUnique(Culture);
-					break;
+					if(AllLocalizationFolders.Contains(CultureName))
+					{
+						OutAvailableCultures.AddUnique(Culture.ToSharedRef());
+						break;
+					}
 				}
 			}
 		}
@@ -569,7 +286,7 @@ void FInternationalization::GetCulturesWithAvailableLocalization(const TArray<FS
 		for(const FString& LocalizationFolder : AllLocalizationFolders)
 		{
 			FCulturePtr Culture = GetCulture(LocalizationFolder);
-			if(Culture.IsValid() && AllCultures.Contains(Culture.ToSharedRef()))
+			if(Culture.IsValid())
 			{
 				OutAvailableCultures.AddUnique(Culture.ToSharedRef());
 			}
@@ -579,190 +296,9 @@ void FInternationalization::GetCulturesWithAvailableLocalization(const TArray<FS
 
 FInternationalization::FInternationalization()
 	:	bIsInitialized(false)
-	,	CurrentCultureIndex(INDEX_NONE)
-{
-
-}
-
-void FInternationalization::PopulateAllCultures(void)
-{
-#if UE_ENABLE_ICU
-	int32_t LocaleCount;
-	const icu::Locale* const AvailableLocales = icu::Locale::getAvailableLocales(LocaleCount);
-
-	// Get available locales and create cultures.
-	AllCultures.Reset(LocaleCount);
-	for(int32_t i = 0; i < LocaleCount; ++i)
-	{
-		AllCultures.Add( MakeShareable( new FCulture( AvailableLocales[i].getName() ) ) );
-	}
-
-	const FString DefaultLocaleName = FPlatformMisc::GetDefaultLocale();
-
-	// Find the default locale as a culture, if present..
-	auto FindDefaultPredicate = [DefaultLocaleName](const FCultureRef& Culture) -> bool
-	{
-		return Culture->GetName() == DefaultLocaleName;
-	};
-	FCultureRef* FoundDefaultCulture = AllCultures.FindByPredicate(FindDefaultPredicate);
-	// If present, set default culture variable.
-	if(FoundDefaultCulture)
-	{
-		DefaultCulture = *FoundDefaultCulture;
-	}
-	// Otherwise, create, set, and add default culture variable.
-	else
-	{
-		DefaultCulture = MakeShareable(new FCulture(DefaultLocaleName));
-		AllCultures.Add( DefaultCulture.ToSharedRef() );
-	}
-
-#else
-	AllCultures.Add( FInvariantCulture::Create() );
-#if ENABLE_LOC_TESTING
-	AllCultures.Add( FLEETCulture::Create() );
-#endif
-	AllCultures.Add( FAmericanEnglishCulture::Create() );
-	AllCultures.Add( FIndiaHindiCulture::Create() );
-	AllCultures.Add( FJapaneseCulture::Create() );
-	AllCultures.Add( FKoreanCulture::Create() );
-
-	DefaultCulture = GetCulture(TEXT(""));
-#endif 
-
-	InvariantCulture = GetCulture(TEXT(""));
-}
-
-#if UE_ENABLE_ICU
-#if STATS
-namespace
-{
-	int64 DataFileBytesInUseCount = 0;
-	int64 CachedDataFileBytesInUseCount = 0;
-}
-#endif
-
-UBool FICUDataCallbacks::OpenDataFile(const void* context, void** fileContext, void** contents, const char* path)
-{
-	auto& PathToCachedFileDataMap = FInternationalization::Get().PathToCachedFileDataMap;
-
-	// Try to find existing buffer
-	FInternationalization::FICUCachedFileData* CachedFileData = PathToCachedFileDataMap.Find(path);
-
-	// If there's no file context, we might have to load the file.
-	if (!CachedFileData)
-	{
-		// Attempt to load the file.
-		FArchive* FileAr = IFileManager::Get().CreateFileReader(StringCast<TCHAR>(path).Get());
-		if (FileAr)
-		{
-			const int64 FileSize = FileAr->TotalSize();
-
-			// Create file data.
-			CachedFileData = &(PathToCachedFileDataMap.Emplace(FString(path), FInternationalization::FICUCachedFileData(FileSize)));
-
-			// Load file into buffer.
-			FileAr->Serialize(CachedFileData->Buffer, FileSize); 
-			delete FileAr;
-
-			// Stat tracking.
-#if STATS
-			DataFileBytesInUseCount += FMemory::GetAllocSize(CachedFileData->Buffer);
-			if (FThreadStats::IsThreadingReady() && CachedDataFileBytesInUseCount != DataFileBytesInUseCount)
-			{
-				SET_MEMORY_STAT(STAT_MemoryICUDataFileAllocationSize, DataFileBytesInUseCount);
-				CachedDataFileBytesInUseCount = DataFileBytesInUseCount;
-			}
-#endif
-		}
-	}
-
-	// Add a reference, either the initial one or an additional one.
-	if (CachedFileData)
-	{
-		++(CachedFileData->ReferenceCount);
-	}
-
-	// Use the file path as the context, so we can look up the cached file data later and decrement its reference count.
-	*fileContext = CachedFileData ? new FString(path) : nullptr;
-
-	// Use the buffer from the cached file data.
-	*contents = CachedFileData ? CachedFileData->Buffer : nullptr;
-
-	// If we have cached file data, we must have loaded new data or found existing data, so we've successfully "opened" and "read" the file into "contents".
-	return CachedFileData != nullptr;
-}
-
-void FICUDataCallbacks::CloseDataFile(const void* context, void* const fileContext, void* const contents)
+	,	Implementation(this)
 {
-	// Early out on null context.
-	if (fileContext == nullptr)
-	{
-		return;
-	}
-
-	auto& PathToCachedFileDataMap = FInternationalization::Get().PathToCachedFileDataMap;
-
-	// The file context is the path to the file.
-	FString* const Path = reinterpret_cast<FString*>(fileContext);
-	check(Path);
 
-	// Look up the cached file data so we can maintain references.
-	FInternationalization::FICUCachedFileData* const CachedFileData = PathToCachedFileDataMap.Find(*Path);
-	check(CachedFileData);
-	check(CachedFileData->Buffer == contents);
-
-	// Remove a reference.
-	--(CachedFileData->ReferenceCount);
-
-	// If the last reference has been removed, the cached file data is not longer needed.
-	if (CachedFileData->ReferenceCount == 0)
-	{
-		// Stat tracking.
-#if STATS
-		DataFileBytesInUseCount -= FMemory::GetAllocSize(CachedFileData->Buffer);
-		if (FThreadStats::IsThreadingReady() && CachedDataFileBytesInUseCount != DataFileBytesInUseCount)
-		{
-			SET_MEMORY_STAT(STAT_MemoryICUDataFileAllocationSize, DataFileBytesInUseCount);
-			CachedDataFileBytesInUseCount = DataFileBytesInUseCount;
-		}
-#endif
-
-		// Delete the cached file data.
-		PathToCachedFileDataMap.Remove(*Path);
-	}
-
-	// The path string we allocated for tracking is no longer necessary.
-	delete Path;
-}
-
-FInternationalization::FICUCachedFileData::FICUCachedFileData(const int64 FileSize)
-	: ReferenceCount(0)
-	, Buffer( FICUOverrides::Malloc(nullptr, FileSize) )
-{
 }
 
-FInternationalization::FICUCachedFileData::FICUCachedFileData(const FICUCachedFileData& Source)
-{
-	checkf(false, TEXT("Cached file data for ICU may not be copy constructed. Something is trying to copy construct FICUCachedFileData."));
-}
-
-FInternationalization::FICUCachedFileData::FICUCachedFileData(FICUCachedFileData&& Source)
-	: ReferenceCount(Source.ReferenceCount)
-	, Buffer( Source.Buffer )
-{
-	// Make sure the moved source object doesn't retain the pointer and free the memory we now point to.
-	Source.Buffer = nullptr;
-}
-
-FInternationalization::FICUCachedFileData::~FICUCachedFileData()
-{
-	if (Buffer)
-	{
-		check(ReferenceCount == 0);
-		FICUOverrides::Free(nullptr, Buffer);
-	}
-}
-#endif
-
 #undef LOCTEXT_NAMESPACE
diff --git a/Engine/Source/Runtime/Core/Private/Internationalization/LegacyCulture.cpp b/Engine/Source/Runtime/Core/Private/Internationalization/LegacyCulture.cpp
--- a/Engine/Source/Runtime/Core/Private/Internationalization/LegacyCulture.cpp
+++ b/Engine/Source/Runtime/Core/Private/Internationalization/LegacyCulture.cpp
@@ -48,20 +48,6 @@ FString FCulture::FLegacyCultureImplementation::GetCanonicalName(const FString&
 	return Name;
 }
 
-FString FCulture::FLegacyCultureImplementation::GetParentName(const FString& CultureName)
-{
-	FString Result;
-
-	int32 UnderscoreIndex = INDEX_NONE;
-	if(CultureName.FindLastChar(TEXT('_'), UnderscoreIndex))
-	{
-		Result = CultureName;
-		Result.RemoveAt(UnderscoreIndex);
-	}
-
-	return Result;
-}
-
 FString FCulture::FLegacyCultureImplementation::GetNativeName() const
 {
 	return NativeName;
diff --git a/Engine/Source/Runtime/Core/Private/Internationalization/LegacyCulture.h b/Engine/Source/Runtime/Core/Private/Internationalization/LegacyCulture.h
--- a/Engine/Source/Runtime/Core/Private/Internationalization/LegacyCulture.h
+++ b/Engine/Source/Runtime/Core/Private/Internationalization/LegacyCulture.h
@@ -26,24 +26,22 @@ class FCulture::FLegacyCultureImplementation
 
 	int GetLCID() const;
 
-	FString GetName() const;
-
 	static FString GetCanonicalName(const FString& Name);
 
-	static FString GetParentName(const FString& CultureName);
+	FString GetName() const;
 
 	FString GetNativeName() const;
 
-	FString GetNativeLanguage() const;
-
-	FString GetNativeRegion() const;
-
 	FString GetUnrealLegacyThreeLetterISOLanguageName() const;
 
 	FString GetThreeLetterISOLanguageName() const;
 
 	FString GetTwoLetterISOLanguageName() const;
 
+	FString GetNativeLanguage() const;
+
+	FString GetNativeRegion() const;
+
 	// Full localized culture name
 	const FText DisplayName;
 
diff --git a/Engine/Source/Runtime/Core/Private/Internationalization/LegacyInternationalization.cpp b/Engine/Source/Runtime/Core/Private/Internationalization/LegacyInternationalization.cpp
new file mode 100644
--- /dev/null
+++ b/Engine/Source/Runtime/Core/Private/Internationalization/LegacyInternationalization.cpp
@@ -0,0 +1,42 @@
+// Copyright 1998-2014 Epic Games, Inc. All Rights Reserved.
+#include "Core.h"
+#include "LegacyInternationalization.h"
+
+#if !UE_ENABLE_ICU
+
+#include "InvariantCulture.h"
+
+FLegacyInternationalization::FLegacyInternationalization(FInternationalization* const InI18N)
+	: I18N(InI18N)
+{
+
+}
+
+bool FLegacyInternationalization::Initialize()
+{
+	I18N->InvariantCulture = FInvariantCulture::Create();
+	I18N->DefaultCulture = I18N->InvariantCulture;
+	I18N->CurrentCulture = I18N->InvariantCulture;
+
+	return true;
+}
+
+void FLegacyInternationalization::Terminate()
+{
+}
+
+void FLegacyInternationalization::SetCurrentCulture(const FString& Name)
+{
+}
+
+void FLegacyInternationalization::GetCultureNames(TArray<FString>& CultureNames) const
+{
+	CultureNames.Add(TEXT(""));
+}
+
+FCulturePtr FLegacyInternationalization::GetCulture(const FString& Name)
+{
+	return Name.IsEmpty() ? I18N->InvariantCulture : nullptr;
+}
+
+#endif
\ No newline at end of file
diff --git a/Engine/Source/Runtime/Core/Private/Internationalization/LegacyInternationalization.h b/Engine/Source/Runtime/Core/Private/Internationalization/LegacyInternationalization.h
new file mode 100644
--- /dev/null
+++ b/Engine/Source/Runtime/Core/Private/Internationalization/LegacyInternationalization.h
@@ -0,0 +1,22 @@
+// Copyright 1998-2014 Epic Games, Inc. All Rights Reserved.
+#pragma once
+
+#if !UE_ENABLE_ICU
+
+class FLegacyInternationalization
+{
+public:
+	FLegacyInternationalization(FInternationalization* const I18N);
+
+	bool Initialize();
+	void Terminate();
+
+	void SetCurrentCulture(const FString& Name);
+	void GetCultureNames(TArray<FString>& CultureNames) const;
+	FCulturePtr GetCulture(const FString& Name);
+
+private:
+	FInternationalization* const I18N;
+};
+
+#endif
\ No newline at end of file
diff --git a/Engine/Source/Runtime/Core/Private/Internationalization/TextLocalizationManager.cpp b/Engine/Source/Runtime/Core/Private/Internationalization/TextLocalizationManager.cpp
--- a/Engine/Source/Runtime/Core/Private/Internationalization/TextLocalizationManager.cpp
+++ b/Engine/Source/Runtime/Core/Private/Internationalization/TextLocalizationManager.cpp
@@ -83,30 +83,29 @@ void EndInitTextLocalization()
 			}
 			LocalizationPaths += FPaths::GetEngineLocalizationPaths();
 
+			// Validate the locale has data or fallback to one that does.
 			TArray< FCultureRef > AvailableCultures;
 			I18N.GetCulturesWithAvailableLocalization(LocalizationPaths, AvailableCultures, false);
 
-			// Query for parent culture name based on request culture name.
-			FString ParentCultureName = RequestedCultureName;
+			TArray<FString> PrioritizedParentCultureNames = I18N.GetCurrentCulture()->GetPrioritizedParentCultureNames();
 				
-			// If we do not have localization data, try the parent culture.
-			for(FCulturePtr ParentCulture = I18N.GetCulture(ParentCultureName); !ParentCulture.IsValid() || !AvailableCultures.Contains(ParentCulture.ToSharedRef()); ParentCulture = I18N.GetCulture(ParentCultureName))
+			FString ValidCultureName;
+			for (const FString& CultureName : PrioritizedParentCultureNames)
 			{
-				ParentCultureName = FCulture::GetParentName(ParentCultureName);
-
-				// No parent culture
-				if(ParentCultureName.IsEmpty())
+				FCulturePtr ValidCulture = I18N.GetCulture(CultureName);
+				if (ValidCulture.IsValid() && AvailableCultures.Contains(ValidCulture.ToSharedRef()))
 				{
+					ValidCultureName = CultureName;
 					break;
 				}
 			}
 
-			if(!ParentCultureName.IsEmpty())
+			if(!ValidCultureName.IsEmpty())
 			{
-				if(RequestedCultureName != ParentCultureName)
+				if(RequestedCultureName != ValidCultureName)
 				{
 					// Make the user aware that the localization data belongs to a parent culture.
-					UE_LOG(LogTextLocalizationManager, Log, TEXT("The requested culture ('%s') has no localization data; parent culture's ('%s') localization data will be used."), *RequestedCultureName, *ParentCultureName);
+					UE_LOG(LogTextLocalizationManager, Log, TEXT("The requested culture ('%s') has no localization data; parent culture's ('%s') localization data will be used."), *RequestedCultureName, *ValidCultureName);
 				}
 			}
 			else
@@ -273,7 +272,6 @@ void FTextLocalizationManager::LoadResources(const bool ShouldLoadEditor, const
 	FInternationalization& I18N = FInternationalization::Get();
 
 	const FString& CultureName = I18N.GetCurrentCulture()->GetName();
-	const FString& BaseLanguageName = I18N.GetCurrentCulture()->GetTwoLetterISOLanguageName();
 
 #if ENABLE_LOC_TESTING
 	if(CultureName == TEXT("LEET"))
@@ -364,37 +362,23 @@ void FTextLocalizationManager::LoadResources(const bool ShouldLoadEditor, const
 	}
 
 	// Read culture localization resources.
-	FLocalizationEntryTracker& CultureTracker = LocalizationEntryTrackers[LocalizationEntryTrackers.Add(FLocalizationEntryTracker())];
-	for (const FString& LocalizationPath : LocalizationPaths)
-	{
-		const FString* const Entry = LocalizationPathToCultureDirectoryMap[LocalizationPath].Find(CultureName);
-		if (Entry)
-		{
-			const FString CulturePath = LocalizationPath / (*Entry);
-
-			CultureTracker.ReadFromDirectory(CulturePath);
-		}
-	}
-	CultureTracker.ReportCollisions();
+	TArray<FString> PrioritizedParentCultureNames = I18N.GetCurrentCulture()->GetPrioritizedParentCultureNames();
 
-	// Read base language localization resources.
-	FLocalizationEntryTracker& BaseLanguageTracker = LocalizationEntryTrackers[LocalizationEntryTrackers.Add(FLocalizationEntryTracker())];
-	for (const FString& LocalizationPath : LocalizationPaths)
+	for (const FString& CultureName : PrioritizedParentCultureNames)
 	{
-		const FString* const BaseEntry = LocalizationPathToCultureDirectoryMap[LocalizationPath].Find(BaseLanguageName);
-		const FString* const Entry = LocalizationPathToCultureDirectoryMap[LocalizationPath].Find(CultureName);
-
-		if (BaseEntry)
+		FLocalizationEntryTracker& CultureTracker = LocalizationEntryTrackers[LocalizationEntryTrackers.Add(FLocalizationEntryTracker())];
+		for (const FString& LocalizationPath : LocalizationPaths)
 		{
-			// If there is no entry or the paths are different.
-			if( !Entry || *BaseEntry != *Entry )
+			const FString* const Entry = LocalizationPathToCultureDirectoryMap[LocalizationPath].Find(FCulture::GetCanonicalName(CultureName));
+			if (Entry)
 			{
-				const FString BaseLanguagePath = LocalizationPath / (*BaseEntry);
-				BaseLanguageTracker.ReadFromDirectory(BaseLanguagePath);
+				const FString CulturePath = LocalizationPath / (*Entry);
+
+				CultureTracker.ReadFromDirectory(CulturePath);
 			}
 		}
+		CultureTracker.ReportCollisions();
 	}
-	BaseLanguageTracker.ReportCollisions();
 
 	UpdateLiveTable(LocalizationEntryTrackers);
 }
diff --git a/Engine/Source/Runtime/Core/Public/Internationalization/Culture.h b/Engine/Source/Runtime/Core/Public/Internationalization/Culture.h
--- a/Engine/Source/Runtime/Core/Public/Internationalization/Culture.h
+++ b/Engine/Source/Runtime/Core/Public/Internationalization/Culture.h
@@ -1,4 +1,4 @@
-// Copyright 1998-2014 Epic Games, Inc. All Rights Reserved.
+﻿// Copyright 1998-2014 Epic Games, Inc. All Rights Reserved.
 #pragma once
 
 #include "Internationalization/Text.h"
@@ -11,26 +11,24 @@
 
 class CORE_API FCulture
 {
-#if UE_ENABLE_ICU
 public:
-	FCulture(const FString& LocaleName);
+#if UE_ENABLE_ICU
+	static FCulturePtr Create(const FString& LocaleName);
 #else
-public:
-	FCulture(	const FText& InDisplayName, 
-				const FString& InEnglishName, 
-				const int InKeyboardLayoutId, 
-				const int InLCID, 
-				const FString& InName, 
-				const FString& InNativeName, 
-				const FString& InUnrealLegacyThreeLetterISOLanguageName, 
-				const FString& InThreeLetterISOLanguageName, 
-				const FString& InTwoLetterISOLanguageName, 
-				const FNumberFormattingRules& InNumberFormattingRule, 
-				const FTextFormattingRules& InTextFormattingRule, 
-				const FDateTimeFormattingRules& InDateTimeFormattingRule);
+	static FCulturePtr Create(	const FText& InDisplayName, 
+									const FString& InEnglishName, 
+									const int InKeyboardLayoutId, 
+									const int InLCID, 
+									const FString& InName, 
+									const FString& InNativeName, 
+									const FString& InUnrealLegacyThreeLetterISOLanguageName, 
+									const FString& InThreeLetterISOLanguageName, 
+									const FString& InTwoLetterISOLanguageName, 
+									const FNumberFormattingRules& InNumberFormattingRule, 
+									const FTextFormattingRules& InTextFormattingRule, 
+									const FDateTimeFormattingRules& InDateTimeFormattingRule);
 #endif
 
-public:
 	FString GetDisplayName() const;
 
 	FString GetEnglishName() const;
@@ -39,46 +37,64 @@ public:
 
 	int GetLCID() const;
 
-	FString GetName() const;
+	TArray<FString> GetPrioritizedParentCultureNames();
 
 	static FString GetCanonicalName(const FString& Name);
 
-	static FString GetParentName(const FString& CultureName);
+	FString GetName() const;
 	
 	FString GetNativeName() const;
 
-	FString GetNativeLanguage() const;
-
-	FString GetNativeRegion() const;
-
 	FString GetUnrealLegacyThreeLetterISOLanguageName() const;
 
 	FString GetThreeLetterISOLanguageName() const;
 
 	FString GetTwoLetterISOLanguageName() const;
 
-	FString GetCountry() const;
+	FString GetNativeLanguage() const;
+
+	FString GetRegion() const;
+
+	FString GetNativeRegion() const;
+
+	FString GetScript() const;
 
 	FString GetVariant() const;
 
-#if UE_ENABLE_ICU
 public:
+#if UE_ENABLE_ICU
 	class FICUCultureImplementation;
 	typedef FICUCultureImplementation FImplementation;
 	TSharedRef<FICUCultureImplementation> Implementation;
 #else
-public:
 	class FLegacyCultureImplementation;
 	typedef FLegacyCultureImplementation FImplementation;
 	TSharedRef<FLegacyCultureImplementation> Implementation;
+#endif
 
-	// The date formating rule for this culture
-	const FDateTimeFormattingRules DateTimeFormattingRule;
+protected:
+#if UE_ENABLE_ICU
+	FCulture(const FString& LocaleName);
+#else
+	FCulture(	const FText& InDisplayName, 
+				const FString& InEnglishName, 
+				const int InKeyboardLayoutId, 
+				const int InLCID, 
+				const FString& InName, 
+				const FString& InNativeName, 
+				const FString& InUnrealLegacyThreeLetterISOLanguageName, 
+				const FString& InThreeLetterISOLanguageName, 
+				const FString& InTwoLetterISOLanguageName, 
+				const FNumberFormattingRules& InNumberFormattingRule, 
+				const FTextFormattingRules& InTextFormattingRule, 
+				const FDateTimeFormattingRules& InDateTimeFormattingRule);
+#endif
 
-	// The text formating rule for this culture
+private:
+#if !UE_ENABLE_ICU
+	friend class FText;
+	const FDateTimeFormattingRules DateTimeFormattingRule;
 	const FTextFormattingRules TextFormattingRule;
-
-	// The number formating rule for this culture
 	const FNumberFormattingRules NumberFormattingRule;
 #endif
 };
diff --git a/Engine/Source/Runtime/Core/Public/Internationalization/Internationalization.h b/Engine/Source/Runtime/Core/Public/Internationalization/Internationalization.h
--- a/Engine/Source/Runtime/Core/Public/Internationalization/Internationalization.h
+++ b/Engine/Source/Runtime/Core/Public/Internationalization/Internationalization.h
@@ -7,6 +7,7 @@
 #include "TextLocalizationManager.h"
 #include "TextLocalizationManagerGlobals.h"
 #include "Culture.h"
+#include "UniqueObj.h"
 
 #define LOC_DEFINE_REGION
 
@@ -24,6 +25,9 @@
  */
 #define NSLOCTEXT( InNamespace, InKey, InTextLiteral ) FInternationalization::ForUseOnlyByLocMacroAndGraphNodeTextLiterals_CreateText( TEXT( InTextLiteral ), TEXT( InNamespace ), TEXT( InKey ) )
 
+class FICUInternationalization;
+class FLegacyInternationalization;
+
 class FInternationalization
 {
 public:
@@ -36,16 +40,17 @@ public:
 		return FText( InTextLiteral, Namespace, Key );
 	}
 
-	CORE_API void GetTimeZonesIDs(TArray<FString>& TimeZonesIDs) const;
-
 	//Set the current culture by name
 	CORE_API void SetCurrentCulture(const FString& Name);
 
 	//@return the current culture
-	CORE_API FCultureRef GetCurrentCulture() const;
+	CORE_API FCultureRef GetCurrentCulture() const
+	{
+		return CurrentCulture.ToSharedRef();
+	}
 
 	//@return culture object by given name, or NULL if not found
-	CORE_API FCulturePtr GetCulture(const FString& Name) const;
+	CORE_API FCulturePtr GetCulture(const FString& Name);
 
 	//@return the default culture
 	CORE_API TSharedRef< FCulture, ESPMode::ThreadSafe > GetDefaultCulture() const
@@ -59,7 +64,6 @@ public:
 		return InvariantCulture.ToSharedRef();
 	}
 
-
 	CORE_API bool IsInitialized() const {return bIsInitialized;}
 
 #if ENABLE_LOC_TESTING
@@ -68,65 +72,31 @@ public:
 
 	CORE_API void GetCultureNames(TArray<FString>& CultureNames) const;
 
-	CORE_API const TArray< FCultureRef >& GetAllCultures() const
-	{
-		return AllCultures;
-	}
-
 	// Given some paths to look at, populate a list of cultures that we have available localization information for. If bIncludeDerivedCultures, include cultures that are derived from those we have localization data for.
-	CORE_API void GetCulturesWithAvailableLocalization(const TArray<FString>& InLocalizationPaths, TArray< FCultureRef >& OutAvailableCultures, const bool bIncludeDerivedCultures) const;
+	CORE_API void GetCulturesWithAvailableLocalization(const TArray<FString>& InLocalizationPaths, TArray< FCultureRef >& OutAvailableCultures, const bool bIncludeDerivedCultures);
 
 private:
 	FInternationalization();
 
-	//@todo how are we going to initialize the list of cultures? below is temporary
-	CORE_API void PopulateAllCultures(void);
-
-	//@return index of a culture by given name, or -1 if not found
-	CORE_API int32 GetCultureIndex(const FString& Name) const;
-
-	CORE_API void Initialize();
-	CORE_API void Terminate();
-
-#if UE_ENABLE_ICU && (IS_PROGRAM || !IS_MONOLITHIC)
-	void LoadDLLs();
-	void UnloadDLLs();
-#endif
+	void Initialize();
+	void Terminate();
 
 private:
 	static FInternationalization* Instance;
 	bool bIsInitialized;
 
-	/** Must be thread safe, cultures may be accessed from more than one thread */
-	TArray< TSharedRef< FCulture, ESPMode::ThreadSafe > > AllCultures;
-
-	int CurrentCultureIndex;
-
-	/** Must be thread safe, cultures may be accessed from more than one thread */
-	TSharedPtr< FCulture, ESPMode::ThreadSafe > DefaultCulture;
-	/** Must be thread safe, cultures may be accessed from more than one thread */
-	TSharedPtr< FCulture, ESPMode::ThreadSafe > InvariantCulture;
-
-	TArray< void* > DLLHandles;
-
 #if UE_ENABLE_ICU
-	friend struct FICUDataCallbacks;
-
-	// Tidy class for storing the count of references for an ICU data file and the file's data itself.
-	struct FICUCachedFileData
-	{
-		FICUCachedFileData(const int64 FileSize);
-		FICUCachedFileData(const FICUCachedFileData& Source);
-		FICUCachedFileData(FICUCachedFileData&& Source);
-		~FICUCachedFileData();
-
-		uint32 ReferenceCount;
-		void* Buffer;
-	};
-
-	// Map for associating ICU data file paths with cached file data, to prevent multiple copies of immutable ICU data files from residing in memory.
-	TMap<FString, FICUCachedFileData> PathToCachedFileDataMap;
+	friend class FICUInternationalization;
+	typedef FICUInternationalization FImplementation;
+#else
+	friend class FLegacyInternationalization;
+	typedef FLegacyInternationalization FImplementation;
 #endif
+	TUniqueObj<FImplementation> Implementation;
+
+	FCulturePtr CurrentCulture;
+	FCulturePtr DefaultCulture;
+	FCulturePtr InvariantCulture;
 };
 
 #undef LOC_DEFINE_REGION
