#ue4 - Add Vector functions: ProjectOnToNormal(), VectorPlaneProject()

diff --git a/Engine/Source/Runtime/Core/Public/Math/Vector.h b/Engine/Source/Runtime/Core/Public/Math/Vector.h
--- a/Engine/Source/Runtime/Core/Public/Math/Vector.h
+++ b/Engine/Source/Runtime/Core/Public/Math/Vector.h
@@ -591,11 +591,19 @@ public:
 	/**
 	 * Gets a copy of this vector projected onto the input vector.
 	 *
-	 * @param A Vector to project onto, does not assume it is normalized.
+	 * @param A	Vector to project onto, does not assume it is normalized.
 	 * @return Projected vector.
 	 */
 	FORCEINLINE FVector ProjectOnTo( const FVector& A ) const ;
 
+	/**
+	 * Gets a copy of this vector projected onto the input vector, which is assumed to be unit length.
+	 * 
+	 * @param  Normal Vector to project onto (assumed to be unit length).
+	 * @return Projected vector.
+	 */
+	FORCEINLINE FVector ProjectOnToNormal(const FVector& Normal ) const;
+
 	/**
 	 * Return the FRotator corresponding to the direction that the vector
 	 * is pointing in.  Sets Yaw and Pitch to the proper numbers, and sets
@@ -710,7 +718,7 @@ public:
 	 *
 	 * @param Point The Point we are checking.
 	 * @param PlaneBase The Base Point in the plane.
-	 * @param PlaneNormal The Normal of the plane.
+	 * @param PlaneNormal The Normal of the plane (assumed to be unit length).
 	 * @return Signed distance between point and plane.
 	 */
 	static float PointPlaneDist( const FVector &Point, const FVector &PlaneBase, const FVector &PlaneNormal );
@@ -740,10 +748,19 @@ public:
 	*
 	* @param Point The point to project onto the plane
 	* @param PlaneBase Point on the plane
-	* @param PlaneNorm Normal of the plane
-	* @return Projection of Point onto plane ABC
+	* @param PlaneNorm Normal of the plane (assumed to be unit length).
+	* @return Projection of Point onto plane
 	*/
-	static FVector PointPlaneProject(const FVector& Point, const FVector& PlaneBase, const FVector& PlaneNorm);
+	static FVector PointPlaneProject(const FVector& Point, const FVector& PlaneBase, const FVector& PlaneNormal);
+
+	/**
+	 * Calculate the projection of a vector on the plane defined by PlaneNormal.
+	 * 
+	 * @param  V The vector to project onto the plane.
+	 * @param  PlaneNormal Normal of the plane (assumed to be unit length).
+	 * @return Projection of V onto plane.
+	 */
+	static FVector VectorPlaneProject(const FVector& V, const FVector& PlaneNormal);
 
 	/**
 	 * Euclidean distance between two points.
@@ -1025,6 +1042,11 @@ inline FVector FVector::PointPlaneProject(const FVector& Point, const FVector& P
 	return Point - FVector::PointPlaneDist(Point,PlaneBase,PlaneNorm) * PlaneNorm;
 }
 
+inline FVector FVector::VectorPlaneProject(const FVector& V, const FVector& PlaneNormal)
+{
+	return V - V.ProjectOnToNormal(PlaneNormal);
+}
+
 inline bool FVector::Parallel( const FVector &Normal1, const FVector &Normal2 )
 {
 	const float NormalDot = Normal1 | Normal2;
@@ -1586,6 +1608,11 @@ FORCEINLINE FVector FVector::ProjectOnTo( const FVector& A ) const
 	return (A * ((*this | A) / (A | A))); 
 }
 
+FORCEINLINE FVector FVector::ProjectOnToNormal(const FVector& Normal) const
+{
+	return (Normal * (*this | Normal));
+}
+
 
 FORCEINLINE bool FVector::ContainsNaN() const
 {
