Conditionally disabled various parts of internationalization related automation tests based on whether or not ICU is enabled.

diff --git a/Engine/Source/Runtime/Core/Private/Tests/Internationalization/DateTimeFormattingRulesTest.cpp b/Engine/Source/Runtime/Core/Private/Tests/Internationalization/DateTimeFormattingRulesTest.cpp
--- a/Engine/Source/Runtime/Core/Private/Tests/Internationalization/DateTimeFormattingRulesTest.cpp
+++ b/Engine/Source/Runtime/Core/Private/Tests/Internationalization/DateTimeFormattingRulesTest.cpp
@@ -21,13 +21,11 @@ namespace
 
 bool FDateTimeFormattingRulesTest::RunTest (const FString& Parameters)
 {
+#if UE_ENABLE_ICU
 	FInternationalization& I18N = FInternationalization::Get();
 
 	const FString OriginalCulture = I18N.GetCurrentCulture()->GetName();
 
-	//////////////////////////////////////////////////////////////////////////
-
-#if UE_ENABLE_ICU
 	const FDateTime UnixEpoch = FDateTime::FromUnixTimestamp(0);
 	const FDateTime UnixBillennium = FDateTime::FromUnixTimestamp(1000000000);
 	const FDateTime UnixOnes = FDateTime::FromUnixTimestamp(1111111111);
@@ -91,12 +89,11 @@ bool FDateTimeFormattingRulesTest::RunTest (const FString& Parameters)
 	Test( this, TEXT("Testing Date-Time"), FText::AsDateTime(TestDateTime, EDateTimeStyle::Medium, EDateTimeStyle::Medium, "GMT"), FText::FromString( TEXT("1990/06/13 12:34:56") ) );
 	Test( this, TEXT("Testing Date-Time"), FText::AsDateTime(TestDateTime, EDateTimeStyle::Long, EDateTimeStyle::Long, "GMT"), FText::FromString( TEXT("1990\x5E74") TEXT("6\x6708") TEXT("13\x65E5") TEXT(" 12:34:56 GMT") ) );
 	Test( this, TEXT("Testing Date-Time"), FText::AsDateTime(TestDateTime, EDateTimeStyle::Full, EDateTimeStyle::Full, "GMT"), FText::FromString( TEXT("1990\x5E74") TEXT("6\x6708") TEXT("13\x65E5\x6C34\x66DC\x65E5") TEXT(" ") TEXT("12\x6642") TEXT("34\x5206") TEXT("56\x79D2 GMT") ) );
+
+	I18N.SetCurrentCulture(OriginalCulture);
 #else
 	AddWarning("ICU is disabled thus locale-aware date/time formatting is disabled.");
 #endif
-	//////////////////////////////////////////////////////////////////////////
-
-	I18N.SetCurrentCulture(OriginalCulture);
 
 	return true;
 }
diff --git a/Engine/Source/Runtime/Core/Private/Tests/Internationalization/NumberFormatingRulesTest.cpp b/Engine/Source/Runtime/Core/Private/Tests/Internationalization/NumberFormatingRulesTest.cpp
--- a/Engine/Source/Runtime/Core/Private/Tests/Internationalization/NumberFormatingRulesTest.cpp
+++ b/Engine/Source/Runtime/Core/Private/Tests/Internationalization/NumberFormatingRulesTest.cpp
@@ -54,6 +54,7 @@ bool FNumberFormattingRulesTest::RunTest (const FString& Parameters)
  *	AsPercent	float 
  */
 
+#if UE_ENABLE_ICU
 	double DoubleValue = 12345678.901;
 	float FloatValue = 1234.567f;
 	double DoubleNegativeValue = -12345678.901;
@@ -93,7 +94,6 @@ bool FNumberFormattingRulesTest::RunTest (const FString& Parameters)
 	Test(this, TEXT("Convert a Negative int32 to a number formatted correct for en-US"),	FText::AsNumber(Int32NegativeValue),	FText::FromString(TEXT("-12,345")));
 	Test(this, TEXT("Convert a Negative int64 to a number formatted correct for en-US"),	FText::AsNumber(Int64NegativeValue),	FText::FromString(TEXT("-12,345")));
 
-#if UE_ENABLE_ICU
 	// Test Number Formatting Options
 	{
 		FNumberFormattingOptions NumberFormattingOptions;
@@ -206,9 +206,6 @@ bool FNumberFormattingRulesTest::RunTest (const FString& Parameters)
 		Test(this, TEXT("Round a Double to a number formatted correct for en-US using ToPositiveInfinity"),	FText::AsNumber( 0.5, &(NumberFormattingOptions)),	FText::FromString(TEXT( "1")));
 		Test(this, TEXT("Round a Double to a number formatted correct for en-US using ToPositiveInfinity"),	FText::AsNumber( 1.0, &(NumberFormattingOptions)),	FText::FromString(TEXT( "1")));
 	}
-#else
-	AddWarning("ICU is disabled thus custom number formatting options are disabled.");
-#endif
 
 	Test(this, TEXT("Convert a Double to a currency formatted correct for en-US"),			FText::AsCurrency(DoubleValue),			FText::FromString(TEXT("$12,345,678.90")));
 	Test(this, TEXT("Convert a Float to a currency formatted correct for en-US"),			FText::AsCurrency(FloatValue),			FText::FromString(TEXT("$1,234.57")));
@@ -251,7 +248,6 @@ bool FNumberFormattingRulesTest::RunTest (const FString& Parameters)
 	Test(this, TEXT("Convert a Negative int32 to a number formatted correct for hi-IN"),	FText::AsNumber(Int32NegativeValue),	FText::FromString(TEXT("-12,345")));
 	Test(this, TEXT("Convert a Negative int64 to a number formatted correct for hi-IN"),	FText::AsNumber(Int64NegativeValue),	FText::FromString(TEXT("-12,345")));
 
-#if UE_ENABLE_ICU
 	{
 		const FCulturePtr& InvariantCulture = I18N.GetInvariantCulture();
 		Test(this, TEXT("Convert a Double to a number formatted correct for hi-IN but as invariant"),				FText::AsNumber(DoubleValue, NULL, InvariantCulture),			FText::FromString(TEXT("12345678.901")));
@@ -271,9 +267,6 @@ bool FNumberFormattingRulesTest::RunTest (const FString& Parameters)
 		Test(this, TEXT("Convert a Negative int32 to a number formatted correct for hi-IN but as invariant"),		FText::AsNumber(Int32NegativeValue, NULL, InvariantCulture),	FText::FromString(TEXT("-12345")));
 		Test(this, TEXT("Convert a Negative int64 to a number formatted correct for hi-IN but as invariant"),		FText::AsNumber(Int64NegativeValue, NULL, InvariantCulture),	FText::FromString(TEXT("-12345")));
 	}
-#else
-	AddWarning("ICU is disabled thus invariant number formatting options are disabled.");
-#endif
 
 	Test(this, TEXT("Convert a Double to a currency formatted correct for hi-IN"),				FText::AsCurrency(DoubleValue),			FText::FromString(TEXT("\x20B9") TEXT("1,23,45,678.90")));
 	Test(this, TEXT("Convert a Float to a currency formatted correct for hi-IN"),				FText::AsCurrency(FloatValue),			FText::FromString(TEXT("\x20B9") TEXT("1,234.57")));
@@ -317,6 +310,9 @@ bool FNumberFormattingRulesTest::RunTest (const FString& Parameters)
 	{
 		AddError( TEXT("Number formatting functions should never produce a Culture Invariant Text") );
 	}
+#else
+	AddWarning("ICU is disabled thus locale-aware number formatting is disabled.");
+#endif
 
 	return true;
 }
diff --git a/Engine/Source/Runtime/Core/Private/Tests/Internationalization/TextTest.cpp b/Engine/Source/Runtime/Core/Private/Tests/Internationalization/TextTest.cpp
--- a/Engine/Source/Runtime/Core/Private/Tests/Internationalization/TextTest.cpp
+++ b/Engine/Source/Runtime/Core/Private/Tests/Internationalization/TextTest.cpp
@@ -59,7 +59,6 @@ bool FTextTest::RunTest (const FString& Parameters)
 
 	FText::SetEnableErrorCheckingResults(true);
 	FText::SetSuppressWarnings(true);
-	I18N.SetCurrentCulture("en-US");
 
 	FText ArgText0 = FText::FromString(TEXT("Arg0"));
 	FText ArgText1 = FText::FromString(TEXT("Arg1"));
@@ -157,8 +156,8 @@ bool FTextTest::RunTest (const FString& Parameters)
 
 	{
 		FFormatNamedArguments Arguments;
-		Arguments.Add( TEXT("Age"), 23 );
-		Arguments.Add( TEXT("Height"), 68 );
+		Arguments.Add( TEXT("Age"), FText::FromString( TEXT("23") ) );
+		Arguments.Add( TEXT("Height"), FText::FromString( TEXT("68") ) );
 		Arguments.Add( TEXT("Gender"), FText::FromString( TEXT("male") ) );
 		Arguments.Add( TEXT("Name"), FText::FromString( TEXT("Saul") ) );
 
@@ -224,39 +223,38 @@ bool FTextTest::RunTest (const FString& Parameters)
 
 		// Not using all the arguments is okay.
 		TestText = INVTEXT("My name is {Name}.");
-		TEST( TestText.ToString(), FText::Format( TestText, ArgumentList), INVTEXT("My name is Saul.") );
+		TEST( TestText.ToString(), FText::Format(TestText, ArgumentList), INVTEXT("My name is Saul.") );
 		TestText = INVTEXT("My age is {Age}.");
-		TEST( TestText.ToString(), FText::Format( TestText, ArgumentList), INVTEXT("My age is 23.") );
+		TEST( TestText.ToString(), FText::Format(TestText, ArgumentList), INVTEXT("My age is 23.") );
 		TestText = INVTEXT("My gender is {Gender}.");
-		TEST( TestText.ToString(), FText::Format( TestText, ArgumentList), INVTEXT("My gender is male.") );
+		TEST( TestText.ToString(), FText::Format(TestText, ArgumentList), INVTEXT("My gender is male.") );
 		TestText = INVTEXT("My height is {Height}.");
-		TEST( TestText.ToString(), FText::Format( TestText, ArgumentList), INVTEXT("My height is 68.") );
+		TEST( TestText.ToString(), FText::Format(TestText, ArgumentList), INVTEXT("My height is 68.") );
 
 		// Using arguments out of order is okay.
 		TestText = INVTEXT("My name is {Name}. My age is {Age}. My gender is {Gender}.");
-		TEST( TestText.ToString(), FText::Format( TestText, ArgumentList), INVTEXT("My name is Saul. My age is 23. My gender is male.") );
+		TEST( TestText.ToString(), FText::Format(TestText, ArgumentList), INVTEXT("My name is Saul. My age is 23. My gender is male.") );
 		TestText = INVTEXT("My age is {Age}. My gender is {Gender}. My name is {Name}.");
-		TEST( TestText.ToString(), FText::Format( TestText, ArgumentList), INVTEXT("My age is 23. My gender is male. My name is Saul.") );
+		TEST( TestText.ToString(), FText::Format(TestText, ArgumentList), INVTEXT("My age is 23. My gender is male. My name is Saul.") );
 		TestText = INVTEXT("My gender is {Gender}. My name is {Name}. My age is {Age}.");
-		TEST( TestText.ToString(), FText::Format( TestText, ArgumentList), INVTEXT("My gender is male. My name is Saul. My age is 23.") );
+		TEST( TestText.ToString(), FText::Format(TestText, ArgumentList), INVTEXT("My gender is male. My name is Saul. My age is 23.") );
 		TestText = INVTEXT("My gender is {Gender}. My age is {Age}. My name is {Name}.");
-		TEST( TestText.ToString(), FText::Format( TestText, ArgumentList), INVTEXT("My gender is male. My age is 23. My name is Saul.") );
+		TEST( TestText.ToString(), FText::Format(TestText, ArgumentList), INVTEXT("My gender is male. My age is 23. My name is Saul.") );
 		TestText = INVTEXT("My age is {Age}. My name is {Name}. My gender is {Gender}.");
-		TEST( TestText.ToString(), FText::Format( TestText, ArgumentList), INVTEXT("My age is 23. My name is Saul. My gender is male.") );
+		TEST( TestText.ToString(), FText::Format(TestText, ArgumentList), INVTEXT("My age is 23. My name is Saul. My gender is male.") );
 		TestText = INVTEXT("My name is {Name}. My gender is {Gender}. My age is {Age}.");
-		TEST( TestText.ToString(), FText::Format( TestText, ArgumentList), INVTEXT("My name is Saul. My gender is male. My age is 23.") );
+		TEST( TestText.ToString(), FText::Format(TestText, ArgumentList), INVTEXT("My name is Saul. My gender is male. My age is 23.") );
 
 		// Reusing arguments is okay.
 		TestText = INVTEXT("If my age is {Age}, I have been alive for {Age} year(s).");
-		TEST( TestText.ToString(), FText::Format( TestText, ArgumentList), INVTEXT("If my age is 23, I have been alive for 23 year(s).") );
+		TEST( TestText.ToString(), FText::Format(TestText, ArgumentList), INVTEXT("If my age is 23, I have been alive for 23 year(s).") );
 
 		// Not providing an argument leaves the parameter as text.
 		TestText = INVTEXT("What... is the air-speed velocity of an unladen swallow? {AirSpeedOfAnUnladenSwallow}.");
-		TEST( TestText.ToString(), FText::Format( TestText, ArgumentList), INVTEXT("What... is the air-speed velocity of an unladen swallow? {AirSpeedOfAnUnladenSwallow}.") );
+		TEST( TestText.ToString(), FText::Format(TestText, ArgumentList), INVTEXT("What... is the air-speed velocity of an unladen swallow? {AirSpeedOfAnUnladenSwallow}.") );
 	}
 
 #undef TEST
-
 #define TEST( Pattern, Actual, Expected ) TestPatternParameterEnumeration(*this, Pattern, Actual, Expected)
 
 	TArray<FString> ActualArguments;
@@ -300,6 +298,7 @@ bool FTextTest::RunTest (const FString& Parameters)
 	FText::SetSuppressWarnings(true);
 
 #if UE_ENABLE_ICU
+	I18N.SetCurrentCulture("en-US");
 #define TEST( A, B, ComparisonLevel ) if( !(FText::FromString(A)).EqualTo(FText::FromString(B), (ComparisonLevel)) ) AddError(FString::Printf(TEXT("Testing comparison of equivalent characters with comparison level (%s). - A=%s B=%s"),TEXT(#ComparisonLevel),(A),(B)))
 
 	// Basic sanity checks
@@ -398,7 +397,11 @@ bool FTextTest::RunTest (const FString& Parameters)
 			}
 		}
 	}
+#else
+	AddWarning("ICU is disabled thus locale-aware string collation is disabled.");
+#endif
 
+#if UE_ENABLE_ICU
 	{
 		I18N.SetCurrentCulture(OriginalCulture);
 
@@ -585,7 +588,7 @@ bool FTextTest::RunTest (const FString& Parameters)
 		}
 	}
 #else
-	AddWarning("ICU is disabled thus locale-aware string collation is disabled.");
+	AddWarning("ICU is disabled thus locale-aware formatting needed in rebuilding source text from history is disabled.");
 #endif
 
 	FText::SetEnableErrorCheckingResults(OriginalEnableErrorCheckingValue);
