Fix for context menu staying up on Linux.

- Ignore focus notifications altogether (for now?) and track clicks only.
- Also fix the error of never passing any other value than Deactivate.
- Remove logspam.

diff --git a/Engine/Source/Runtime/Core/Private/Linux/LinuxApplication.cpp b/Engine/Source/Runtime/Core/Private/Linux/LinuxApplication.cpp
--- a/Engine/Source/Runtime/Core/Private/Linux/LinuxApplication.cpp
+++ b/Engine/Source/Runtime/Core/Private/Linux/LinuxApplication.cpp
@@ -177,7 +177,7 @@ void FLinuxApplication::TrackActivationChanges(const TSharedPtr<FLinuxWindow> Wi
 {
 	bool bNeedToNotify = (Window != CurrentlyActiveWindow) || Event == EWindowActivation::Deactivate;
 	
-	UE_LOG(LogLinuxWindow, Log, TEXT("TrackActivationChanges: %s (Window: %p, CurrentlyActiveWindow: %p, Event: %d)"),
+	UE_LOG(LogLinuxWindow, Verbose, TEXT("TrackActivationChanges: bNeedToNotify=%s (Window: %p, CurrentlyActiveWindow: %p, Event: %d)"),
 		bNeedToNotify ? TEXT("true") : TEXT("false"),
 		Window.Get(),
 		CurrentlyActiveWindow.Get(),
@@ -199,7 +199,7 @@ void FLinuxApplication::TrackActivationChanges(const TSharedPtr<FLinuxWindow> Wi
 			
 			if (CurrentlyActiveWindow.IsValid())
 			{
-				MessageHandler->OnWindowActivationChanged(CurrentlyActiveWindow.ToSharedRef(), EWindowActivation::Deactivate);
+				MessageHandler->OnWindowActivationChanged(CurrentlyActiveWindow.ToSharedRef(), Event);
 			}
 		}
 		else
@@ -700,24 +700,8 @@ void FLinuxApplication::ProcessDeferredMessage( SDL_Event Event )
 					}
 					break;
 
-				case SDL_WINDOWEVENT_FOCUS_GAINED:
-					{
-						if (CurrentEventWindow.IsValid())
-						{
-							TrackActivationChanges(CurrentEventWindow, EWindowActivation::Activate);                            
-						}
-					}
-					break;
-
-				case SDL_WINDOWEVENT_FOCUS_LOST:
-					{
-						if (CurrentEventWindow.IsValid())
-						{
-							TrackActivationChanges(CurrentEventWindow, EWindowActivation::Deactivate);
-						}
-					}
-				break;
-
+				case SDL_WINDOWEVENT_FOCUS_GAINED:	// seems to be spurious and does not always reflect actual focus changes, ignore for now
+				case SDL_WINDOWEVENT_FOCUS_LOST:	// seems to be spurious and does not always reflect actual focus changes, ignore for now
 				case SDL_WINDOWEVENT_HIDDEN:		// intended fall-through
 				case SDL_WINDOWEVENT_EXPOSED:		// intended fall-through
 				case SDL_WINDOWEVENT_MINIMIZED:		// intended fall-through
