Resolved non-Unity build error with new ICU half of internationalization implementation.

diff --git a/Engine/Source/Runtime/Core/Private/Internationalization/ICUInternationalization.cpp b/Engine/Source/Runtime/Core/Private/Internationalization/ICUInternationalization.cpp
--- a/Engine/Source/Runtime/Core/Private/Internationalization/ICUInternationalization.cpp
+++ b/Engine/Source/Runtime/Core/Private/Internationalization/ICUInternationalization.cpp
@@ -10,7 +10,6 @@
 #include <unicode/timezone.h>
 #include <unicode/uclean.h>
 #include <unicode/udata.h>
-#include <unicode/umachine.h>
 
 namespace
 {
diff --git a/Engine/Source/Runtime/Core/Private/Internationalization/ICUInternationalization.h b/Engine/Source/Runtime/Core/Private/Internationalization/ICUInternationalization.h
--- a/Engine/Source/Runtime/Core/Private/Internationalization/ICUInternationalization.h
+++ b/Engine/Source/Runtime/Core/Private/Internationalization/ICUInternationalization.h
@@ -3,6 +3,8 @@
 
 #if UE_ENABLE_ICU
 
+#include <unicode/umachine.h>
+
 class FICUInternationalization
 {
 public:
