UE4 - task graph randomization test code

diff --git a/Engine/Source/Runtime/Core/Private/Async/TaskGraph.cpp b/Engine/Source/Runtime/Core/Private/Async/TaskGraph.cpp
--- a/Engine/Source/Runtime/Core/Private/Async/TaskGraph.cpp
+++ b/Engine/Source/Runtime/Core/Private/Async/TaskGraph.cpp
@@ -24,6 +24,100 @@ namespace ENamedThreads
 **/
 static class FTaskGraphImplementation* TaskGraphImplementationSingleton = NULL;
 
+#if !UE_BUILD_SHIPPING && !UE_BUILD_TEST
+
+static struct FChaosMode 
+{
+	enum 
+	{
+		NumSamples = 45771
+	};
+	FThreadSafeCounter Current;
+	float DelayTimes[NumSamples + 1]; 
+	int32 Enabled;
+
+	FChaosMode()
+		: Enabled(0)
+	{
+		FRandomStream Stream((int32)FPlatformTime::Cycles());
+		for (int32 Index = 0; Index < NumSamples; Index++)
+		{
+			DelayTimes[Index] = Stream.GetFraction();
+		}
+		// ave = .5
+		for (int32 Cube = 0; Cube < 2; Cube++)
+		{
+			for (int32 Index = 0; Index < NumSamples; Index++)
+			{
+				DelayTimes[Index] *= Stream.GetFraction();
+			}
+		}
+		// ave = 1/8
+		for (int32 Index = 0; Index < NumSamples; Index++)
+		{
+			DelayTimes[Index] *= 0.0001f;
+		}
+		// ave = 0.0000125s
+		for (int32 Zeros = 0; Zeros < NumSamples / 2; Zeros++)
+		{
+			int32 Index = Stream.RandHelper(NumSamples);
+			DelayTimes[Index] = 0.0f;
+		}
+		// half the samples are now zero
+		for (int32 Zeros = 0; Zeros < NumSamples / 100; Zeros++)
+		{
+			int32 Index = Stream.RandHelper(NumSamples);
+			DelayTimes[Index] = .005f;
+		}
+		// 1% of the samples are 5ms
+	}
+	FORCEINLINE void Delay()
+	{
+		if (Enabled)
+		{
+			uint32 MyIndex = (uint32)Current.Increment();
+			MyIndex %= NumSamples;
+			float DelayS = DelayTimes[MyIndex];
+			if (DelayS > 0.0f)
+			{
+				FPlatformProcess::Sleep(DelayS);
+			}
+		}
+	}
+} GChaosMode;
+
+static void EnableRandomizedThreads(const TArray<FString>& Args)
+{
+	GChaosMode.Enabled = !GChaosMode.Enabled;
+	if (GChaosMode.Enabled)
+	{
+		UE_LOG(LogConsoleResponse, Display, TEXT("Random sleeps are enabled."));
+	}
+	else
+	{
+		UE_LOG(LogConsoleResponse, Display, TEXT("Random sleeps are disabled."));
+	}
+}
+
+static FAutoConsoleCommand TestRandomizedThreadsCommand(
+	TEXT("TaskGraph.Randomize"),
+	TEXT("Useful for debugging, adds random sleeps throughout the task graph."),
+	FConsoleCommandWithArgsDelegate::CreateStatic(&EnableRandomizedThreads)
+	);
+
+FORCEINLINE void TestRandomizedThreads()
+{
+	GChaosMode.Delay();
+}
+
+#else
+
+FORCEINLINE void TestRandomizedThreads()
+{
+}
+
+#endif
+
 /** 
  *	FTaskQueue
  *	High performance, SINGLE threaded, FIFO task queue for the private queue on named threads.
@@ -319,6 +413,7 @@ public:
 			}
 			if (Task)
 			{
+				TestRandomizedThreads();
 #if STATS
 				if (!bTasksOpen && FThreadStats::IsCollectingData(StatName))
 				{
@@ -330,6 +425,7 @@ public:
 				{
 					Task->Execute(NewTasks, ENamedThreads::Type(ThreadId | (QueueIndex << ENamedThreads::QueueIndexShift)));
 				}
+				TestRandomizedThreads();
 			}
 			else
 			{
@@ -410,6 +506,7 @@ public:
 	 **/
 	bool EnqueueFromOtherThread(int32 QueueIndex, FBaseGraphTask* Task)
 	{
+		TestRandomizedThreads();
 		checkThreadGraph(Queue(QueueIndex).StallRestartEvent); // make sure we are started up
 		bool bWasReopenedByMe = Queue(QueueIndex).IncomingQueue.ReopenIfClosedAndPush(Task);
 		if (bWasReopenedByMe)
@@ -537,6 +634,8 @@ private:
 	 */
 	bool Stall(int32 QueueIndex, TStatId StallStatId, bool bCountAsStall)
 	{
+		TestRandomizedThreads();
+
 		checkThreadGraph(Queue(QueueIndex).StallRestartEvent); // make sure we are started up
 		if (Queue(QueueIndex).QuitWhenIdle.GetValue() == 0)
 		{
@@ -552,7 +651,9 @@ private:
 					int32 NewValue = IsStalled.Increment();
 					NotifyStalling();
 					checkThreadGraph(NewValue == 1); // there should be no concurrent calls to Stall!
+					TestRandomizedThreads();
 					Queue(QueueIndex).StallRestartEvent->Wait(MAX_uint32, bCountAsStall);
+					TestRandomizedThreads();
 					NewValue = IsStalled.Decrement();
 					checkThreadGraph(NewValue == 0); // there should be no concurrent calls to Stall!
 					return true;
@@ -738,6 +839,7 @@ public:
 	**/
 	virtual void QueueTask(FBaseGraphTask* Task, ENamedThreads::Type ThreadToExecuteOn, ENamedThreads::Type CurrentThreadIfKnown = ENamedThreads::AnyThread) override
 	{
+		TestRandomizedThreads();
 		checkThreadGraph(NextUnnamedThreadMod);
 		if (CurrentThreadIfKnown == ENamedThreads::AnyThread)
 		{
@@ -901,6 +1003,7 @@ public:
 		}
 		if (!bAnyPending)
 		{
+			TestRandomizedThreads();
 			InEvent->Trigger();
 			return;
 		}
@@ -916,6 +1019,7 @@ public:
 	**/
 	FBaseGraphTask* FindWork(ENamedThreads::Type ThreadInNeed)
 	{
+		TestRandomizedThreads();
 		{
 			FBaseGraphTask* Task = SortedAnyThreadTasks.Pop();
 			if (Task)
