Core - GetMachineId/GetEpicAccountId is cached

diff --git a/Engine/Source/Runtime/Core/Private/GenericPlatform/GenericPlatformMisc.cpp b/Engine/Source/Runtime/Core/Private/GenericPlatform/GenericPlatformMisc.cpp
--- a/Engine/Source/Runtime/Core/Private/GenericPlatform/GenericPlatformMisc.cpp
+++ b/Engine/Source/Runtime/Core/Private/GenericPlatform/GenericPlatformMisc.cpp
@@ -801,11 +801,11 @@ bool FGenericPlatformMisc::IsRunningOnBattery()
 
 FGuid FGenericPlatformMisc::GetMachineId()
 {
-	FGuid MachineId;
+	static FGuid MachineId;
 	FString MachineIdStr;
 
 	// Check to see if we already have a valid machine ID to use
-	if( !FPlatformMisc::GetStoredValue( TEXT( "Epic Games" ), TEXT( "Unreal Engine/Identifiers" ), TEXT( "MachineId" ), MachineIdStr ) || !FGuid::Parse( MachineIdStr, MachineId ) )
+	if( !MachineId.IsValid() && (!FPlatformMisc::GetStoredValue( TEXT( "Epic Games" ), TEXT( "Unreal Engine/Identifiers" ), TEXT( "MachineId" ), MachineIdStr ) || !FGuid::Parse( MachineIdStr, MachineId )) )
 	{
 		// No valid machine ID, generate and save a new one
 		MachineId = FGuid::NewGuid();
@@ -821,16 +821,25 @@ FGuid FGenericPlatformMisc::GetMachineId()
 	return MachineId;
 }
 
+namespace
+{
+	/** Cached EpicAccountId. */
+	FString EpicAccountId;
+}
+
 FString FGenericPlatformMisc::GetEpicAccountId()
 {
-	FString AccountId;
-	FPlatformMisc::GetStoredValue( TEXT( "Epic Games" ), TEXT( "Unreal Engine/Identifiers" ), TEXT( "AccountId" ), AccountId );
-	return AccountId;
+	if( EpicAccountId.IsEmpty() )
+	{
+		FPlatformMisc::GetStoredValue( TEXT( "Epic Games" ), TEXT( "Unreal Engine/Identifiers" ), TEXT( "AccountId" ), EpicAccountId );
+	}
+	return EpicAccountId;
 }
 
 void FGenericPlatformMisc::SetEpicAccountId( const FString& AccountId )
 {
 	FPlatformMisc::SetStoredValue( TEXT( "Epic Games" ), TEXT( "Unreal Engine/Identifiers" ), TEXT( "AccountId" ), AccountId );
+	EpicAccountId = AccountId;
 }
 
 const TCHAR* FGenericPlatformMisc::GetEngineMode()
