Core - GetOSVersions is cached

diff --git a/Engine/Source/Runtime/Core/Private/Windows/WindowsPlatformMisc.cpp b/Engine/Source/Runtime/Core/Private/Windows/WindowsPlatformMisc.cpp
--- a/Engine/Source/Runtime/Core/Private/Windows/WindowsPlatformMisc.cpp
+++ b/Engine/Source/Runtime/Core/Private/Windows/WindowsPlatformMisc.cpp
@@ -361,10 +361,27 @@ static void InitSHAHashes()
 	}
 }
 
+/**
+ *	Sets process memory limit using the job object, may fail under some situation like when Program Compatibility Assistant is enabled.
+ *	Debugging purpose only.
+ */
+static void SetProcessMemoryLimit( SIZE_T ProcessMemoryLimitMB )
+{
+	HANDLE JobObject = ::CreateJobObject(nullptr, TEXT("UE4-JobObject"));
+	JOBOBJECT_EXTENDED_LIMIT_INFORMATION JobLimitInfo = {0};
+	JobLimitInfo.ProcessMemoryLimit = 1024*1024*ProcessMemoryLimitMB;
+	JobLimitInfo.BasicLimitInformation.LimitFlags = JOB_OBJECT_LIMIT_PROCESS_MEMORY;
+	const BOOL bSetJob = ::SetInformationJobObject(JobObject,JobObjectExtendedLimitInformation,&JobLimitInfo,sizeof(JobLimitInfo));
 
+	const BOOL bAssign = ::AssignProcessToJobObject(JobObject, GetCurrentProcess());
+}
 
 void FWindowsPlatformMisc::PlatformPreInit()
 {
+	//SetProcessMemoryLimit( 92 );
+
+	FGenericPlatformMisc::PlatformPreInit();
+
 	// Use our own handler for pure virtuals being called.
 	DefaultPureCallHandler = _set_purecall_handler( PureCallHandler );
 
@@ -2175,7 +2192,16 @@ FString FWindowsPlatformMisc::GetPrimaryGPUBrand()
 
 void FWindowsPlatformMisc::GetOSVersions( FString& out_OSVersionLabel, FString& out_OSSubVersionLabel )
 {
-	FWindowsOSVersionHelper::GetOSVersions( out_OSVersionLabel, out_OSSubVersionLabel );
+	static FString OSVersionLabel;
+	static FString OSSubVersionLabel;
+
+	if( OSVersionLabel.IsEmpty() && OSSubVersionLabel.IsEmpty() )
+	{
+		FWindowsOSVersionHelper::GetOSVersions( OSVersionLabel, OSSubVersionLabel );
+	}
+
+	out_OSVersionLabel = OSVersionLabel;
+	out_OSSubVersionLabel = OSSubVersionLabel;
 }
 
 
