Make CoreFoundation and all other Core*-level APIs in Mac OS X go through our allocation routines by creating a custom CFAllocatorRef when setting up our base allocator & then setting it as the default CFAllocator. This means CFString & friends will benefit from faster allocators & will use our special crash heaps during signal handlers.

diff --git a/Engine/Source/Runtime/Core/Private/Mac/MacPlatformMemory.cpp b/Engine/Source/Runtime/Core/Private/Mac/MacPlatformMemory.cpp
--- a/Engine/Source/Runtime/Core/Private/Mac/MacPlatformMemory.cpp
+++ b/Engine/Source/Runtime/Core/Private/Mac/MacPlatformMemory.cpp
@@ -9,6 +9,26 @@
 #include "MallocAnsi.h"
 #include "MallocBinned.h"
 
+static void* FMacAllocatorAllocate(CFIndex AllocSize, CFOptionFlags Hint, void* Info)
+{
+	return FMemory::Malloc(AllocSize);
+}
+
+static void* FMacAllocatorReallocate(void* Ptr, CFIndex Newsize, CFOptionFlags Hint, void* Info)
+{
+	return FMemory::Realloc(Ptr, Newsize);
+}
+
+static void FMacAllocatorDeallocate(void* Ptr, void* Info)
+{
+	return FMemory::Free(Ptr);
+}
+
+static CFIndex FMacAllocatorPreferredSize(CFIndex Size, CFOptionFlags Hint, void* Info)
+{
+	return FMemory::MallocQuantizeSize(Size);
+}
+
 void FMacPlatformMemory::Init()
 {
 	const FPlatformMemoryConstants& MemoryConstants = FPlatformMemory::GetConstants();
@@ -21,13 +41,32 @@ void FMacPlatformMemory::Init()
 
 FMalloc* FMacPlatformMemory::BaseAllocator()
 {
+	FMalloc* Malloc = nullptr;
+	
 #if FORCE_ANSI_ALLOCATOR
-	return new FMallocAnsi();
+	Malloc = new FMallocAnsi();
 #elif (WITH_EDITORONLY_DATA || IS_PROGRAM) && TBB_ALLOCATOR_ALLOWED
-	return new FMallocTBB();
+	Malloc = new FMallocTBB();
 #else
-	return new FMallocBinned((uint32)(GetConstants().PageSize&MAX_uint32), 0x100000000);
+	Malloc = new FMallocBinned((uint32)(GetConstants().PageSize&MAX_uint32), 0x100000000);
 #endif
+	
+	// Configure CoreFoundation's default allocator to use our allocation routines too.
+	CFAllocatorContext AllocatorContext;
+	AllocatorContext.version = 0;
+	AllocatorContext.info = nullptr;
+	AllocatorContext.retain = nullptr;
+	AllocatorContext.release = nullptr;
+	AllocatorContext.copyDescription = nullptr;
+	AllocatorContext.allocate = &FMacAllocatorAllocate;
+	AllocatorContext.reallocate = &FMacAllocatorReallocate;
+	AllocatorContext.deallocate = &FMacAllocatorDeallocate;
+	AllocatorContext.preferredSize = &FMacAllocatorPreferredSize;
+	
+	CFAllocatorRef Alloc = CFAllocatorCreate(kCFAllocatorDefault, &AllocatorContext);
+	CFAllocatorSetDefault(Alloc);
+	
+	return Malloc;
 }
 
 FPlatformMemoryStats FMacPlatformMemory::GetStats()
