Fix some faulty FMacWindow event order assumptions to ensure correct behaviour. In FMacWindow::Show always set the window as visible - it will be shortly if it isn't already. Don't resize a window if it is already the same size, infinite recursion may occur in some cases. Remove some obsolete code that tried to ensure that a window became visible in Show - we don't need this anymore.

diff --git a/Engine/Source/Runtime/Core/Private/Mac/MacWindow.cpp b/Engine/Source/Runtime/Core/Private/Mac/MacWindow.cpp
--- a/Engine/Source/Runtime/Core/Private/Mac/MacWindow.cpp
+++ b/Engine/Source/Runtime/Core/Private/Mac/MacWindow.cpp
@@ -231,35 +231,40 @@ void FMacWindow::ReshapeWindow( int32 X, int32 Y, int32 Width, int32 Height )
 		
 		if(WindowMode == EWindowMode::Windowed || WindowMode == EWindowMode::WindowedFullscreen)
 		{
-			MainThreadCall(^{
-				SCOPED_AUTORELEASE_POOL;
-				BOOL DisplayIfNeeded = (WindowMode == EWindowMode::Windowed);
-				
-				const int32 InvertedY = FPlatformMisc::ConvertSlateYPositionToCocoa(Y) - Height + 1;
-				if (Definition->HasOSWindowBorder)
-				{
-					[WindowHandle setFrame: [WindowHandle frameRectForContentRect: NSMakeRect(X, InvertedY, FMath::Max(Width, 1), FMath::Max(Height, 1))] display:DisplayIfNeeded];
-				}
-				else
-				{
-					[WindowHandle setFrame: NSMakeRect(X, InvertedY, FMath::Max(Width, 1), FMath::Max(Height, 1)) display:DisplayIfNeeded];
-				}
-				
-				// Force resize back to screen size in fullscreen - not ideally pretty but means we don't
-				// have to subvert the OS X or UE fullscreen handling events elsewhere.
-				if(WindowMode != EWindowMode::Windowed)
-				{
-					[WindowHandle setFrame: [WindowHandle screen].frame display:YES];
-				}
-				
-				WindowHandle->bZoomed = [WindowHandle isZoomed];
-			}, UE4ResizeEventMode, true);
+			const int32 InvertedY = FPlatformMisc::ConvertSlateYPositionToCocoa(Y) - Height + 1;
+			NSRect Rect = NSMakeRect(X, InvertedY, FMath::Max(Width, 1), FMath::Max(Height, 1));
+			if (Definition->HasOSWindowBorder)
+			{
+				Rect = [WindowHandle frameRectForContentRect: Rect];
+			}
+			
+			if ( !NSEqualRects([WindowHandle frame], Rect) )
+			{
+				MainThreadCall(^{
+					SCOPED_AUTORELEASE_POOL;
+					BOOL DisplayIfNeeded = (WindowMode == EWindowMode::Windowed);
+					
+					[WindowHandle setFrame: Rect display:DisplayIfNeeded];
+					
+					// Force resize back to screen size in fullscreen - not ideally pretty but means we don't
+					// have to subvert the OS X or UE fullscreen handling events elsewhere.
+					if(WindowMode != EWindowMode::Windowed)
+					{
+						[WindowHandle setFrame: [WindowHandle screen].frame display:YES];
+					}
+					
+					WindowHandle->bZoomed = [WindowHandle isZoomed];
+				}, UE4ResizeEventMode, true);
+			}
 		}
 		else
 		{
-			WindowHandle.PreFullScreenRect = NSMakeRect(WindowHandle.PreFullScreenRect.origin.x, WindowHandle.PreFullScreenRect.origin.y, (CGFloat)Width, (CGFloat)Height);
-			
-			FMacEvent::SendToGameRunLoop([NSNotification notificationWithName:NSWindowDidResizeNotification object:WindowHandle], WindowHandle, EMacEventSendMethod::Sync, @[ NSDefaultRunLoopMode, UE4ResizeEventMode, UE4ShowEventMode, UE4FullscreenEventMode ]);
+			NSRect NewRect = NSMakeRect(WindowHandle.PreFullScreenRect.origin.x, WindowHandle.PreFullScreenRect.origin.y, (CGFloat)Width, (CGFloat)Height);
+			if ( !NSEqualRects(WindowHandle.PreFullScreenRect, NewRect) )
+			{
+				WindowHandle.PreFullScreenRect = NewRect;
+				FMacEvent::SendToGameRunLoop([NSNotification notificationWithName:NSWindowDidResizeNotification object:WindowHandle], WindowHandle, EMacEventSendMethod::Sync, @[ NSDefaultRunLoopMode, UE4ResizeEventMode, UE4ShowEventMode, UE4FullscreenEventMode ]);
+			}
 		}
 		
 		MessageHandler->FinishedReshapingWindow( SharedThis( this ) );
@@ -387,22 +392,7 @@ void FMacWindow::Show()
 			[WindowHandle orderFrontAndMakeMain:bMakeMainAndKey andKey:bMakeMainAndKey];
 		}, UE4ShowEventMode, true);
 		
-		bIsVisible = [WindowHandle isVisible];
-		static bool bCannotRecurse = false;
-		if(bCannotRecurse)
-		{
-			bCannotRecurse = true;
-			// For the movie code we don't pump the event loop - so this function must do so before it returns
-			// or the window will never become main/key!
-			bool isMainAndKey = false;
-			do
-			{
-				FPlatformMisc::PumpMessages(true);
-				bIsVisible = [WindowHandle isVisible];
-				isMainAndKey = [WindowHandle isKeyWindow];
-			} while(!bIsVisible && isMainAndKey != bMakeMainAndKey);
-			bCannotRecurse = false;
-		}
+		bIsVisible = true;
 	}
 }
 
