Fixed a problem with game thread sometime freezing on tooltip update when opening menus on Mac

diff --git a/Engine/Source/Runtime/Core/Private/Mac/CocoaThread.cpp b/Engine/Source/Runtime/Core/Private/Mac/CocoaThread.cpp
--- a/Engine/Source/Runtime/Core/Private/Mac/CocoaThread.cpp
+++ b/Engine/Source/Runtime/Core/Private/Mac/CocoaThread.cpp
@@ -405,7 +405,7 @@ void MainThreadCall(dispatch_block_t Block, NSString* WaitMode, bool const bWait
 	if ( [NSThread mainThread] != [NSThread currentThread] )
 	{
 		FCocoaRunLoopSource& MainRunLoopSource = FCocoaRunLoopSource::GetMainRunLoopSource();
-		PerformBlockOnThread(MainRunLoopSource, Block, @[ NSDefaultRunLoopMode, NSModalPanelRunLoopMode ], WaitMode, bWait);
+		PerformBlockOnThread(MainRunLoopSource, Block, @[ NSDefaultRunLoopMode, NSModalPanelRunLoopMode, NSEventTrackingRunLoopMode ], WaitMode, bWait);
 	}
 	else
 	{
