CrashReport - Added missing override

diff --git a/Engine/Source/Runtime/Core/Private/Windows/WindowsPlatformExceptionHandling.cpp b/Engine/Source/Runtime/Core/Private/Windows/WindowsPlatformExceptionHandling.cpp
--- a/Engine/Source/Runtime/Core/Private/Windows/WindowsPlatformExceptionHandling.cpp
+++ b/Engine/Source/Runtime/Core/Private/Windows/WindowsPlatformExceptionHandling.cpp
@@ -194,7 +194,7 @@ void AddMiniDump(HREPORT ReportHandle, EXCEPTION_POINTERS* ExceptionInfo)
 
 struct FWindowsPlatformCrashContext : public FGenericCrashContext
 {
-	void AddPlatformSpecificProperties()
+	void AddPlatformSpecificProperties() override
 	{
 		AddCrashProperty( TEXT( "Platform.IsRunningWindows" ), 1 );
 	}
