Missing file from CL# 2322809

diff --git a/Engine/Source/Runtime/Core/Private/IOS/IOSPlatformFramePacer.cpp b/Engine/Source/Runtime/Core/Private/IOS/IOSPlatformFramePacer.cpp
--- a/Engine/Source/Runtime/Core/Private/IOS/IOSPlatformFramePacer.cpp
+++ b/Engine/Source/Runtime/Core/Private/IOS/IOSPlatformFramePacer.cpp
@@ -52,27 +52,46 @@ static TArray<FEvent*> ListeningEvents;
  * FIOSPlatformRHIFramePacer implementation
  *******************************************************************/
 
+
+namespace IOSDisplayConstants
+{
+    const uint32 MaxRefreshRate = 60;
+}
+
 uint32 FIOSPlatformRHIFramePacer::FrameInterval = 1;
 FIOSFramePacer* FIOSPlatformRHIFramePacer::FramePacer = nil;
 
+
 bool FIOSPlatformRHIFramePacer::IsEnabled()
 {
-    bool bIsRHIFramePacerEnabled = false;
-	GConfig->GetBool( TEXT( "PlatformFramePacer" ), TEXT( "IsEnabled" ), bIsRHIFramePacerEnabled, GEngineIni );
+    static bool bIsRHIFramePacerEnabled = false;
     
+    FString FrameRateLockAsEnum;
+    GConfig->GetString( TEXT("/Script/IOSRuntimeSettings.IOSRuntimeSettings"), TEXT("FrameRateLock"), FrameRateLockAsEnum, GEngineIni );
+    
+    uint32 FrameRateLock = 0;
+    if( FParse::Value(*FrameRateLockAsEnum, TEXT("PUFRL_"), FrameRateLock))
+    {
+        if( !bIsRHIFramePacerEnabled && FrameRateLock > 0 )
+        {
+            check( (IOSDisplayConstants::MaxRefreshRate % FrameRateLock) == 0 );
+            FrameInterval = IOSDisplayConstants::MaxRefreshRate / FrameRateLock;
+        
+            bIsRHIFramePacerEnabled = (FrameInterval > 0);
+        }
+    }
+    
+
     
     return bIsRHIFramePacerEnabled;
 }
 
-
-void FIOSPlatformRHIFramePacer::InitWithEvent(FEvent* TriggeredEvent, uint32 InFrameInterval)
+void FIOSPlatformRHIFramePacer::InitWithEvent(FEvent* TriggeredEvent)
 {
-    FrameInterval = InFrameInterval;
-    
     // Create display link thread
-	FramePacer = [[FIOSFramePacer alloc] init];
-	[NSThread detachNewThreadSelector:@selector(run:) toTarget:FramePacer withObject:nil];
-    
+    FramePacer = [[FIOSFramePacer alloc] init];
+    [NSThread detachNewThreadSelector:@selector(run:) toTarget:FramePacer withObject:nil];
+        
     // Only one supported for now, we may want more eventually.
     ListeningEvents.Add( TriggeredEvent );
 }
