Changed FMacCursor so it depends on mouse position from mouse move event. Also, changed handling of mouse drag event while dragging a window so that Slate doesn't get updates about window and mouse movement until we receive a windowDidMove notification.

diff --git a/Engine/Source/Runtime/Core/Private/Mac/MacApplication.cpp b/Engine/Source/Runtime/Core/Private/Mac/MacApplication.cpp
--- a/Engine/Source/Runtime/Core/Private/Mac/MacApplication.cpp
+++ b/Engine/Source/Runtime/Core/Private/Mac/MacApplication.cpp
@@ -106,6 +106,7 @@ FMacApplication::FMacApplication()
 	, bUsingTrackpad( false )
 	, HIDInput( HIDInputInterface::Create( MessageHandler ) )
 	, DraggedWindow( NULL )
+	, WindowDragCursorPos( FVector2D::ZeroVector )
 	, MouseCaptureWindow( NULL )
 	, bIsMouseCaptureEnabled( false )
 	, bIsMouseCursorLocked( false )
@@ -329,16 +330,6 @@ static TSharedPtr< FMacWindow > FindWindowByNSWindow( const TArray< TSharedRef<
 	return TSharedPtr< FMacWindow >( NULL );
 }
 
-FCocoaWindow* FMacApplication::FindMacEventWindow( NSEvent* CocoaEvent )
-{
-	return MacApplication->FindEventWindow( CocoaEvent );
-}
-
-TSharedPtr< FMacWindow > FMacApplication::FindMacWindowByNSWindow( FCocoaWindow* const WindowHandle )
-{
-	return FindWindowByNSWindow(MacApplication->Windows, &MacApplication->WindowsMutex, WindowHandle);
-}
-
 void FMacApplication::ProcessEvent(FMacEvent const* const Event)
 {
 	// Must have an event
@@ -573,25 +564,16 @@ void FMacApplication::ProcessNSEvent(NSEvent* const Event, TSharedPtr< FMacWindo
 				[NativeWindow setMovable: IsMovable];
 				[NativeWindow setMovableByWindowBackground: IsMouseOverTitleBar];
 			}
-			
-			// Cocoa does not update NSWindow's frame until user stops dragging the window, so while window is being dragged, we calculate
-			// its position based on mouse move delta
+
+			FMacCursor* MacCursor = (FMacCursor*)Cursor.Get();
+
 			if( DraggedWindow )
 			{
-				const int32 X = FMath::TruncToInt(CurrentEventWindow->PositionX + [Event deltaX]);
-				const int32 Y = FMath::TruncToInt(CurrentEventWindow->PositionY + [Event deltaY]);
-				if( CurrentEventWindow.IsValid() )
-				{
-					MessageHandler->OnMovedWindow( CurrentEventWindow.ToSharedRef(), X, Y );
-					CurrentEventWindow->PositionX = X;
-					CurrentEventWindow->PositionY = Y;
-				}
+				// Cocoa does not update NSWindow's frame until user stops dragging the window, so while the window is being dragged we don't let Slate know anything's moving
+				MacCursor->UpdateCurrentPosition(WindowDragCursorPos);
 			}
-
-			if( CurrentEventWindow.IsValid() )
+			else if( CurrentEventWindow.IsValid() )
 			{
-				FMacCursor* MacCursor = (FMacCursor*)Cursor.Get();
-				
 				// Under OS X we disassociate the cursor and mouse position during hi-precision mouse input.
 				// The game snaps the mouse cursor back to the starting point when this is disabled, which
 				// accumulates mouse delta that we want to ignore.
@@ -637,11 +619,15 @@ void FMacApplication::ProcessNSEvent(NSEvent* const Event, TSharedPtr< FMacWindo
 				}
 				else
 				{
-					FVector2D CurrentPosition = MacCursor->GetPosition();
+					FVector2D CurrentPosition = MousePosition * MacCursor->GetMouseScaling();
 					if( MacCursor->UpdateCursorClipping( CurrentPosition ) )
 					{
 						MacCursor->SetPosition( CurrentPosition.X, CurrentPosition.Y );
 					}
+					else
+					{
+						MacCursor->UpdateCurrentPosition(MousePosition);
+					}
 					MessageHandler->OnMouseMove();
 				}
 
@@ -893,10 +879,8 @@ void FMacApplication::ProcessEvent( NSEvent* Event )
 
 	FCocoaWindow* NativeWindow = FindEventWindow( Event );
 	TSharedPtr< FMacWindow > CurrentEventWindow = FindWindowByNSWindow( Windows, &WindowsMutex, NativeWindow );
-	
-	FMacCursor* MacCursor = (FMacCursor*)Cursor.Get();
-	
-	ProcessNSEvent(Event, CurrentEventWindow, MacCursor->GetPosition());
+
+	ProcessNSEvent(Event, CurrentEventWindow, ((FMacCursor*)Cursor.Get())->GetPosition());
 }
 
 void FMacApplication::ResendEvent(NSEvent* Event)
@@ -954,8 +938,8 @@ FCocoaWindow* FMacApplication::FindEventWindow( NSEvent* Event )
 		}
 		else
 		{
-			const FVector2D CursorPos = bUsingHighPrecisionMouseInput ? HighPrecisionMousePos : FVector2D([NSEvent mouseLocation].x, FPlatformMisc::ConvertSlateYPositionToCocoa([NSEvent mouseLocation].y));
-			TSharedPtr<FMacWindow> WindowUnderCursor = StaticCastSharedPtr<FMacWindow>(LocateWindowUnderCursor(CursorPos));
+			const FVector2D CursorPos = bUsingHighPrecisionMouseInput ? HighPrecisionMousePos : ((FMacCursor*)Cursor.Get())->GetPosition();
+			TSharedPtr<FMacWindow> WindowUnderCursor = LocateWindowUnderCursor(CursorPos);
 			if (WindowUnderCursor.IsValid())
 			{
 				EventWindow = WindowUnderCursor->GetWindowHandle();
@@ -966,7 +950,7 @@ FCocoaWindow* FMacApplication::FindEventWindow( NSEvent* Event )
 	return EventWindow;
 }
 
-TSharedPtr<FGenericWindow> FMacApplication::LocateWindowUnderCursor( const FVector2D& CursorPos )
+TSharedPtr<FMacWindow> FMacApplication::LocateWindowUnderCursor( const FVector2D& CursorPos )
 {
 	const NSPoint Position = NSMakePoint(CursorPos.X, FPlatformMisc::ConvertSlateYPositionToCocoa(CursorPos.Y));
 
@@ -1300,6 +1284,12 @@ void FMacApplication::OnWindowWillMove( FCocoaWindow* Window )
 	SCOPED_AUTORELEASE_POOL;
 
 	DraggedWindow = Window;
+
+	TSharedPtr< FMacWindow > EventWindow = FindWindowByNSWindow(Windows, &WindowsMutex, Window);
+	if (EventWindow.IsValid())
+	{
+		WindowDragCursorPos = ((FMacCursor*)Cursor.Get())->GetPosition();
+	}
 }
 
 void FMacApplication::OnWindowDidMove( FCocoaWindow* Window )
diff --git a/Engine/Source/Runtime/Core/Private/Mac/MacCursor.cpp b/Engine/Source/Runtime/Core/Private/Mac/MacCursor.cpp
--- a/Engine/Source/Runtime/Core/Private/Mac/MacCursor.cpp
+++ b/Engine/Source/Runtime/Core/Private/Mac/MacCursor.cpp
@@ -8,6 +8,7 @@
 FMacCursor::FMacCursor()
 :	bIsVisible(true)
 ,	bAssociateMouseCursor(false)
+,	CurrentPosition(FVector2D::ZeroVector)
 ,	MouseWarpDelta(FVector2D::ZeroVector)
 ,	MouseScale(1.0f, 1.0f)
 {
@@ -107,6 +108,12 @@ FMacCursor::FMacCursor()
 
 	// Set the default cursor
 	SetType( EMouseCursor::Default );
+
+	CGEventRef Event = CGEventCreate(NULL);
+	CGPoint CursorPos = CGEventGetLocation(Event);
+	CFRelease(Event);
+
+	CurrentPosition = FVector2D(FMath::TruncToFloat(CursorPos.x), FMath::TruncToFloat(CursorPos.y));
 }
 
 FMacCursor::~FMacCursor()
@@ -150,11 +157,7 @@ FMacCursor::~FMacCursor()
 
 FVector2D FMacCursor::GetPosition() const
 {
-	CGEventRef Event = CGEventCreate(NULL);
-	CGPoint CursorPos = CGEventGetLocation(Event);
-	CFRelease(Event);
-	
-	return FVector2D(FMath::TruncToFloat(CursorPos.x * MouseScale.X), FMath::TruncToFloat(CursorPos.y * MouseScale.X));
+	return FVector2D(FMath::TruncToFloat(CurrentPosition.X * MouseScale.X), FMath::TruncToFloat(CurrentPosition.Y * MouseScale.Y));
 }
 
 void FMacCursor::SetPosition( const int32 X, const int32 Y )
@@ -208,10 +211,10 @@ void FMacCursor::Lock( const RECT* const Bounds )
 		CusorClipRect.Max.Y = FMath::TruncToInt(Bounds->bottom) - 1;
 	}
 
-	FVector2D CurrentPosition = GetPosition();
-	if( UpdateCursorClipping( CurrentPosition ) )
+	FVector2D Position = GetPosition();
+	if( UpdateCursorClipping( Position ) )
 	{
-		SetPosition( CurrentPosition.X, CurrentPosition.Y );
+		SetPosition( Position.X, Position.Y );
 	}
 }
 
@@ -274,6 +277,11 @@ void FMacCursor::UpdateVisibility()
 #pragma clang diagnostic pop
 }
 
+void FMacCursor::UpdateCurrentPosition(const FVector2D &Position)
+{
+	CurrentPosition = Position;
+}
+
 void FMacCursor::WarpCursor( const int32 X, const int32 Y )
 {
 	// Apple suppress mouse events for 0.25 seconds after a call to Warp, unless we call CGAssociateMouseAndMouseCursorPosition.
@@ -293,6 +301,8 @@ void FMacCursor::WarpCursor( const int32 X, const int32 Y )
 	{
 		CGAssociateMouseAndMouseCursorPosition( true );
 	}
+
+	CurrentPosition = FVector2D(X, Y);
 }
 
 FVector2D FMacCursor::GetMouseWarpDelta( bool const bClearAccumulatedDelta )
diff --git a/Engine/Source/Runtime/Core/Private/Mac/MacEvent.cpp b/Engine/Source/Runtime/Core/Private/Mac/MacEvent.cpp
--- a/Engine/Source/Runtime/Core/Private/Mac/MacEvent.cpp
+++ b/Engine/Source/Runtime/Core/Private/Mac/MacEvent.cpp
@@ -13,7 +13,7 @@ FMacEvent::FMacEvent(NSEvent* const Event)
 	check(Event);
 	[Event retain];
 	
-	CocoaWindow = FMacApplication::FindMacEventWindow(Event);
+	CocoaWindow = MacApplication->FindEventWindow(Event);
 
 	CGEventRef CGEvent = [Event CGEvent];
 	CGPoint CGMouseLocation = CGEventGetLocation(CGEvent);
diff --git a/Engine/Source/Runtime/Core/Public/Mac/MacApplication.h b/Engine/Source/Runtime/Core/Public/Mac/MacApplication.h
--- a/Engine/Source/Runtime/Core/Public/Mac/MacApplication.h
+++ b/Engine/Source/Runtime/Core/Public/Mac/MacApplication.h
@@ -110,8 +110,6 @@ public:
 
 	void OnMouseCursorLock( bool bLockEnabled );
 
-	static FCocoaWindow* FindMacEventWindow( NSEvent* CocoaEvent );
-	static TSharedPtr< FMacWindow > FindMacWindowByNSWindow( FCocoaWindow* const WindowHandle );
 	static void ProcessEvent(FMacEvent const* const Event);
 
 	const TArray<TSharedRef<FMacWindow>>& GetAllWindows() const { return Windows; }
@@ -133,7 +131,7 @@ private:
 
 	void ResendEvent( NSEvent* Event );
 	FCocoaWindow* FindEventWindow( NSEvent* CocoaEvent );
-	TSharedPtr<FGenericWindow> LocateWindowUnderCursor( const FVector2D& CursorPos );
+	TSharedPtr<FMacWindow> LocateWindowUnderCursor( const FVector2D& CursorPos );
 
 	NSScreen* FindScreenByPoint( int32 X, int32 Y ) const;
 
@@ -163,6 +161,7 @@ private:
 	TSharedRef< class HIDInputInterface > HIDInput;
 
 	FCocoaWindow* DraggedWindow;
+	FVector2D WindowDragCursorPos;
 
 	FMouseCaptureWindow* MouseCaptureWindow;
 	bool bIsMouseCaptureEnabled;
@@ -208,6 +207,7 @@ private:
 	void* EventMonitor;
 
 	friend class FMacWindow;
+	friend class FMacEvent;
 };
 
 extern FMacApplication* MacApplication;
diff --git a/Engine/Source/Runtime/Core/Public/Mac/MacCursor.h b/Engine/Source/Runtime/Core/Public/Mac/MacCursor.h
--- a/Engine/Source/Runtime/Core/Public/Mac/MacCursor.h
+++ b/Engine/Source/Runtime/Core/Public/Mac/MacCursor.h
@@ -43,6 +43,8 @@ public:
 	
 	FVector2D GetMouseScaling( void );
 
+	void UpdateCurrentPosition( const FVector2D& Position );
+
 private:
 
 	void UpdateVisibility();
@@ -55,7 +57,8 @@ private:
 	bool bIsVisible;
 	bool bAssociateMouseCursor;
 	NSCursor* CurrentCursor;
-	
+
+	FVector2D CurrentPosition;
 	FVector2D MouseWarpDelta;
 	FVector2D MouseScale;
 };
