#ue4 - Add FVector::Coincident(A, B, Tolerance) to see if two normal vectors are coincident within a tolerance (nearly parallel and pointing in the same direction).

- Add tolerance param to FVector::Parallel.

diff --git a/Engine/Source/Runtime/Core/Public/Math/Vector.h b/Engine/Source/Runtime/Core/Public/Math/Vector.h
--- a/Engine/Source/Runtime/Core/Public/Math/Vector.h
+++ b/Engine/Source/Runtime/Core/Public/Math/Vector.h
@@ -794,9 +794,20 @@ public:
 	 *
 	 * @param Normal1 First normalized vector.
 	 * @param Normal1 Second normalized vector.
+	 * @param  DotTolerance Vectors are parallel if dot product varies less than this.
 	 * @return true if vectors are nearly parallel, false otherwise.
 	 */
-	static bool Parallel( const FVector &Normal1, const FVector &Normal2 );
+	static bool Parallel(const FVector& Normal1, const FVector& Normal2, float DotTolerance = THRESH_VECTORS_ARE_PARALLEL);
+
+	/**
+	 * See if two normal vectors (or plane normals) are coincident (nearly parallel and point in the same direction).
+	 * 
+	 * @param  Normal1 First normalized vector.
+	 * @param  Normal2 Second normalized vector.
+	 * @param  DotTolerance Vectors are parallel if dot product varies less than this.
+	 * @return true if vectors are coincident (nearly parallel and point in the same direction).
+	 */
+	static bool Coincident(const FVector& Normal1, const FVector& Normal2, float DotTolerance = THRESH_VECTORS_ARE_PARALLEL);
 
 	/**
 	 * See if two planes are coplanar.
@@ -1047,10 +1058,16 @@ inline FVector FVector::VectorPlaneProject(const FVector& V, const FVector& Plan
 	return V - V.ProjectOnToNormal(PlaneNormal);
 }
 
-inline bool FVector::Parallel( const FVector &Normal1, const FVector &Normal2 )
+inline bool FVector::Parallel(const FVector &Normal1, const FVector &Normal2, float DotTolerance)
+{
+	const float NormalDot = Normal1 | Normal2;
+	return (FMath::Abs(NormalDot - 1.f) <= DotTolerance);
+}
+
+inline bool FVector::Coincident(const FVector &Normal1, const FVector &Normal2, float DotTolerance)
 {
 	const float NormalDot = Normal1 | Normal2;
-	return (FMath::Abs(NormalDot - 1.f) <= THRESH_VECTORS_ARE_PARALLEL);
+	return (1.f - NormalDot) <= DotTolerance;
 }
 
 inline bool FVector::Coplanar( const FVector &Base1, const FVector &Normal1, const FVector &Base2, const FVector &Normal2 )
