Changed invariant culture in ICU internationalization implementation to use "en-US-POSIX" locale.

diff --git a/Engine/Source/Runtime/Core/Private/Internationalization/ICUInternationalization.cpp b/Engine/Source/Runtime/Core/Private/Internationalization/ICUInternationalization.cpp
--- a/Engine/Source/Runtime/Core/Private/Internationalization/ICUInternationalization.cpp
+++ b/Engine/Source/Runtime/Core/Private/Internationalization/ICUInternationalization.cpp
@@ -116,7 +116,7 @@ bool FICUInternationalization::Initialize()
 
 	FICUBreakIteratorManager::Create();
 
-	I18N->InvariantCulture = FindOrMakeCulture(TEXT(""));
+	I18N->InvariantCulture = FindOrMakeCulture(TEXT("en-US-POSIX"));
 	I18N->DefaultCulture = FindOrMakeCulture(icu::Locale::getDefault().getName());
 	SetCurrentCulture( I18N->GetDefaultCulture()->GetName() );
 	return U_SUCCESS(ICUStatus) ? true : false;
