Fixed a crash caused by out-of-bounds array access

Added a check to ensure that the array is populated first.
This addresses TTP#349083 - CRITICAL: Regression: EDITOR: CRASH: Occurs when changing maps using an automation test.

diff --git a/Engine/Source/Runtime/Core/Public/Misc/FeedbackContext.h b/Engine/Source/Runtime/Core/Public/Misc/FeedbackContext.h
--- a/Engine/Source/Runtime/Core/Public/Misc/FeedbackContext.h
+++ b/Engine/Source/Runtime/Core/Public/Misc/FeedbackContext.h
@@ -243,7 +243,7 @@ public:
 		verifySlow(ExpectedWorkThisFrame <= WorkRemaining);
 		CurrentFrameScope = FMath::Min(WorkRemaining, ExpectedWorkThisFrame);
 
-		Context.RequestUpdateUI(bCreatedDialog || Context.ScopeStack[0] == this);
+		Context.RequestUpdateUI(bCreatedDialog || (bEnabled && Context.ScopeStack[0] == this));
 	}
 
 	/**
