Merge Github pull request 477. [GitHub] 477 : Fix syntax error in TMultiMap::MultiFindPointer)

diff --git a/Engine/Source/Runtime/Core/Public/Containers/Map.h b/Engine/Source/Runtime/Core/Public/Containers/Map.h
--- a/Engine/Source/Runtime/Core/Public/Containers/Map.h
+++ b/Engine/Source/Runtime/Core/Public/Containers/Map.h
@@ -1012,7 +1012,7 @@ public:
 			OrderedValues.Empty(OutValues.Num());
 			for(int32 Index = OutValues.Num() - 1;Index >= 0;Index--)
 			{
-				new(OrderedValues) const ValueType*(OutValues(Index));
+				new(OrderedValues) const ValueType*(OutValues[Index]);
 			}
 
 			// Swap the ordered array into the output array.
