CrashReport - CrashContext.runtime-xml is now a part of the minidump file, so the whole crash report is just one file

diff --git a/Engine/Source/Runtime/Core/Private/Windows/WindowsPlatformExceptionHandling.cpp b/Engine/Source/Runtime/Core/Private/Windows/WindowsPlatformExceptionHandling.cpp
--- a/Engine/Source/Runtime/Core/Private/Windows/WindowsPlatformExceptionHandling.cpp
+++ b/Engine/Source/Runtime/Core/Private/Windows/WindowsPlatformExceptionHandling.cpp
@@ -23,12 +23,22 @@ namespace
 {
 static int32 ReportCrashCallCount = 0;
 
+struct FWindowsPlatformCrashContext : public FGenericCrashContext
+{
+	virtual void AddPlatformSpecificProperties() override
+	{
+		AddCrashProperty( TEXT( "Platform.IsRunningWindows" ), 1 );
+	}
+};
+
 /**
  * Write a Windows minidump to disk
  * @param Path Full path of file to write (normally a .dmp file)
  * @param ExceptionInfo Pointer to structure containing the exception information
  * @return Success or failure
  */
+
+// @TODO yrx 2014-10-08 Move to FWindowsPlatformCrashContext
 bool WriteMinidump(const TCHAR* Path, LPEXCEPTION_POINTERS ExceptionInfo)
 {
 	// Try to create file for minidump.
@@ -46,10 +56,28 @@ bool WriteMinidump(const TCHAR* Path, LPEXCEPTION_POINTERS ExceptionInfo)
 	DumpExceptionInfo.ExceptionPointers	= ExceptionInfo;
 	DumpExceptionInfo.ClientPointers	= FALSE;
 
-	BOOL result = MiniDumpWriteDump(GetCurrentProcess(), GetCurrentProcessId(), FileHandle, MiniDumpNormal, &DumpExceptionInfo, NULL, NULL);
+	// CrashContext.runtime-xml is now a part of the minidump file.
+	FWindowsPlatformCrashContext CrashContext;
+	CrashContext.SerializeContentToBuffer();
+
+	enum 
+	{
+		UE4_MINIDUMP_CRASHCONTEXT = LastReservedStream + 1,
+	};
+	
+	MINIDUMP_USER_STREAM CrashContextStream ={0};
+	CrashContextStream.Type = UE4_MINIDUMP_CRASHCONTEXT;
+	CrashContextStream.BufferSize = CrashContext.GetBuffer().GetAllocatedSize();
+	CrashContextStream.Buffer = (void*)*CrashContext.GetBuffer();
+
+	MINIDUMP_USER_STREAM_INFORMATION CrashContextStreamInformation = {0};
+	CrashContextStreamInformation.UserStreamCount = 1;
+	CrashContextStreamInformation.UserStreamArray = &CrashContextStream;
+
+	const BOOL Result = MiniDumpWriteDump(GetCurrentProcess(), GetCurrentProcessId(), FileHandle, MiniDumpNormal, &DumpExceptionInfo, &CrashContextStreamInformation, NULL);
 	CloseHandle(FileHandle);
 
-	return result == TRUE;
+	return Result == TRUE;
 }
 
 #if WINVER > 0x502	// Windows Error Reporting is not supported on Windows XP
@@ -184,22 +212,15 @@ void SetReportParameters( HREPORT ReportHandle, EXCEPTION_POINTERS* ExceptionInf
  */
 void AddMiniDump(HREPORT ReportHandle, EXCEPTION_POINTERS* ExceptionInfo)
 {
-	FString MinidumpFileName = FString::Printf(TEXT("%sDump%d.dmp"), *FPaths::GameLogDir(), FDateTime::UtcNow().GetTicks());
+	// Add GetMinidumpFile
+	const FString MinidumpFileName = FString::Printf(TEXT("%sDump%d.dmp"), *FPaths::GameLogDir(), FDateTime::UtcNow().GetTicks());
 	
 	if (WriteMinidump(*MinidumpFileName, ExceptionInfo))
 	{
-		WerReportAddFile(ReportHandle, *MinidumpFileName, WerFileTypeMinidump, WER_FILE_ANONYMOUS_DATA); 
+		WerReportAddFile(ReportHandle, *MinidumpFileName, WerFileTypeMinidump, WER_FILE_ANONYMOUS_DATA);
 	}
 }
 
-struct FWindowsPlatformCrashContext : public FGenericCrashContext
-{
-	void AddPlatformSpecificProperties() override
-	{
-		AddCrashProperty( TEXT( "Platform.IsRunningWindows" ), 1 );
-	}
-};
-
 /** 
  * Add miscellaneous files to the report. Currently the log and the video file
  */
@@ -210,13 +231,6 @@ void AddMiscFiles( HREPORT ReportHandle )
 
 	FString CrashVideoPath = FPaths::GameLogDir() / TEXT( "CrashVideo.avi" );
 	WerReportAddFile( ReportHandle, *CrashVideoPath, WerFileTypeOther, WER_FILE_ANONYMOUS_DATA );
-
-	// Introduces a new runtime crash context. Will replace all Windows related crash reporting.
-	const FString CrashContextXMLPath = FPaths::Combine( *FPaths::GameLogDir(), FWindowsPlatformCrashContext::CrashContextRuntimeXMLNameW );
-	FWindowsPlatformCrashContext CrashContext;
-	CrashContext.SerializeAsXML( *CrashContextXMLPath );
-
-	WerReportAddFile( ReportHandle, *CrashContextXMLPath, WerFileTypeOther, WER_FILE_ANONYMOUS_DATA );
 }
 
 /**
