Core/CrashReport - FGenericCrashContext::Initialize caches several platform specific data so getting that data should be safe after a crash, called in PlatformPreInit

diff --git a/Engine/Source/Runtime/Core/Private/Android/AndroidMisc.cpp b/Engine/Source/Runtime/Core/Private/Android/AndroidMisc.cpp
--- a/Engine/Source/Runtime/Core/Private/Android/AndroidMisc.cpp
+++ b/Engine/Source/Runtime/Core/Private/Android/AndroidMisc.cpp
@@ -61,6 +61,7 @@ void FAndroidMisc::LoadPreInitModules()
 
 void FAndroidMisc::PlatformPreInit()
 {
+	FGenericPlatformMisc::PlatformPreInit();
 	FAndroidAppEntry::PlatformInit();
 }
 
diff --git a/Engine/Source/Runtime/Core/Private/GenericPlatform/GenericPlatformContext.cpp b/Engine/Source/Runtime/Core/Private/GenericPlatform/GenericPlatformContext.cpp
--- a/Engine/Source/Runtime/Core/Private/GenericPlatform/GenericPlatformContext.cpp
+++ b/Engine/Source/Runtime/Core/Private/GenericPlatform/GenericPlatformContext.cpp
@@ -6,11 +6,42 @@
 
 const ANSICHAR* FGenericCrashContext::CrashContextRuntimeXMLNameA = "CrashContext.runtime-xml";
 const TCHAR* FGenericCrashContext::CrashContextRuntimeXMLNameW = TEXT( "CrashContext.runtime-xml" );
+bool FGenericCrashContext::bIsInitialized = false;
 
-namespace
+namespace NCachedCrashContextProperties
 {
 	static FString CrashDescriptionCommonBuffer;
+	static FString BaseDir;
+	static FString EpicAccountId;
+	static FString MachineIdStr;
+	static FString OsVersion;
+	static FString OsSubVersion;
+	static int32 NumberOfCores;
+	static int32 NumberOfCoresIncludingHyperthreads;
+	static FString CPUVendor;
+	static FString CPUBrand;
+	static FString PrimaryGPUBrand;
+	static FString UserName;
 }
+using namespace NCachedCrashContextProperties;
+
+void FGenericCrashContext::Initialize()
+{
+	BaseDir = FPlatformProcess::BaseDir();
+	EpicAccountId = FPlatformMisc::GetEpicAccountId();
+	MachineIdStr = FPlatformMisc::GetMachineId().ToString( EGuidFormats::Digits );
+	FPlatformMisc::GetOSVersions(OsVersion,OsSubVersion);
+	NumberOfCores = FPlatformMisc::NumberOfCores();
+	NumberOfCoresIncludingHyperthreads = FPlatformMisc::NumberOfCoresIncludingHyperthreads();
+
+	CPUVendor = FPlatformMisc::GetCPUVendor();
+	CPUBrand = FPlatformMisc::GetCPUBrand();
+	PrimaryGPUBrand = FPlatformMisc::GetPrimaryGPUBrand();
+	UserName = FPlatformProcess::UserName();
+
+	bIsInitialized = true;
+}
+
 
 FGenericCrashContext::FGenericCrashContext() :
 CommonBuffer( CrashDescriptionCommonBuffer )
@@ -22,13 +53,17 @@ void FGenericCrashContext::SerializeContentToBuffer()
 {
 	AddHeader();
 
+	// ProcessId will be used by the crash report client to get more information about the crashed process.
+	// It should be enough to open the process, get the memory stats, a list of loaded modules etc.
+	// It means that the crashed process needs to be alive as long as the crash report client is working on the process.
+	AddCrashProperty( TEXT( "ProcessId" ), FPlatformProcess::GetCurrentProcessId() );
+
 	// Add common crash properties.
 	AddCrashProperty( TEXT( "GameName" ), *FString::Printf( TEXT( "UE4-%s" ), FApp::GetGameName() ) );
 	AddCrashProperty( TEXT( "PlatformName" ), *FString( FPlatformProperties::PlatformName() ) );
 	AddCrashProperty( TEXT( "EngineMode" ), FPlatformMisc::GetEngineMode() );
 	AddCrashProperty( TEXT( "EngineVersion" ), ENGINE_VERSION_STRING );
-	AddCrashProperty( TEXT( "CommandLine" ), FCommandLine::Get() );
-	AddCrashProperty( TEXT( "BaseDir" ), FPlatformProcess::BaseDir() );
+	AddCrashProperty( TEXT( "CommandLine" ), FCommandLine::IsInitialized() ? FCommandLine::Get() : TEXT("") );
 	AddCrashProperty( TEXT( "LanguageLCID" ), FInternationalization::Get().GetCurrentCulture()->GetLCID() );
 
 	// Get the user name only for non-UE4 releases.
@@ -37,52 +72,69 @@ void FGenericCrashContext::SerializeContentToBuffer()
 		// Remove periods from internal user names to match AutoReporter user names
 		// The name prefix is read by CrashRepository.AddNewCrash in the website code
 
-		AddCrashProperty( TEXT( "UserName" ), *FString( FPlatformProcess::UserName() ).Replace( TEXT( "." ), TEXT( "" ) ) );
+		AddCrashProperty( TEXT( "UserName" ), *UserName.Replace( TEXT( "." ), TEXT( "" ) ) );
 	}
 
-	AddCrashProperty( TEXT( "MachineId" ), *FPlatformMisc::GetMachineId().ToString( EGuidFormats::Digits ) );
-	AddCrashProperty( TEXT( "EpicAccountId" ), *FPlatformMisc::GetEpicAccountId() );
+
+	AddCrashProperty( TEXT( "BaseDir" ), *BaseDir );
+	AddCrashProperty( TEXT( "MachineId" ), *MachineIdStr );
+	AddCrashProperty( TEXT( "EpicAccountId" ), *EpicAccountId );
 
 	AddCrashProperty( TEXT( "CallStack" ), TEXT( "" ) );
 	AddCrashProperty( TEXT( "SourceContext" ), TEXT( "" ) );
 	AddCrashProperty( TEXT( "UserDescription" ), TEXT( "" ) );
 	AddCrashProperty( TEXT( "ErrorMessage" ), (const TCHAR*)GErrorMessage );
 
-	AddCrashProperty( TEXT( "TimeOfCrash" ), FDateTime::UtcNow().GetTicks() );
+	// @TODO yrx 2014-10-08 Move to the crash report client.
+	/*if( CanUseUnsafeAPI() )
+	{
+		AddCrashProperty( TEXT( "TimeOfCrash" ), FDateTime::UtcNow().GetTicks() );
+	}
+	*/
 
 	// Add misc stats.
-	AddCrashProperty( TEXT( "Misc.NumberOfCores" ), FPlatformMisc::NumberOfCores() );
-	AddCrashProperty( TEXT( "Misc.NumberOfCoresIncludingHyperthreads" ), FPlatformMisc::NumberOfCoresIncludingHyperthreads() );
+	AddCrashProperty( TEXT( "Misc.NumberOfCores" ), NumberOfCores );
+	AddCrashProperty( TEXT( "Misc.NumberOfCoresIncludingHyperthreads" ), NumberOfCoresIncludingHyperthreads );
 	AddCrashProperty( TEXT( "Misc.Is64bitOperatingSystem" ), (int32)FPlatformMisc::Is64bitOperatingSystem() );
 
-	AddCrashProperty( TEXT( "Misc.CPUVendor" ), *FPlatformMisc::GetCPUVendor() );
-	AddCrashProperty( TEXT( "Misc.CPUBrand" ), *FPlatformMisc::GetCPUBrand() );
-	AddCrashProperty( TEXT( "Misc.PrimaryGPUBrand" ), *FPlatformMisc::GetPrimaryGPUBrand() );
-
-	FString OSVersion, OSSubVersion;
-	FPlatformMisc::GetOSVersions( OSVersion, OSSubVersion );
-	AddCrashProperty( TEXT( "Misc.OSVersion" ), *FString::Printf( TEXT( "%s, %s" ), *OSVersion, *OSSubVersion ) );
-
-	uint64 AppDiskTotalNumberOfBytes = 0;
-	uint64 AppDiskNumberOfFreeBytes = 0;
-	FPlatformMisc::GetDiskTotalAndFreeSpace( FPlatformProcess::BaseDir(), AppDiskTotalNumberOfBytes, AppDiskNumberOfFreeBytes );
-	AddCrashProperty( TEXT( "Misc.AppDiskTotalNumberOfBytes" ), AppDiskTotalNumberOfBytes );
-	AddCrashProperty( TEXT( "Misc.AppDiskNumberOfFreeBytes" ), AppDiskNumberOfFreeBytes );
-
-	// Add memory stats.
-	const FPlatformMemoryStats MemStats = FPlatformMemory::GetStats();
-
-	AddCrashProperty( TEXT( "MemoryStats.TotalPhysical" ), (uint64)MemStats.TotalPhysical );
-	AddCrashProperty( TEXT( "MemoryStats.TotalVirtual" ), (uint64)MemStats.TotalVirtual );
-	AddCrashProperty( TEXT( "MemoryStats.PageSize" ), (uint64)MemStats.PageSize );
-	AddCrashProperty( TEXT( "MemoryStats.TotalPhysicalGB" ), MemStats.TotalPhysicalGB );
-	AddCrashProperty( TEXT( "MemoryStats.AvailablePhysical" ), (uint64)MemStats.AvailablePhysical );
-	AddCrashProperty( TEXT( "MemoryStats.AvailableVirtual" ), (uint64)MemStats.AvailableVirtual );
-	AddCrashProperty( TEXT( "MemoryStats.UsedPhysical" ), (uint64)MemStats.UsedPhysical );
-	AddCrashProperty( TEXT( "MemoryStats.PeakUsedPhysical" ), (uint64)MemStats.PeakUsedPhysical );
-	AddCrashProperty( TEXT( "MemoryStats.UsedVirtual" ), (uint64)MemStats.UsedVirtual );
-	AddCrashProperty( TEXT( "MemoryStats.PeakUsedVirtual" ), (uint64)MemStats.PeakUsedVirtual );
-	AddCrashProperty( TEXT( "MemoryStats.bIsOOM" ), (int32)FPlatformMemory::bIsOOM );
+	AddCrashProperty( TEXT( "Misc.CPUVendor" ), *CPUVendor );
+	AddCrashProperty( TEXT( "Misc.CPUBrand" ), *CPUBrand );
+	AddCrashProperty( TEXT( "Misc.PrimaryGPUBrand" ), *PrimaryGPUBrand );
+	AddCrashProperty( TEXT( "Misc.OSVersion" ), *FString::Printf( TEXT( "%s, %s" ), *OsVersion, *OsSubVersion ) );
+
+	// @TODO yrx 2014-10-08 Move to the crash report client.
+	/*if( CanUseUnsafeAPI() )
+	{
+		uint64 AppDiskTotalNumberOfBytes = 0;
+		uint64 AppDiskNumberOfFreeBytes = 0;
+		FPlatformMisc::GetDiskTotalAndFreeSpace( FPlatformProcess::BaseDir(), AppDiskTotalNumberOfBytes, AppDiskNumberOfFreeBytes );
+		AddCrashProperty( TEXT( "Misc.AppDiskTotalNumberOfBytes" ), AppDiskTotalNumberOfBytes );
+		AddCrashProperty( TEXT( "Misc.AppDiskNumberOfFreeBytes" ), AppDiskNumberOfFreeBytes );
+	}*/
+
+	// FPlatformMemory::GetConstants is called in the GCreateMalloc, so we can assume it is always valid.
+	{
+		// Add memory stats.
+		const FPlatformMemoryConstants& MemConstants = FPlatformMemory::GetConstants();
+
+		AddCrashProperty( TEXT( "MemoryStats.TotalPhysical" ), (uint64)MemConstants.TotalPhysical );
+		AddCrashProperty( TEXT( "MemoryStats.TotalVirtual" ), (uint64)MemConstants.TotalVirtual );
+		AddCrashProperty( TEXT( "MemoryStats.PageSize" ), (uint64)MemConstants.PageSize );
+		AddCrashProperty( TEXT( "MemoryStats.TotalPhysicalGB" ), MemConstants.TotalPhysicalGB );
+	}
+
+	// @TODO yrx 2014-10-08 Move to the crash report client.
+	/*if( CanUseUnsafeAPI() )
+	{
+		const FPlatformMemoryStats MemStats = FPlatformMemory::GetStats();
+		AddCrashProperty( TEXT( "MemoryStats.AvailablePhysical" ), (uint64)MemStats.AvailablePhysical );
+		AddCrashProperty( TEXT( "MemoryStats.AvailableVirtual" ), (uint64)MemStats.AvailableVirtual );
+		AddCrashProperty( TEXT( "MemoryStats.UsedPhysical" ), (uint64)MemStats.UsedPhysical );
+		AddCrashProperty( TEXT( "MemoryStats.PeakUsedPhysical" ), (uint64)MemStats.PeakUsedPhysical );
+		AddCrashProperty( TEXT( "MemoryStats.UsedVirtual" ), (uint64)MemStats.UsedVirtual );
+		AddCrashProperty( TEXT( "MemoryStats.PeakUsedVirtual" ), (uint64)MemStats.PeakUsedVirtual );
+		AddCrashProperty( TEXT( "MemoryStats.bIsOOM" ), (int32)FPlatformMemory::bIsOOM );
+	}*/
 
 	// @TODO yrx 2014-09-10 
 	//Architecture
@@ -99,6 +151,8 @@ void FGenericCrashContext::SerializeContentToBuffer()
 void FGenericCrashContext::SerializeAsXML( const TCHAR* Filename )
 {
 	SerializeContentToBuffer();
+	// @TODO yrx 2014-10-08 We need mechanism to store the file that is aware of crashes and OOMs.
+	// Probably will be replaced with some kind of Inter-Process Communication
 	FFileHelper::SaveStringToFile( CommonBuffer, Filename, FFileHelper::EEncodingOptions::ForceUTF8 );
 }
 
diff --git a/Engine/Source/Runtime/Core/Private/GenericPlatform/GenericPlatformMisc.cpp b/Engine/Source/Runtime/Core/Private/GenericPlatform/GenericPlatformMisc.cpp
--- a/Engine/Source/Runtime/Core/Private/GenericPlatform/GenericPlatformMisc.cpp
+++ b/Engine/Source/Runtime/Core/Private/GenericPlatform/GenericPlatformMisc.cpp
@@ -11,6 +11,7 @@
 #include "ExceptionHandling.h"
 #include "Containers/Map.h"
 #include "../../Launch/Resources/Version.h"
+#include "GenericPlatformContext.h"
 
 #include "UProjectInfo.h"
 
@@ -821,25 +822,16 @@ FGuid FGenericPlatformMisc::GetMachineId()
 	return MachineId;
 }
 
-namespace
-{
-	/** Cached EpicAccountId. */
-	FString EpicAccountId;
-}
-
 FString FGenericPlatformMisc::GetEpicAccountId()
 {
-	if( EpicAccountId.IsEmpty() )
-	{
-		FPlatformMisc::GetStoredValue( TEXT( "Epic Games" ), TEXT( "Unreal Engine/Identifiers" ), TEXT( "AccountId" ), EpicAccountId );
-	}
+	FString EpicAccountId;
+	FPlatformMisc::GetStoredValue( TEXT( "Epic Games" ), TEXT( "Unreal Engine/Identifiers" ), TEXT( "AccountId" ), EpicAccountId );
 	return EpicAccountId;
 }
 
 void FGenericPlatformMisc::SetEpicAccountId( const FString& AccountId )
 {
 	FPlatformMisc::SetStoredValue( TEXT( "Epic Games" ), TEXT( "Unreal Engine/Identifiers" ), TEXT( "AccountId" ), AccountId );
-	EpicAccountId = AccountId;
 }
 
 const TCHAR* FGenericPlatformMisc::GetEngineMode()
@@ -849,4 +841,9 @@ const TCHAR* FGenericPlatformMisc::GetEngineMode()
 		GIsEditor ? TEXT( "Editor" ) :
 		IsRunningDedicatedServer() ? TEXT( "Server" ) :
 		TEXT( "Game" );
-}
\ No newline at end of file
+}
+
+void FGenericPlatformMisc::PlatformPreInit()
+{
+	FGenericCrashContext::Initialize();
+}
diff --git a/Engine/Source/Runtime/Core/Private/WinRT/WinRTMisc.cpp b/Engine/Source/Runtime/Core/Private/WinRT/WinRTMisc.cpp
--- a/Engine/Source/Runtime/Core/Private/WinRT/WinRTMisc.cpp
+++ b/Engine/Source/Runtime/Core/Private/WinRT/WinRTMisc.cpp
@@ -71,6 +71,8 @@ static void InitSHAHashes()
 
 void FWinRTMisc::PlatformPreInit()
 {
+	FGenericPlatformMisc::PlatformPreInit();
+
 	// Use our own handler for pure virtuals being called.
 	DefaultPureCallHandler = _set_purecall_handler(PureCallHandler);
 
diff --git a/Engine/Source/Runtime/Core/Private/Windows/WindowsPlatformStackWalk.cpp b/Engine/Source/Runtime/Core/Private/Windows/WindowsPlatformStackWalk.cpp
--- a/Engine/Source/Runtime/Core/Private/Windows/WindowsPlatformStackWalk.cpp
+++ b/Engine/Source/Runtime/Core/Private/Windows/WindowsPlatformStackWalk.cpp
@@ -13,6 +13,7 @@
 
 /*-----------------------------------------------------------------------------
 	Stack walking.
+	@TODO To be removed
 -----------------------------------------------------------------------------*/
 
 /** Whether appInitStackWalking() has been called successfully or not. */
@@ -359,8 +360,7 @@ static void LoadProcessModules()
 	{
 		// Keep track of the fact that we need to free it again.
 		bNeedToFreeModuleHandlePointer = true;
-		// @TODO yrx 2014-09-05 Calling malloc after a crash is not a good idea, fix it.
-		ModuleHandlePointer = (HMODULE*) FMemory::Malloc( BytesRequired );
+		ModuleHandlePointer = (HMODULE*) GMalloc->Malloc( BytesRequired );
 		FEnumProcessModules( ProcessHandle, ModuleHandlePointer, sizeof(ModuleHandleArray), (::DWORD *)&BytesRequired );
 	}
 
@@ -404,7 +404,7 @@ static void LoadProcessModules()
 	// Free the module handle pointer allocated in case the static array was insufficient.
 	if( bNeedToFreeModuleHandlePointer )
 	{
-		FMemory::Free( ModuleHandlePointer );
+		GMalloc->Free( ModuleHandlePointer );
 	}
 }
 
diff --git a/Engine/Source/Runtime/Core/Public/GenericPlatform/GenericPlatformContext.h b/Engine/Source/Runtime/Core/Public/GenericPlatform/GenericPlatformContext.h
--- a/Engine/Source/Runtime/Core/Public/GenericPlatform/GenericPlatformContext.h
+++ b/Engine/Source/Runtime/Core/Public/GenericPlatform/GenericPlatformContext.h
@@ -24,6 +24,20 @@ public:
 	static const ANSICHAR* CrashContextRuntimeXMLNameA;
 	static const TCHAR* CrashContextRuntimeXMLNameW;
 
+	/**	Whether the Initialize() has been called */
+	static bool bIsInitialized;
+
+	/** Initializes crash context related platform specific data that can be impossible to obtain after a crash. */
+	static void Initialize();
+
+	/**
+	 * @return true, if the generic crash context has been initialized.
+	 */
+	static bool IsInitalized()
+	{
+		return bIsInitialized;
+	}
+
 	/** Default constructor. */
 	FGenericCrashContext();
 
diff --git a/Engine/Source/Runtime/Core/Public/GenericPlatform/GenericPlatformMisc.h b/Engine/Source/Runtime/Core/Public/GenericPlatform/GenericPlatformMisc.h
--- a/Engine/Source/Runtime/Core/Public/GenericPlatform/GenericPlatformMisc.h
+++ b/Engine/Source/Runtime/Core/Public/GenericPlatform/GenericPlatformMisc.h
@@ -193,9 +193,7 @@ struct CORE_API FGenericPlatformMisc
 	/**
 	 * Called during appInit() after cmd line setup
 	 */
-	static void PlatformPreInit()
-	{
-	}
+	static void PlatformPreInit();
 	static void PlatformInit()
 	{
 	}
