CrashReport - Caching platform specific data, disabled on Mac

diff --git a/Engine/Source/Runtime/Core/Private/Mac/MacPlatformMisc.cpp b/Engine/Source/Runtime/Core/Private/Mac/MacPlatformMisc.cpp
--- a/Engine/Source/Runtime/Core/Private/Mac/MacPlatformMisc.cpp
+++ b/Engine/Source/Runtime/Core/Private/Mac/MacPlatformMisc.cpp
@@ -213,6 +213,8 @@ bool FMacPlatformMisc::bChachedMacMenuStateNeedsUpdate = true;
 
 void FMacPlatformMisc::PlatformPreInit()
 {
+	//FGenericPlatformMisc::PlatformPreInit();
+
 	// Increase the maximum number of simultaneously open files
 	uint32 MaxFilesPerProc = OPEN_MAX;
 	size_t UInt32Size = sizeof(uint32);
@@ -254,6 +256,7 @@ void FMacPlatformMisc::PlatformInit()
 	// Timer resolution.
 	UE_LOG(LogInit, Log, TEXT("High frequency timer resolution =%f MHz"), 0.000001 / FPlatformTime::GetSecondsPerCycle() );
 	
+	// Can we move this to FMacPlatformMisc::PlatformPreInit?
 	GMacAppInfo.Init();
 	
 	UE_LOG(LogInit, Log, TEXT("Power Source: %s"), GMacAppInfo.RunningOnBattery ? TEXT(kIOPSBatteryPowerValue) : TEXT(kIOPSACPowerValue) );
