CrashReport - Renamed WindowsPlatformExceptionHandling to *CrashContext, to be consistent with other platforms

diff --git a/Engine/Source/Runtime/Core/Private/Windows/WindowsPlatformCrashContext.cpp b/Engine/Source/Runtime/Core/Private/Windows/WindowsPlatformCrashContext.cpp
new file mode 100644
--- /dev/null
+++ b/Engine/Source/Runtime/Core/Private/Windows/WindowsPlatformCrashContext.cpp
@@ -0,0 +1,529 @@
+// Copyright 1998-2014 Epic Games, Inc. All Rights Reserved.
+
+#include "Core.h"
+#include "MallocCrash.h"
+#include "ExceptionHandling.h"
+#include "GenericPlatform/GenericPlatformContext.h"
+#include "../../Launch/Resources/Version.h"
+
+#include "AllowWindowsPlatformTypes.h"
+
+	#include <strsafe.h>
+#if WINVER > 0x502	// Windows Error Reporting is not supported on Windows XP
+	#include <werapi.h>
+	#pragma comment( lib, "wer.lib" )
+#endif	// WINVER
+	#include <dbghelp.h>
+	#include <Shlwapi.h>
+
+#pragma comment( lib, "version.lib" )
+#pragma comment( lib, "Shlwapi.lib" )
+
+namespace
+{
+static int32 ReportCrashCallCount = 0;
+
+struct FWindowsPlatformCrashContext : public FGenericCrashContext
+{
+	virtual void AddPlatformSpecificProperties() override
+	{
+		AddCrashProperty( TEXT( "Platform.IsRunningWindows" ), 1 );
+	}
+};
+
+/**
+ * Write a Windows minidump to disk
+ * @param Path Full path of file to write (normally a .dmp file)
+ * @param ExceptionInfo Pointer to structure containing the exception information
+ * @return Success or failure
+ */
+
+// @TODO yrx 2014-10-08 Move to FWindowsPlatformCrashContext
+bool WriteMinidump(const TCHAR* Path, LPEXCEPTION_POINTERS ExceptionInfo)
+{
+	// Try to create file for minidump.
+	HANDLE FileHandle = CreateFileW(Path, GENERIC_WRITE, 0, NULL, CREATE_ALWAYS, FILE_ATTRIBUTE_NORMAL, NULL);
+	
+	if (FileHandle == INVALID_HANDLE_VALUE)
+	{
+		return false;
+	}
+
+	// Initialise structure required by MiniDumpWriteDumps
+	MINIDUMP_EXCEPTION_INFORMATION DumpExceptionInfo = {};
+
+	DumpExceptionInfo.ThreadId			= GetCurrentThreadId();
+	DumpExceptionInfo.ExceptionPointers	= ExceptionInfo;
+	DumpExceptionInfo.ClientPointers	= FALSE;
+
+	// CrashContext.runtime-xml is now a part of the minidump file.
+	FWindowsPlatformCrashContext CrashContext;
+	CrashContext.SerializeContentToBuffer();
+
+	enum 
+	{
+		UE4_MINIDUMP_CRASHCONTEXT = LastReservedStream + 1,
+	};
+	
+	MINIDUMP_USER_STREAM CrashContextStream ={0};
+	CrashContextStream.Type = UE4_MINIDUMP_CRASHCONTEXT;
+	CrashContextStream.BufferSize = CrashContext.GetBuffer().GetAllocatedSize();
+	CrashContextStream.Buffer = (void*)*CrashContext.GetBuffer();
+
+	MINIDUMP_USER_STREAM_INFORMATION CrashContextStreamInformation = {0};
+	CrashContextStreamInformation.UserStreamCount = 1;
+	CrashContextStreamInformation.UserStreamArray = &CrashContextStream;
+
+	const BOOL Result = MiniDumpWriteDump(GetCurrentProcess(), GetCurrentProcessId(), FileHandle, MiniDumpNormal, &DumpExceptionInfo, &CrashContextStreamInformation, NULL);
+	CloseHandle(FileHandle);
+
+	return Result == TRUE;
+}
+
+#if WINVER > 0x502	// Windows Error Reporting is not supported on Windows XP
+
+/** 
+ * Get one line of text to describe the crash
+ */
+void GetCrashDescription(WER_REPORT_INFORMATION& ReportInformation)
+{
+	StringCchCopy( ReportInformation.wzDescription, ARRAYSIZE( ReportInformation.wzDescription ), TEXT( "The application crashed while running " ) );
+
+	const TCHAR* Description = IsRunningCommandlet() ? TEXT( "a commandlet" ) :
+		GIsEditor ? TEXT( "the editor" ) :
+		IsRunningDedicatedServer() ? TEXT( "a server" ) :
+		TEXT( "the game" );
+
+	StringCchCat( ReportInformation.wzDescription, ARRAYSIZE( ReportInformation.wzDescription ), Description );
+}
+
+#endif	// WINVER
+
+/**
+ * Get the 4 element version number of the module
+ */
+void GetModuleVersion( TCHAR* ModuleName, TCHAR* StringBuffer, DWORD MaxSize )
+{
+	StringCchCopy( StringBuffer, MaxSize, TEXT( "0.0.0.0" ) );
+	
+	DWORD Handle = 0;
+	DWORD InfoSize = GetFileVersionInfoSize( ModuleName, &Handle );
+	if( InfoSize > 0 )
+	{
+		TArray<char> VersionInfo;
+		VersionInfo.SetNum(InfoSize);
+
+		if( GetFileVersionInfo( ModuleName, 0, InfoSize, VersionInfo.GetData() ) )
+		{
+			VS_FIXEDFILEINFO* FixedFileInfo;
+
+			UINT InfoLength = 0;
+			if( VerQueryValue( VersionInfo.GetData(), TEXT( "\\" ), ( void** )&FixedFileInfo, &InfoLength ) )
+			{
+				StringCchPrintf( StringBuffer, MaxSize, TEXT( "%u.%u.%u.%u" ), 
+					HIWORD( FixedFileInfo->dwProductVersionMS ), LOWORD( FixedFileInfo->dwProductVersionMS ), HIWORD( FixedFileInfo->dwProductVersionLS ), LOWORD( FixedFileInfo->dwProductVersionLS ) );
+			}
+		}
+	}
+}
+
+
+#if WINVER > 0x502	// Windows Error Reporting is not supported on Windows XP
+
+/** 
+ * Set the ten Windows Error Reporting parameters
+ *
+ * Parameters 0 through 7 are predefined for Windows
+ * Parameters 8 and 9 are user defined
+ */
+void SetReportParameters( HREPORT ReportHandle, EXCEPTION_POINTERS* ExceptionInfo, const TCHAR* ErrorMessage )
+{
+	HRESULT Result;
+	TCHAR StringBuffer[MAX_SPRINTF] = {0};
+	TCHAR LocalBuffer[MAX_SPRINTF] = {0};
+
+	// Set the parameters for the standard problem signature
+	HMODULE ModuleHandle = GetModuleHandle( NULL );
+
+	StringCchPrintf( StringBuffer, MAX_SPRINTF, TEXT( "UE4-%s" ), FApp::GetGameName() );
+	Result = WerReportSetParameter( ReportHandle, WER_P0, TEXT( "Application Name" ), StringBuffer );
+
+	GetModuleFileName( ModuleHandle, LocalBuffer, MAX_SPRINTF );
+	PathStripPath( LocalBuffer );
+	GetModuleVersion( LocalBuffer, StringBuffer, MAX_SPRINTF );
+	Result = WerReportSetParameter( ReportHandle, WER_P1, TEXT( "Application Version" ), StringBuffer );
+
+	StringCchPrintf( StringBuffer, MAX_SPRINTF, TEXT( "%08x" ), GetTimestampForLoadedLibrary( ModuleHandle ) );
+	Result = WerReportSetParameter( ReportHandle, WER_P2, TEXT( "Application Timestamp" ), StringBuffer );
+
+	HMODULE FaultModuleHandle = NULL;
+	GetModuleHandleEx( GET_MODULE_HANDLE_EX_FLAG_FROM_ADDRESS | GET_MODULE_HANDLE_EX_FLAG_UNCHANGED_REFCOUNT, ( LPCTSTR )ExceptionInfo->ExceptionRecord->ExceptionAddress, &FaultModuleHandle );
+
+	GetModuleFileName( FaultModuleHandle, LocalBuffer, MAX_SPRINTF );
+	PathStripPath( LocalBuffer );
+	Result = WerReportSetParameter( ReportHandle, WER_P3, TEXT( "Fault Module Name" ), LocalBuffer );
+
+	GetModuleVersion( LocalBuffer, StringBuffer, MAX_SPRINTF );
+	Result = WerReportSetParameter( ReportHandle, WER_P4, TEXT( "Fault Module Version" ), StringBuffer );
+
+	StringCchPrintf( StringBuffer, MAX_SPRINTF, TEXT( "%08x" ), GetTimestampForLoadedLibrary( FaultModuleHandle ) );
+	Result = WerReportSetParameter( ReportHandle, WER_P5, TEXT( "Fault Module Timestamp" ), StringBuffer );
+
+	StringCchPrintf( StringBuffer, MAX_SPRINTF, TEXT( "%08x" ), ExceptionInfo->ExceptionRecord->ExceptionCode );
+	Result = WerReportSetParameter( ReportHandle, WER_P6, TEXT( "Exception Code" ), StringBuffer );
+
+	INT_PTR ExceptionOffset = ( char* )( ExceptionInfo->ExceptionRecord->ExceptionAddress ) - ( char* )FaultModuleHandle;
+	CA_SUPPRESS(6066) // The format specifier should probably be something like %tX, but VS 2013 doesn't support 't'.
+		StringCchPrintf( StringBuffer, MAX_SPRINTF, TEXT( "%p" ), ExceptionOffset );
+	Result = WerReportSetParameter( ReportHandle, WER_P7, TEXT( "Exception Offset" ), StringBuffer );
+
+	// Use LocalBuffer to store the error message.
+	FCString::Strncpy( LocalBuffer, ErrorMessage, MAX_SPRINTF );
+
+	// Replace " with ' and replace \n with #
+	for (TCHAR& Char: LocalBuffer)
+	{
+		if (Char == 0)
+		{
+			break;
+		}
+
+		switch (Char)
+		{
+			default: break;
+			case '"':	Char = '\'';	break;
+			case '\r':	Char = '#';		break;
+			case '\n':	Char = '#';		break;
+		}
+	}
+
+	// AssertLog should be ErrorMessage, but this require crash server changes, so don't change this.
+	StringCchPrintf( StringBuffer, MAX_SPRINTF, TEXT( "!%s!AssertLog=\"%s\"" ), FCommandLine::Get(), LocalBuffer );
+	Result = WerReportSetParameter( ReportHandle, WER_P8, TEXT( "Commandline" ), StringBuffer );
+
+	StringCchPrintf( StringBuffer, MAX_SPRINTF, TEXT( "%s!%s!%s!%d" ), TEXT( BRANCH_NAME ), FPlatformProcess::BaseDir(), FPlatformMisc::GetEngineMode(), BUILT_FROM_CHANGELIST );
+	Result = WerReportSetParameter( ReportHandle, WER_P9, TEXT( "BranchBaseDir" ), StringBuffer );
+}
+
+/**
+ * Add a minidump to the Windows Error Report
+ *
+ * Note this has to be a minidump (and not a microdump) for the dbgeng functions to work correctly
+ */
+void AddMiniDump(HREPORT ReportHandle, EXCEPTION_POINTERS* ExceptionInfo)
+{
+	// Add GetMinidumpFile
+	const FString MinidumpFileName = FString::Printf(TEXT("%sDump%d.dmp"), *FPaths::GameLogDir(), FDateTime::UtcNow().GetTicks());
+	
+	if (WriteMinidump(*MinidumpFileName, ExceptionInfo))
+	{
+		WerReportAddFile(ReportHandle, *MinidumpFileName, WerFileTypeMinidump, WER_FILE_ANONYMOUS_DATA);
+	}
+}
+
+/** 
+ * Add miscellaneous files to the report. Currently the log and the video file
+ */
+void AddMiscFiles( HREPORT ReportHandle )
+{
+	FString LogFileName = FPaths::GameLogDir() / FApp::GetGameName() + TEXT( ".log" );
+	WerReportAddFile( ReportHandle, *LogFileName, WerFileTypeOther, WER_FILE_ANONYMOUS_DATA ); 
+
+	FString CrashVideoPath = FPaths::GameLogDir() / TEXT( "CrashVideo.avi" );
+	WerReportAddFile( ReportHandle, *CrashVideoPath, WerFileTypeOther, WER_FILE_ANONYMOUS_DATA );
+}
+
+/**
+ * Force WER queuing on or off
+ *
+ * This is required so that the crash report doesn't immediately get sent to Microsoft, and we can't intercept it
+ */
+void ForceWERQueuing( bool bEnable )
+{
+	HKEY Key = 0;
+	HRESULT KeyResult = RegOpenKeyEx( HKEY_CURRENT_USER, TEXT( "Software\\Microsoft\\Windows\\Windows Error Reporting" ), 0, KEY_WRITE, &Key );
+	if( KeyResult == S_OK )
+	{
+		DWORD ForceQueue = ( DWORD )bEnable;
+		KeyResult = RegSetValueExW( Key, TEXT( "ForceQueue" ), 0, REG_DWORD, ( const BYTE* )&ForceQueue, sizeof( DWORD ) );
+	}
+}
+
+
+/**
+ * Enum indicating whether to run the crash reporter UI
+ */
+enum class EErrorReportUI
+{
+	/** Ask the user for a description */
+	ShowDialog,
+
+	/** Silently uploaded the report */
+	ReportInUnattendedMode	
+};
+
+/**
+ * Scoped setter of WER flag, to ensure it always gets put back
+ */
+class FScopedWERQueuing
+{
+public:
+	/**	Default constructor - force WER to queue, not send, reports */
+	FScopedWERQueuing()
+	{
+		ForceWERQueuing(true);
+	}
+
+	/**	Destructor - Reset the setting after we've had our way with it */
+	~FScopedWERQueuing()
+	{
+		ForceWERQueuing(false);
+	}
+
+private:
+	// disallow copying
+	FScopedWERQueuing(const FScopedWERQueuing&);
+	FScopedWERQueuing& operator=(const FScopedWERQueuing&);
+};
+
+
+/** 
+ * Create a Windows Error Report, add the user log and video, and add it to the WER queue
+ * Launch CrashReportClient.exe to intercept the report and upload to our local site
+ */
+int32 ReportCrashUsingCrashReportClient(EXCEPTION_POINTERS* ExceptionInfo, const TCHAR* ErrorMessage, EErrorReportUI ReportUI)
+{
+	// Flush out the log
+	GLog->Flush();
+
+	// Prevent CrashReportClient from spawning another CrashReportClient.
+	const TCHAR* ExecutableName = FPlatformProcess::ExecutableName();
+	const bool bCanRunCrashReportClient = FCString::Stristr( ExecutableName, TEXT( "CrashReportClient" ) ) == nullptr;
+	if( bCanRunCrashReportClient )
+	{
+		// Set the report to force queue
+		FScopedWERQueuing ScopedQueueForcer;
+
+		// Construct the report details
+		WER_REPORT_INFORMATION ReportInformation = {sizeof( WER_REPORT_INFORMATION )};
+
+		StringCchCopy( ReportInformation.wzConsentKey, ARRAYSIZE( ReportInformation.wzConsentKey ), TEXT( "" ) );
+
+		StringCchCopy( ReportInformation.wzApplicationName, ARRAYSIZE( ReportInformation.wzApplicationName ), TEXT("UE4-") );
+		StringCchCat( ReportInformation.wzApplicationName, ARRAYSIZE( ReportInformation.wzApplicationName ), FApp::GetGameName() );
+
+		StringCchCopy( ReportInformation.wzApplicationPath, ARRAYSIZE( ReportInformation.wzApplicationPath ), FPlatformProcess::BaseDir() );
+		StringCchCat( ReportInformation.wzApplicationPath, ARRAYSIZE( ReportInformation.wzApplicationPath ), FPlatformProcess::ExecutableName() );
+		StringCchCat( ReportInformation.wzApplicationPath, ARRAYSIZE( ReportInformation.wzApplicationPath ), TEXT( ".exe" ) );
+
+		GetCrashDescription( ReportInformation );
+
+		// Create a crash event report
+		HREPORT ReportHandle = NULL;
+		if( WerReportCreate( APPCRASH_EVENT, WerReportApplicationCrash, &ReportInformation, &ReportHandle ) == S_OK )
+		{
+			// Set the standard set of a crash parameters
+			SetReportParameters( ReportHandle, ExceptionInfo, ErrorMessage );
+
+			// Add a manually generated minidump
+			AddMiniDump( ReportHandle, ExceptionInfo );
+
+			// Add the log and video
+			AddMiscFiles( ReportHandle );
+
+			// Submit
+			WER_SUBMIT_RESULT SubmitResult;
+			WerReportSubmit( ReportHandle, WerConsentAlwaysPrompt, WER_SUBMIT_QUEUE | WER_SUBMIT_BYPASS_DATA_THROTTLING, &SubmitResult );
+
+			// Cleanup
+			WerReportCloseHandle( ReportHandle );
+		}
+
+		// Build machines do not upload these automatically since it is not okay to have lingering processes after the build completes.
+		if( GIsBuildMachine )
+		{
+			return EXCEPTION_CONTINUE_EXECUTION;
+		}
+
+		FString CrashReportClientArguments;
+
+		// Suppress the user input dialog if we're running in unattended mode
+		bool bNoDialog = FApp::IsUnattended() || ReportUI == EErrorReportUI::ReportInUnattendedMode;
+		if( bNoDialog )
+		{
+			CrashReportClientArguments += TEXT( " -Unattended" );
+		}
+
+		if( FApp::IsInstalled() )
+		{
+			// Temporary workaround for CrashReportClient being built in Development, not Shipping (TTP328030). The
+			// following ensures that logs are saved to the user directory when UE4 is installed.
+			CrashReportClientArguments += TEXT( " -Installed" );
+		}
+
+		CrashReportClientArguments += FString( TEXT( " -AppName=" ) ) + ReportInformation.wzApplicationName;
+
+		static const TCHAR CrashReportClientExeName[] = TEXT( "CrashReportClient.exe" );
+		FString CrashClientPath = FString( TEXT( "..\\..\\..\\Engine\\Binaries" ) ) / FPlatformProcess::GetBinariesSubdirectory() / CrashReportClientExeName;
+
+		bool bCrashReporterRan = FPlatformProcess::CreateProc( *CrashClientPath, *CrashReportClientArguments, true, false, false, NULL, 0, NULL, NULL ).IsValid();
+
+		if( !bCrashReporterRan && !bNoDialog )
+		{
+			UE_LOG( LogWindows, Log, TEXT( "Could not start %s" ), CrashReportClientExeName );
+			FPlatformMemory::DumpStats( *GWarn );
+			FText MessageTitle( FText::Format(
+				NSLOCTEXT( "MessageDialog", "AppHasCrashed", "The {0} {1} has crashed and will close" ),
+				FText::FromString( ReportInformation.wzApplicationName ),
+				FText::FromString( FPlatformMisc::GetEngineMode() )
+				) );
+			FMessageDialog::Open( EAppMsgType::Ok, FText::FromString( GErrorHist ), &MessageTitle );
+		}
+	}
+
+	// Let the system take back over (return value only used by NewReportEnsure)
+	return EXCEPTION_CONTINUE_EXECUTION;
+}
+
+#endif		// WINVER
+
+} // end anonymous namespace
+
+/** 
+ * Report an ensure to the crash reporting system
+ */
+void NewReportEnsure( const TCHAR* ErrorMessage )
+{
+#if WINVER > 0x502	// Windows Error Reporting is not supported on Windows XP
+#if !PLATFORM_SEH_EXCEPTIONS_DISABLED
+	__try
+#endif
+	{
+		FPlatformMisc::RaiseException( 1 );
+	}
+#if !PLATFORM_SEH_EXCEPTIONS_DISABLED
+	__except(ReportCrashUsingCrashReportClient(GetExceptionInformation(), ErrorMessage, EErrorReportUI::ReportInUnattendedMode))
+	CA_SUPPRESS(6322)
+	{
+	}
+#endif
+#endif	// WINVER
+}
+
+#include "HideWindowsPlatformTypes.h"
+
+// Original code below
+
+#include "AllowWindowsPlatformTypes.h"
+	#include <ErrorRep.h>
+	#include <DbgHelp.h>
+#include "HideWindowsPlatformTypes.h"
+
+#pragma comment(lib, "Faultrep.lib")
+
+/** 
+ * Creates an info string describing the given exception record.
+ * See MSDN docs on EXCEPTION_RECORD.
+ */
+#include "AllowWindowsPlatformTypes.h"
+void CreateExceptionInfoString(EXCEPTION_RECORD* ExceptionRecord)
+{
+	// @TODO yrx 2014-08-18 Fix FString usage?
+	FString ErrorString = TEXT("Unhandled Exception: ");
+
+#define HANDLE_CASE(x) case x: ErrorString += TEXT(#x); break;
+
+	switch (ExceptionRecord->ExceptionCode)
+	{
+	case EXCEPTION_ACCESS_VIOLATION:
+		ErrorString += TEXT("EXCEPTION_ACCESS_VIOLATION ");
+		if (ExceptionRecord->ExceptionInformation[0] == 0)
+		{
+			ErrorString += TEXT("reading address ");
+		}
+		else if (ExceptionRecord->ExceptionInformation[0] == 1)
+		{
+			ErrorString += TEXT("writing address ");
+		}
+		ErrorString += FString::Printf(TEXT("0x%08x"), (uint32)ExceptionRecord->ExceptionInformation[1]);
+		break;
+	HANDLE_CASE(EXCEPTION_ARRAY_BOUNDS_EXCEEDED)
+	HANDLE_CASE(EXCEPTION_DATATYPE_MISALIGNMENT)
+	HANDLE_CASE(EXCEPTION_FLT_DENORMAL_OPERAND)
+	HANDLE_CASE(EXCEPTION_FLT_DIVIDE_BY_ZERO)
+	HANDLE_CASE(EXCEPTION_FLT_INVALID_OPERATION)
+	HANDLE_CASE(EXCEPTION_ILLEGAL_INSTRUCTION)
+	HANDLE_CASE(EXCEPTION_INT_DIVIDE_BY_ZERO)
+	HANDLE_CASE(EXCEPTION_PRIV_INSTRUCTION)
+	HANDLE_CASE(EXCEPTION_STACK_OVERFLOW)
+	default:
+		ErrorString += FString::Printf(TEXT("0x%08x"), (uint32)ExceptionRecord->ExceptionCode);
+	}
+
+#if WITH_EDITORONLY_DATA
+	FCString::Strncpy(GErrorExceptionDescription, *ErrorString, ARRAY_COUNT(GErrorExceptionDescription));
+#endif
+#undef HANDLE_CASE
+}
+#include "HideWindowsPlatformTypes.h"
+
+/** 
+ * A null crash handler to suppress error report generation
+ */
+int32 NullReportCrash( LPEXCEPTION_POINTERS ExceptionInfo )
+{
+	return EXCEPTION_EXECUTE_HANDLER;
+}
+
+int32 ReportCrash( LPEXCEPTION_POINTERS ExceptionInfo )
+{
+	// Switch to malloc crash.
+	FMallocCrash::Get().SetAsGMalloc();
+
+	// Only create a minidump the first time this function is called.
+	// (Can be called the first time from the RenderThread, then a second time from the MainThread.)
+	if( FPlatformAtomics::InterlockedIncrement( &ReportCrashCallCount ) != 1 )
+	{
+		return EXCEPTION_EXECUTE_HANDLER;
+	}
+
+	FCoreDelegates::OnHandleSystemError.Broadcast();
+
+	// Note: ReportCrashUsingCrashReportClient reads the callstack from GErrorHist
+	const SIZE_T StackTraceSize = 65535;
+	ANSICHAR* StackTrace = (ANSICHAR*) GMalloc->Malloc( StackTraceSize );
+	StackTrace[0] = 0;
+	// Walk the stack and dump it to the allocated memory.
+	FPlatformStackWalk::StackWalkAndDump( StackTrace, StackTraceSize, 0, ExceptionInfo->ContextRecord );
+
+	// @TODO yrx 2014-08-20 Make this constant.
+	if( ExceptionInfo->ExceptionRecord->ExceptionCode != 1 )
+	{
+		CreateExceptionInfoString( ExceptionInfo->ExceptionRecord );
+		FCString::Strncat( GErrorHist, GErrorExceptionDescription, ARRAY_COUNT( GErrorHist ) );
+		FCString::Strncat( GErrorHist, TEXT( "\r\n\r\n" ), ARRAY_COUNT( GErrorHist ) );
+	}
+
+	FCString::Strncat( GErrorHist, ANSI_TO_TCHAR(StackTrace), ARRAY_COUNT(GErrorHist) );
+
+	GMalloc->Free( StackTrace );
+
+#if WINVER > 0x502	// Windows Error Reporting is not supported on Windows XP
+	if (GUseCrashReportClient)
+	{
+		ReportCrashUsingCrashReportClient(ExceptionInfo, GErrorMessage, EErrorReportUI::ShowDialog);
+	}
+	else
+#endif		// WINVER
+	{
+		WriteMinidump(MiniDumpFilenameW, ExceptionInfo);
+
+#if UE_BUILD_SHIPPING && WITH_EDITOR
+		uint32 dwOpt = 0;
+		EFaultRepRetVal repret = ReportFault( ExceptionInfo, dwOpt);
+#endif
+	}
+
+	return EXCEPTION_EXECUTE_HANDLER;
+}
+
diff --git a/Engine/Source/Runtime/Core/Private/Windows/WindowsPlatformExceptionHandling.cpp b/Engine/Source/Runtime/Core/Private/Windows/WindowsPlatformExceptionHandling.cpp
deleted file mode 100644
