[GitHub] 480 : FMath::GetReflectionVector(Dir,SurfaceNormal), For Reflecting Lasers off Mirrors and Projectile Bounces

TTP 348794

diff --git a/Engine/Source/Runtime/Core/Public/Math/UnrealMathUtility.h b/Engine/Source/Runtime/Core/Public/Math/UnrealMathUtility.h
--- a/Engine/Source/Runtime/Core/Public/Math/UnrealMathUtility.h
+++ b/Engine/Source/Runtime/Core/Public/Math/UnrealMathUtility.h
@@ -127,6 +127,17 @@ struct FMath : public FPlatformMath
 	/** Returns a random point within the passed in bounding box */
 	static CORE_API FVector RandPointInBox(const FBox& Box);
 
+	/** 
+	 * Given a direction vector and a surface normal, returns the vector reflected across the surface normal.
+	 * Produces a result like shining a laser at a mirror!
+	 *
+	 * @param Direction Direction vector the ray is comming from.
+	 * @param SurfaceNormal A normal of the surface the ray should be reflected on.
+	 *
+	 * @returns Reflected vector.
+	 */
+	static CORE_API FVector GetReflectionVector(const FVector& Direction, const FVector& SurfaceNormal);
+	
 	// Predicates
 
 	/** Checks if value is within a range, exclusive on MaxValue) */
diff --git a/Engine/Source/Runtime/Engine/Private/KismetMathLibrary.cpp b/Engine/Source/Runtime/Engine/Private/KismetMathLibrary.cpp
--- a/Engine/Source/Runtime/Engine/Private/KismetMathLibrary.cpp
+++ b/Engine/Source/Runtime/Engine/Private/KismetMathLibrary.cpp
@@ -768,8 +768,7 @@ FVector UKismetMathLibrary::RandomUnitVectorInCone(FVector ConeDir, float ConeHa
 
 FVector UKismetMathLibrary::MirrorVectorByNormal(FVector A, FVector B)
 {
-	B = B.SafeNormal();
-	return A - 2.f * B * (B | A);
+	return FMath::GetReflectionVector(A, B);
 }
 
 FRotator UKismetMathLibrary::RandomRotator(bool bRoll)
