[GitHub] 480 : FMath::GetReflectionVector(Dir,SurfaceNormal), For Reflecting Lasers off Mirrors and Projectile Bounces

Missing file for CL2323910 (GitHub https://github.com/EpicGames/UnrealEngine/commit/2c1bbfc11d3a2003f3d3076109adb85bf0d3f468).

TTP 348794

diff --git a/Engine/Source/Runtime/Core/Private/Math/UnrealMath.cpp b/Engine/Source/Runtime/Core/Private/Math/UnrealMath.cpp
--- a/Engine/Source/Runtime/Core/Private/Math/UnrealMath.cpp
+++ b/Engine/Source/Runtime/Core/Private/Math/UnrealMath.cpp
@@ -2247,6 +2247,11 @@ FVector FMath::RandPointInBox(const FBox& Box)
 					FRandRange(Box.Min.Z, Box.Max.Z) );
 }
 
+FVector FMath::GetReflectionVector(const FVector& Direction, const FVector& SurfaceNormal)
+{
+	return Direction - 2 * (Direction | SurfaceNormal.SafeNormal()) * SurfaceNormal.SafeNormal();
+}
+
 struct FClusterMovedHereToMakeCompile
 {
 	FVector ClusterPosAccum;
