Adding FORCENOINLINE to the generic implementation of MemoryBarrier

diff --git a/Engine/Source/Runtime/Core/Public/GenericPlatform/GenericPlatformMisc.h b/Engine/Source/Runtime/Core/Public/GenericPlatform/GenericPlatformMisc.h
--- a/Engine/Source/Runtime/Core/Public/GenericPlatform/GenericPlatformMisc.h
+++ b/Engine/Source/Runtime/Core/Public/GenericPlatform/GenericPlatformMisc.h
@@ -347,7 +347,7 @@ struct CORE_API FGenericPlatformMisc
 	/**
 	 * Enforces strict memory load/store ordering across the memory barrier call.
 	 */
-	static void MemoryBarrier();
+	FORCENOINLINE static void MemoryBarrier();
 
 	/**
 	 * Handles IO failure by ending gameplay.
