Cache the crash reporting values in PlatformPreInit on Mac OS X, should work just fine.

diff --git a/Engine/Source/Runtime/Core/Private/Mac/MacPlatformMisc.cpp b/Engine/Source/Runtime/Core/Private/Mac/MacPlatformMisc.cpp
--- a/Engine/Source/Runtime/Core/Private/Mac/MacPlatformMisc.cpp
+++ b/Engine/Source/Runtime/Core/Private/Mac/MacPlatformMisc.cpp
@@ -213,7 +213,9 @@ bool FMacPlatformMisc::bChachedMacMenuStateNeedsUpdate = true;
 
 void FMacPlatformMisc::PlatformPreInit()
 {
-	//FGenericPlatformMisc::PlatformPreInit();
+	FGenericPlatformMisc::PlatformPreInit();
+	
+	GMacAppInfo.Init();
 
 	// Increase the maximum number of simultaneously open files
 	uint32 MaxFilesPerProc = OPEN_MAX;
@@ -256,9 +258,6 @@ void FMacPlatformMisc::PlatformInit()
 	// Timer resolution.
 	UE_LOG(LogInit, Log, TEXT("High frequency timer resolution =%f MHz"), 0.000001 / FPlatformTime::GetSecondsPerCycle() );
 	
-	// Can we move this to FMacPlatformMisc::PlatformPreInit?
-	GMacAppInfo.Init();
-	
 	UE_LOG(LogInit, Log, TEXT("Power Source: %s"), GMacAppInfo.RunningOnBattery ? TEXT(kIOPSBatteryPowerValue) : TEXT(kIOPSACPowerValue) );
 }
 
