Copy the crash video if there is one into the crash report folder on OS X.

diff --git a/Engine/Source/Runtime/Core/Private/Mac/MacPlatformMisc.cpp b/Engine/Source/Runtime/Core/Private/Mac/MacPlatformMisc.cpp
--- a/Engine/Source/Runtime/Core/Private/Mac/MacPlatformMisc.cpp
+++ b/Engine/Source/Runtime/Core/Private/Mac/MacPlatformMisc.cpp
@@ -1715,7 +1715,7 @@ void FMacCrashContext::GenerateCrashInfoAndLaunchReporter() const
 		FCStringAnsi::Strncpy(FilePath, CrashInfoFolder, PATH_MAX);
 		FCStringAnsi::Strcat(FilePath, PATH_MAX, "/" );
 		FCStringAnsi::Strcat(FilePath, PATH_MAX, FGenericCrashContext::CrashContextRuntimeXMLNameA );
-		//SerializeAsXML( FilePath ); @todo uncomment after verification
+		//SerializeAsXML( FilePath ); @todo uncomment after verification - need to do a bit more work on this for OS X
 		
 		// copy log
 		FCStringAnsi::Strncpy(FilePath, CrashInfoFolder, PATH_MAX);
@@ -1735,6 +1735,24 @@ void FMacCrashContext::GenerateCrashInfoAndLaunchReporter() const
 		close(LogSrc);
 		// best effort, so don't care about result: couldn't copy -> tough, no log
 		
+		// copy crash video if there is one
+		if ( access(GMacAppInfo.CrashReportVideo, R_OK|F_OK) == 0 )
+		{
+			FCStringAnsi::Strncpy(FilePath, CrashInfoFolder, PATH_MAX);
+			FCStringAnsi::Strcat(FilePath, PATH_MAX, "/");
+			FCStringAnsi::Strcat(FilePath, PATH_MAX, "CrashVideo.avi");
+			int VideoSrc = open(GMacAppInfo.CrashReportVideo, O_RDONLY);
+			int VideoDst = open(FilePath, O_CREAT|O_WRONLY, 0766);
+			
+			int Bytes = 0;
+			while((Bytes = read(VideoSrc, Data, PATH_MAX)) > 0)
+			{
+				write(VideoDst, Data, Bytes);
+			}
+			close(VideoDst);
+			close(VideoSrc);
+		}
+		
 		// try launching the tool and wait for its exit, if at all
 		// Use fork() & execl() as they are async-signal safe, CreateProc can fail in Cocoa
 		int32 ReturnCode = 0;
