CrashReport - Refactoring continued

diff --git a/Engine/Source/Programs/CrashReporter/MinidumpDiagnostics/Private/Windows/MinidumpDiagnosticsMainWindows.cpp b/Engine/Source/Programs/CrashReporter/MinidumpDiagnostics/Private/Windows/MinidumpDiagnosticsMainWindows.cpp
--- a/Engine/Source/Programs/CrashReporter/MinidumpDiagnostics/Private/Windows/MinidumpDiagnosticsMainWindows.cpp
+++ b/Engine/Source/Programs/CrashReporter/MinidumpDiagnostics/Private/Windows/MinidumpDiagnosticsMainWindows.cpp
@@ -8,6 +8,15 @@
 
 IMPLEMENT_APPLICATION( MinidumpDiagnostics, "MinidumpDiagnostics" )
 
+	
+/** 
+ * A null crash handler to suppress error report generation
+ */
+int32 EmptyCrashHandler( LPEXCEPTION_POINTERS ExceptionInfo )
+{
+	return EXCEPTION_EXECUTE_HANDLER;
+}
+
 // More Windows glue
 int32 GuardedMain(int32 Argc, TCHAR* Argv[])
 {
@@ -41,7 +50,7 @@ int32 GuardedMainWrapper(int32 ArgC, TCHAR* ArgV[])
 			ReturnCode = GuardedMain( ArgC, ArgV );
 			GIsGuarded = 0;
 		}
-		__except( NullReportCrash( GetExceptionInformation() ) )
+		__except( EmptyCrashHandler( GetExceptionInformation() ) )
 		{
 		}
 	}
diff --git a/Engine/Source/Runtime/Core/Private/GenericPlatform/GenericPlatformContext.cpp b/Engine/Source/Runtime/Core/Private/GenericPlatform/GenericPlatformContext.cpp
--- a/Engine/Source/Runtime/Core/Private/GenericPlatform/GenericPlatformContext.cpp
+++ b/Engine/Source/Runtime/Core/Private/GenericPlatform/GenericPlatformContext.cpp
@@ -66,6 +66,8 @@ void FGenericCrashContext::SerializeContentToBuffer()
 	AddCrashProperty( TEXT( "CommandLine" ), FCommandLine::IsInitialized() ? FCommandLine::Get() : TEXT("") );
 	AddCrashProperty( TEXT( "LanguageLCID" ), FInternationalization::Get().GetCurrentCulture()->GetLCID() );
 
+	AddCrashProperty( TEXT( "IsUE4Release" ), (int32)FRocketSupport::IsRocket() );
+
 	// Get the user name only for non-UE4 releases.
 	if( !FRocketSupport::IsRocket() )
 	{
diff --git a/Engine/Source/Runtime/Core/Private/Windows/WindowsPlatformCrashContext.cpp b/Engine/Source/Runtime/Core/Private/Windows/WindowsPlatformCrashContext.cpp
--- a/Engine/Source/Runtime/Core/Private/Windows/WindowsPlatformCrashContext.cpp
+++ b/Engine/Source/Runtime/Core/Private/Windows/WindowsPlatformCrashContext.cpp
@@ -3,7 +3,7 @@
 #include "Core.h"
 #include "MallocCrash.h"
 #include "ExceptionHandling.h"
-#include "GenericPlatform/GenericPlatformContext.h"
+#include "WindowsPlatformCrashContext.h"
 #include "../../Launch/Resources/Version.h"
 
 #include "AllowWindowsPlatformTypes.h"
@@ -23,14 +23,6 @@ namespace
 {
 static int32 ReportCrashCallCount = 0;
 
-struct FWindowsPlatformCrashContext : public FGenericCrashContext
-{
-	virtual void AddPlatformSpecificProperties() override
-	{
-		AddCrashProperty( TEXT( "Platform.IsRunningWindows" ), 1 );
-	}
-};
-
 /**
  * Write a Windows minidump to disk
  * @param Path Full path of file to write (normally a .dmp file)
@@ -59,14 +51,9 @@ bool WriteMinidump(const TCHAR* Path, LPEXCEPTION_POINTERS ExceptionInfo)
 	// CrashContext.runtime-xml is now a part of the minidump file.
 	FWindowsPlatformCrashContext CrashContext;
 	CrashContext.SerializeContentToBuffer();
-
-	enum 
-	{
-		UE4_MINIDUMP_CRASHCONTEXT = LastReservedStream + 1,
-	};
 	
 	MINIDUMP_USER_STREAM CrashContextStream ={0};
-	CrashContextStream.Type = UE4_MINIDUMP_CRASHCONTEXT;
+	CrashContextStream.Type = FWindowsPlatformCrashContext::UE4_MINIDUMP_CRASHCONTEXT;
 	CrashContextStream.BufferSize = CrashContext.GetBuffer().GetAllocatedSize();
 	CrashContextStream.Buffer = (void*)*CrashContext.GetBuffer();
 
@@ -467,14 +454,6 @@ void CreateExceptionInfoString(EXCEPTION_RECORD* ExceptionRecord)
 }
 #include "HideWindowsPlatformTypes.h"
 
-/** 
- * A null crash handler to suppress error report generation
- */
-int32 NullReportCrash( LPEXCEPTION_POINTERS ExceptionInfo )
-{
-	return EXCEPTION_EXECUTE_HANDLER;
-}
-
 int32 ReportCrash( LPEXCEPTION_POINTERS ExceptionInfo )
 {
 	// Switch to malloc crash.
diff --git a/Engine/Source/Runtime/Core/Public/Android/AndroidPlatformCrashContext.h b/Engine/Source/Runtime/Core/Public/Android/AndroidPlatformCrashContext.h
--- a/Engine/Source/Runtime/Core/Public/Android/AndroidPlatformCrashContext.h
+++ b/Engine/Source/Runtime/Core/Public/Android/AndroidPlatformCrashContext.h
@@ -39,4 +39,6 @@ struct CORE_API FAndroidCrashContext : public FGenericCrashContext
 		Info = InInfo;
 		Context = InContext;
 	}
-};
\ No newline at end of file
+};
+
+typedef FAndroidCrashContext FPlatformCrashContext;
\ No newline at end of file
diff --git a/Engine/Source/Runtime/Core/Public/HAL/ExceptionHandling.h b/Engine/Source/Runtime/Core/Public/HAL/ExceptionHandling.h
--- a/Engine/Source/Runtime/Core/Public/HAL/ExceptionHandling.h
+++ b/Engine/Source/Runtime/Core/Public/HAL/ExceptionHandling.h
@@ -15,10 +15,11 @@ extern CORE_API TCHAR MiniDumpFilenameW[1024];
 
 // @TODO yrx 2014-09-11 Move to PlatformExceptionHandling
 #if PLATFORM_WINDOWS
-extern CORE_API int32 NullReportCrash( LPEXCEPTION_POINTERS ExceptionInfo );
+// @TODO yrx 2014-10-09 These methods are specific to windows, remove from here.
 extern CORE_API int32 ReportCrash( LPEXCEPTION_POINTERS ExceptionInfo );
 extern CORE_API void NewReportEnsure( const TCHAR* ErrorMessage );
 #elif PLATFORM_MAC
+// @TODO yrx 2014-10-09 Should be move to another file
 extern CORE_API int32 ReportCrash( ucontext_t *Context, int32 Signal, struct __siginfo* Info );
 extern CORE_API void NewReportEnsure( const TCHAR* ErrorMessage );
 #elif PLATFORM_LINUX
diff --git a/Engine/Source/Runtime/Core/Public/Linux/LinuxPlatformCrashContext.h b/Engine/Source/Runtime/Core/Public/Linux/LinuxPlatformCrashContext.h
--- a/Engine/Source/Runtime/Core/Public/Linux/LinuxPlatformCrashContext.h
+++ b/Engine/Source/Runtime/Core/Public/Linux/LinuxPlatformCrashContext.h
@@ -79,3 +79,5 @@ struct CORE_API FLinuxCrashContext : public FGenericCrashContext
 	 */
 	void GenerateReport(const FString & DiagnosticsPath) const;
 };
+
+typedef FLinuxCrashContext FPlatformCrashContext;
diff --git a/Engine/Source/Runtime/Core/Public/Mac/MacPlatformCrashContext.h b/Engine/Source/Runtime/Core/Public/Mac/MacPlatformCrashContext.h
--- a/Engine/Source/Runtime/Core/Public/Mac/MacPlatformCrashContext.h
+++ b/Engine/Source/Runtime/Core/Public/Mac/MacPlatformCrashContext.h
@@ -16,3 +16,4 @@ struct CORE_API FMacCrashContext : public FApplePlatformCrashContext
 	void GenerateCrashInfoAndLaunchReporter() const;
 };
 
+typedef FMacCrashContext FPlatformCrashContext;
diff --git a/Engine/Source/Runtime/Core/Public/Windows/WindowsPlatformCrashContext.h b/Engine/Source/Runtime/Core/Public/Windows/WindowsPlatformCrashContext.h
new file mode 100644
--- /dev/null
+++ b/Engine/Source/Runtime/Core/Public/Windows/WindowsPlatformCrashContext.h
@@ -0,0 +1,23 @@
+// Copyright 1998-2014 Epic Games, Inc. All Rights Reserved.
+
+#include "GenericPlatform/GenericPlatformContext.h"
+
+#include "AllowWindowsPlatformTypes.h"
+#include "DbgHelp.h"
+#include "HideWindowsPlatformTypes.h"
+
+struct CORE_API FWindowsPlatformCrashContext : public FGenericCrashContext
+{
+	/** Platform specific constants. */
+	enum EConstants
+	{
+		UE4_MINIDUMP_CRASHCONTEXT = LastReservedStream + 1,
+	};
+
+	virtual void AddPlatformSpecificProperties() override
+	{
+		AddCrashProperty( TEXT( "Platform.IsRunningWindows" ), 1 );
+	}
+};
+
+typedef FWindowsPlatformCrashContext FPlatformCrashContext;
\ No newline at end of file
