Slate: Added support for convertible laptops on Windows (GitHub PR #380 by JeffRous)

https://github.com/EpicGames/UnrealEngine/pull/380

diff --git a/Engine/Source/Runtime/Core/Private/Windows/WindowsApplication.cpp b/Engine/Source/Runtime/Core/Private/Windows/WindowsApplication.cpp
--- a/Engine/Source/Runtime/Core/Private/Windows/WindowsApplication.cpp
+++ b/Engine/Source/Runtime/Core/Private/Windows/WindowsApplication.cpp
@@ -27,6 +27,11 @@
 #define WM_MOUSEHWHEEL                  0x020E
 #endif
 
+// This might not be defined by Windows when maintaining backwards-compatibility to pre-Win8 builds
+#ifndef SM_CONVERTIBLESLATEMODE
+#define SM_CONVERTIBLESLATEMODE			0x2003
+#endif
+
 DEFINE_LOG_CATEGORY(LogWindowsDesktop);
 
 const FIntPoint FWindowsApplication::MinimizedWindowPosition(-32000,-32000);
@@ -889,6 +894,16 @@ int32 FWindowsApplication::ProcessMessage( HWND hwnd, uint32 msg, WPARAM wParam,
 			}
 			break;
 
+		case WM_SETTINGCHANGE:
+			{
+				// Convertible mode change
+				if ((lParam != NULL) && (wcscmp(TEXT("ConvertibleSlateMode"), (TCHAR *)lParam) == 0))
+				{
+					DeferMessage(CurrentNativeEventWindowPtr, hwnd, msg, wParam, lParam);
+				}
+			}
+			break;
+
 		case WM_PAINT:
 			{
 				if( bInModalSizeLoop && IsInGameThread() )
@@ -1447,6 +1462,11 @@ int32 FWindowsApplication::ProcessDeferredMessage( const FDeferredWindowsMessage
 			MessageHandler->OnApplicationActivationChanged( !!wParam );
 			break;
 
+		case WM_SETTINGCHANGE:
+			MessageHandler->OnConvertibleDeviceModeChanged((GetSystemMetrics(SM_CONVERTIBLESLATEMODE) == 0)
+				? EConvertibleLaptopModes::Tablet
+				: EConvertibleLaptopModes::Laptop);
+			break;
 	
 		case WM_NCACTIVATE:
 			{
diff --git a/Engine/Source/Runtime/Core/Public/GenericPlatform/GenericApplicationMessageHandler.h b/Engine/Source/Runtime/Core/Public/GenericPlatform/GenericApplicationMessageHandler.h
--- a/Engine/Source/Runtime/Core/Public/GenericPlatform/GenericApplicationMessageHandler.h
+++ b/Engine/Source/Runtime/Core/Public/GenericPlatform/GenericApplicationMessageHandler.h
@@ -151,6 +151,20 @@ namespace EGestureEvent
 	};
 }
 
+
+/**
+ * Enumerates the modes a convertible laptop can be in.
+ */
+enum class EConvertibleLaptopModes
+{
+	/** Laptop arranged as a tablet. */
+	Tablet,
+
+	/** Laptop arranged as a laptop. */
+	Laptop
+};
+
+
 class FGenericApplicationMessageHandler
 {
 public:
@@ -300,6 +314,11 @@ public:
 		return false;
 	}
 
+	virtual bool OnConvertibleDeviceModeChanged(const EConvertibleLaptopModes NewMode)
+	{
+		return false;
+	}
+
 	virtual EWindowZone::Type GetWindowZoneForPoint( const TSharedRef< FGenericWindow >& Window, const int32 X, const int32 Y )
 	{
 		return EWindowZone::NotInWindow;
diff --git a/Engine/Source/Runtime/Engine/Private/PlayerController.cpp b/Engine/Source/Runtime/Engine/Private/PlayerController.cpp
--- a/Engine/Source/Runtime/Engine/Private/PlayerController.cpp
+++ b/Engine/Source/Runtime/Engine/Private/PlayerController.cpp
@@ -1022,7 +1022,7 @@ void APlayerController::CreateTouchInterface()
 	ULocalPlayer* LocalPlayer = Cast<ULocalPlayer>(Player);
 
 	// do we want to show virtual joysticks?
-	if (LocalPlayer && LocalPlayer->ViewportClient && SVirtualJoystick::ShouldDisplayTouchInterface())
+	if (LocalPlayer && LocalPlayer->ViewportClient && !VirtualJoystick.IsValid() && SVirtualJoystick::ShouldDisplayTouchInterface())
 	{
 		// in case we already had one, remove it
 		if (VirtualJoystick.IsValid())
@@ -3671,6 +3671,17 @@ void APlayerController::TickPlayerInput(const float DeltaSeconds, const bool bGa
 				CurrentTouchablePrimitives[TouchIndexInt] = CurrentComponent;
 			}
 		}
+
+		// force touch UI on when it should be
+		if (SVirtualJoystick::ShouldDisplayTouchInterface())
+		{
+			CreateTouchInterface();
+			VirtualJoystick->SetVisibility(true, true);
+		}
+		else
+		{
+			CleanupGameViewport();
+		}
 	}
 
 	if (PlayerInput != NULL)
diff --git a/Engine/Source/Runtime/Slate/Private/Framework/Application/SlateApplication.cpp b/Engine/Source/Runtime/Slate/Private/Framework/Application/SlateApplication.cpp
--- a/Engine/Source/Runtime/Slate/Private/Framework/Application/SlateApplication.cpp
+++ b/Engine/Source/Runtime/Slate/Private/Framework/Application/SlateApplication.cpp
@@ -4464,6 +4464,27 @@ void FSlateApplication::ProcessApplicationActivationEvent( bool InAppActivated )
 	}
 }
 
+
+bool FSlateApplication::OnConvertibleDeviceModeChanged(const EConvertibleLaptopModes NewMode)
+{
+	// Notify that we want the mobile experience when in tablet mode, otherwise use mouse and keyboard
+	if (!(FParse::Param(FCommandLine::Get(), TEXT("simmobile")) || FParse::Param(FCommandLine::Get(), TEXT("faketouches"))))
+	{
+		// Not sure what the correct long-term strategy is. Use bIsFakingTouch for now to get things going.
+		if (NewMode == EConvertibleLaptopModes::Tablet)
+		{
+			bIsFakingTouch = true;
+		}
+		else
+		{
+			bIsFakingTouch = false;
+		}
+	}
+
+	return true;
+}
+
+
 EWindowZone::Type FSlateApplication::GetWindowZoneForPoint( const TSharedRef< FGenericWindow >& PlatformWindow, const int32 X, const int32 Y )
 {
 	TSharedPtr< SWindow > Window = FSlateWindowHelper::FindWindowByPlatformWindow( SlateWindows, PlatformWindow );
diff --git a/Engine/Source/Runtime/Slate/Public/Framework/Application/SlateApplication.h b/Engine/Source/Runtime/Slate/Public/Framework/Application/SlateApplication.h
--- a/Engine/Source/Runtime/Slate/Public/Framework/Application/SlateApplication.h
+++ b/Engine/Source/Runtime/Slate/Public/Framework/Application/SlateApplication.h
@@ -1052,6 +1052,8 @@ public:
 
 	virtual bool OnApplicationActivationChanged( const bool IsActive ) override;
 
+	virtual bool OnConvertibleDeviceModeChanged(const EConvertibleLaptopModes NewMode) override;
+
 	virtual EWindowZone::Type GetWindowZoneForPoint( const TSharedRef< FGenericWindow >& PlatformWindow, const int32 X, const int32 Y ) override;
 
 	virtual void OnWindowClose( const TSharedRef< FGenericWindow >& PlatformWindow ) override;
diff --git a/Engine/Source/Runtime/Slate/Public/Widgets/Input/SVirtualJoystick.h b/Engine/Source/Runtime/Slate/Public/Widgets/Input/SVirtualJoystick.h
--- a/Engine/Source/Runtime/Slate/Public/Widgets/Input/SVirtualJoystick.h
+++ b/Engine/Source/Runtime/Slate/Public/Widgets/Input/SVirtualJoystick.h
@@ -114,10 +114,7 @@ public:
 
 	void SetControls(const TArray<FControlInfo>& InControls)
 	{
-		if (ShouldDisplayTouchInterface())
-		{
-			Controls = InControls;
-		}
+		Controls = InControls;
 	}
 
 	/**
