Core: Fixed incorrect assertion in FTimespan::FromSeconds()

diff --git a/Engine/Source/Runtime/Core/Private/Misc/Timespan.cpp b/Engine/Source/Runtime/Core/Private/Misc/Timespan.cpp
--- a/Engine/Source/Runtime/Core/Private/Misc/Timespan.cpp
+++ b/Engine/Source/Runtime/Core/Private/Misc/Timespan.cpp
@@ -113,7 +113,7 @@ FTimespan FTimespan::FromMinutes( double Minutes )
 
 FTimespan FTimespan::FromSeconds( double Seconds )
 {
-	check((Seconds >= MinValue().GetTotalMinutes()) && (Seconds <= MaxValue().GetTotalMinutes()));
+	check((Seconds >= MinValue().GetTotalSeconds()) && (Seconds <= MaxValue().GetTotalSeconds()));
 
 	return FTimespan(Seconds * ETimespan::TicksPerSecond);
 }
