Editor/Slate - The multiline text editor now allows you to specify a modifer key that is required to create a new line.  Normally this is none, but you can set it to Shift, for example to require shift +enter to create a new line.  Failing to press Shift+Enter and instead just Enter would now commit the text.  The multiline property editors for the editor now use this.  Enter by default now always commits the values, but in multiline property fields you can now add new lines by using Shift+Enter.

diff --git a/Engine/Source/Editor/PropertyEditor/Private/UserInterface/PropertyEditor/SPropertyEditorText.cpp b/Engine/Source/Editor/PropertyEditor/Private/UserInterface/PropertyEditor/SPropertyEditorText.cpp
--- a/Engine/Source/Editor/PropertyEditor/Private/UserInterface/PropertyEditor/SPropertyEditorText.cpp
+++ b/Engine/Source/Editor/PropertyEditor/Private/UserInterface/PropertyEditor/SPropertyEditorText.cpp
@@ -27,6 +27,7 @@ void SPropertyEditorText::Construct( const FArguments& InArgs, const TSharedRef<
 			.SelectAllTextOnCommit( false )
 			.IsReadOnly(this, &SPropertyEditorText::IsReadOnly)
 			.AutoWrapText(true)
+			//.ModiferKeyForNewLine(EModifierKey::Shift)
 		];
 	}
 	else
diff --git a/Engine/Source/Runtime/Core/Public/GenericPlatform/GenericApplication.h b/Engine/Source/Runtime/Core/Public/GenericPlatform/GenericApplication.h
--- a/Engine/Source/Runtime/Core/Public/GenericPlatform/GenericApplication.h
+++ b/Engine/Source/Runtime/Core/Public/GenericPlatform/GenericApplication.h
@@ -9,6 +9,29 @@ class ITextInputMethodSystem;
 class IForceFeedbackSystem;
 class IAnalyticsProvider;
 
+/**
+* Enumerates available modifier keys for input gestures.
+*/
+namespace EModifierKey
+{
+	typedef uint8 Type;
+
+	/** No key. */
+	const Type None	= 0;
+
+	/** Ctrl key (Command key on Mac, Control key on Windows). */
+	const Type Control = 1 << 0;
+
+	/** Alt key. */
+	const Type Alt = 1 << 1;
+
+	/** Shift key. */
+	const Type Shift = 1 << 2;
+
+	/** Cmd key (Control key on Mac, Win key on Windows) */
+	const Type Command = 1 << 3;
+};
+
 namespace EPopUpOrientation
 {
 	enum Type
@@ -37,10 +60,10 @@ public:
 	 * @param  bInIsRightAltDown  True if right alt is pressed
 	 */
 	FModifierKeysState( const bool bInIsLeftShiftDown,
-					    const bool bInIsRightShiftDown,
-					    const bool bInIsLeftControlDown,
-					    const bool bInIsRightControlDown,
-					    const bool bInIsLeftAltDown,
+						const bool bInIsRightShiftDown,
+						const bool bInIsLeftControlDown,
+						const bool bInIsRightControlDown,
+						const bool bInIsLeftAltDown,
 						const bool bInIsRightAltDown,
 						const bool bInIsLeftCommandDown,
 						const bool bInIsRightCommandDown,
@@ -199,6 +222,34 @@ public:
 		return bAreCapsLocked;
 	}
 
+	/**
+	 * @param ModifierKeys the modifier keys to test to see if they are pressed.  Returns true if no modifiers are specified.
+	 * @return true if all modifier keys are pressed specified in the modifier keys.
+	 */
+	bool AreModifersDown(EModifierKey::Type ModiferKeys) const
+	{
+		bool AllModifersDown = true;
+
+		if ( (ModiferKeys & EModifierKey::Shift) == EModifierKey::Shift )
+		{
+			AllModifersDown &= IsShiftDown();
+		}
+		if ( (ModiferKeys & EModifierKey::Command) == EModifierKey::Command )
+		{
+			AllModifersDown &= IsCommandDown();
+		}
+		if ( (ModiferKeys & EModifierKey::Control) == EModifierKey::Control )
+		{
+			AllModifersDown &= IsControlDown();
+		}
+		if ( (ModiferKeys & EModifierKey::Alt) == EModifierKey::Alt )
+		{
+			AllModifersDown &= IsAltDown();
+		}
+
+		return AllModifersDown;
+	}
+
 private:
 
 	/** True if the left shift key was down when this event occurred. */
@@ -376,9 +427,9 @@ public:
 
 	/** Function to return the current implementation of the Text Input Method System */
 	virtual ITextInputMethodSystem *GetTextInputMethodSystem() { return NULL; }
-    
-    /** Send any analytics captured by the application */
-    virtual void SendAnalytics(IAnalyticsProvider* Provider) { }
+	
+	/** Send any analytics captured by the application */
+	virtual void SendAnalytics(IAnalyticsProvider* Provider) { }
 
 	virtual bool SupportsSystemHelp() const { return false; }
 
diff --git a/Engine/Source/Runtime/Slate/Private/Widgets/Input/SMultiLineEditableTextBox.cpp b/Engine/Source/Runtime/Slate/Private/Widgets/Input/SMultiLineEditableTextBox.cpp
--- a/Engine/Source/Runtime/Slate/Private/Widgets/Input/SMultiLineEditableTextBox.cpp
+++ b/Engine/Source/Runtime/Slate/Private/Widgets/Input/SMultiLineEditableTextBox.cpp
@@ -98,6 +98,7 @@ void SMultiLineEditableTextBox::Construct( const FArguments& InArgs )
 					.AutoWrapText(InArgs._AutoWrapText)
 					.HScrollBar(HScrollBar)
 					.VScrollBar(VScrollBar)
+					.ModiferKeyForNewLine(InArgs._ModiferKeyForNewLine)
 				]
 
 				+SVerticalBox::Slot()
diff --git a/Engine/Source/Runtime/Slate/Private/Widgets/Text/SMultiLineEditableText.cpp b/Engine/Source/Runtime/Slate/Private/Widgets/Text/SMultiLineEditableText.cpp
--- a/Engine/Source/Runtime/Slate/Private/Widgets/Text/SMultiLineEditableText.cpp
+++ b/Engine/Source/Runtime/Slate/Private/Widgets/Text/SMultiLineEditableText.cpp
@@ -229,6 +229,8 @@ void SMultiLineEditableText::Construct( const FArguments& InArgs )
 	SelectionStart = TOptional<FTextLocation>();
 	PreferredCursorScreenOffsetInLine = 0;
 
+	ModiferKeyForNewLine = InArgs._ModiferKeyForNewLine;
+
 	WrapTextAt = InArgs._WrapTextAt;
 	AutoWrapText = InArgs._AutoWrapText;
 	CachedSize = FVector2D::ZeroVector;
@@ -1710,22 +1712,39 @@ FReply SMultiLineEditableText::OnEscape()
 
 void SMultiLineEditableText::OnEnter()
 {
-	if (AnyTextSelected())
+	if ( FSlateApplication::Get().GetModifierKeys().AreModifersDown(ModiferKeyForNewLine) )
 	{
-		// Delete selected text
-		DeleteSelectedText();
-	}
+		if ( AnyTextSelected() )
+		{
+			// Delete selected text
+			DeleteSelectedText();
+		}
 
-	const FTextLocation CursorInteractionPosition = CursorInfo.GetCursorInteractionLocation();
-	if (TextLayout->SplitLineAt(CursorInteractionPosition))
-	{
-		// Adjust the cursor position to be at the beginning of the new line
-		const FTextLocation NewCursorPosition = FTextLocation(CursorInteractionPosition.GetLineIndex() + 1, 0);
-		CursorInfo.SetCursorLocationAndCalculateAlignment(TextLayout, NewCursorPosition);
+		const FTextLocation CursorInteractionPosition = CursorInfo.GetCursorInteractionLocation();
+		if ( TextLayout->SplitLineAt(CursorInteractionPosition) )
+		{
+			// Adjust the cursor position to be at the beginning of the new line
+			const FTextLocation NewCursorPosition = FTextLocation(CursorInteractionPosition.GetLineIndex() + 1, 0);
+			CursorInfo.SetCursorLocationAndCalculateAlignment(TextLayout, NewCursorPosition);
+		}
+
+		ClearSelection();
+		UpdateCursorHighlight();
 	}
+	else
+	{
+		//Always clear the local undo chain on commit.
+		ClearUndoStates();
 
-	ClearSelection();
-	UpdateCursorHighlight();
+		const FText EditedText = GetEditableText();
+
+		// When enter is pressed text is committed.  Let anyone interested know about it.
+		OnTextCommitted.ExecuteIfBound(EditedText, ETextCommit::OnEnter);
+
+		// Reload underlying value now it is committed  (commit may alter the value) 
+		// so it can be re-displayed in the edit box if it retains focus
+		LoadText();
+	}
 }
 
 bool SMultiLineEditableText::CanExecuteCut() const
diff --git a/Engine/Source/Runtime/Slate/Public/Framework/Commands/InputGesture.h b/Engine/Source/Runtime/Slate/Public/Framework/Commands/InputGesture.h
--- a/Engine/Source/Runtime/Slate/Public/Framework/Commands/InputGesture.h
+++ b/Engine/Source/Runtime/Slate/Public/Framework/Commands/InputGesture.h
@@ -3,30 +3,6 @@
 #pragma once
 
 
-/**
- * Enumerates available modifier keys for input gestures.
- */
-namespace EModifierKey
-{
-	typedef uint8 Type;
-
-	/** No key. */
-	const Type None	= 0;
-
-	/** Ctrl key (Command key on Mac, Control key on Windows). */
-	const Type Control = 1 << 0;
-
-	/** Alt key. */
-	const Type Alt = 1 << 1;
-
-	/** Shift key. */
-	const Type Shift = 1 << 2;
-
-	/** Cmd key (Control key on Mac, Win key on Windows) */
-	const Type Command = 1 << 3;
-};
-
-
 /**
  * Raw input gesture that defines input that must be valid when           
  */
diff --git a/Engine/Source/Runtime/Slate/Public/Widgets/Input/SMultiLineEditableTextBox.h b/Engine/Source/Runtime/Slate/Public/Widgets/Input/SMultiLineEditableTextBox.h
--- a/Engine/Source/Runtime/Slate/Public/Widgets/Input/SMultiLineEditableTextBox.h
+++ b/Engine/Source/Runtime/Slate/Public/Widgets/Input/SMultiLineEditableTextBox.h
@@ -37,6 +37,7 @@ public:
 		, _Padding()
 		, _Margin()
 		, _ErrorReporting()
+		, _ModiferKeyForNewLine(EModifierKey::None)
 		{}
 
 		/** The styling of the textbox */
@@ -49,10 +50,10 @@ public:
 		SLATE_ARGUMENT(TSharedPtr< ITextLayoutMarshaller >, Marshaller)
 
 		/** Sets the text content for this editable text box widget */
- 		SLATE_ATTRIBUTE( FText, Text )
+		SLATE_ATTRIBUTE( FText, Text )
 
 		/** Hint text that appears when there is no text in the text box */
- 		SLATE_ATTRIBUTE( FText, HintText )
+		SLATE_ATTRIBUTE( FText, HintText )
 
 		/** Font color and opacity (overrides Style) */
 		SLATE_ATTRIBUTE( FSlateFontInfo, Font )
@@ -118,7 +119,7 @@ public:
 		SLATE_ATTRIBUTE( float, WrapTextAt )
 
 		/** Whether to wrap text automatically based on the widget's computed horizontal space.  IMPORTANT: Using automatic wrapping can result
-		    in visual artifacts, as the the wrapped size will computed be at least one frame late!  Consider using WrapTextAt instead.  The initial 
+			in visual artifacts, as the the wrapped size will computed be at least one frame late!  Consider using WrapTextAt instead.  The initial 
 			desired size will not be clamped.  This works best in cases where the text block's size is not affecting other widget's layout. */
 		SLATE_ATTRIBUTE( bool, AutoWrapText )
 
@@ -138,6 +139,9 @@ public:
 		/** Provide a alternative mechanism for error reporting. */
 		SLATE_ARGUMENT( TSharedPtr<class IErrorReportingWidget>, ErrorReporting )
 
+		/** The optional modifier key necessary to create a newline when typing into the editor. */
+		SLATE_ARGUMENT( EModifierKey::Type, ModiferKeyForNewLine)
+
 	SLATE_END_ARGS()
 	
 	/**
diff --git a/Engine/Source/Runtime/Slate/Public/Widgets/Text/SMultiLineEditableText.h b/Engine/Source/Runtime/Slate/Public/Widgets/Text/SMultiLineEditableText.h
--- a/Engine/Source/Runtime/Slate/Public/Widgets/Text/SMultiLineEditableText.h
+++ b/Engine/Source/Runtime/Slate/Public/Widgets/Text/SMultiLineEditableText.h
@@ -32,6 +32,7 @@ public:
 		, _OnTextCommitted()
 		, _OnCursorMoved()
 		, _ContextMenuExtender()
+		, _ModiferKeyForNewLine(EModifierKey::None)
 	{}
 		/** The initial text that will appear in the widget. */
 		SLATE_ATTRIBUTE(FText, Text)
@@ -83,6 +84,9 @@ public:
 		/** Menu extender for the right-click context menu */
 		SLATE_EVENT(FMenuExtensionDelegate, ContextMenuExtender)
 
+		/** The optional modifier key necessary to create a newline when typing into the editor. */
+		SLATE_ARGUMENT(EModifierKey::Type, ModiferKeyForNewLine)
+
 	SLATE_END_ARGS()
 
 	SMultiLineEditableText();
@@ -643,6 +647,9 @@ private:
 
 	/** Whether the text has been changed by a virtual keyboard */
 	bool bTextChangedByVirtualKeyboard;
+
+	/** The optional modifier key necessary to create a newline when typing into the editor. */
+	EModifierKey::Type ModiferKeyForNewLine;
 };
 
 
