Make Core module public headers individually compilable. Part 2 of n.

diff --git a/Engine/Source/Runtime/Core/Public/Containers/ContainersFwd.h b/Engine/Source/Runtime/Core/Public/Containers/ContainersFwd.h
new file mode 100644
--- /dev/null
+++ b/Engine/Source/Runtime/Core/Public/Containers/ContainersFwd.h
@@ -0,0 +1,12 @@
+// Copyright 1998-2014 Epic Games, Inc. All Rights Reserved.
+
+#pragma once
+
+class FDefaultAllocator;
+class FDefaultSetAllocator;
+template<typename KeyType, typename ValueType> class TPair;
+template<typename T, typename Allocator = FDefaultAllocator> class TArray;
+template<typename T> class TTransArray;
+template<typename KeyType, typename ValueType, bool bInAllowDuplicateKeys> struct TDefaultMapKeyFuncs;
+template<typename KeyType, typename ValueType, typename SetAllocator = FDefaultSetAllocator, typename KeyFuncs = TDefaultMapKeyFuncs<KeyType, ValueType, false> > class TMap;
+template<typename KeyType, typename ValueType, typename SetAllocator = FDefaultSetAllocator, typename KeyFuncs = TDefaultMapKeyFuncs<KeyType, ValueType, true > > class TMultiMap;
\ No newline at end of file
diff --git a/Engine/Source/Runtime/Core/Public/Core.h b/Engine/Source/Runtime/Core/Public/Core.h
--- a/Engine/Source/Runtime/Core/Public/Core.h
+++ b/Engine/Source/Runtime/Core/Public/Core.h
@@ -12,8 +12,7 @@
 #include "CoreMiscDefines.h"					// Misc defines and enumerations
 class FString;									// @todo Figure out better include order in Core.h
 inline uint32 GetTypeHash (const int64 A);		// @todo Figure out better include order in Core.h
-class FDefaultAllocator;						// @todo Figure out better include order in Core.h
-template<typename T,typename Allocator = FDefaultAllocator> class TArray; // @todo Figure out better include order in Core.h
+#include "Containers/ContainersFwd.h"
 #include "Timespan.h"							// Time span definition
 #include "DateTime.h"							// Date and time handling
 #include "PlatformIncludes.h"					// Include the main and misc platform headers
@@ -46,12 +45,6 @@ template<typename T,typename Allocator = FDefaultAllocator> class TArray; // @to
 // Container forward declarations
 #include "Containers/ContainerAllocationPolicies.h"
 
-template<typename T> class TTransArray;
-template<typename KeyType,typename ValueType,bool bInAllowDuplicateKeys> struct TDefaultMapKeyFuncs;
-template<typename KeyType,typename ValueType,typename SetAllocator = FDefaultSetAllocator, typename KeyFuncs = TDefaultMapKeyFuncs<KeyType,ValueType,false> > class TMap;
-template<typename KeyType,typename ValueType,typename SetAllocator = FDefaultSetAllocator, typename KeyFuncs = TDefaultMapKeyFuncs<KeyType,ValueType,true > > class TMultiMap;
-
-
 // Objects
 class UObjectBase;
 class UObjectBaseUtility;
diff --git a/Engine/Source/Runtime/Core/Public/GenericPlatform/GenericPlatformAffinity.h b/Engine/Source/Runtime/Core/Public/GenericPlatform/GenericPlatformAffinity.h
--- a/Engine/Source/Runtime/Core/Public/GenericPlatform/GenericPlatformAffinity.h
+++ b/Engine/Source/Runtime/Core/Public/GenericPlatform/GenericPlatformAffinity.h
@@ -1,8 +1,7 @@
 // Copyright 1998-2014 Epic Games, Inc. All Rights Reserved.
 
 #pragma once
-
-#include "Core.h"
+#include "HAL/Platform.h"
 
 #define		MAKEAFFINITYMASK1(x)			(1<<x)
 #define		MAKEAFFINITYMASK2(x,y)			((1<<x)+(1<<y))
diff --git a/Engine/Source/Runtime/Core/Public/GenericPlatform/GenericPlatformAtomics.h b/Engine/Source/Runtime/Core/Public/GenericPlatform/GenericPlatformAtomics.h
--- a/Engine/Source/Runtime/Core/Public/GenericPlatform/GenericPlatformAtomics.h
+++ b/Engine/Source/Runtime/Core/Public/GenericPlatform/GenericPlatformAtomics.h
@@ -6,6 +6,7 @@
 ==============================================================================================*/
 
 #pragma once
+#include "HAL/Platform.h"
 
 /** Helper struct used to hold 128-bit values, internally represented as two 64-bit integers. */
 struct MS_ALIGN(16) FInt128 
diff --git a/Engine/Source/Runtime/Core/Public/GenericPlatform/GenericPlatformFile.h b/Engine/Source/Runtime/Core/Public/GenericPlatform/GenericPlatformFile.h
--- a/Engine/Source/Runtime/Core/Public/GenericPlatform/GenericPlatformFile.h
+++ b/Engine/Source/Runtime/Core/Public/GenericPlatform/GenericPlatformFile.h
@@ -6,6 +6,11 @@
 ==============================================================================================*/
 
 #pragma once
+#include "HAL/Platform.h"
+
+class FArchive;
+class FString;
+struct FDateTime;
 
 /** 
  * File handle interface. 
@@ -135,9 +140,9 @@ public:
 	/** Return the last access time of a file. Returns FDateTime::MinValue() on failure **/
 	virtual FDateTime	GetAccessTimeStamp(const TCHAR* Filename) = 0;
 
-	/** Attempt to open a file for reading. If successful will return a non-NULL pointer. Close the file by delete'ing the handle. **/
+	/** Attempt to open a file for reading. If successful will return a non-nullptr pointer. Close the file by delete'ing the handle. **/
 	virtual IFileHandle*	OpenRead(const TCHAR* Filename) = 0;
-	/** Attempt to open a file for writing. If successful will return a non-NULL pointer. Close the file by delete'ing the handle. **/
+	/** Attempt to open a file for writing. If successful will return a non-nullptr pointer. Close the file by delete'ing the handle. **/
 	virtual IFileHandle*	OpenWrite(const TCHAR* Filename, bool bAppend = 0, bool bAllowRead = 0) = 0;
 
 	/** Return true if the directory exists. **/
@@ -268,7 +273,7 @@ public:
 	virtual bool Initialize(IPlatformFile* Inner, const TCHAR* CmdLine) override;
 	virtual IPlatformFile* GetLowerLevel() override
 	{
-		return NULL;
+		return nullptr;
 	}
 	virtual const TCHAR* GetName() const override
 	{
diff --git a/Engine/Source/Runtime/Core/Public/GenericPlatform/GenericPlatformFramePacer.h b/Engine/Source/Runtime/Core/Public/GenericPlatform/GenericPlatformFramePacer.h
--- a/Engine/Source/Runtime/Core/Public/GenericPlatform/GenericPlatformFramePacer.h
+++ b/Engine/Source/Runtime/Core/Public/GenericPlatform/GenericPlatformFramePacer.h
@@ -8,6 +8,7 @@
 
 #pragma once
 
+class FEvent;
 
 /**
  * Generic implementation for most platforms, these tend to be unused and unimplemented
@@ -25,7 +26,7 @@ struct CORE_API FGenericPlatformRHIFramePacer
      * @param InTriggeredEvent - The event we wish to trigger when the frame interval has been triggered by the hardware.
      * @param InFrameInterval - How often should the event be triggered, in Frames.
      */
-	static void InitWithEvent(class FEvent* InTriggeredEvent) {}
+	static void InitWithEvent(FEvent* InTriggeredEvent) {}
     
     /**
      * Teardown the Frame Pacer.
diff --git a/Engine/Source/Runtime/Core/Public/GenericPlatform/GenericPlatformHttp.h b/Engine/Source/Runtime/Core/Public/GenericPlatform/GenericPlatformHttp.h
--- a/Engine/Source/Runtime/Core/Public/GenericPlatform/GenericPlatformHttp.h
+++ b/Engine/Source/Runtime/Core/Public/GenericPlatform/GenericPlatformHttp.h
@@ -2,7 +2,11 @@
 
 
 #pragma once
+#include "HAL/Platform.h"
 
+class FHttpManager;
+class FString;
+class IHttpRequest;
 
 /**
  * Platform specific Http implementations
@@ -19,11 +23,11 @@ public:
 	/**
 	 * Creates a platform-specific HTTP manager.
 	 *
-	 * @return NULL if default implementation is to be used
+	 * @return nullptr if default implementation is to be used
 	 */
-	static class FHttpManager * CreatePlatformHttpManager()
+	static FHttpManager* CreatePlatformHttpManager()
 	{
-		return NULL;
+		return nullptr;
 	}
 
 	/**
@@ -36,7 +40,7 @@ public:
 	 *
 	 * @return request object
 	 */
-	static class IHttpRequest* ConstructRequest();
+	static IHttpRequest* ConstructRequest();
 
 	/**
 	 * Returns a percent-encoded version of the passed in string
diff --git a/Engine/Source/Runtime/Core/Public/GenericPlatform/GenericPlatformMemory.h b/Engine/Source/Runtime/Core/Public/GenericPlatform/GenericPlatformMemory.h
--- a/Engine/Source/Runtime/Core/Public/GenericPlatform/GenericPlatformMemory.h
+++ b/Engine/Source/Runtime/Core/Public/GenericPlatform/GenericPlatformMemory.h
@@ -6,6 +6,12 @@
 ==============================================================================================*/
 
 #pragma once
+#include "HAL/Platform.h"
+#include "HAL/PlatformCodeAnalysis.h"
+#include <wchar.h>
+#include <string.h>
+
+class FString;
 
 /** Holds generic memory stats, internally implemented as a map. */
 struct FGenericMemoryStats;
diff --git a/Engine/Source/Runtime/Core/Public/GenericPlatform/GenericPlatformMisc.h b/Engine/Source/Runtime/Core/Public/GenericPlatform/GenericPlatformMisc.h
--- a/Engine/Source/Runtime/Core/Public/GenericPlatform/GenericPlatformMisc.h
+++ b/Engine/Source/Runtime/Core/Public/GenericPlatform/GenericPlatformMisc.h
@@ -6,13 +6,20 @@
 ==============================================================================================*/
 
 #pragma once
+#include "Containers/ContainersFwd.h"
+#include "HAL/Platform.h"
 
 struct FGenericCrashContext;
 struct FGenericMemoryWarningContext;
 
+class FDefaultAllocator;
+class FOutputDevice;
+class FString;
 class FText;
+class GenericApplication;
+class IPlatformChunkInstall;
+class UWorld;
 
-template<typename KeyType, typename ValueType> class TPair;
 
 namespace EBuildConfigurations
 {
@@ -208,11 +215,11 @@ struct CORE_API FGenericPlatformMisc
 	{
 	}
 
-	static class GenericApplication* CreateApplication();
+	static GenericApplication* CreateApplication();
 
 	static void* GetHardwareWindow()
 	{
-		return NULL;
+		return nullptr;
 	}
 
 	/**
@@ -352,7 +359,7 @@ struct CORE_API FGenericPlatformMisc
 	/**
 	 * Handles IO failure by ending gameplay.
 	 *
-	 * @param Filename	If not NULL, name of the file the I/O error occured with
+	 * @param Filename	If not nullptr, name of the file the I/O error occured with
 	 */
 	void static HandleIOFailure( const TCHAR* Filename );
 
@@ -614,18 +621,18 @@ public:
 	static int32 NumberOfWorkerThreadsToSpawn();
 
 	/**
-	 * Return the platform specific async IO system, or NULL if the standard one should be used.
+	 * Return the platform specific async IO system, or nullptr if the standard one should be used.
 	 */
 	static struct FAsyncIOSystemBase* GetPlatformSpecificAsyncIOSystem()
 	{
-		return NULL;
+		return nullptr;
 	}
 
-	/** Return the name of the platform features module. Can be NULL if there are no extra features for this platform */
+	/** Return the name of the platform features module. Can be nullptr if there are no extra features for this platform */
 	static const TCHAR* GetPlatformFeaturesModuleName()
 	{
 		// by default, no module
-		return NULL;
+		return nullptr;
 	}
 
 	/** Get the application root directory. */
@@ -666,7 +673,7 @@ public:
 	 *
 	 * @return	Returns the platform specific chunk based install implementation
 	 */
-	static class IPlatformChunkInstall* GetPlatformChunkInstall();
+	static IPlatformChunkInstall* GetPlatformChunkInstall();
 
 	/**
 	 * Has the OS execute a command and path pair (such as launch a browser)
@@ -718,7 +725,7 @@ public:
 	 *
 	 *	@return	bool		true if command was processed, false if not
 	 */
-	static bool Exec(class UWorld* InWorld, const TCHAR* Cmd, class FOutputDevice& Out)
+	static bool Exec(UWorld* InWorld, const TCHAR* Cmd, FOutputDevice& Out)
 	{
 		return false;
 	}
@@ -777,7 +784,7 @@ public:
 	 * Return an ordered list of target platforms this runtime can support (ie Android_DXT, Android
 	 * would mean that it prefers Android_DXT, but can use Android as well)
 	 */
-	static void GetValidTargetPlatforms(class TArray<class FString>& TargetPlatformNames);
+	static void GetValidTargetPlatforms(TArray<FString>& TargetPlatformNames);
 
 	/**
 	 * Returns whether the platform wants to use a touch screen for virtual joysticks.
diff --git a/Engine/Source/Runtime/Core/Public/GenericPlatform/GenericPlatformNamedPipe.h b/Engine/Source/Runtime/Core/Public/GenericPlatform/GenericPlatformNamedPipe.h
--- a/Engine/Source/Runtime/Core/Public/GenericPlatform/GenericPlatformNamedPipe.h
+++ b/Engine/Source/Runtime/Core/Public/GenericPlatform/GenericPlatformNamedPipe.h
@@ -6,7 +6,9 @@
 
 #pragma once
 
-#include "Platform.h"
+#include "HAL/Platform.h"
+
+class FString;
 
 #if PLATFORM_SUPPORTS_NAMED_PIPES
 
diff --git a/Engine/Source/Runtime/Core/Public/GenericPlatform/GenericPlatformOutputDevices.h b/Engine/Source/Runtime/Core/Public/GenericPlatform/GenericPlatformOutputDevices.h
--- a/Engine/Source/Runtime/Core/Public/GenericPlatform/GenericPlatformOutputDevices.h
+++ b/Engine/Source/Runtime/Core/Public/GenericPlatform/GenericPlatformOutputDevices.h
@@ -6,6 +6,13 @@
 ==============================================================================================*/
 
 #pragma once
+#include "HAL/Platform.h"
+
+class FFeedbackContext;
+class FOutputDevice;
+class FOutputDeviceConsole;
+class FOutputDeviceError;
+class FString;
 
 /**
 * Generic implementation for most platforms
@@ -19,16 +26,16 @@ struct CORE_API FGenericPlatformOutputDevices
 
 	static FString						GetAbsoluteLogFilename();
 
-	static class FOutputDevice*			GetLog();
-	static class FOutputDevice*			GetEventLog()
+	static FOutputDevice*				GetLog();
+	static FOutputDevice*				GetEventLog()
 	{
-		return NULL; // normally only used for dedicated servers
+		return nullptr; // normally only used for dedicated servers
 	}
-	static class FOutputDeviceConsole*	GetLogConsole()
+	static FOutputDeviceConsole*		GetLogConsole()
 	{
-		return NULL; // normally only used for PC
+		return nullptr; // normally only used for PC
 	}
-	static class FOutputDeviceError*	GetError();
-	static class FFeedbackContext*		GetWarn();
+	static FOutputDeviceError*			GetError();
+	static FFeedbackContext*			GetWarn();
 };
 
diff --git a/Engine/Source/Runtime/Core/Public/HAL/PlatformCodeAnalysis.h b/Engine/Source/Runtime/Core/Public/HAL/PlatformCodeAnalysis.h
new file mode 100644
--- /dev/null
+++ b/Engine/Source/Runtime/Core/Public/HAL/PlatformCodeAnalysis.h
@@ -0,0 +1,8 @@
+// Copyright 1998-2014 Epic Games, Inc. All Rights Reserved.
+#pragma once
+#include "HAL/Platform.h"
+#include "Misc/CoreMiscDefines.h"
+
+#if PLATFORM_WINDOWS
+#include "Windows/WindowsPlatformCodeAnalysis.h"
+#endif
