Make Core module public headers individually compilable. Part 3 of n.

diff --git a/Engine/Source/Runtime/Core/Public/GenericPlatform/GenericPlatformProcess.h b/Engine/Source/Runtime/Core/Public/GenericPlatform/GenericPlatformProcess.h
--- a/Engine/Source/Runtime/Core/Public/GenericPlatform/GenericPlatformProcess.h
+++ b/Engine/Source/Runtime/Core/Public/GenericPlatform/GenericPlatformProcess.h
@@ -6,8 +6,11 @@
 ==============================================================================================*/
 
 #pragma once
+#include "Containers/ContainersFwd.h"
+#include "GenericPlatform/GenericPlatformMisc.h"
+#include "HAL/Platform.h"
 
-#include "Platform.h"
+class FString;
 
 namespace EProcessResource
 {
diff --git a/Engine/Source/Runtime/Core/Public/GenericPlatform/GenericPlatformProperties.h b/Engine/Source/Runtime/Core/Public/GenericPlatform/GenericPlatformProperties.h
--- a/Engine/Source/Runtime/Core/Public/GenericPlatform/GenericPlatformProperties.h
+++ b/Engine/Source/Runtime/Core/Public/GenericPlatform/GenericPlatformProperties.h
@@ -9,7 +9,8 @@
 ==============================================================================================*/
 
 #pragma once
-
+#include "GenericPlatform/GenericPlatformMisc.h"
+#include "HAL/Platform.h"
 
 /**
  * Base class for platform properties.
diff --git a/Engine/Source/Runtime/Core/Public/GenericPlatform/GenericPlatformSplash.h b/Engine/Source/Runtime/Core/Public/GenericPlatform/GenericPlatformSplash.h
--- a/Engine/Source/Runtime/Core/Public/GenericPlatform/GenericPlatformSplash.h
+++ b/Engine/Source/Runtime/Core/Public/GenericPlatform/GenericPlatformSplash.h
@@ -5,6 +5,7 @@
 ==============================================================================================*/
 
 #pragma once
+#include "HAL/Platform.h"
 
 /**
  * SplashTextType defines the types of text on the splash screen
diff --git a/Engine/Source/Runtime/Core/Public/GenericPlatform/GenericPlatformStackWalk.h b/Engine/Source/Runtime/Core/Public/GenericPlatform/GenericPlatformStackWalk.h
--- a/Engine/Source/Runtime/Core/Public/GenericPlatform/GenericPlatformStackWalk.h
+++ b/Engine/Source/Runtime/Core/Public/GenericPlatform/GenericPlatformStackWalk.h
@@ -6,6 +6,9 @@
 ==============================================================================================*/
 
 #pragma once
+#include "HAL/Platform.h"
+
+struct FGenericCrashContext;
 
 /**
 * This is used to capture all of the module information needed to load pdb's.
@@ -100,7 +103,7 @@ struct CORE_API FGenericPlatformStackWalk
 	 * @param	Context					Pointer to crash context, if any
 	 * @return	true if the symbol was found, otherwise false
 	 */ 
-	static bool ProgramCounterToHumanReadableString( int32 CurrentCallDepth, uint64 ProgramCounter, ANSICHAR* HumanReadableString, SIZE_T HumanReadableStringSize, FGenericCrashContext* Context = NULL );
+	static bool ProgramCounterToHumanReadableString( int32 CurrentCallDepth, uint64 ProgramCounter, ANSICHAR* HumanReadableString, SIZE_T HumanReadableStringSize, FGenericCrashContext* Context = nullptr );
 
 	/**
 	 * Converts the passed in program counter address to a symbol info struct, filling in module and filename, line number and displacement.
@@ -121,7 +124,7 @@ struct CORE_API FGenericPlatformStackWalk
 	 * @param	MaxDepth			Entries in BackTrace array
 	 * @param	Context				Optional thread context information
 	 */
-	static void CaptureStackBackTrace( uint64* BackTrace, uint32 MaxDepth, void* Context = NULL );
+	static void CaptureStackBackTrace( uint64* BackTrace, uint32 MaxDepth, void* Context = nullptr );
 
 	/**
 	 * Walks the stack and appends the human readable string to the passed in one.
@@ -132,7 +135,7 @@ struct CORE_API FGenericPlatformStackWalk
 	 * @param	IgnoreCount			Number of stack entries to ignore (some are guaranteed to be in the stack walking code)
 	 * @param	Context				Optional thread context information
 	 */ 
-	static void StackWalkAndDump( ANSICHAR* HumanReadableString, SIZE_T HumanReadableStringSize, int32 IgnoreCount, void* Context = NULL );
+	static void StackWalkAndDump( ANSICHAR* HumanReadableString, SIZE_T HumanReadableStringSize, int32 IgnoreCount, void* Context = nullptr );
 
 	/**
 	 * Returns the number of modules loaded by the currently running process.
diff --git a/Engine/Source/Runtime/Core/Public/GenericPlatform/GenericPlatformStricmp.h b/Engine/Source/Runtime/Core/Public/GenericPlatform/GenericPlatformStricmp.h
--- a/Engine/Source/Runtime/Core/Public/GenericPlatform/GenericPlatformStricmp.h
+++ b/Engine/Source/Runtime/Core/Public/GenericPlatform/GenericPlatformStricmp.h
@@ -1,9 +1,10 @@
 // Copyright 1998-2014 Epic Games, Inc. All Rights Reserved.
 
 #pragma once
-
-#include "Char.h"
-#include "UnrealTypeTraits.h"
+#include "HAL/Platform.h"
+#include "Misc/Char.h"
+#include "Templates/EnableIf.h"
+#include "Templates/UnrealTypeTraits.h"
 
 /**
  * This trait tells if given CharTypeA is comparison compatible with CharTypeB, i.e.
diff --git a/Engine/Source/Runtime/Core/Public/GenericPlatform/GenericPlatformString.h b/Engine/Source/Runtime/Core/Public/GenericPlatform/GenericPlatformString.h
--- a/Engine/Source/Runtime/Core/Public/GenericPlatform/GenericPlatformString.h
+++ b/Engine/Source/Runtime/Core/Public/GenericPlatform/GenericPlatformString.h
@@ -6,7 +6,7 @@
 ==============================================================================================*/
 
 #pragma once
-
+#include "HAL/Platform.h"
 #include "Templates/EnableIf.h"
 
 /**
diff --git a/Engine/Source/Runtime/Core/Public/GenericPlatform/GenericPlatformSurvey.h b/Engine/Source/Runtime/Core/Public/GenericPlatform/GenericPlatformSurvey.h
--- a/Engine/Source/Runtime/Core/Public/GenericPlatform/GenericPlatformSurvey.h
+++ b/Engine/Source/Runtime/Core/Public/GenericPlatform/GenericPlatformSurvey.h
@@ -6,7 +6,7 @@
 ==============================================================================================*/
 
 #pragma once
-
+#include "HAL/Platform.h"
 
 // time and amount of work that was measured
 struct FTimeSample
diff --git a/Engine/Source/Runtime/Core/Public/GenericPlatform/GenericPlatformTLS.h b/Engine/Source/Runtime/Core/Public/GenericPlatform/GenericPlatformTLS.h
--- a/Engine/Source/Runtime/Core/Public/GenericPlatform/GenericPlatformTLS.h
+++ b/Engine/Source/Runtime/Core/Public/GenericPlatform/GenericPlatformTLS.h
@@ -6,6 +6,7 @@
 ==============================================================================================*/
 
 #pragma once
+#include "HAL/Platform.h"
 
 /**
  * It should be possible to provide a generic implementation as long as a threadID is provided. We don't do that yet.
diff --git a/Engine/Source/Runtime/Core/Public/GenericPlatform/GenericPlatformTime.h b/Engine/Source/Runtime/Core/Public/GenericPlatform/GenericPlatformTime.h
--- a/Engine/Source/Runtime/Core/Public/GenericPlatform/GenericPlatformTime.h
+++ b/Engine/Source/Runtime/Core/Public/GenericPlatform/GenericPlatformTime.h
@@ -7,7 +7,9 @@
 
 #pragma once
 
-#include "Platform.h"
+#include "HAL/Platform.h"
+
+class FString;
 
 #if PLATFORM_HAS_BSD_TIME 
 #include <sys/time.h>
diff --git a/Engine/Source/Runtime/Core/Public/GenericPlatform/ICursor.h b/Engine/Source/Runtime/Core/Public/GenericPlatform/ICursor.h
--- a/Engine/Source/Runtime/Core/Public/GenericPlatform/ICursor.h
+++ b/Engine/Source/Runtime/Core/Public/GenericPlatform/ICursor.h
@@ -2,7 +2,16 @@
 
 
 #pragma once
+#include "HAL/Platform.h"
 
+struct FVector2D;
+
+#if PLATFORM_WINDOWS || PLATFORM_XBOXONE || PLATFORM_WINRT
+struct tagRECT;
+typedef struct tagRECT RECT;
+#else
+struct RECT;
+#endif
 /**
  * Mouse cursor types
  */
diff --git a/Engine/Source/Runtime/Core/Public/GenericPlatform/IForceFeedbackSystem.h b/Engine/Source/Runtime/Core/Public/GenericPlatform/IForceFeedbackSystem.h
--- a/Engine/Source/Runtime/Core/Public/GenericPlatform/IForceFeedbackSystem.h
+++ b/Engine/Source/Runtime/Core/Public/GenericPlatform/IForceFeedbackSystem.h
@@ -1,6 +1,7 @@
 // Copyright 1998-2014 Epic Games, Inc. All Rights Reserved.
 
 #pragma once
+#include "HAL/Platform.h"
 
 // General identifiers for potential force feedback channels. These will be mapped according to the
 // platform specific implementation.
diff --git a/Engine/Source/Runtime/Core/Public/GenericPlatform/StandardPlatformString.h b/Engine/Source/Runtime/Core/Public/GenericPlatform/StandardPlatformString.h
--- a/Engine/Source/Runtime/Core/Public/GenericPlatform/StandardPlatformString.h
+++ b/Engine/Source/Runtime/Core/Public/GenericPlatform/StandardPlatformString.h
@@ -1,9 +1,11 @@
 // Copyright 1998-2014 Epic Games, Inc. All Rights Reserved.
 
 #pragma once
-
-#include "Char.h"
-#include "GenericPlatformStricmp.h"
+#include "GenericPlatform/GenericPlatformStricmp.h"
+#include "GenericPlatform/GenericPlatformString.h"
+#include "HAL/Platform.h"
+#include "Misc/Char.h"
+#include <wchar.h>
 
 /**
 * Standard implementation
diff --git a/Engine/Source/Runtime/Core/Public/Windows/WindowsHWrapper.h b/Engine/Source/Runtime/Core/Public/Windows/WindowsHWrapper.h
new file mode 100644
--- /dev/null
+++ b/Engine/Source/Runtime/Core/Public/Windows/WindowsHWrapper.h
@@ -0,0 +1,9 @@
+// Copyright 1998-2014 Epic Games, Inc. All Rights Reserved.
+#pragma once
+
+#include "Windows/PreWindowsApi.h"
+#ifndef STRICT
+#define STRICT
+#endif
+#include "Windows/MinWindows.h"
+#include "Windows/PostWindowsApi.h"
\ No newline at end of file
