Use FCocoaWindow.DeferFrame instead of actual window frame for hidden windows when calling SlateApplication's OnMovedWindow in FMacApplication::OnWindowDidMove

diff --git a/Engine/Source/Runtime/Core/Private/Mac/CocoaWindow.cpp b/Engine/Source/Runtime/Core/Private/Mac/CocoaWindow.cpp
--- a/Engine/Source/Runtime/Core/Private/Mac/CocoaWindow.cpp
+++ b/Engine/Source/Runtime/Core/Private/Mac/CocoaWindow.cpp
@@ -34,8 +34,8 @@ NSString* NSPerformDragOperation = @"NSPerformDragOperation";
 	bDeferOrderFront = false;
 	DeferOpacity = 0.0f;
 	bRenderInitialised = false;
-	bDeferSetFrame = false;
-	bDeferSetOrigin = false;
+	self.bDeferSetFrame = false;
+	self.bDeferSetOrigin = false;
 
 	id NewSelf = [super initWithContentRect:ContentRect styleMask:Style backing:BufferingType defer:Flag];
 	if(NewSelf)
@@ -44,8 +44,8 @@ NSString* NSPerformDragOperation = @"NSPerformDragOperation";
 		self.bForwardEvents = true;
 		self.TargetWindowMode = EWindowMode::Windowed;
 		[super setAlphaValue:DeferOpacity];
-		DeferFrame = [super frame];
-		self.PreFullScreenRect = DeferFrame;
+		self.DeferFrame = [super frame];
+		self.PreFullScreenRect = self.DeferFrame;
 	}
 	return NewSelf;
 }
@@ -64,11 +64,11 @@ NSString* NSPerformDragOperation = @"NSPerformDragOperation";
 	}
 	else if([self styleMask] & (NSTexturedBackgroundWindowMask))
 	{
-		return (!bDeferSetFrame ? [self frame] : DeferFrame);
+		return (!self.bDeferSetFrame ? [self frame] : self.DeferFrame);
 	}
 	else
 	{
-		return (!bDeferSetFrame ? [[self contentView] frame] : [self contentRectForFrameRect:DeferFrame]);
+		return (!self.bDeferSetFrame ? [[self contentView] frame] : [self contentRectForFrameRect:self.DeferFrame]);
 	}
 }
 
@@ -104,7 +104,7 @@ NSString* NSPerformDragOperation = @"NSPerformDragOperation";
 	if(bDeferOrderFront)
 	{
 		SCOPED_AUTORELEASE_POOL;
-		if(!(bDeferSetFrame || bDeferSetOrigin))
+		if(!(self.bDeferSetFrame || self.bDeferSetOrigin))
 		{
 			bDeferOrderFront = false;
 			[super setAlphaValue:DeferOpacity];
@@ -118,17 +118,17 @@ NSString* NSPerformDragOperation = @"NSPerformDragOperation";
 
 - (void)performDeferredSetFrame
 {
-	if(bRenderInitialised && (bDeferSetFrame || bDeferSetOrigin))
+	if(bRenderInitialised && (self.bDeferSetFrame || self.bDeferSetOrigin))
 	{
 		SCOPED_AUTORELEASE_POOL;
 		dispatch_block_t Block = ^{
 			SCOPED_AUTORELEASE_POOL;
-			if(!bDeferSetFrame && bDeferSetOrigin)
+			if(!self.bDeferSetFrame && self.bDeferSetOrigin)
 			{
-				DeferFrame.size = [self frame].size;
+				self.DeferFrame.size = [self frame].size;
 			}
 			
-			[super setFrame:DeferFrame display:YES];
+			[super setFrame:self.DeferFrame display:YES];
 		};
 		
 		if([NSThread isMainThread])
@@ -140,8 +140,8 @@ NSString* NSPerformDragOperation = @"NSPerformDragOperation";
 			dispatch_async(dispatch_get_main_queue(), Block);
 		}
 		
-		bDeferSetFrame = false;
-		bDeferSetOrigin = false;
+		self.bDeferSetFrame = false;
+		self.bDeferSetOrigin = false;
 	}
 }
 
@@ -307,12 +307,12 @@ NSString* NSPerformDragOperation = @"NSPerformDragOperation";
 	if(!bRenderInitialised || ([self isVisible] && [super alphaValue] > 0.0f && (Size.width > 1 || Size.height > 1 || NewSize.width > 1 || NewSize.height > 1)))
 	{
 		[super setFrame:FrameRect display:Flag];
-		bDeferSetFrame = false;
+		self.bDeferSetFrame = false;
 	}
 	else
 	{
-		bDeferSetFrame = true;
-		DeferFrame = FrameRect;
+		self.bDeferSetFrame = true;
+		self.DeferFrame = FrameRect;
 		if(self.bForwardEvents)
 		{
 			NSNotification* Notification = [NSNotification notificationWithName:NSWindowDidResizeNotification object:self];
@@ -331,12 +331,12 @@ NSString* NSPerformDragOperation = @"NSPerformDragOperation";
 			SCOPED_AUTORELEASE_POOL;
 			[super setFrameOrigin:Point];
 		});
-		bDeferSetOrigin = false;
+		self.bDeferSetOrigin = false;
 	}
 	else
 	{
-		bDeferSetOrigin = true;
-		DeferFrame.origin = Point;
+		self.bDeferSetOrigin = true;
+		self.DeferFrame.origin = Point;
 		NSNotification* Notification = [NSNotification notificationWithName:NSWindowDidMoveNotification object:self];
 		FMacEvent::SendToGameRunLoop(Notification, self, EMacEventSendMethod::Async, @[ NSDefaultRunLoopMode, UE4ResizeEventMode, UE4ShowEventMode ]);
 	}
diff --git a/Engine/Source/Runtime/Core/Private/Mac/MacApplication.cpp b/Engine/Source/Runtime/Core/Private/Mac/MacApplication.cpp
--- a/Engine/Source/Runtime/Core/Private/Mac/MacApplication.cpp
+++ b/Engine/Source/Runtime/Core/Private/Mac/MacApplication.cpp
@@ -1296,7 +1296,7 @@ void FMacApplication::OnWindowDidMove( FCocoaWindow* Window )
 {
 	SCOPED_AUTORELEASE_POOL;
 
-	NSRect WindowFrame = [Window frame];
+	NSRect WindowFrame = (Window.bDeferSetFrame || Window.bDeferSetOrigin) ? Window.DeferFrame : [Window frame];
 	NSRect OpenGLFrame = [Window openGLFrame];
 	
 	const int32 X = (int32)WindowFrame.origin.x;
diff --git a/Engine/Source/Runtime/Core/Public/Mac/CocoaWindow.h b/Engine/Source/Runtime/Core/Public/Mac/CocoaWindow.h
--- a/Engine/Source/Runtime/Core/Public/Mac/CocoaWindow.h
+++ b/Engine/Source/Runtime/Core/Public/Mac/CocoaWindow.h
@@ -10,14 +10,11 @@
 @interface FCocoaWindow : NSWindow <NSWindowDelegate>
 {
 	EWindowMode::Type WindowMode;
-	NSRect DeferFrame;
 	CGFloat DeferOpacity;
 	bool bAcceptsInput;
 	bool bRoundedCorners;
 	bool bDisplayReconfiguring;
 	bool bDeferOrderFront;
-	bool bDeferSetFrame;
-	bool bDeferSetOrigin;
 	bool bRenderInitialised;
 	bool bNeedsRedraw;
 @public
@@ -27,6 +24,9 @@
 @property (assign) NSRect PreFullScreenRect;
 @property (assign) EWindowMode::Type TargetWindowMode;
 @property (assign) bool bForwardEvents;
+@property (assign) bool bDeferSetFrame;
+@property (assign) bool bDeferSetOrigin;
+@property (assign) NSRect DeferFrame;
 
 /** Get the frame filled by a child OpenGL view, which may cover the window or fill the content view depending upon the window style.
  @return The NSRect for a child OpenGL view. */
