Core: Replaced namespaced enums; added more unit tests

diff --git a/Engine/Source/Runtime/Core/Private/Misc/DateTime.cpp b/Engine/Source/Runtime/Core/Private/Misc/DateTime.cpp
--- a/Engine/Source/Runtime/Core/Private/Misc/DateTime.cpp
+++ b/Engine/Source/Runtime/Core/Private/Misc/DateTime.cpp
@@ -84,7 +84,7 @@ void FDateTime::GetDate( int32& OutYear, int32& OutMonth, int32& OutDay ) const
 }
 
 
-int32 FDateTime::GetDay( ) const
+int32 FDateTime::GetDay() const
 {
 	int32 Year, Month, Day;
 	GetDate(Year, Month, Day);
@@ -93,14 +93,14 @@ int32 FDateTime::GetDay( ) const
 }
 
 
-EDayOfWeek::Type FDateTime::GetDayOfWeek( ) const
+EDayOfWeek FDateTime::GetDayOfWeek() const
 {
 	// January 1, 0001 was a Monday
-	return static_cast<EDayOfWeek::Type>((Ticks / ETimespan::TicksPerDay) % 7);
+	return static_cast<EDayOfWeek>((Ticks / ETimespan::TicksPerDay) % 7);
 }
 
 
-int32 FDateTime::GetDayOfYear( ) const
+int32 FDateTime::GetDayOfYear() const
 {
 	int32 Year, Month, Day;
 	GetDate(Year, Month, Day);
@@ -114,7 +114,7 @@ int32 FDateTime::GetDayOfYear( ) const
 }
 
 
-int32 FDateTime::GetHour12( ) const
+int32 FDateTime::GetHour12() const
 {
 	int32 Hour = GetHour();
 
@@ -132,7 +132,7 @@ int32 FDateTime::GetHour12( ) const
 }
 
 
-int32 FDateTime::GetMonth( ) const
+int32 FDateTime::GetMonth() const
 {
 	int32 Year, Month, Day;
 	GetDate(Year, Month, Day);
@@ -141,7 +141,7 @@ int32 FDateTime::GetMonth( ) const
 }
 
 
-int32 FDateTime::GetYear( ) const
+int32 FDateTime::GetYear() const
 {
 	int32 Year, Month, Day;
 	GetDate(Year, Month, Day);
@@ -176,13 +176,13 @@ bool FDateTime::Serialize( FArchive& Ar )
 }
 
 
-FString FDateTime::ToIso8601( ) const
+FString FDateTime::ToIso8601() const
 {
 	return ToString(TEXT("%Y-%m-%dT%H:%M:%S.%sZ"));
 }
 
 
-FString FDateTime::ToString( ) const
+FString FDateTime::ToString() const
 {
 	return ToString(TEXT("%Y.%m.%d-%H.%M.%S"));
 }
@@ -236,7 +236,7 @@ int32 FDateTime::DaysInMonth( int32 Year, int32 Month )
 {
 	check((Month >= 1) && (Month <= 12));
 
-	if ((Month == EMonthOfYear::February) && IsLeapYear(Year))
+	if ((Month == 2) && IsLeapYear(Year))
 	{
 		return 29;
 	}
@@ -267,7 +267,7 @@ bool FDateTime::IsLeapYear( int32 Year )
 }
 
 
-FDateTime FDateTime::Now( )
+FDateTime FDateTime::Now()
 {
 	int32 Year, Month, Day, DayOfWeek;
 	int32 Hour, Minute, Second, Millisecond;
@@ -439,7 +439,7 @@ bool FDateTime::ParseIso8601( const TCHAR* DateTimeString, FDateTime& OutDateTim
 }
 
 
-FDateTime FDateTime::UtcNow( )
+FDateTime FDateTime::UtcNow()
 {
 	int32 Year, Month, Day, DayOfWeek;
 	int32 Hour, Minute, Second, Millisecond;
diff --git a/Engine/Source/Runtime/Core/Private/Misc/Guid.cpp b/Engine/Source/Runtime/Core/Private/Misc/Guid.cpp
--- a/Engine/Source/Runtime/Core/Private/Misc/Guid.cpp
+++ b/Engine/Source/Runtime/Core/Private/Misc/Guid.cpp
@@ -32,7 +32,7 @@ bool FGuid::ImportTextItem( const TCHAR*& Buffer, int32 PortFlags, class UObject
 }
 
 
-FString FGuid::ToString( EGuidFormats::Type Format ) const
+FString FGuid::ToString( EGuidFormats Format ) const
 {
 	switch (Format)
 	{
@@ -60,7 +60,7 @@ FString FGuid::ToString( EGuidFormats::Type Format ) const
 /* FGuid static interface
  *****************************************************************************/
 
-FGuid FGuid::NewGuid( )
+FGuid FGuid::NewGuid()
 {
 	FGuid Result(0,0,0,0);
 	FPlatformMisc::CreateGuid(Result);
@@ -104,7 +104,7 @@ bool FGuid::Parse( const FString& GuidString, FGuid& OutGuid )
 }
 
 
-bool FGuid::ParseExact( const FString& GuidString, EGuidFormats::Type Format, FGuid& OutGuid )
+bool FGuid::ParseExact( const FString& GuidString, EGuidFormats Format, FGuid& OutGuid )
 {
 	FString NormalizedGuidString;
 
diff --git a/Engine/Source/Runtime/Core/Private/Tests/Misc/DateTimeTest.cpp b/Engine/Source/Runtime/Core/Private/Tests/Misc/DateTimeTest.cpp
--- a/Engine/Source/Runtime/Core/Private/Tests/Misc/DateTimeTest.cpp
+++ b/Engine/Source/Runtime/Core/Private/Tests/Misc/DateTimeTest.cpp
@@ -1,16 +1,13 @@
 // Copyright 1998-2014 Epic Games, Inc. All Rights Reserved.
 
-/*=============================================================================
-	DateTimeTest.cpp: Unit test for the FDateTime class.
-=============================================================================*/
-
 #include "Core.h"
 #include "AutomationTest.h"
 
 
 IMPLEMENT_SIMPLE_AUTOMATION_TEST(FDateTimeTest, "Core.Misc.DateTime", EAutomationTestFlags::ATF_SmokeTest)
 
-bool FDateTimeTest::RunTest( const FString& Parameters )
+
+	bool FDateTimeTest::RunTest( const FString& Parameters )
 {
 	const int32 UnixEpochTimestamp = 0;
 	const int32 UnixBillenniumTimestamp = 1000000000;
@@ -34,7 +31,8 @@ bool FDateTimeTest::RunTest( const FString& Parameters )
 #undef TestUnixEquivalent
 
 #define TestYear( Desc, A, B ) if( (A.GetYear()) != (B) ) AddError(FString::Printf(TEXT("%s - A=%d B=%d"),Desc,(A.GetYear()),(B)));
-#define TestMonth( Desc, A, B ) if( (A.GetMonth()) != (B) ) AddError(FString::Printf(TEXT("%s - A=%d B=%d"),Desc,(A.GetMonth()),static_cast<int32>(B)));
+#define TestMonth( Desc, A, B ) if( (A.GetMonth()) != (B) ) AddError(FString::Printf(TEXT("%s - A=%d B=%d"),Desc,(A.GetMonth()),B));
+#define TestMonthOfYear( Desc, A, B ) if( (A.GetMonthOfYear()) != (B) ) AddError(FString::Printf(TEXT("%s - A=%d B=%d"),Desc,(static_cast<int32>(A.GetMonthOfYear())),static_cast<int32>(B)));
 #define TestDay( Desc, A, B ) if( (A.GetDay()) != (B) ) AddError(FString::Printf(TEXT("%s - A=%d B=%d"),Desc,(A.GetDay()),(B)));
 #define TestHour( Desc, A, B ) if( (A.GetHour()) != (B) ) AddError(FString::Printf(TEXT("%s - A=%d B=%d"),Desc,(A.GetHour()),(B)));
 #define TestMinute( Desc, A, B ) if( (A.GetMinute()) != (B) ) AddError(FString::Printf(TEXT("%s - A=%d B=%d"),Desc,(A.GetMinute()),(B)));
@@ -42,7 +40,8 @@ bool FDateTimeTest::RunTest( const FString& Parameters )
 #define TestMillisecond( Desc, A, B ) if( (A.GetMillisecond()) != (B) ) AddError(FString::Printf(TEXT("%s - A=%d B=%d"),Desc,(A.GetMillisecond()),(B)));
 
 	TestYear( TEXT("Testing Unix Epoch Year"), UnixEpoch, 1970 );
-	TestMonth( TEXT("Testing Unix Epoch Month"), UnixEpoch, EMonthOfYear::January );
+	TestMonth( TEXT("Testing Unix Epoch Month"), UnixEpoch, 1 );
+	TestMonthOfYear( TEXT("Testing Unix Epoch MonthOfYear"), UnixEpoch, EMonthOfYear::January );
 	TestDay( TEXT("Testing Unix Epoch Day"), UnixEpoch, 1 );
 	TestHour( TEXT("Testing Unix Epoch Hour"), UnixEpoch, 0 );
 	TestMinute( TEXT("Testing Unix Epoch Minute"), UnixEpoch, 0 );
@@ -50,7 +49,8 @@ bool FDateTimeTest::RunTest( const FString& Parameters )
 	TestMillisecond( TEXT("Testing Unix Epoch Millisecond"), UnixEpoch, 0 );
 
 	TestYear( TEXT("Testing Unix Billennium Year"), UnixBillennium, 2001 );
-	TestMonth( TEXT("Testing Unix Billennium Month"), UnixBillennium, EMonthOfYear::September );
+	TestMonth( TEXT("Testing Unix Billennium MonthOfYear"), UnixBillennium, 9 );
+	TestMonthOfYear( TEXT("Testing Unix Billennium Month"), UnixBillennium, EMonthOfYear::September );
 	TestDay( TEXT("Testing Unix Billennium Day"), UnixBillennium, 9 );
 	TestHour( TEXT("Testing Unix Billennium Hour"), UnixBillennium, 1 );
 	TestMinute( TEXT("Testing Unix Billennium Minute"), UnixBillennium, 46 );
@@ -58,7 +58,8 @@ bool FDateTimeTest::RunTest( const FString& Parameters )
 	TestMillisecond( TEXT("Testing Unix Billennium Millisecond"), UnixBillennium, 0 );
 
 	TestYear( TEXT("Testing Unix Ones Year"), UnixOnes, 2005 );
-	TestMonth( TEXT("Testing Unix Ones Month"), UnixOnes, EMonthOfYear::March );
+	TestMonth( TEXT("Testing Unix Ones Month"), UnixOnes, 3 );
+	TestMonthOfYear( TEXT("Testing Unix Ones MonthOfYear"), UnixOnes, EMonthOfYear::March );
 	TestDay( TEXT("Testing Unix Ones Day"), UnixOnes, 18 );
 	TestHour( TEXT("Testing Unix Ones Hour"), UnixOnes, 1 );
 	TestMinute( TEXT("Testing Unix Ones Minute"), UnixOnes, 58 );
@@ -66,7 +67,8 @@ bool FDateTimeTest::RunTest( const FString& Parameters )
 	TestMillisecond( TEXT("Testing Unix Ones Millisecond"), UnixOnes, 0 );
 
 	TestYear( TEXT("Testing Unix Decimal Sequence Year"), UnixDecimalSequence, 2009 );
-	TestMonth( TEXT("Testing Unix Decimal Sequence Month"), UnixDecimalSequence, EMonthOfYear::February );
+	TestMonth( TEXT("Testing Unix Decimal Sequence Month"), UnixDecimalSequence, 2 );
+	TestMonthOfYear( TEXT("Testing Unix Decimal Sequence MonthOfYear"), UnixDecimalSequence, EMonthOfYear::February );
 	TestDay( TEXT("Testing Unix Decimal Sequence Day"), UnixDecimalSequence, 13 );
 	TestHour( TEXT("Testing Unix Decimal Sequence Hour"), UnixDecimalSequence, 23 );
 	TestMinute( TEXT("Testing Unix Decimal Sequence Minute"), UnixDecimalSequence, 31 );
@@ -75,6 +77,7 @@ bool FDateTimeTest::RunTest( const FString& Parameters )
 
 	TestYear( TEXT("Testing Test Date Time Year"), TestDateTime, 2013 );
 	TestMonth( TEXT("Testing Test Date Time Month"), TestDateTime, 8 );
+	TestMonthOfYear( TEXT("Testing Test Date Time MonthOfYear"), TestDateTime, EMonthOfYear::August );
 	TestDay( TEXT("Testing Test Date Time Day"), TestDateTime, 14 );
 	TestHour( TEXT("Testing Test Date Time Hour"), TestDateTime, 12 );
 	TestMinute( TEXT("Testing Test Date Time Minute"), TestDateTime, 34 );
diff --git a/Engine/Source/Runtime/Core/Public/Misc/DateTime.h b/Engine/Source/Runtime/Core/Public/Misc/DateTime.h
--- a/Engine/Source/Runtime/Core/Public/Misc/DateTime.h
+++ b/Engine/Source/Runtime/Core/Public/Misc/DateTime.h
@@ -1,53 +1,49 @@
 // Copyright 1998-2014 Epic Games, Inc. All Rights Reserved.
 
 #pragma once
+
 #include "HAL/Platform.h"
 #include "Misc/Timespan.h"
 
+
 class FArchive;
 class FOutputDevice;
 class FString;
 class UObject;
 
 
-namespace EDayOfWeek
+/**
+ * Enumerates the days of the week in 7-day calendars.
+ */
+enum class EDayOfWeek
 {
-	/**
-	 * Enumerates the days of the week in 7-day calendars.
-	 */
-	enum Type
-	{
-		Monday = 0,
-		Tuesday,
-		Wednesday,
-		Thursday,
-		Friday,
-		Saturday,
-		Sunday
-	};
+	Monday = 0,
+	Tuesday,
+	Wednesday,
+	Thursday,
+	Friday,
+	Saturday,
+	Sunday
 };
 
 
-namespace EMonthOfYear
+/**
+ * Enumerates the months of the year in 12-month calendars.
+ */
+enum class EMonthOfYear
 {
-	/**
-	 * Enumerates the months of the year in 12-month calendars.
-	 */
-	enum Type
-	{
-		January = 1,
-		February,
-		March,
-		April,
-		May,
-		June,
-		July,
-		August,
-		September,
-		October,
-		November,
-		December
-	};
+	January = 1,
+	February,
+	March,
+	April,
+	May,
+	June,
+	July,
+	August,
+	September,
+	October,
+	November,
+	December
 };
 
 
@@ -77,10 +73,8 @@ struct FDateTime
 {
 public:
 
-	/**
-	 * Default constructor (no initialization).
-	 */
-	FDateTime( ) { }
+	/** Default constructor (no initialization). */
+	FDateTime() { }
 
 	/**
 	 * Creates and initializes a new instance with the specified number of ticks.
@@ -110,7 +104,6 @@ public:
 	 * Returns result of adding the given time span to this date.
 	 *
 	 * @return A date whose value is the sum of this date and the given time span.
-	 *
 	 * @see FTimespan
 	 */
 	FDateTime operator+( const FTimespan& Other ) const
@@ -254,7 +247,7 @@ public:
 	 *
 	 * @return A FDateTime object containing the date.
 	 */
-	FDateTime GetDate( ) const
+	FDateTime GetDate() const
 	{
 		return FDateTime(Ticks - (Ticks % ETimespan::TicksPerDay));
 	}
@@ -274,7 +267,7 @@ public:
 	 * @return Day of the month.
 	 * @see GetHour, GetHour12, GetMillisecond, GetMinute, GetMonth, GetSecond, GetYear
 	 */
-	CORE_API int32 GetDay( ) const;
+	CORE_API int32 GetDay() const;
 
 	/**
 	 * Calculates this date's day of the week (Sunday - Saturday).
@@ -282,7 +275,7 @@ public:
 	 * @return The week day.
 	 * @see GetDayOfYear, GetMonthOfYear, GetTimeOfDay
 	 */
-	CORE_API EDayOfWeek::Type GetDayOfWeek( ) const;
+	CORE_API EDayOfWeek GetDayOfWeek() const;
 
 	/**
 	 * Gets this date's day of the year.
@@ -290,7 +283,7 @@ public:
 	 * @return The day of year.
 	 * @see GetDayOfWeek, GetMonthOfYear, GetTimeOfDay
 	 */
-	CORE_API int32 GetDayOfYear( ) const;
+	CORE_API int32 GetDayOfYear() const;
 
 	/**
 	 * Gets this date's hour part in 24-hour clock format (0 to 23).
@@ -298,7 +291,7 @@ public:
 	 * @return The hour.
 	 * @see GetDay, GetDayOfWeek, GetDayOfYear, GetHour12, GetMillisecond, GetMinute, GetMonth, GetSecond, GetYear
 	 */
-	int32 GetHour( ) const
+	int32 GetHour() const
 	{
 		return (int32)((Ticks / ETimespan::TicksPerHour) % 24);
 	}
@@ -309,7 +302,7 @@ public:
 	 * @return The hour in AM/PM format.
 	 * @see GetDay, GetHour, GetMillisecond, GetMinute, GetMonth, GetSecond, GetYear
 	 */
-	CORE_API int32 GetHour12( ) const;
+	CORE_API int32 GetHour12() const;
 
 	/**
 	 * Returns the Julian Day for this date.
@@ -321,7 +314,7 @@ public:
 	 * @return Julian Day.
 	 * @see FromJulianDay, GetModifiedJulianDay
 	 */
-	double GetJulianDay( ) const
+	double GetJulianDay() const
 	{
 		return (double)(1721425.5 + Ticks / ETimespan::TicksPerDay);
 	}
@@ -335,7 +328,7 @@ public:
 	 * @return Modified Julian Day
 	 * @see GetJulianDay
 	 */
-	double GetModifiedJulianDay( ) const
+	double GetModifiedJulianDay() const
 	{
 		return (GetJulianDay() - 2400000.5);
 	}
@@ -346,7 +339,7 @@ public:
 	 * @return The millisecond.
 	 * @see GetDay, GetHour, GetHour12, GetMinute, GetMonth, GetSecond, GetYear
 	 */
-	int32 GetMillisecond( ) const
+	int32 GetMillisecond() const
 	{
 		return (int32)((Ticks / ETimespan::TicksPerMillisecond) % 1000);
 	}
@@ -357,7 +350,7 @@ public:
 	 * @return The minute.
 	 * @see GetDay, GetHour, GetHour12, GetMillisecond, GetMonth, GetSecond, GetYear
 	 */
-	int32 GetMinute( ) const
+	int32 GetMinute() const
 	{
 		return (int32)((Ticks / ETimespan::TicksPerMinute) % 60);
 	}
@@ -368,7 +361,7 @@ public:
 	 * @return The month.
 	 * @see GetDay, GetHour, GetHour12, GetMillisecond, GetMinute, GetSecond, GetYear
 	 */
-	CORE_API int32 GetMonth( ) const;
+	CORE_API int32 GetMonth() const;
 
 	/**
 	 * Gets the date's month of the year (January to December).
@@ -376,9 +369,9 @@ public:
 	 * @return Month of year.
 	 * @see GetDayOfWeek, GetDayOfYear, GetTimeOfDay
 	 */
-	EMonthOfYear::Type GetMonthOfYear( ) const
+	EMonthOfYear GetMonthOfYear() const
 	{
-		return static_cast<EMonthOfYear::Type>(GetMonth());
+		return static_cast<EMonthOfYear>(GetMonth());
 	}
 
 	/**
@@ -387,7 +380,7 @@ public:
 	 * @return The second.
 	 * @see GetDay, GetHour, GetHour12, GetMillisecond, GetMinute, GetMonth, GetYear
 	 */
-	int32 GetSecond( ) const
+	int32 GetSecond() const
 	{
 		return (int32)((Ticks / ETimespan::TicksPerSecond) % 60);
 	}
@@ -397,7 +390,7 @@ public:
 	 *
 	 * @return Number of ticks since midnight, January 1, 0001.
 	 */
-	int64 GetTicks( ) const
+	int64 GetTicks() const
 	{
 		return Ticks;
 	}
@@ -408,7 +401,7 @@ public:
 	 * @param Time of day since midnight.
 	 * @see GetDayOfWeek, GetDayOfYear, GetMonthOfYear
 	 */
-	FTimespan GetTimeOfDay( ) const
+	FTimespan GetTimeOfDay() const
 	{
 		return FTimespan(Ticks % ETimespan::TicksPerDay);
 	}
@@ -419,7 +412,7 @@ public:
 	 * @return The year.
 	 * @see GetDay, GetHour, GetHour12, GetMillisecond, GetMinute, GetMonth, GetSecond
 	 */
-	CORE_API int32 GetYear( ) const;
+	CORE_API int32 GetYear() const;
 
 	/**
 	 * Imports a date and time value from a text buffer.
@@ -439,7 +432,7 @@ public:
 	 * @param true if it is in the afternoon, false otherwise.
 	 * @see IsMorning
 	 */
-	bool IsAfternoon( ) const
+	bool IsAfternoon() const
 	{
 		return (GetHour() >= 12);
 	}
@@ -450,7 +443,7 @@ public:
 	 * @param true if it is in the morning, false otherwise.
 	 * @see IsAfternoon
 	 */
-	bool IsMorning( ) const
+	bool IsMorning() const
 	{
 		return (GetHour() < 12);
 	}
@@ -471,7 +464,7 @@ public:
 	 * @return String representation.
 	 * @see ParseIso8601, ToString
 	 */
-	CORE_API FString ToIso8601( ) const;
+	CORE_API FString ToIso8601() const;
 
 	/**
 	 * Returns the string representation of this date using a default format.
@@ -482,7 +475,7 @@ public:
 	 * @return String representation.
 	 * @see Parse, ToIso8601
 	 */
-	CORE_API FString ToString( ) const;
+	CORE_API FString ToString() const;
 
 	/**
 	 * Returns the string representation of this date.
@@ -499,7 +492,7 @@ public:
 	 * @return Time of day.
 	 * @see FromUnixTimestamp
 	 */
-	CORE_API int32 ToUnixTimestamp( ) const
+	CORE_API int32 ToUnixTimestamp() const
 	{
 		return static_cast<int32>((Ticks - FDateTime(1970, 1, 1).Ticks) / ETimespan::TicksPerSecond);
 	}
@@ -575,7 +568,7 @@ public:
 	 *
 	 * @see MinValue
 	 */
-	static FDateTime MaxValue( )
+	static FDateTime MaxValue()
 	{
 		return FDateTime(3652059 * ETimespan::TicksPerDay - 1);
 	}
@@ -587,7 +580,7 @@ public:
 	 *
 	 * @see MaxValue
 	 */
-	static FDateTime MinValue( )
+	static FDateTime MinValue()
 	{
 		return FDateTime(0);
 	}
@@ -602,7 +595,7 @@ public:
 	 * @return Current date and time.
 	 * @see Today, UtcNow
 	 */
-	static CORE_API FDateTime Now( );
+	static CORE_API FDateTime Now();
 
 	/**
 	 * Converts a string to a date and time.
@@ -634,7 +627,7 @@ public:
 	 * @return Current date.
 	 * @see Now, UtcNow
 	 */
-	static FDateTime Today( )
+	static FDateTime Today()
 	{
 		return Now().GetDate();
 	}
@@ -650,7 +643,7 @@ public:
 	 * @return Current date and time.
 	 * @see Now
 	 */
-	static CORE_API FDateTime UtcNow( );
+	static CORE_API FDateTime UtcNow();
 
 public:
 
@@ -678,14 +671,10 @@ public:
 
 protected:
 
-	/**
-	 * Holds the days per month in a non-leap year.
-	 */
+	/** Holds the days per month in a non-leap year. */
 	static const int32 DaysPerMonth[];
 
-	/**
-	 * Holds the cumulative days per month in a non-leap year.
-	 */
+	/** Holds the cumulative days per month in a non-leap year. */
 	static const int32 DaysToMonth[];
 
 private:
diff --git a/Engine/Source/Runtime/Core/Public/Misc/Guid.h b/Engine/Source/Runtime/Core/Public/Misc/Guid.h
--- a/Engine/Source/Runtime/Core/Public/Misc/Guid.h
+++ b/Engine/Source/Runtime/Core/Public/Misc/Guid.h
@@ -7,56 +7,53 @@ class UObject;
 class FOutputDevice;
 
 
-namespace EGuidFormats
+/**
+ * Enumerates known GUID formats.
+ */
+enum class EGuidFormats
 {
 	/**
-	 * Enumerates known GUID formats.
+	 * 32 digits.
+	 *
+	 * For example: "00000000000000000000000000000000"
 	 */
-	enum Type
-	{
-		/**
-		 * 32 digits.
-		 *
-		 * For example: "00000000000000000000000000000000"
-		 */
-		Digits,
-
-		/**
-		 * 32 digits separated by hyphens.
-		 *
-		 * For example: 00000000-0000-0000-0000-000000000000
-		 */
-		DigitsWithHyphens,
-
-		/**
-		 * 32 digits separated by hyphens and enclosed in braces.
-		 *
-		 * For example: {00000000-0000-0000-0000-000000000000}
-		 */
-		DigitsWithHyphensInBraces,
-
-		/**
-		 * 32 digits separated by hyphens and enclosed in parentheses.
-		 *
-		 * For example: (00000000-0000-0000-0000-000000000000)
-		 */
-		DigitsWithHyphensInParentheses,
-
-		/**
-		 * Comma-separated hexadecimal values enclosed in braces.
-		 *
-		 * For example: {0x00000000,0x0000,0x0000,{0x00,0x00,0x00,0x00,0x00,0x00,0x00,0x00}}
-		 */
-		HexValuesInBraces,
-
-		/**
-		 * This format is currently used by the FUniqueObjectGuid class.
-		 *
-		 * For example: 00000000-00000000-00000000-00000000
-		 */
-		UniqueObjectGuid,
-	};
-}
+	Digits,
+
+	/**
+	 * 32 digits separated by hyphens.
+	 *
+	 * For example: 00000000-0000-0000-0000-000000000000
+	 */
+	DigitsWithHyphens,
+
+	/**
+	 * 32 digits separated by hyphens and enclosed in braces.
+	 *
+	 * For example: {00000000-0000-0000-0000-000000000000}
+	 */
+	DigitsWithHyphensInBraces,
+
+	/**
+	 * 32 digits separated by hyphens and enclosed in parentheses.
+	 *
+	 * For example: (00000000-0000-0000-0000-000000000000)
+	 */
+	DigitsWithHyphensInParentheses,
+
+	/**
+	 * Comma-separated hexadecimal values enclosed in braces.
+	 *
+	 * For example: {0x00000000,0x0000,0x0000,{0x00,0x00,0x00,0x00,0x00,0x00,0x00,0x00}}
+	 */
+	HexValuesInBraces,
+
+	/**
+	 * This format is currently used by the FUniqueObjectGuid class.
+	 *
+	 * For example: 00000000-00000000-00000000-00000000
+	*/
+	UniqueObjectGuid,
+};
 
 
 /**
@@ -66,10 +63,8 @@ class FGuid
 {
 public:
 
-	/**
-	 * Default constructor.
-	 */
-	FGuid( )
+	/** Default constructor. */
+	FGuid()
 		: A(0)
 		, B(0)
 		, C(0)
@@ -219,7 +214,7 @@ public:
 	 * @return true if valid, false otherwise
 	 * @see Invalidate
 	 */
-	bool IsValid( ) const
+	bool IsValid() const
 	{
 		return ((A | B | C | D) != 0);
 	}
@@ -229,7 +224,7 @@ public:
 	 *
 	 * @see IsValid
 	 */
-	void Invalidate( )
+	void Invalidate()
 	{
 		A = B = C = D = 0;
 	}
@@ -239,7 +234,7 @@ public:
 	 *
 	 * @return The string representation.
 	 */
-	FString ToString( ) const
+	FString ToString() const
 	{
 		return ToString(EGuidFormats::Digits);
 	}
@@ -250,7 +245,7 @@ public:
 	 * @param Format The string format to use.
 	 * @return The string representation.
 	 */
-	CORE_API FString ToString( EGuidFormats::Type Format ) const;
+	CORE_API FString ToString( EGuidFormats Format ) const;
 
 public:
 
@@ -272,7 +267,7 @@ public:
 	 *
 	 * @return A new GUID.
 	 */
-	static CORE_API FGuid NewGuid( );
+	static CORE_API FGuid NewGuid();
 
 	/**
 	 * Converts a string to a GUID.
@@ -293,20 +288,20 @@ public:
 	 * @return true if the string was converted successfully, false otherwise.
 	 * @see Parse, ToString
 	 */
-	static CORE_API bool ParseExact( const FString& GuidString, EGuidFormats::Type Format, FGuid& OutGuid );
+	static CORE_API bool ParseExact( const FString& GuidString, EGuidFormats Format, FGuid& OutGuid );
 
 //private:
 public:
 
-	// Holds the first component.
+	/** Holds the first component. */
 	uint32 A;
 
-	// Holds the second component.
+	/** Holds the second component. */
 	uint32 B;
 
-	// Holds the third component.
+	/** Holds the third component. */
 	uint32 C;
 
-	// Holds the fourth component.
+	/** Holds the fourth component. */
 	uint32 D;
 };
diff --git a/Engine/Source/Runtime/Core/Public/Misc/Timespan.h b/Engine/Source/Runtime/Core/Public/Misc/Timespan.h
--- a/Engine/Source/Runtime/Core/Public/Misc/Timespan.h
+++ b/Engine/Source/Runtime/Core/Public/Misc/Timespan.h
@@ -1,8 +1,10 @@
 // Copyright 1998-2014 Epic Games, Inc. All Rights Reserved.
 
 #pragma once
+
 #include "HAL/Platform.h"
 
+
 class FArchive;
 class FOutputDevice;
 class FString;
@@ -10,6 +12,7 @@ class UObject;
 
 inline uint32 GetTypeHash(const int64 A);
 
+
 namespace ETimespan
 {
 	/** The number of timespan ticks per day. */
@@ -50,10 +53,8 @@ class FTimespan
 {
 public:
 
-	/**
-	 * Default constructor.
-	 */
-	FTimespan( ) { }
+	/** Default constructor (no initialization). */
+	FTimespan() { }
 
 	/**
 	 * Creates and initializes a new time interval with the specified number of ticks.
@@ -69,7 +70,8 @@ public:
 	 *
 	 * @param Hours The hours component.
 	 * @param Minutes The minutes component.
-	 * @param Seconds The seconds component.	 */
+	 * @param Seconds The seconds component.
+	 */
 	FTimespan( int32 Hours, int32 Minutes, int32 Seconds )
 	{
 		Assign(0, Hours, Minutes, Seconds, 0);
@@ -133,7 +135,7 @@ public:
 	 *
 	 * @return Inverse of this time span.
 	 */
-	FTimespan operator-( ) const
+	FTimespan operator-() const
 	{
 		return FTimespan(-Ticks);
 	}
@@ -272,7 +274,7 @@ public:
 	 *
 	 * @return Days component.
 	 */
-	int32 GetDays( ) const
+	int32 GetDays() const
 	{
 		return (int32)(Ticks / ETimespan::TicksPerDay);
 	}
@@ -285,7 +287,7 @@ public:
 	 * @return Duration of this time span.
 	 * @see MinValue
 	 */
-	FTimespan GetDuration( )
+	FTimespan GetDuration()
 	{
 		return FTimespan(Ticks >= 0 ? Ticks : -Ticks);
 	}
@@ -296,7 +298,7 @@ public:
 	 * @return Hours component.
 	 * @see GetTotalHours
 	 */
-	int32 GetHours( ) const
+	int32 GetHours() const
 	{
 		return (int32)((Ticks / ETimespan::TicksPerHour) % 24);
 	}
@@ -307,7 +309,7 @@ public:
 	 * @return Milliseconds component.
 	 * @see GetTotalMilliseconds
 	 */
-	int32 GetMilliseconds( ) const
+	int32 GetMilliseconds() const
 	{
 		return (int32)((Ticks / ETimespan::TicksPerMillisecond) % 1000);
 	}
@@ -318,7 +320,7 @@ public:
 	 * @return Minutes component.
 	 * @see GetTotalMinutes
 	 */
-	int32 GetMinutes( ) const
+	int32 GetMinutes() const
 	{
 		return (int32)((Ticks / ETimespan::TicksPerMinute) % 60);
 	}
@@ -329,7 +331,7 @@ public:
 	 * @return Seconds component.
 	 * @see GetTotalSeconds
 	 */
-	int32 GetSeconds( ) const
+	int32 GetSeconds() const
 	{
 		return (int32)((Ticks / ETimespan::TicksPerSecond) % 60);
 	}
@@ -339,7 +341,7 @@ public:
 	 *
 	 * @return Number of ticks.
 	 */
-	int64 GetTicks( ) const
+	int64 GetTicks() const
 	{
 		return Ticks;
 	}
@@ -350,7 +352,7 @@ public:
 	 * @return Number of days.
 	 * @see GetDays
 	 */
-	double GetTotalDays( ) const
+	double GetTotalDays() const
 	{
 		return ((double)Ticks / ETimespan::TicksPerDay);
 	}
@@ -361,7 +363,7 @@ public:
 	 * @return Number of hours.
 	 * @see GetHours
 	 */
-	double GetTotalHours( ) const
+	double GetTotalHours() const
 	{
 		return ((double)Ticks / ETimespan::TicksPerHour);
 	}
@@ -372,7 +374,7 @@ public:
 	 * @return Number of milliseconds.
 	 * @see GetMilliseconds
 	 */
-	double GetTotalMilliseconds( ) const
+	double GetTotalMilliseconds() const
 	{
 		return ((double)Ticks / ETimespan::TicksPerMillisecond);
 	}
@@ -383,7 +385,7 @@ public:
 	 * @return Number of minutes.
 	 * @see GetMinutes
 	 */
-	double GetTotalMinutes( ) const
+	double GetTotalMinutes() const
 	{
 		return ((double)Ticks / ETimespan::TicksPerMinute);
 	}
@@ -394,7 +396,7 @@ public:
 	 * @return Number of seconds.
 	 * @see GetSeconds
 	 */
-	double GetTotalSeconds( ) const
+	double GetTotalSeconds() const
 	{
 		return ((double)Ticks / ETimespan::TicksPerSecond);
 	}
@@ -429,7 +431,7 @@ public:
 	 * @return String representation.
 	 * @see Parse
 	 */
-	CORE_API FString ToString( ) const;
+	CORE_API FString ToString() const;
 
 	/**
 	 * Converts this time span to its string representation.
@@ -509,7 +511,7 @@ public:
 	 * @return Maximum time span.
 	 * @see MinValue,Zero
 	 */
-	static FTimespan MaxValue( )
+	static FTimespan MaxValue()
 	{
 		return FTimespan(9223372036854775807);
 	}
@@ -522,7 +524,7 @@ public:
 	 * @return Minimum time span.
 	 * @see MaxValue, ZeroValue
 	 */
-	static FTimespan MinValue( )
+	static FTimespan MinValue()
 	{
 		return FTimespan(-9223372036854775807 - 1);
 	}
@@ -548,7 +550,7 @@ public:
 	 * @return Zero time span.
 	 * @see MaxValue, MinValue
 	 */
-	static FTimespan Zero( )
+	static FTimespan Zero()
 	{
 		return FTimespan(0);
 	}
diff --git a/Engine/Source/Runtime/Core/Public/Misc/Variant.h b/Engine/Source/Runtime/Core/Public/Misc/Variant.h
--- a/Engine/Source/Runtime/Core/Public/Misc/Variant.h
+++ b/Engine/Source/Runtime/Core/Public/Misc/Variant.h
@@ -1,61 +1,55 @@
 // Copyright 1998-2014 Epic Games, Inc. All Rights Reserved.
 
-/*=============================================================================
-	Variant.h: Declares the FVariant class.
-=============================================================================*/
-
 #pragma once
 
 
+/**
+ * Enumerates the built-in types that can be stored in instances of FVariant.
+ */
 namespace EVariantTypes
 {
-	/**
-	 * Enumerates the built-in types that can be stored in instances of FVariant.
-	 */
-	enum Type
-	{
-		Empty,
-		Ansichar,
-		Bool,
-		Box,
-		BoxSphereBounds,
-		ByteArray,
-		Color,
-		DateTime,
-		Double,
-		Enum,
-		Float,
-		Guid,
-		Int8,
-		Int16,
-		Int32,
-		Int64,
-		IntRect,
-		LinearColor,
-		Matrix,
-		Name,
-		Plane,
-		Quat,
-		RandomStream,
-		Rotator,
-		String,
-		Widechar,
-		Timespan,
-		Transform,
-		TwoVectors,
-		UInt8,
-		UInt16,
-		UInt32,
-		UInt64,
-		Vector,
-		Vector2d,
-		Vector4,
-		IntPoint,
-		IntVector,
-		NetworkGUID,
-		Custom = 0x40
-	};
-}
+	const int32 Empty = 0;
+	const int32 Ansichar = 1;
+	const int32 Bool = 2;
+	const int32 Box = 3;
+	const int32 BoxSphereBounds = 4;
+	const int32 ByteArray = 5;
+	const int32 Color = 6;
+	const int32 DateTime = 7;
+	const int32 Double = 8;
+	const int32 Enum = 9;
+	const int32 Float = 10;
+	const int32 Guid = 11;
+	const int32 Int8 = 12;
+	const int32 Int16 = 13;
+	const int32 Int32 = 14;
+	const int32 Int64 = 15;
+	const int32 IntRect = 16;
+	const int32 LinearColor = 17;
+	const int32 Matrix = 18;
+	const int32 Name = 19;
+	const int32 Plane = 20;
+	const int32 Quat = 21;
+	const int32 RandomStream = 22;
+	const int32 Rotator = 23;
+	const int32 String = 24;
+	const int32 Widechar = 25;
+	const int32 Timespan = 26;
+	const int32 Transform = 27;
+	const int32 TwoVectors = 28;
+	const int32 UInt8 = 29;
+	const int32 UInt16 = 30;
+	const int32 UInt32 = 31;
+	const int32 UInt64 = 32;
+	const int32 Vector = 33;
+	const int32 Vector2d = 34;
+	const int32 Vector4 = 35;
+	const int32 IntPoint = 36;
+	const int32 IntVector = 37;
+	const int32 NetworkGUID = 38;
+
+	const int32 Custom = 0x40;
+};
 
 
 /**
@@ -69,11 +63,11 @@ namespace EVariantTypes
  * and then assigned to a variant. Note that you will be responsible for ensuring
  * correct byte ordering when serializing those types.
  *
- * @param T - The type to be used in FVariant.
+ * @param T The type to be used in FVariant.
  */
 template<typename T> struct TVariantTraits
 {
-	static int32 GetType( )
+	static int32 GetType()
 	{
 		static_assert(!sizeof(T), "Variant trait must be specialized for this type.");
 	}
@@ -91,16 +85,12 @@ class FVariant
 {
 public:
 
-	/**
-	 * Default constructor.
-	 */
-	FVariant( )
+	/** Default constructor. */
+	FVariant()
 		: Type(EVariantTypes::Empty)
 	{ }
 
-	/**
-	 * Copy constructor.
-	 */
+	/** Copy constructor. */
 	FVariant( const FVariant& Other )
 		: Type(Other.Type)
 		, Value(Other.Value)
@@ -109,7 +99,7 @@ public:
 	/**
 	 * Creates and initializes a new instance with the specified value.
 	 *
-	 * @param InValue - The initial value.
+	 * @param InValue The initial value.
 	 */
 	template<typename T>
 	FVariant( T InValue )
@@ -220,7 +210,7 @@ public:
 	 * @return The value converted to the specified type.
 	 */
 	template<typename T>
-	operator T( ) const
+	operator T() const
 	{
 		return GetValue<T>();
 	}
@@ -256,7 +246,7 @@ public:
 	 *
 	 * @see IsEmpty
 	 */
-	void Empty( )
+	void Empty()
 	{
 		Type = EVariantTypes::Empty;
 
@@ -270,7 +260,7 @@ public:
 	 *
 	 * @see Empty
 	 */
-	bool IsEmpty( ) const
+	bool IsEmpty() const
 	{
 		return (Type == EVariantTypes::Empty);
 	}
@@ -285,7 +275,7 @@ public:
 	 * @return Byte array.
 	 * @see GetValue
 	 */
-	const TArray<uint8>& GetBytes( ) const
+	const TArray<uint8>& GetBytes() const
 	{
 		return Value;
 	}
@@ -296,7 +286,7 @@ public:
 	 * @return Size of the value.
 	 * @see GetType, GetValue
 	 */
-	int32 GetSize( ) const
+	int32 GetSize() const
 	{
 		return Value.Num();
 	}
@@ -307,7 +297,7 @@ public:
 	 * @return Type of the value.
 	 * @see GetSize, GetValue
 	 */
-	int32 GetType( ) const
+	int32 GetType() const
 	{
 		return Type;
 	}
@@ -322,7 +312,7 @@ public:
 	 * @see GetSize, GetType
 	 */
 	template<typename T>
-	T GetValue( ) const
+	T GetValue() const
 	{
 		check((Type == TVariantTraits<T>::GetType()) || ((TVariantTraits<T>::GetType() == EVariantTypes::UInt8) && (Type == EVariantTypes::Enum)));
 
@@ -350,10 +340,10 @@ public:
 
 private:
 
-	// Holds the type of the variant.
+	/** Holds the type of the variant. */
 	int32 Type;
 
-	// Holds the serialized value.
+	/** Holds the serialized value. */
 	TArray<uint8> Value;
 };
 
@@ -371,7 +361,7 @@ private:
  * @see GetBytes
  */
 template<>
-FORCEINLINE TArray<uint8> FVariant::GetValue<TArray<uint8> >( ) const
+FORCEINLINE TArray<uint8> FVariant::GetValue<TArray<uint8> >() const
 {
 	check(Type == EVariantTypes::ByteArray);
 
@@ -382,343 +372,267 @@ FORCEINLINE TArray<uint8> FVariant::GetValue<TArray<uint8> >( ) const
 /* Default FVariant traits for built-in types
  *****************************************************************************/
 
-/**
- * Implements variant type traits for the built-in ANSICHAR type.
- */
+/** Implements variant type traits for the built-in ANSICHAR type. */
 template<> struct TVariantTraits<ANSICHAR>
 {
-	static int32 GetType( ) { return EVariantTypes::Ansichar; }
+	static int32 GetType() { return EVariantTypes::Ansichar; }
 };
 
 
-/**
- * Implements variant type traits for the built-in bool type.
- */
+/** Implements variant type traits for the built-in bool type. */
 template<> struct TVariantTraits<bool>
 {
-	static int32 GetType( ) { return EVariantTypes::Bool; }
+	static int32 GetType() { return EVariantTypes::Bool; }
 };
 
 
-/**
- * Implements variant type traits for the built-in FBox type.
- */
+/** Implements variant type traits for the built-in FBox type. */
 template<> struct TVariantTraits<FBox>
 {
-	static int32 GetType( ) { return EVariantTypes::Box; }
+	static int32 GetType() { return EVariantTypes::Box; }
 };
 
 
-/**
- * Implements variant type traits for the built-in FBoxSphereBounds type.
- */
+/** Implements variant type traits for the built-in FBoxSphereBounds type. */
 template<> struct TVariantTraits<FBoxSphereBounds>
 {
-	static int32 GetType( ) { return EVariantTypes::BoxSphereBounds; }
+	static int32 GetType() { return EVariantTypes::BoxSphereBounds; }
 };
 
 
-/**
- * Implements variant type traits for byte arrays.
- */
+/** Implements variant type traits for byte arrays. */
 template<> struct TVariantTraits<TArray<uint8> >
 {
-	static int32 GetType( ) { return EVariantTypes::ByteArray; }
+	static int32 GetType() { return EVariantTypes::ByteArray; }
 };
 
 
-/**
- * Implements variant type traits for the built-in FColor type.
- */
+/** Implements variant type traits for the built-in FColor type. */
 template<> struct TVariantTraits<FColor>
 {
-	static int32 GetType( ) { return EVariantTypes::Color; }
+	static int32 GetType() { return EVariantTypes::Color; }
 };
 
 
-/**
- * Implements variant type traits for the built-in FDateTime type.
- */
+/** Implements variant type traits for the built-in FDateTime type. */
 template<> struct TVariantTraits<FDateTime>
 {
-	static int32 GetType( ) { return EVariantTypes::DateTime; }
+	static int32 GetType() { return EVariantTypes::DateTime; }
 };
 
 
-/**
- * Implements variant type traits for the built-in double type.
- */
+/** Implements variant type traits for the built-in double type. */
 template<> struct TVariantTraits<double>
 {
-	static int32 GetType( ) { return EVariantTypes::Double; }
+	static int32 GetType() { return EVariantTypes::Double; }
 };
 
 
-/**
- * Implements variant type traits for enumeration types.
- */
+/** Implements variant type traits for enumeration types. */
 template<typename EnumType> struct TVariantTraits<TEnumAsByte<EnumType> >
 {
-	static int32 GetType( ) { return EVariantTypes::Enum; }
+	static int32 GetType() { return EVariantTypes::Enum; }
 };
 
 
-/**
- * Implements variant type traits for the built-in float type.
- */
+/** Implements variant type traits for the built-in float type. */
 template<> struct TVariantTraits<float>
 {
-	static int32 GetType( ) { return EVariantTypes::Float; }
+	static int32 GetType() { return EVariantTypes::Float; }
 };
 
 
-/**
- * Implements variant type traits for the built-in FGuid type.
- */
+/** Implements variant type traits for the built-in FGuid type. */
 template<> struct TVariantTraits<FGuid>
 {
-	static int32 GetType( ) { return EVariantTypes::Guid; }
+	static int32 GetType() { return EVariantTypes::Guid; }
 };
 
 
-/**
- * Implements variant type traits for the built-in int8 type.
- */
+/** Implements variant type traits for the built-in int8 type. */
 template<> struct TVariantTraits<int8>
 {
-	static int32 GetType( ) { return EVariantTypes::Int8; }
+	static int32 GetType() { return EVariantTypes::Int8; }
 };
 
 
-/**
- * Implements variant type traits for the built-in int16 type.
- */
+/** Implements variant type traits for the built-in int16 type. */
 template<> struct TVariantTraits<int16>
 {
-	static int32 GetType( ) { return EVariantTypes::Int16; }
+	static int32 GetType() { return EVariantTypes::Int16; }
 };
 
 
-/**
- * Implements variant type traits for the built-in int32 type.
- */
+/** Implements variant type traits for the built-in int32 type. */
 template<> struct TVariantTraits<int32>
 {
-	static int32 GetType( ) { return EVariantTypes::Int32; }
+	static int32 GetType() { return EVariantTypes::Int32; }
 };
 
 
-/**
- * Implements variant type traits for the built-in int64 type.
- */
+/** Implements variant type traits for the built-in int64 type. */
 template<> struct TVariantTraits<int64>
 {
-	static int32 GetType( ) { return EVariantTypes::Int64; }
+	static int32 GetType() { return EVariantTypes::Int64; }
 };
 
 
-/**
- * Implements variant type traits for the built-in FIntPoint type.
- */
+/** Implements variant type traits for the built-in FIntPoint type. */
 template<> struct TVariantTraits<FIntPoint>
 {
-	static int32 GetType( ) { return EVariantTypes::IntPoint; }
+	static int32 GetType() { return EVariantTypes::IntPoint; }
 };
 
 
-/**
- * Implements variant type traits for the built-in FIntVector type.
- */
+/** Implements variant type traits for the built-in FIntVector type. */
 template<> struct TVariantTraits<FIntVector>
 {
-	static int32 GetType( ) { return EVariantTypes::IntVector; }
+	static int32 GetType() { return EVariantTypes::IntVector; }
 };
 
 
-/**
- * Implements variant type traits for the built-in FIntRect type.
- */
+/** Implements variant type traits for the built-in FIntRect type. */
 template<> struct TVariantTraits<FIntRect>
 {
-	static int32 GetType( ) { return EVariantTypes::IntRect; }
+	static int32 GetType() { return EVariantTypes::IntRect; }
 };
 
 
-/**
- * Implements variant type traits for the built-in FLinearColor type.
- */
+/** Implements variant type traits for the built-in FLinearColor type. */
 template<> struct TVariantTraits<FLinearColor>
 {
-	static int32 GetType( ) { return EVariantTypes::LinearColor; }
+	static int32 GetType() { return EVariantTypes::LinearColor; }
 };
 
 
-/**
- * Implements variant type traits for the built-in FMatrix type.
- */
+/** Implements variant type traits for the built-in FMatrix type. */
 template<> struct TVariantTraits<FMatrix>
 {
-	static int32 GetType( ) { return EVariantTypes::Matrix; }
+	static int32 GetType() { return EVariantTypes::Matrix; }
 };
 
 
-/**
- * Implements variant type traits for the built-in FPlane type.
- */
+/** Implements variant type traits for the built-in FPlane type. */
 template<> struct TVariantTraits<FPlane>
 {
-	static int32 GetType( ) { return EVariantTypes::Plane; }
+	static int32 GetType() { return EVariantTypes::Plane; }
 };
 
 
-/**
- * Implements variant type traits for the built-in FQuat type.
- */
+/** Implements variant type traits for the built-in FQuat type. */
 template<> struct TVariantTraits<FQuat>
 {
-	static int32 GetType( ) { return EVariantTypes::Quat; }
+	static int32 GetType() { return EVariantTypes::Quat; }
 };
 
 
-/**
- * Implements variant type traits for the built-in FName type.
- */
+/** Implements variant type traits for the built-in FName type. */
 template<> struct TVariantTraits<FName>
 {
-	static int32 GetType( ) { return EVariantTypes::Name; }
+	static int32 GetType() { return EVariantTypes::Name; }
 };
 
 
-/**
- * Implements variant type traits for the built-in FRandomStream type.
- */
+/** Implements variant type traits for the built-in FRandomStream type. */
 template<> struct TVariantTraits<FRandomStream>
 {
-	static int32 GetType( ) { return EVariantTypes::RandomStream; }
+	static int32 GetType() { return EVariantTypes::RandomStream; }
 };
 
 
-/**
- * Implements variant type traits for the built-in FRotator type.
- */
+/** Implements variant type traits for the built-in FRotator type. */
 template<> struct TVariantTraits<FRotator>
 {
-	static int32 GetType( ) { return EVariantTypes::Rotator; }
+	static int32 GetType() { return EVariantTypes::Rotator; }
 };
 
 
-/**
- * Implements variant type traits for the built-in FString type.
- */
+/** Implements variant type traits for the built-in FString type. */
 template<> struct TVariantTraits<FString>
 {
-	static int32 GetType( ) { return EVariantTypes::String; }
+	static int32 GetType() { return EVariantTypes::String; }
 };
 
 
-/**
- * Implements variant type traits for the built-in WIDECHAR type.
- */
+/** Implements variant type traits for the built-in WIDECHAR type. */
 template<> struct TVariantTraits<WIDECHAR>
 {
-	static int32 GetType( ) { return EVariantTypes::Widechar; }
+	static int32 GetType() { return EVariantTypes::Widechar; }
 };
 
 
-/**
- * Implements variant type traits for the built-in FTimespan type.
- */
+/** Implements variant type traits for the built-in FTimespan type. */
 template<> struct TVariantTraits<FTimespan>
 {
-	static int32 GetType( ) { return EVariantTypes::Timespan; }
+	static int32 GetType() { return EVariantTypes::Timespan; }
 };
 
 
-/**
- * Implements variant type traits for the built-in FTransform type.
- */
+/** Implements variant type traits for the built-in FTransform type. */
 template<> struct TVariantTraits<FTransform>
 {
-	static int32 GetType( ) { return EVariantTypes::Transform; }
+	static int32 GetType() { return EVariantTypes::Transform; }
 };
 
 
-/**
- * Implements variant type traits for the built-in FTwoVectors type.
- */
+/** Implements variant type traits for the built-in FTwoVectors type. */
 template<> struct TVariantTraits<FTwoVectors>
 {
-	static int32 GetType( ) { return EVariantTypes::TwoVectors; }
+	static int32 GetType() { return EVariantTypes::TwoVectors; }
 };
 
 
-/**
- * Implements variant type traits for the built-in uint8 type.
- */
+/** Implements variant type traits for the built-in uint8 type. */
 template<> struct TVariantTraits<uint8>
 {
-	static int32 GetType( ) { return EVariantTypes::UInt8; }
+	static int32 GetType() { return EVariantTypes::UInt8; }
 };
 
 
-/**
- * Implements variant type traits for the built-in uint16 type.
- */
+/** Implements variant type traits for the built-in uint16 type. */
 template<> struct TVariantTraits<uint16>
 {
-	static int32 GetType( ) { return EVariantTypes::UInt16; }
+	static int32 GetType() { return EVariantTypes::UInt16; }
 };
 
 
-/**
- * Implements variant type traits for the built-in uint32 type.
- */
+/** Implements variant type traits for the built-in uint32 type. */
 template<> struct TVariantTraits<uint32>
 {
-	static int32 GetType( ) { return EVariantTypes::UInt32; }
+	static int32 GetType() { return EVariantTypes::UInt32; }
 };
 
 
-/**
- * Implements variant type traits for the built-in uint64 type.
- */
+/** Implements variant type traits for the built-in uint64 type. */
 template<> struct TVariantTraits<uint64>
 {
-	static int32 GetType( ) { return EVariantTypes::UInt64; }
+	static int32 GetType() { return EVariantTypes::UInt64; }
 };
 
 
-/**
- * Implements variant type traits for the built-in FVector type.
- */
+/** Implements variant type traits for the built-in FVector type. */
 template<> struct TVariantTraits<FVector>
 {
-	static int32 GetType( ) { return EVariantTypes::Vector; }
+	static int32 GetType() { return EVariantTypes::Vector; }
 };
 
 
-/**
- * Implements variant type traits for the built-in FVector2D type.
- */
+/** Implements variant type traits for the built-in FVector2D type. */
 template<> struct TVariantTraits<FVector2D>
 {
-	static int32 GetType( ) { return EVariantTypes::Vector2d; }
+	static int32 GetType() { return EVariantTypes::Vector2d; }
 };
 
 
-/**
- * Implements variant type traits for the built-in FVector4 type.
- */
+/** Implements variant type traits for the built-in FVector4 type. */
 template<> struct TVariantTraits<FVector4>
 {
-	static int32 GetType( ) { return EVariantTypes::Vector4; }
+	static int32 GetType() { return EVariantTypes::Vector4; }
 };
 
 
-/**
- * Implements variant type traits for the built-in NetworkGUID type.
- */
+/** Implements variant type traits for the built-in NetworkGUID type. */
 template<> struct TVariantTraits<FNetworkGUID>
 {
-	static int32 GetType( ) { return EVariantTypes::NetworkGUID; }
+	static int32 GetType() { return EVariantTypes::NetworkGUID; }
 };
diff --git a/Engine/Source/Runtime/Engine/Classes/Kismet/KismetGuidLibrary.h b/Engine/Source/Runtime/Engine/Classes/Kismet/KismetGuidLibrary.h
--- a/Engine/Source/Runtime/Engine/Classes/Kismet/KismetGuidLibrary.h
+++ b/Engine/Source/Runtime/Engine/Classes/Kismet/KismetGuidLibrary.h
@@ -35,7 +35,7 @@ class UKismetGuidLibrary
 	UFUNCTION(BlueprintPure, meta=(FriendlyName="ToString (Guid)", CompactNodeTitle="->"), Category="Guid")
 	static FString Conv_GuidToString( const FGuid& InGuid );
 
-	/** Converts a String of format EGuidFormats::Type to a Guid. Returns Guid OutGuid, Returns bool Success */
+	/** Converts a String of format EGuidFormats to a Guid. Returns Guid OutGuid, Returns bool Success */
 	UFUNCTION(BlueprintPure, meta=(FriendlyName="Parse String to Guid"), Category="Guid")
 	static void Parse_StringToGuid( const FString& GuidString, FGuid& OutGuid, bool& Success );
 };
