Core: Misc documentation cleanup

diff --git a/Engine/Source/Runtime/Core/Private/Misc/AES.cpp b/Engine/Source/Runtime/Core/Private/Misc/AES.cpp
--- a/Engine/Source/Runtime/Core/Private/Misc/AES.cpp
+++ b/Engine/Source/Runtime/Core/Private/Misc/AES.cpp
@@ -1,12 +1,9 @@
 // Copyright 1998-2014 Epic Games, Inc. All Rights Reserved.
 
-/*=============================================================================
-	AES.cpp: Handling of Advanced Encryption Standard
-=============================================================================*/
-
 #include "Core.h"
 #include "AES.h"
 
+
 // Uncomment this line to skip encryption
 // #undef AES_KEY
 
@@ -22,7 +19,6 @@
 
 // It may be useful to integrate this into a class derived from FArchive to allow streaming encryption/decryption - but that's a mission for another day.
 
-
 #ifdef AES_KEY
 
 #define TEST_ENCRYPTION			0
diff --git a/Engine/Source/Runtime/Core/Private/Misc/App.cpp b/Engine/Source/Runtime/Core/Private/Misc/App.cpp
--- a/Engine/Source/Runtime/Core/Private/Misc/App.cpp
+++ b/Engine/Source/Runtime/Core/Private/Misc/App.cpp
@@ -1,9 +1,5 @@
 // Copyright 1998-2014 Epic Games, Inc. All Rights Reserved.
 
-/*=============================================================================
-	App.cpp: Implements the FApp class.
-=============================================================================*/
-
 #include "Core.h"
 #include "Runtime/Launch/Resources/Version.h"
 
@@ -22,6 +18,7 @@ double FApp::CurrentTime = 0.0;
 double FApp::LastTime = 0.0;
 double FApp::DeltaTime = 1 / 30.0;
 
+
 /* FApp static interface
  *****************************************************************************/
 
@@ -30,6 +27,7 @@ FString FApp::GetBranchName( )
 	return FString(TEXT(BRANCH_NAME));
 }
 
+
 EBuildConfigurations::Type FApp::GetBuildConfiguration()
 {
 #if UE_BUILD_DEBUG
@@ -56,6 +54,7 @@ EBuildConfigurations::Type FApp::GetBuildConfiguration()
 #endif
 }
 
+
 FString FApp::GetBuildDate( )
 {
 	return FString(ANSI_TO_TCHAR(__DATE__));
@@ -102,6 +101,7 @@ void FApp::InitializeSession( )
 	}
 }
 
+
 bool FApp::IsInstalled()
 {
 #if UE_BUILD_SHIPPING && PLATFORM_DESKTOP && !UE_SERVER
@@ -112,16 +112,20 @@ bool FApp::IsInstalled()
 	return bIsInstalled;
 }
 
+
 bool FApp::IsEngineInstalled()
 {
 	static bool bIsInstalledEngine = IsInstalled() || (FRocketSupport::IsRocket() ? !FParse::Param(FCommandLine::Get(), TEXT("EngineNotInstalled")) : FParse::Param(FCommandLine::Get(), TEXT("EngineInstalled")));
 	return bIsInstalledEngine;
 }
 
+
 #if HAVE_RUNTIME_THREADING_SWITCHES
 bool FApp::ShouldUseThreadingForPerformance( )
 {
 	static bool OnlyOneThread = FParse::Param(FCommandLine::Get(), TEXT("ONETHREAD")) || IsRunningDedicatedServer() || !FPlatformProcess::SupportsMultithreading() || FPlatformMisc::NumberOfCores() < 2;
 	return !OnlyOneThread;
 }
+
+
 #endif // HAVE_RUNTIME_THREADING_SWITCHES
diff --git a/Engine/Source/Runtime/Core/Private/Misc/AutomationTest.cpp b/Engine/Source/Runtime/Core/Private/Misc/AutomationTest.cpp
--- a/Engine/Source/Runtime/Core/Private/Misc/AutomationTest.cpp
+++ b/Engine/Source/Runtime/Core/Private/Misc/AutomationTest.cpp
@@ -1,12 +1,9 @@
 // Copyright 1998-2014 Epic Games, Inc. All Rights Reserved.
 
-/*=============================================================================
-	AutomationTest.cpp
-=============================================================================*/
-
 #include "Core.h"
 #include "ModuleManager.h"
 
+
 DEFINE_LOG_CATEGORY_STATIC(LogAutomationTest, Warning, All);
 
 
@@ -456,32 +453,38 @@ void FAutomationTestFramework::SetDeveloperDirectoryIncluded(const bool bInDevel
 	bDeveloperDirectoryIncluded = bInDeveloperDirectoryIncluded;
 }
 
+
 void FAutomationTestFramework::SetVisualCommandletFilter(const bool bInVisualCommandletFilterOn)
 {
 	bVisualCommandletFilterOn = bInVisualCommandletFilterOn;
 }
 
+
 FOnTestScreenshotCaptured& FAutomationTestFramework::OnScreenshotCaptured()
 {
 	return TestScreenshotCapturedDelegate;
 }
 
+
 void FAutomationTestFramework::SetScreenshotOptions( const bool bInScreenshotsEnabled, const bool bInUseFullSizeScreenshots )
 {
 	bScreenshotsEnabled = bInScreenshotsEnabled;
 	bUseFullSizeScreenShots = bInUseFullSizeScreenshots;
 }
 
+
 bool FAutomationTestFramework::IsScreenshotAllowed() const
 {
 	return bScreenshotsEnabled;
 }
 
+
 bool FAutomationTestFramework::ShouldUseFullSizeScreenshots() const
 {
 	return bUseFullSizeScreenShots;
 }
 
+
 void FAutomationTestFramework::PrepForAutomationTests()
 {
 	check(!GIsAutomationTesting);
@@ -646,8 +649,7 @@ FAutomationTestFramework::FAutomationTestFramework()
 ,	bUseFullSizeScreenShots(false)
 ,	NetworkRoleIndex(0)
 ,	bForceSmokeTests(false)
-{
-}
+{ }
 
 
 FAutomationTestFramework::~FAutomationTestFramework()
@@ -656,6 +658,7 @@ FAutomationTestFramework::~FAutomationTestFramework()
 	AutomationTestClassNameToInstanceMap.Empty();
 }
 
+
 void FAutomationTestBase::ClearExecutionInfo()
 {
 	ExecutionInfo.Clear();
@@ -744,4 +747,3 @@ void FAutomationTestBase::GenerateTestNames(TArray<FAutomationTestInfo>& TestInf
 		TestInfo.Add( NewTestInfo );
 	}
 }
-
diff --git a/Engine/Source/Runtime/Core/Private/Misc/Compression.cpp b/Engine/Source/Runtime/Core/Private/Misc/Compression.cpp
--- a/Engine/Source/Runtime/Core/Private/Misc/Compression.cpp
+++ b/Engine/Source/Runtime/Core/Private/Misc/Compression.cpp
@@ -1,16 +1,11 @@
 // Copyright 1998-2014 Epic Games, Inc. All Rights Reserved.
 
-/*=============================================================================
-	Compression.cpp: Compression classes
-=============================================================================*/
-
 #include "Core.h"
 #include "CompressedGrowableBuffer.h"
-
-DEFINE_LOG_CATEGORY_STATIC(LogCompression, Log, All);
+#include "ThirdParty/zlib/zlib-1.2.5/Inc/zlib.h"
 
 
-#include "ThirdParty/zlib/zlib-1.2.5/Inc/zlib.h"
+DEFINE_LOG_CATEGORY_STATIC(LogCompression, Log, All);
 
 /**
  * Thread-safe abstract compression routine. Compresses memory from uncompressed buffer and writes it to compressed
diff --git a/Engine/Source/Runtime/Core/Private/Misc/ConfigCacheIni.cpp b/Engine/Source/Runtime/Core/Private/Misc/ConfigCacheIni.cpp
--- a/Engine/Source/Runtime/Core/Private/Misc/ConfigCacheIni.cpp
+++ b/Engine/Source/Runtime/Core/Private/Misc/ConfigCacheIni.cpp
@@ -1,9 +1,5 @@
 // Copyright 1998-2014 Epic Games, Inc. All Rights Reserved.
 
-/*=============================================================================
-	ConfigCacheIni.cpp: Unreal config file reading/writing.
-=============================================================================*/
-
 #include "Core.h"
 #include "ConfigCacheIni.h"
 #include "RemoteConfigIni.h"
@@ -12,8 +8,10 @@
 #include "DefaultValueHelper.h"
 #include "EngineBuildSettings.h"
 
+
 DEFINE_LOG_CATEGORY(LogConfig);
 
+
 class FTextFriendHelper
 {
 	friend class FConfigFile;
diff --git a/Engine/Source/Runtime/Core/Private/Misc/Core.cpp b/Engine/Source/Runtime/Core/Private/Misc/Core.cpp
--- a/Engine/Source/Runtime/Core/Private/Misc/Core.cpp
+++ b/Engine/Source/Runtime/Core/Private/Misc/Core.cpp
@@ -1,14 +1,12 @@
 // Copyright 1998-2014 Epic Games, Inc. All Rights Reserved.
 
-/*=============================================================================
-	Core.cpp: Unreal core.
-=============================================================================*/
-
 #include "Core.h"
 #include "ModuleManager.h"
 
+
 #define LOCTEXT_NAMESPACE "Core"
 
+
 class FCoreModule : public FDefaultModuleImpl
 {
 public:
@@ -18,8 +16,11 @@ public:
 		return false;
 	}
 };
+
+
 IMPLEMENT_MODULE( FCoreModule, Core );
 
+
 /*-----------------------------------------------------------------------------
 	Global variables.
 -----------------------------------------------------------------------------*/
diff --git a/Engine/Source/Runtime/Core/Private/Misc/CoreMisc.cpp b/Engine/Source/Runtime/Core/Private/Misc/CoreMisc.cpp
--- a/Engine/Source/Runtime/Core/Private/Misc/CoreMisc.cpp
+++ b/Engine/Source/Runtime/Core/Private/Misc/CoreMisc.cpp
@@ -1,9 +1,5 @@
 // Copyright 1998-2014 Epic Games, Inc. All Rights Reserved.
 
-/*=============================================================================
-	UnMisc.cpp: Various core platform-independent functions.
-=============================================================================*/
-
 // Core includes.
 #include "Core.h"
 #include "ExceptionHandling.h"
@@ -21,6 +17,7 @@
 #include "Ticker.h"
 #include "DerivedDataCacheInterface.h"
 
+
 DEFINE_LOG_CATEGORY(LogSHA);
 DEFINE_LOG_CATEGORY(LogStats);
 DEFINE_LOG_CATEGORY(LogStreaming);
diff --git a/Engine/Source/Runtime/Core/Private/Misc/Database.cpp b/Engine/Source/Runtime/Core/Private/Misc/Database.cpp
--- a/Engine/Source/Runtime/Core/Private/Misc/Database.cpp
+++ b/Engine/Source/Runtime/Core/Private/Misc/Database.cpp
@@ -1,9 +1,5 @@
 // Copyright 1998-2014 Epic Games, Inc. All Rights Reserved.
 
-/*=============================================================================
-	Database.cpp: Database integration
-=============================================================================*/
-
 // Precompiled header include. Can't add anything above this line.
 #include "Core.h"
 #include "Database.h"
diff --git a/Engine/Source/Runtime/Core/Private/Misc/GuidGenerator.cpp b/Engine/Source/Runtime/Core/Private/Misc/GuidGenerator.cpp
--- a/Engine/Source/Runtime/Core/Private/Misc/GuidGenerator.cpp
+++ b/Engine/Source/Runtime/Core/Private/Misc/GuidGenerator.cpp
@@ -1,9 +1,5 @@
 // Copyright 1998-2014 Epic Games, Inc. All Rights Reserved.
 
-/*============================================================================
-	UnGUID.cpp: Globally Unique IDentifiers.
-============================================================================*/
-
 #include "Core.h"
 
 /*
@@ -82,7 +78,7 @@ void appGetGUID( void* GUID )
 	}
 	uuid_create( (UUID_t*)GUID );
 }
-	
+
 /*----------------------------------------------------------------------------
 	Support functions.
 ----------------------------------------------------------------------------*/
diff --git a/Engine/Source/Runtime/Core/Private/Misc/LocalTimestampDirectoryVisitor.cpp b/Engine/Source/Runtime/Core/Private/Misc/LocalTimestampDirectoryVisitor.cpp
--- a/Engine/Source/Runtime/Core/Private/Misc/LocalTimestampDirectoryVisitor.cpp
+++ b/Engine/Source/Runtime/Core/Private/Misc/LocalTimestampDirectoryVisitor.cpp
@@ -1,9 +1,5 @@
 // Copyright 1998-2014 Epic Games, Inc. All Rights Reserved.
 
-/*=============================================================================
-	LocalTimestampDirectoryVisitor.cpp: Implements the FLocalTimestampDirectoryVisitor class.
-=============================================================================*/
-
 #include "Core.h"
 
 
@@ -90,4 +86,4 @@ bool FLocalTimestampDirectoryVisitor::Visit(const TCHAR* FilenameOrDirectory, bo
 	}
 
 	return true;
-}
\ No newline at end of file
+}
diff --git a/Engine/Source/Runtime/Core/Private/Misc/MemStack.cpp b/Engine/Source/Runtime/Core/Private/Misc/MemStack.cpp
--- a/Engine/Source/Runtime/Core/Private/Misc/MemStack.cpp
+++ b/Engine/Source/Runtime/Core/Private/Misc/MemStack.cpp
@@ -1,18 +1,14 @@
 // Copyright 1998-2014 Epic Games, Inc. All Rights Reserved.
 
-/*=============================================================================
-	MemStack.cpp: Unreal memory grabbing functions
-=============================================================================*/
-
 #include "Core.h"
 
+
 DECLARE_THREAD_SINGLETON( FMemStack );
 
 DECLARE_MEMORY_STAT(TEXT("MemStack Large Block"), STAT_MemStackLargeBLock,STATGROUP_Memory);
 DECLARE_MEMORY_STAT(TEXT("PageAllocator Free"), STAT_PageAllocatorFree, STATGROUP_Memory);
 DECLARE_MEMORY_STAT(TEXT("PageAllocator Used"), STAT_PageAllocatorUsed, STATGROUP_Memory);
 
-
 TLockFreeFixedSizeAllocator<FPageAllocator::PageSize, FThreadSafeCounter> FPageAllocator::TheAllocator;
 TLockFreeFixedSizeAllocator<FPageAllocator::SmallPageSize, FThreadSafeCounter> FPageAllocator::TheSmallAllocator;
 
@@ -23,6 +19,7 @@ void FPageAllocator::UpdateStats()
 	SET_MEMORY_STAT(STAT_PageAllocatorUsed, BytesUsed());
 }
 #endif
+
 /*-----------------------------------------------------------------------------
 	FMemStack implementation.
 -----------------------------------------------------------------------------*/
diff --git a/Engine/Source/Runtime/Core/Private/Misc/MonitoredProcess.cpp b/Engine/Source/Runtime/Core/Private/Misc/MonitoredProcess.cpp
--- a/Engine/Source/Runtime/Core/Private/Misc/MonitoredProcess.cpp
+++ b/Engine/Source/Runtime/Core/Private/Misc/MonitoredProcess.cpp
@@ -12,16 +12,16 @@ FMonitoredProcess::FMonitoredProcess( const FString& InURL, const FString& InPar
 	, Hidden(InHidden)
 	, KillTree(false)
 	, Params(InParams)
-	, ReadPipe(NULL)
+	, ReadPipe(nullptr)
+	, ReturnCode(0)
 	, StartTime(0)
-	, Thread(NULL)
+	, Thread(nullptr)
 	, URL(InURL)
-	, WritePipe(NULL)
-	, ReturnCode(0)
+	, WritePipe(nullptr)
 { }
 
 
-FMonitoredProcess::~FMonitoredProcess( )
+FMonitoredProcess::~FMonitoredProcess()
 {
 	if (IsRunning())
 	{
@@ -35,7 +35,7 @@ FMonitoredProcess::~FMonitoredProcess( )
 /* FMonitoredProcess interface
  *****************************************************************************/
 
-FTimespan FMonitoredProcess::GetDuration( ) const
+FTimespan FMonitoredProcess::GetDuration() const
 {
 	if (IsRunning())
 	{
@@ -46,7 +46,7 @@ FTimespan FMonitoredProcess::GetDuration( ) const
 }
 
 
-bool FMonitoredProcess::Launch( )
+bool FMonitoredProcess::Launch()
 {
 	if (IsRunning())
 	{
@@ -58,7 +58,7 @@ bool FMonitoredProcess::Launch( )
 		return false;
 	}
 
-	ProcessHandle = FPlatformProcess::CreateProc(*URL, *Params, false, Hidden, Hidden, NULL, 0, *FPaths::RootDir(), WritePipe);
+	ProcessHandle = FPlatformProcess::CreateProc(*URL, *Params, false, Hidden, Hidden, nullptr, 0, *FPaths::RootDir(), WritePipe);
 
 	if (!ProcessHandle.IsValid())
 	{
@@ -94,7 +94,7 @@ void FMonitoredProcess::ProcessOutput( const FString& Output )
 /* FRunnable interface
  *****************************************************************************/
 
-uint32 FMonitoredProcess::Run( )
+uint32 FMonitoredProcess::Run()
 {
 	// monitor the process
 	StartTime = FDateTime::UtcNow();
@@ -109,7 +109,7 @@ uint32 FMonitoredProcess::Run( )
 			{
 				FPlatformProcess::TerminateProc(ProcessHandle, KillTree);
 				CanceledDelegate.ExecuteIfBound();
-				Thread = NULL;
+				Thread = nullptr;
 
 				return 0;
 			}
@@ -120,7 +120,7 @@ uint32 FMonitoredProcess::Run( )
 
 	// close output pipes
 	FPlatformProcess::ClosePipe(ReadPipe, WritePipe);
-	ReadPipe = WritePipe = NULL;
+	ReadPipe = WritePipe = nullptr;
 
 	// get completion status
 	if (!FPlatformProcess::GetProcReturnCode(ProcessHandle, &ReturnCode))
@@ -129,7 +129,7 @@ uint32 FMonitoredProcess::Run( )
 	}
 
 	CompletedDelegate.ExecuteIfBound(ReturnCode);
-	Thread = NULL;
+	Thread = nullptr;
 
 	return 0;
 }
diff --git a/Engine/Source/Runtime/Core/Private/Misc/NetworkGuid.cpp b/Engine/Source/Runtime/Core/Private/Misc/NetworkGuid.cpp
--- a/Engine/Source/Runtime/Core/Private/Misc/NetworkGuid.cpp
+++ b/Engine/Source/Runtime/Core/Private/Misc/NetworkGuid.cpp
@@ -1,9 +1,5 @@
 // Copyright 1998-2014 Epic Games, Inc. All Rights Reserved.
 
-/*=============================================================================
-	NetworkGuid.cpp: Implements the FNetworkGUID class.
-=============================================================================*/
-
 #include "Core.h"
 
 
@@ -14,4 +10,4 @@ bool FNetworkGUID::NetSerialize( FArchive& Ar, class UPackageMap* Map, bool& bOu
 {
 	Ar << Value;
 	return true;
-}
\ No newline at end of file
+}
diff --git a/Engine/Source/Runtime/Core/Private/Misc/ObjectThumbnail.cpp b/Engine/Source/Runtime/Core/Private/Misc/ObjectThumbnail.cpp
--- a/Engine/Source/Runtime/Core/Private/Misc/ObjectThumbnail.cpp
+++ b/Engine/Source/Runtime/Core/Private/Misc/ObjectThumbnail.cpp
@@ -1,9 +1,5 @@
 // Copyright 1998-2014 Epic Games, Inc. All Rights Reserved.
 
-/*=============================================================================
-	ObjectThumbnail.cpp: Stored thumbnail support for Unreal objects
-=============================================================================*/
-
 #include "Core.h"
 
 
@@ -11,20 +7,15 @@
 FThumbnailCompressionInterface* FObjectThumbnail::ThumbnailCompressor = NULL;
 
 
-
-/** Default constructor */
 FObjectThumbnail::FObjectThumbnail()
 	: ImageWidth( 0 ),
 	  ImageHeight( 0 ),
 	  bIsDirty( false ),
 	  bLoadedFromDisk(false),
 	  bCreatedAfterCustomThumbForSharedTypesEnabled(false)
-{
-}
+{ }
 
 
-
-/** Returns uncompressed image data, decompressing it on demand if needed */
 const TArray< uint8 >& FObjectThumbnail::GetUncompressedImageData() const
 {
 	if( ImageData.Num() == 0 )
@@ -37,8 +28,6 @@ const TArray< uint8 >& FObjectThumbnail::GetUncompressedImageData() const
 }
 
 
-
-/** Serializer */
 void FObjectThumbnail::Serialize( FArchive& Ar )
 {
 	Ar << ImageWidth;
@@ -74,8 +63,6 @@ void FObjectThumbnail::Serialize( FArchive& Ar )
 }
 
 
-	
-/** Compress image data */
 void FObjectThumbnail::CompressImageData()
 {
 	CompressedImageData.Reset();
@@ -86,8 +73,6 @@ void FObjectThumbnail::CompressImageData()
 }
 
 
-
-/** Decompress image data */
 void FObjectThumbnail::DecompressImageData()
 {
 	ImageData.Reset();
@@ -97,11 +82,7 @@ void FObjectThumbnail::DecompressImageData()
 	}
 }
 
-/**
- * Calculates the memory usage of this FObjectThumbnail.
- *
- * @param	Ar	the FArchiveCountMem (or similar) archive that will store the results of the memory usage calculation.
- */
+
 void FObjectThumbnail::CountBytes( FArchive& Ar ) const
 {
 	SIZE_T StaticSize = sizeof(FObjectThumbnail);
@@ -112,31 +93,19 @@ void FObjectThumbnail::CountBytes( FArchive& Ar ) const
 	UnconstThis->ImageData.CountBytes(Ar);
 }
 
-/**
- * Calculates the amount of memory used by the compressed bytes array
- *
- * @param	Ar	the FArchiveCountMem (or similar) archive that will store the results of the memory usage calculation.
- */
+
 void FObjectThumbnail::CountImageBytes_Compressed( FArchive& Ar ) const
 {
 	const_cast<FObjectThumbnail*>(this)->CompressedImageData.CountBytes(Ar);
 }
 
-/**
- * Calculates the amount of memory used by the uncompressed bytes array
- *
- * @param	Ar	the FArchiveCountMem (or similar) archive that will store the results of the memory usage calculation.
- */
+
 void FObjectThumbnail::CountImageBytes_Uncompressed( FArchive& Ar ) const
 {
 	const_cast<FObjectThumbnail*>(this)->ImageData.CountBytes(Ar);
 }
 
-/**
- * Calculates the memory usage of this FObjectFullNameAndThumbnail.
- *
- * @param	Ar	the FArchiveCountMem (or similar) archive that will store the results of the memory usage calculation.
- */
+
 void FObjectFullNameAndThumbnail::CountBytes( FArchive& Ar ) const
 {
 	SIZE_T StaticSize = sizeof(FObjectFullNameAndThumbnail);
@@ -147,6 +116,3 @@ void FObjectFullNameAndThumbnail::CountBytes( FArchive& Ar ) const
 		ObjectThumbnail->CountBytes(Ar);
 	}
 }
-
-
-
diff --git a/Engine/Source/Runtime/Core/Private/Misc/OutputDevice.cpp b/Engine/Source/Runtime/Core/Private/Misc/OutputDevice.cpp
--- a/Engine/Source/Runtime/Core/Private/Misc/OutputDevice.cpp
+++ b/Engine/Source/Runtime/Core/Private/Misc/OutputDevice.cpp
@@ -1,16 +1,14 @@
 // Copyright 1998-2014 Epic Games, Inc. All Rights Reserved.
 
-/*=============================================================================
-	OutputDevice.cpp: Implementation of OutputDevice
-=============================================================================*/
-
 #include "Core.h"
 #include "ExceptionHandling.h"
 #include "VarargsHelper.h"
 #include "MallocCrash.h"
 
+
 DEFINE_LOG_CATEGORY_STATIC(LogOutputDevice, Log, All);
 
+
 namespace
 {
 	/**
diff --git a/Engine/Source/Runtime/Core/Private/Misc/RemoteConfigIni.cpp b/Engine/Source/Runtime/Core/Private/Misc/RemoteConfigIni.cpp
--- a/Engine/Source/Runtime/Core/Private/Misc/RemoteConfigIni.cpp
+++ b/Engine/Source/Runtime/Core/Private/Misc/RemoteConfigIni.cpp
@@ -12,7 +12,6 @@ FRemoteConfigAsyncTaskManager GRemoteConfigIOManager;
    FRemoteConfigAsyncIOInfo
 -----------------------------------------------------------------------------*/
 
-/** Constructor */
 FRemoteConfigAsyncIOInfo::FRemoteConfigAsyncIOInfo(const TCHAR* InDefaultIniFile)
 	: bReadIOFailed(false)
 	, bWasProcessed(false)
@@ -20,7 +19,7 @@ FRemoteConfigAsyncIOInfo::FRemoteConfigAsyncIOInfo(const TCHAR* InDefaultIniFile
 	FCString::Strcpy(DefaultIniFile, InDefaultIniFile);
 }
 
-/** Assignment operator */
+
 FRemoteConfigAsyncIOInfo& FRemoteConfigAsyncIOInfo::operator=(const FRemoteConfigAsyncIOInfo& Other)
 {
 	Buffer = Other.Buffer;
@@ -39,7 +38,6 @@ FRemoteConfigAsyncIOInfo& FRemoteConfigAsyncIOInfo::operator=(const FRemoteConfi
    FRemoteConfigAsyncWorker
 -----------------------------------------------------------------------------*/
 
-/** Constructor */
 FRemoteConfigAsyncWorker::FRemoteConfigAsyncWorker(const TCHAR* InFilename, FRemoteConfigAsyncIOInfo& InIOInfo, FString* InContents, bool bInIsRead)
 {
 	check(FCString::Strlen(InFilename) < 1024);
@@ -53,7 +51,7 @@ FRemoteConfigAsyncWorker::FRemoteConfigAsyncWorker(const TCHAR* InFilename, FRem
 	}
 }
 
-/** Performs the actual IO operations */
+
 void FRemoteConfigAsyncWorker::DoWork()
 {
 	if (bIsRead)
@@ -72,32 +70,32 @@ void FRemoteConfigAsyncWorker::DoWork()
 	}
 }
 
-/** Returns true if the read IO operation succeeded */
+
 bool FRemoteConfigAsyncWorker::IsReadSuccess() const
 {
 	return !IOInfo.bReadIOFailed;
 }
 
-/** Returns the local IO info object */
+
 FRemoteConfigAsyncIOInfo& FRemoteConfigAsyncWorker::GetIOInfo()
 {
 	return IOInfo;
 }
 
-/** Give the name for external event viewers */
+
 const TCHAR* FRemoteConfigAsyncWorker::Name()
 {
 	return TEXT("FRemoteConfigAsyncWorker");
 }
 
-/** Indicates to the thread pool that this task is abandonable */
+
 bool FRemoteConfigAsyncWorker::CanAbandon()
 {
 	//return !bIsRead;
 	return false;
 }
 
-/** Abandon routine */
+
 void FRemoteConfigAsyncWorker::Abandon()
 {
 	// @todo what goes in here?
@@ -108,13 +106,12 @@ void FRemoteConfigAsyncWorker::Abandon()
    FRemoteConfigAsyncTaskManager
 -----------------------------------------------------------------------------*/
 
-/** Returns a reference to the global FRemoteConfigAsyncTaskManager object */
 FRemoteConfigAsyncTaskManager* FRemoteConfigAsyncTaskManager::Get()
 {
 	return &GRemoteConfigIOManager;
 }
 
-/** Handles cached write tasks */
+
 void FRemoteConfigAsyncTaskManager::Tick()
 {
 	FScopeLock ScopeLock(&SynchronizationObject);
@@ -128,7 +125,7 @@ void FRemoteConfigAsyncTaskManager::Tick()
 	}
 }
 
-/** Returns true if file exists in the cached task list */
+
 bool FRemoteConfigAsyncTaskManager::FindCachedWriteTask(const TCHAR* InFilename, bool bCompareContents, const TCHAR* InContents)
 {
 	for (int32 Idx = 0; Idx < CachedWriteTasks.Num(); ++Idx)
@@ -142,7 +139,7 @@ bool FRemoteConfigAsyncTaskManager::FindCachedWriteTask(const TCHAR* InFilename,
 	return false;
 }
 
-/** Add an async IO task to the queue and kick it off */
+
 bool FRemoteConfigAsyncTaskManager::StartTask(const TCHAR* InFilename, const TCHAR* RemotePath, FRemoteConfigAsyncIOInfo& InIOInfo, FString* InContents, bool bInIsRead)
 {
 	FScopeLock ScopeLock(&SynchronizationObject);
@@ -182,7 +179,7 @@ bool FRemoteConfigAsyncTaskManager::StartTask(const TCHAR* InFilename, const TCH
 	return true;
 }
 
-/** Returns true if the task has completed */
+
 bool FRemoteConfigAsyncTaskManager::IsFinished(const TCHAR* InFilename)
 {
 	FScopeLock ScopeLock(&SynchronizationObject);
@@ -192,7 +189,7 @@ bool FRemoteConfigAsyncTaskManager::IsFinished(const TCHAR* InFilename)
 	return AsyncTask? AsyncTask->IsDone(): true;
 }
 
-/** Returns true if the all tasks in the queue have completed (or, if the queue is empty) */
+
 bool FRemoteConfigAsyncTaskManager::AreAllTasksFinished(bool bDoRemoval)
 {
 	FScopeLock ScopeLock(&SynchronizationObject);
@@ -226,7 +223,7 @@ bool FRemoteConfigAsyncTaskManager::AreAllTasksFinished(bool bDoRemoval)
 	return true;
 }
 
-/** Safely retrieve the read data from the completed async task */
+
 bool FRemoteConfigAsyncTaskManager::GetReadData(const TCHAR* InFilename, FRemoteConfigAsyncIOInfo& OutIOInfo)
 {
 	FScopeLock ScopeLock(&SynchronizationObject);
@@ -252,20 +249,19 @@ bool FRemoteConfigAsyncTaskManager::GetReadData(const TCHAR* InFilename, FRemote
    FRemoteConfig
 -----------------------------------------------------------------------------*/
 
-/** Constructor */
 FRemoteConfig::FRemoteConfig()
 	: Timeout(-1.0f)
 	, bIsEnabled(true)
 	, bHasCachedFilenames(false)
 { }
 
-/** Returns a reference to the global FRemoteConfig object */
+
 FRemoteConfig* FRemoteConfig::Get()
 {
 	return &GRemoteConfig;
 }
 
-/** Returns true if the specified config file has been flagged as being remote **/
+
 bool FRemoteConfig::IsRemoteFile(const TCHAR* Filename)
 {
 	FString IniFileName(Filename);
@@ -300,25 +296,25 @@ bool FRemoteConfig::IsRemoteFile(const TCHAR* Filename)
 	return false;
 }
 
-/** Returns true if the specified file is remote and still needs to be read **/
+
 bool FRemoteConfig::ShouldReadRemoteFile(const TCHAR* Filename)
 {
 	return IsRemoteFile(Filename) && !FindConfig(Filename);
 }
 
-/** Simple accessor function **/
+
 FRemoteConfigAsyncIOInfo* FRemoteConfig::FindConfig(const TCHAR* Filename)
 {
 	return ConfigBuffers.Find(FString(Filename));
 }
 
-/** Returns true if the task has completed */
+
 bool FRemoteConfig::IsFinished(const TCHAR* InFilename)
 {
 	return GRemoteConfigIOManager.IsFinished(InFilename);
 }
 	
-/** Queues up a new async task for reading a remote config file **/
+
 bool FRemoteConfig::Read(const TCHAR* GeneratedIniFile, const TCHAR* DefaultIniFile)
 {
 	FString FullPath = GenerateRemotePath(GeneratedIniFile);
@@ -335,7 +331,7 @@ bool FRemoteConfig::Read(const TCHAR* GeneratedIniFile, const TCHAR* DefaultIniF
 	return GRemoteConfigIOManager.StartTask(GeneratedIniFile, *FullPath, IOInfo, NULL, true);
 }
 
-/** Queues up a new async task for writing a remote config file **/
+
 bool FRemoteConfig::Write(const TCHAR* Filename, FString& Contents)
 {
 	FRemoteConfigAsyncIOInfo* IOInfo = FindConfig(Filename);
@@ -351,7 +347,7 @@ bool FRemoteConfig::Write(const TCHAR* Filename, FString& Contents)
 	return true;
 }
 
-/** Waits on the async read if it hasn't finished yet... times out if the operation has taken too long **/
+
 void FRemoteConfig::FinishRead(const TCHAR* Filename)
 {
 	FRemoteConfigAsyncIOInfo* IOInfo = FindConfig(Filename);
@@ -375,7 +371,7 @@ void FRemoteConfig::FinishRead(const TCHAR* Filename)
 	}
 }
 
-/** Finishes all pending async IO tasks **/
+
 void FRemoteConfig::Flush()
 {
 	// @todo this is hacky... figure out how to clean up properly
@@ -402,7 +398,7 @@ static const TCHAR* SpecialCharMap[NUM_SPECIAL_CHARS][2] =
 	{ TEXT("\""), TEXT("~Quote~")  }
 };
 
-/* Replaces chars used by the ini parser with "special chars" */
+
 FString FRemoteConfig::ReplaceIniCharWithSpecialChar(const FString& Str)
 {
 	FString Result = Str;
@@ -413,7 +409,7 @@ FString FRemoteConfig::ReplaceIniCharWithSpecialChar(const FString& Str)
 	return Result;
 }
 
-/* Replaces "special chars" that have been inserted to avoid problems with the ini parser with equivalent regular chars */
+
 FString FRemoteConfig::ReplaceIniSpecialCharWithChar(const FString& Str)
 {
 	FString Result = Str;
@@ -424,7 +420,7 @@ FString FRemoteConfig::ReplaceIniSpecialCharWithChar(const FString& Str)
 	return Result;
 }
 
-/** Creates the remote path string for the specified file */
+
 FString FRemoteConfig::GenerateRemotePath(const TCHAR* Filename)
 {
 	FString IniFileName(Filename);
@@ -534,4 +530,4 @@ void MakeLocalCopy(const TCHAR* Filename)
 		FCString::Strcat(FilenameLocal, TEXT("_Local.ini"));
 		IFileManager::Get().Copy(FilenameLocal, Filename);
 	}
-}
\ No newline at end of file
+}
diff --git a/Engine/Source/Runtime/Core/Private/Misc/RocketSupport.cpp b/Engine/Source/Runtime/Core/Private/Misc/RocketSupport.cpp
--- a/Engine/Source/Runtime/Core/Private/Misc/RocketSupport.cpp
+++ b/Engine/Source/Runtime/Core/Private/Misc/RocketSupport.cpp
@@ -1,9 +1,5 @@
 // Copyright 1998-2014 Epic Games, Inc. All Rights Reserved.
 
-/*=============================================================================
-	RocketSupport.cpp: Implements the FRocketSupport class.
-=============================================================================*/
-
 #include "Core.h"
 #include "RocketSupport.h"
 #include "EngineVersion.h"
diff --git a/Engine/Source/Runtime/Core/Private/Misc/SecureHash.cpp b/Engine/Source/Runtime/Core/Private/Misc/SecureHash.cpp
--- a/Engine/Source/Runtime/Core/Private/Misc/SecureHash.cpp
+++ b/Engine/Source/Runtime/Core/Private/Misc/SecureHash.cpp
@@ -1,14 +1,12 @@
 // Copyright 1998-2014 Epic Games, Inc. All Rights Reserved.
 
-/*=============================================================================
-	SecureHash.cpp: SHA and MD5 has functions
-=============================================================================*/
-
 #include "Core.h"
 #include "SecureHash.h"
 
+
 DEFINE_LOG_CATEGORY_STATIC(LogSecureHash, Log, All);
 
+
 /*-----------------------------------------------------------------------------
 	MD5 functions, adapted from MD5 RFC by Brandon Reinhart
 -----------------------------------------------------------------------------*/
diff --git a/Engine/Source/Runtime/Core/Private/Misc/StringUtility.cpp b/Engine/Source/Runtime/Core/Private/Misc/StringUtility.cpp
--- a/Engine/Source/Runtime/Core/Private/Misc/StringUtility.cpp
+++ b/Engine/Source/Runtime/Core/Private/Misc/StringUtility.cpp
@@ -1,9 +1,6 @@
 // Copyright 1998-2014 Epic Games, Inc. All Rights Reserved.
 
-/*=============================================================================
-	StringUtility.cpp: String Utility Functions
-=============================================================================*/
-
 #include "Core.h"
 
-DEFINE_LOG_CATEGORY_STATIC(LogStringUtility, Log, All);
\ No newline at end of file
+
+DEFINE_LOG_CATEGORY_STATIC(LogStringUtility, Log, All);
diff --git a/Engine/Source/Runtime/Core/Private/Misc/WildcardString.cpp b/Engine/Source/Runtime/Core/Private/Misc/WildcardString.cpp
--- a/Engine/Source/Runtime/Core/Private/Misc/WildcardString.cpp
+++ b/Engine/Source/Runtime/Core/Private/Misc/WildcardString.cpp
@@ -1,9 +1,5 @@
 // Copyright 1998-2014 Epic Games, Inc. All Rights Reserved.
 
-/*=============================================================================
-	WildcardString.cpp: Implements the FWildcardString class.
-=============================================================================*/
-
 #include "Core.h"
 
 
@@ -12,7 +8,7 @@
 
 bool FWildcardString::ContainsWildcards(const TCHAR* Pattern)
 {
-	if (Pattern != NULL)
+	if (Pattern != nullptr)
 	{
 		while (*Pattern != EndOfString)
 		{
@@ -31,13 +27,13 @@ bool FWildcardString::ContainsWildcards(const TCHAR* Pattern)
 
 bool FWildcardString::IsMatch(const TCHAR* Pattern, const TCHAR* Input)
 {
-	if ((Pattern == NULL) || (Input == NULL))
+	if ((Pattern == nullptr) || (Input == nullptr))
 	{
 		return false;
 	}
 
 	const TCHAR* pp = Pattern;
-	const TCHAR* pinput = NULL;
+	const TCHAR* pinput = nullptr;
 
 	while (*Input != EndOfString)
 	{
@@ -58,7 +54,7 @@ bool FWildcardString::IsMatch(const TCHAR* Pattern, const TCHAR* Input)
 			Pattern++;
 			Input++;
 		}
-		else if (pinput == NULL)
+		else if (pinput == nullptr)
 		{
 			// single mismatch
 			return false;
@@ -84,4 +80,4 @@ bool FWildcardString::IsMatch(const TCHAR* Pattern, const TCHAR* Input)
 
 	// matching incomplete
 	return false;
-}
\ No newline at end of file
+}
diff --git a/Engine/Source/Runtime/Core/Public/Misc/AES.h b/Engine/Source/Runtime/Core/Public/Misc/AES.h
--- a/Engine/Source/Runtime/Core/Public/Misc/AES.h
+++ b/Engine/Source/Runtime/Core/Public/Misc/AES.h
@@ -1,12 +1,10 @@
 // Copyright 1998-2014 Epic Games, Inc. All Rights Reserved.
 
-/*=============================================================================
-	AES.h: Handling of Advanced Encryption Standard
-=============================================================================*/
-
 #pragma once
 
-#define AES_BLOCK_SIZE				16
+
+#define AES_BLOCK_SIZE 16
+
 
 struct CORE_API FAES
 {
@@ -14,6 +12,7 @@ struct CORE_API FAES
 
 	static void EncryptData( uint8 *Contents, uint32 NumBytes );
 	static void DecryptData( uint8 *Contents, uint32 NumBytes );
+
 	/**
 	 * Encrypts a chunk of data using a specific key
 	 *
@@ -22,6 +21,7 @@ struct CORE_API FAES
 	 * @param Key a null terminated string that is a 32 byte multiple length
 	 */
 	static void EncryptData(uint8* Contents, uint32 NumBytes, ANSICHAR* Key);
+
 	/**
 	 * Decrypts a chunk of data using a specific key
 	 *
@@ -31,5 +31,3 @@ struct CORE_API FAES
 	 */
 	static void DecryptData(uint8* Contents, uint32 NumBytes, ANSICHAR* Key);
 };
-
-
diff --git a/Engine/Source/Runtime/Core/Public/Misc/App.h b/Engine/Source/Runtime/Core/Public/Misc/App.h
--- a/Engine/Source/Runtime/Core/Public/Misc/App.h
+++ b/Engine/Source/Runtime/Core/Public/Misc/App.h
@@ -1,14 +1,12 @@
 // Copyright 1998-2014 Epic Games, Inc. All Rights Reserved.
 
-/*=============================================================================
-	App.h: Declares the FApp class.
-=============================================================================*/
-
 #pragma once
 
+
 // platforms which can have runtime threading switches
 #define HAVE_RUNTIME_THREADING_SWITCHES			(PLATFORM_DESKTOP || PLATFORM_ANDROID || PLATFORM_IOS)
 
+
 /**
  * Provides information about the application.
  */
@@ -64,7 +62,6 @@ public:
 	 * that can be used to identify it from anywhere on the network.
 	 *
 	 * @return Instance identifier, or an invalid GUID if there is no local instance.
-	 *
 	 * @see GetSessionId
 	 */
 	FORCEINLINE static FGuid GetInstanceId()
@@ -88,7 +85,6 @@ public:
 	 * Gets the name of the application, i.e. "UE4" or "Rocket".
 	 *
 	 * @todo need better application name discovery. this is quite horrible and may not work on future platforms.
-	 *
 	 * @return Application name string.
 	 */
 	static FString GetName()
@@ -119,7 +115,6 @@ public:
 	 * Conversely, sessions that were launched separately will have different session identifiers.
 	 *
 	 * @return Session identifier, or an invalid GUID if there is no local instance.
-	 *
 	 * @see GetInstanceId
 	 */
 	FORCEINLINE static FGuid GetSessionId()
@@ -276,7 +271,7 @@ public:
 	/**
 	 * Sets application benchmarking mode.
 	 *
-	 * @param bVal - True sets application in benchmark mode, false sets to non-benchmark mode.
+	 * @param bVal True sets application in benchmark mode, false sets to non-benchmark mode.
 	 */
 	static void SetBenchmarking(bool bVal)
 	{
@@ -296,7 +291,7 @@ public:
 	/**
 	 * Sets time step in seconds if a fixed delta time is wanted.
 	 *
-	 * @param seconds - Time step in seconds for fixed delta time.
+	 * @param seconds Time step in seconds for fixed delta time.
 	 */
 	static void SetFixedDeltaTime(double Seconds)
 	{
@@ -344,9 +339,7 @@ public:
 		return LastTime;
 	}
 
-	/**
-	 * Updates Last time to CurrentTime.
-	 */
+	/** Updates Last time to CurrentTime. */
 	static void UpdateLastTime()
 	{
 		LastTime = CurrentTime;
@@ -365,7 +358,7 @@ public:
 	/**
 	 * Sets time delta in seconds.
 	 *
-	 * @param seconds - Time in seconds.
+	 * @param seconds Time in seconds.
 	 */
 	static void SetDeltaTime(double Seconds)
 	{
@@ -374,43 +367,42 @@ public:
 
 private:
 
-	// Holds the instance identifier.
+	/** Holds the instance identifier. */
 	static FGuid InstanceId;
 
-	// Holds the session identifier.
+	/** Holds the session identifier. */
 	static FGuid SessionId;
 
-	// Holds the session name.
+	/** Holds the session name. */
 	static FString SessionName;
 
-	// Holds the name of the user that launched session.
+	/** Holds the name of the user that launched session. */
 	static FString SessionOwner;
 
-	// Holds a flag indicating whether this is a standalone session.
+	/** Holds a flag indicating whether this is a standalone session. */
 	static bool Standalone;
 
-	// Holds a flag Whether we are in benchmark mode or not.
+	/** Holds a flag Whether we are in benchmark mode or not. */
 	static bool bIsBenchmarking;
 
-	// Holds time step if a fixed delta time is wanted.
+	/** Holds time step if a fixed delta time is wanted. */
 	static double FixedDeltaTime;
 
-	// Holds current time.
+	/** Holds current time. */
 	static double CurrentTime;
 
-	// Holds previous value of CurrentTime.
+	/** Holds previous value of CurrentTime. */
 	static double LastTime;
 
-	// Holds current delta time in seconds.
+	/** Holds current delta time in seconds. */
 	static double DeltaTime;
 };
 
 
-/**
- * Called to determine the result of IsServerForOnlineSubsystems()
- */
+/** Called to determine the result of IsServerForOnlineSubsystems() */
 DECLARE_DELEGATE_RetVal_OneParam(bool, FQueryIsRunningServer, FName);
 
+
 /**
  * @return true if there is a running game world that is a server (including listen servers), false otherwise
  */
diff --git a/Engine/Source/Runtime/Core/Public/Misc/AssertionMacros.h b/Engine/Source/Runtime/Core/Public/Misc/AssertionMacros.h
--- a/Engine/Source/Runtime/Core/Public/Misc/AssertionMacros.h
+++ b/Engine/Source/Runtime/Core/Public/Misc/AssertionMacros.h
@@ -1,13 +1,11 @@
 // Copyright 1998-2014 Epic Games, Inc. All Rights Reserved.
 
-/*=============================================================================
-	AssertionMacros.h: Macros for assertions and other compile-time control
-=============================================================================*/
-
 #pragma once
+
 #include "HAL/Platform.h"
 #include "Misc/Build.h"
 
+
 /*----------------------------------------------------------------------------
 	Check, verify, etc macros
 ----------------------------------------------------------------------------*/
diff --git a/Engine/Source/Runtime/Core/Public/Misc/Attribute.h b/Engine/Source/Runtime/Core/Public/Misc/Attribute.h
--- a/Engine/Source/Runtime/Core/Public/Misc/Attribute.h
+++ b/Engine/Source/Runtime/Core/Public/Misc/Attribute.h
@@ -2,9 +2,9 @@
 
 #pragma once
 
-
 #include "Delegate.h"
 
+
 /**
  * Attribute object
  */
@@ -24,16 +24,12 @@ public:
 	DECLARE_DELEGATE_RetVal( ObjectType, FGetter );
 
 
-	/**
-	 * Default constructor
-	 */
+	/** Default constructor. */
 	TAttribute()
 		: Value()         // NOTE: Potentially uninitialized for atomics!!
 		, Getter()
 		, bIsSet(false)
-	{
-	}
-
+	{ }
 
 	/**
 	 * Construct implicitly from an initial value
@@ -45,9 +41,7 @@ public:
 		: Value( InInitialValue )
 		, Getter()
 		, bIsSet(true)
-	{
-	}
-
+	{ }
 
 	/**
 	 * Constructs by binding an arbitrary function that will be called to generate this attribute's value on demand. 
@@ -62,9 +56,7 @@ public:
 		: Value()
 		, Getter( FGetter::CreateSP( InUserObject, InMethodPtr ) )
 		, bIsSet(true)
-	{
-	}
-
+	{ }
 	
 	/**
 	 * Constructs by binding an arbitrary function that will be called to generate this attribute's value on demand. 
@@ -79,9 +71,7 @@ public:
 		: Value()
 		, Getter( FGetter::CreateSP( InUserObject, InMethodPtr ) )	
 		, bIsSet(true)
-	{
-	}
-
+	{ }
 
 	/**
 	 * Static: Creates an attribute that's pre-bound to the specified 'getter' delegate
@@ -93,7 +83,6 @@ public:
 		const bool bExplicitConstructor = true;
 		return TAttribute( InGetter, bExplicitConstructor );
 	}
-
 	
 	/**
 	 * Creates an attribute by binding an arbitrary function that will be called to generate this attribute's value on demand. 
@@ -108,7 +97,6 @@ public:
 		return TAttribute( FGetter::CreateStatic( InFuncPtr ), bExplicitConstructor );
 	}
 
-
 	/**
 	 * Sets the attribute's value
 	 *
@@ -160,7 +148,6 @@ public:
 		return bIsSet ? Get() : DefaultValue;
 	}
 
-
 	/**
 	 * Binds an arbitrary function that will be called to generate this attribute's value on demand.  After
 	 * binding, the attribute will no longer have a value that can be accessed directly, and instead the bound
@@ -174,7 +161,6 @@ public:
 		Getter = InGetter;
 	}
 
-
 	/**
 	 * Binds an arbitrary function that will be called to generate this attribute's value on demand. 
 	 * After binding, the attribute will no longer have a value that can be accessed directly, and instead the bound
@@ -187,7 +173,6 @@ public:
 		bIsSet = true;
 		Getter.BindStatic( InFuncPtr );
 	}
-
 	
 	/**
 	 * Binds an arbitrary function that will be called to generate this attribute's value on demand. 
@@ -204,7 +189,6 @@ public:
 		Getter.BindRaw( InUserObject, InMethodPtr );
 	}
 
-
 	/**
 	 * Binds an arbitrary function that will be called to generate this attribute's value on demand. 
 	 * After binding, the attribute will no longer have a value that can be accessed directly, and instead the bound
@@ -219,7 +203,6 @@ public:
 		bIsSet = true;
 		Getter.BindSP( InUserObject, InMethodPtr );
 	}
-
 	
 	/**
 	 * Binds an arbitrary function that will be called to generate this attribute's value on demand. 
@@ -235,7 +218,6 @@ public:
 		bIsSet = true;
 		Getter.BindSP( InUserObject, InMethodPtr );
 	}
-
 	
 	/**
 	 * Binds an arbitrary function that will be called to generate this attribute's value on demand. 
@@ -252,7 +234,6 @@ public:
 		Getter.BindUObject( InUserObject, InMethodPtr );
 	}
 
-
 	/**
 	 * Binds an arbitrary function that will be called to generate this attribute's value on demand.
 	 * After binding, the attribute will no longer have a value that can be accessed directly, and instead the bound
@@ -268,7 +249,6 @@ public:
 		Getter.BindUFunction(InUserObject, InFunctionName);
 	}
 
-
 	/**
 	 * Creates an attribute by binding an arbitrary function that will be called to generate this attribute's value on demand.
 	 * After binding, the attribute will no longer have a value that can be accessed directly, and instead the bound
@@ -285,7 +265,6 @@ public:
 		return Attrib;
 	}
 
-
 	/**
 	 * Checks to see if this attribute has a 'getter' function bound
 	 *
@@ -306,9 +285,11 @@ public:
 		return Getter;
 	}
 
-
 	/**
 	 * Equality operator
+	 *
+	 * @param InOther The other attribute to compare with.
+	 * @return true if the attributes are equal, false otherwise.
 	 */
 	bool operator==( const TAttribute& InOther ) const
 	{
@@ -317,13 +298,15 @@ public:
 
 
 	/**
-	 * Not equal operator
+	 * Not equal operator.
+	 *
+	 * @param InOther The other attribute to compare with.
+	 * @return true if the attributes are not equal, false otherwise.
 	 */
 	bool operator!=( const TAttribute& InOther ) const
 	{
 		return InOther.Get() != Get();
 	}
-
 	
 private:
 
@@ -332,8 +315,7 @@ private:
 		: Value()
 		, Getter( InGetter )
 		, bIsSet( true )
-	{
-	}
+	{ }
 
 	// We declare ourselves as a friend (templated using OtherType) so we can access members as needed
 	template< class OtherType > friend class TAttribute;
@@ -348,7 +330,4 @@ private:
 
 	/** true when this attribute was explicitly set by a consumer, false when the attribute's value is set to the default*/
 	bool bIsSet;
-
 };
-
-
diff --git a/Engine/Source/Runtime/Core/Public/Misc/AutomationTest.h b/Engine/Source/Runtime/Core/Public/Misc/AutomationTest.h
--- a/Engine/Source/Runtime/Core/Public/Misc/AutomationTest.h
+++ b/Engine/Source/Runtime/Core/Public/Misc/AutomationTest.h
@@ -1,8 +1,5 @@
 // Copyright 1998-2014 Epic Games, Inc. All Rights Reserved.
 
-/*=============================================================================
-	AutomationTest.h
-=============================================================================*/
 #pragma once
 
 
diff --git a/Engine/Source/Runtime/Core/Public/Misc/Build.h b/Engine/Source/Runtime/Core/Public/Misc/Build.h
--- a/Engine/Source/Runtime/Core/Public/Misc/Build.h
+++ b/Engine/Source/Runtime/Core/Public/Misc/Build.h
@@ -1,9 +1,5 @@
 // Copyright 1998-2014 Epic Games, Inc. All Rights Reserved.
 
-/*=============================================================================
-	Build.h: Unreal build settings.
-=============================================================================*/
-
 #pragma once
 
 
diff --git a/Engine/Source/Runtime/Core/Public/Misc/ByteSwap.h b/Engine/Source/Runtime/Core/Public/Misc/ByteSwap.h
--- a/Engine/Source/Runtime/Core/Public/Misc/ByteSwap.h
+++ b/Engine/Source/Runtime/Core/Public/Misc/ByteSwap.h
@@ -1,9 +1,5 @@
 // Copyright 1998-2014 Epic Games, Inc. All Rights Reserved.
 
-/*=============================================================================
-	ByteSwap.h: Implements byte swapping utilities
-=============================================================================*/
-
 #pragma once
 
 
diff --git a/Engine/Source/Runtime/Core/Public/Misc/CallbackDevice.h b/Engine/Source/Runtime/Core/Public/Misc/CallbackDevice.h
--- a/Engine/Source/Runtime/Core/Public/Misc/CallbackDevice.h
+++ b/Engine/Source/Runtime/Core/Public/Misc/CallbackDevice.h
@@ -1,15 +1,16 @@
 // Copyright 1998-2014 Epic Games, Inc. All Rights Reserved.
 
-/*=============================================================================
-	CallbackDevice.h:  Allows the engine to do callbacks into the editor
-=============================================================================*/
-
 #pragma once
 
 #include "ScopedCallback.h"
 
+
+class AActor;
+class FUniqueNetIdString;
+
+
 /**
- * The list of events that can be fired from the engine to unknown listeners
+ * The list of events that can be fired from the engine to unknown listeners.
  */
 enum ECallbackQueryType
 {
@@ -18,7 +19,6 @@ enum ECallbackQueryType
 	CALLBACK_QueryCount,
 };
 
-class AActor;
 
 // delegates for hotfixes
 namespace EHotfixDelegates
@@ -29,6 +29,7 @@ namespace EHotfixDelegates
 	};
 }
 
+
 // this is an example of a hotfix arg and return value structure. Once we have other examples, it can be deleted.
 struct FTestHotFixPayload
 {
@@ -37,8 +38,6 @@ struct FTestHotFixPayload
 	bool Result;
 };
 
-//forward declares
-class FUniqueNetIdString;
 
 class CORE_API FCoreDelegates
 {
@@ -198,8 +197,10 @@ public:
 	static FApplicationLicenseChange ApplicationLicenseChange;
 
 private:
+
 	// Callbacks for hotfixes
 	static TArray<FHotFixDelegate> HotFixDelegates;
+
 	// This class is only for namespace use
 	FCoreDelegates() {}
 };
diff --git a/Engine/Source/Runtime/Core/Public/Misc/CompilationResult.h b/Engine/Source/Runtime/Core/Public/Misc/CompilationResult.h
--- a/Engine/Source/Runtime/Core/Public/Misc/CompilationResult.h
+++ b/Engine/Source/Runtime/Core/Public/Misc/CompilationResult.h
@@ -2,9 +2,6 @@
 
 #pragma once
 
-/*=============================================================================
-	CompilationResult.h:  Defines an enum for return codes between build tools.
-=============================================================================*/
 
 /**
  * Enumerates possible results of a compilation operation.
diff --git a/Engine/Source/Runtime/Core/Public/Misc/CompressedGrowableBuffer.h b/Engine/Source/Runtime/Core/Public/Misc/CompressedGrowableBuffer.h
--- a/Engine/Source/Runtime/Core/Public/Misc/CompressedGrowableBuffer.h
+++ b/Engine/Source/Runtime/Core/Public/Misc/CompressedGrowableBuffer.h
@@ -1,11 +1,8 @@
 // Copyright 1998-2014 Epic Games, Inc. All Rights Reserved.
 
-/*=============================================================================
-	CompressedGrowableBuffer.h: FCompressedGrowableBuffer classes
-=============================================================================*/
-
 #pragma once
 
+
 /**
  * Growable compressed buffer. Usage is to append frequently but only request and therefore decompress
  * very infrequently. The prime usage case is the memory profiler keeping track of full call stacks.
@@ -71,6 +68,7 @@ struct CORE_API FCompressedGrowableBuffer
 	}
 
 private:
+
 	/** Helper structure for book keeping. */
 	struct FBufferBookKeeping
 	{
@@ -103,7 +101,3 @@ private:
 	/** Book keeping information for decompression/ access.					*/
 	TArray<FBufferBookKeeping>	BookKeepingInfo;
 };
-
-
-
-
diff --git a/Engine/Source/Runtime/Core/Public/Misc/Compression.h b/Engine/Source/Runtime/Core/Public/Misc/Compression.h
--- a/Engine/Source/Runtime/Core/Public/Misc/Compression.h
+++ b/Engine/Source/Runtime/Core/Public/Misc/Compression.h
@@ -1,11 +1,8 @@
 // Copyright 1998-2014 Epic Games, Inc. All Rights Reserved.
 
-/*=============================================================================
-	Compression.h: Compression classes
-=============================================================================*/
-
 #pragma once
 
+
 /**
  * Flags controlling [de]compression
  */
@@ -21,12 +18,14 @@ enum ECompressionFlags
 	COMPRESS_BiasSpeed				= 0x20,
 };
 
+
 // Define global current platform default to current platform.
 #define COMPRESS_Default			COMPRESS_ZLIB
 
 /** Compression Flag Masks **/
 /** mask out compression type flags */
 #define COMPRESSION_FLAGS_TYPE_MASK		0x0F
+
 /** mask out compression type */
 #define COMPRESSION_FLAGS_OPTIONS_MASK	0xF0
 
diff --git a/Engine/Source/Runtime/Core/Public/Misc/ConfigCacheIni.h b/Engine/Source/Runtime/Core/Public/Misc/ConfigCacheIni.h
--- a/Engine/Source/Runtime/Core/Public/Misc/ConfigCacheIni.h
+++ b/Engine/Source/Runtime/Core/Public/Misc/ConfigCacheIni.h
@@ -1,9 +1,5 @@
 // Copyright 1998-2014 Epic Games, Inc. All Rights Reserved.
 
-/*=============================================================================
-	ConfigCacheIni.h: Unreal config file reading/writing.
-=============================================================================*/
-
 /*-----------------------------------------------------------------------------
 	Config cache.
 -----------------------------------------------------------------------------*/
diff --git a/Engine/Source/Runtime/Core/Public/Misc/CoreMisc.h b/Engine/Source/Runtime/Core/Public/Misc/CoreMisc.h
--- a/Engine/Source/Runtime/Core/Public/Misc/CoreMisc.h
+++ b/Engine/Source/Runtime/Core/Public/Misc/CoreMisc.h
@@ -1,13 +1,10 @@
 // Copyright 1998-2014 Epic Games, Inc. All Rights Reserved.
 
-/*=============================================================================
-	CoreMisc.h: General-purpose file utilities.
-=============================================================================*/
-
 #pragma once
 
 #include "IntPoint.h"
- 
+
+
 /**
  * Exec handler that registers itself and is being routed via StaticExec.
  * Note: Not intended for use with UObjects!
diff --git a/Engine/Source/Runtime/Core/Public/Misc/CoreMiscDefines.h b/Engine/Source/Runtime/Core/Public/Misc/CoreMiscDefines.h
--- a/Engine/Source/Runtime/Core/Public/Misc/CoreMiscDefines.h
+++ b/Engine/Source/Runtime/Core/Public/Misc/CoreMiscDefines.h
@@ -1,11 +1,8 @@
 // Copyright 1998-2014 Epic Games, Inc. All Rights Reserved.
 
-/*=============================================================================
-	CoreMiscDefines.h: Misc defines and enums for core
-=============================================================================*/
-
 #pragma once
 
+
 // This file is included in some resource files, which issue a warning:
 //
 // warning RC4011: identifier truncated to 'PLATFORM_CAN_SUPPORT_EDITORONLY'
diff --git a/Engine/Source/Runtime/Core/Public/Misc/CoreStats.h b/Engine/Source/Runtime/Core/Public/Misc/CoreStats.h
--- a/Engine/Source/Runtime/Core/Public/Misc/CoreStats.h
+++ b/Engine/Source/Runtime/Core/Public/Misc/CoreStats.h
@@ -1,13 +1,10 @@
 // Copyright 1998-2014 Epic Games, Inc. All Rights Reserved.
 
-/*=============================================================================
-	CoreStats.h: Stats enums for core
-=============================================================================*/
-
 #pragma once
 
 #include "Stats.h"
 
+
 CORE_API DECLARE_LOG_CATEGORY_EXTERN(LogSHA, Warning, All);
 CORE_API DECLARE_LOG_CATEGORY_EXTERN(LogStats, Log, All);
 CORE_API DECLARE_LOG_CATEGORY_EXTERN(LogStreaming, Warning, All);
diff --git a/Engine/Source/Runtime/Core/Public/Misc/DataCodec.h b/Engine/Source/Runtime/Core/Public/Misc/DataCodec.h
--- a/Engine/Source/Runtime/Core/Public/Misc/DataCodec.h
+++ b/Engine/Source/Runtime/Core/Public/Misc/DataCodec.h
@@ -1,13 +1,11 @@
 // Copyright 1998-2014 Epic Games, Inc. All Rights Reserved.
 
-/*=============================================================================
-	DataCodec.h: Data compression codecs.
-=============================================================================*/
+#pragma once
+
 
 /*-----------------------------------------------------------------------------
 	Coder/decoder base class.
 -----------------------------------------------------------------------------*/
-#pragma once
 
 DECLARE_LOG_CATEGORY_EXTERN(LogDataCodex, Log, All);
 
@@ -19,6 +17,7 @@ public:
 	virtual ~FCodec(){}
 };
 
+
 /*-----------------------------------------------------------------------------
 	Burrows-Wheeler inspired data compressor.
 -----------------------------------------------------------------------------*/
diff --git a/Engine/Source/Runtime/Core/Public/Misc/Database.h b/Engine/Source/Runtime/Core/Public/Misc/Database.h
--- a/Engine/Source/Runtime/Core/Public/Misc/Database.h
+++ b/Engine/Source/Runtime/Core/Public/Misc/Database.h
@@ -1,11 +1,8 @@
 // Copyright 1998-2014 Epic Games, Inc. All Rights Reserved.
 
-/*=============================================================================
-	Database.h: Database integration
-=============================================================================*/
-
 #pragma once
 
+
 /**
  * Whether to compile in support for database connectivity and SQL execution.
  */
@@ -13,7 +10,6 @@
 	#define WITH_DATABASE_SUPPORT (!UE_BUILD_MINIMAL && !UE_BUILD_SHIPPING)
 #endif
 
-
 // Only use ADO on windows, if support is enabled and not for shipping games.
 // @todo clang: #import is not supported by Clang on Windows platform, but we currently need this for ADO symbol importing.  For now we disable ADO support in Clang builds.
 #define USE_ADO_INTEGRATION (PLATFORM_WINDOWS && !(UE_BUILD_SHIPPING && WITH_EDITOR) && WITH_DATABASE_SUPPORT && !defined(__clang__))
diff --git a/Engine/Source/Runtime/Core/Public/Misc/DefaultValueHelper.h b/Engine/Source/Runtime/Core/Public/Misc/DefaultValueHelper.h
--- a/Engine/Source/Runtime/Core/Public/Misc/DefaultValueHelper.h
+++ b/Engine/Source/Runtime/Core/Public/Misc/DefaultValueHelper.h
@@ -1,11 +1,8 @@
 // Copyright 1998-2014 Epic Games, Inc. All Rights Reserved.
 
-/*=============================================================================
-	CppDefaultValueHelper.h: Parses default values for C++ function parameters
-=============================================================================*/
-
 #pragma once
 
+
 //
 //	FDefaultValueHelper
 //
diff --git a/Engine/Source/Runtime/Core/Public/Misc/EngineVersion.h b/Engine/Source/Runtime/Core/Public/Misc/EngineVersion.h
--- a/Engine/Source/Runtime/Core/Public/Misc/EngineVersion.h
+++ b/Engine/Source/Runtime/Core/Public/Misc/EngineVersion.h
@@ -1,11 +1,8 @@
 // Copyright 1998-2014 Epic Games, Inc. All Rights Reserved.
 
-/*=============================================================================
-	EngineVersion.h: Encapsulates version information for the engine.
-=============================================================================*/
-
 #pragma once
 
+
 /** Enum for the components of a version string. */
 namespace EVersionComponent
 {
@@ -19,6 +16,7 @@ namespace EVersionComponent
 	};
 }
 
+
 /** Components of a version string. */
 namespace EVersionComparison
 {
@@ -30,10 +28,12 @@ namespace EVersionComparison
 	};
 }
 
+
 /** Utility functions. */
 class CORE_API FEngineVersion
 {
 public:
+
 	/** Empty constructor. Initializes the version to 0.0.0-0. */
 	FEngineVersion();
 
@@ -74,6 +74,7 @@ public:
 	friend CORE_API void operator<<(FArchive &Ar, FEngineVersion &Version);
 
 private:
+
 	/** Major version number. */
 	uint16 Major;
 
@@ -90,5 +91,6 @@ private:
 	FString Branch;
 };
 
+
 /** Global instance of the current engine version. */
 CORE_API extern const FEngineVersion GEngineVersion;
diff --git a/Engine/Source/Runtime/Core/Public/Misc/FeedbackContext.h b/Engine/Source/Runtime/Core/Public/Misc/FeedbackContext.h
--- a/Engine/Source/Runtime/Core/Public/Misc/FeedbackContext.h
+++ b/Engine/Source/Runtime/Core/Public/Misc/FeedbackContext.h
@@ -1,15 +1,13 @@
 // Copyright 1998-2014 Epic Games, Inc. All Rights Reserved.
 
-/*=============================================================================
-	FeedbackContext.h: FFeedbackContext
-=============================================================================*/
-
 #pragma once
 
 #include "UniquePtr.h"
 
+
 struct FScopedSlowTask;
 
+
 /** A stack of feedback scopes */
 struct FScopedSlowTaskStack : TArray<FScopedSlowTask*>
 {
@@ -17,8 +15,10 @@ struct FScopedSlowTaskStack : TArray<FScopedSlowTask*>
 	CORE_API float GetProgressFraction(int32 Index) const;
 };
 
+
 /** A context for displaying modal warning messages. */
-class CORE_API FFeedbackContext : public FOutputDevice
+class CORE_API FFeedbackContext
+	: public FOutputDevice
 {
 public:
 
@@ -105,7 +105,6 @@ protected:
 };
 
 
-
 /**
  * A scope block representing an amount of work divided up into sections.
  * Use one scope at the top of each function to give accurate feedback to the user of a slow operation's progress.
diff --git a/Engine/Source/Runtime/Core/Public/Misc/FeedbackContextMarkup.h b/Engine/Source/Runtime/Core/Public/Misc/FeedbackContextMarkup.h
--- a/Engine/Source/Runtime/Core/Public/Misc/FeedbackContextMarkup.h
+++ b/Engine/Source/Runtime/Core/Public/Misc/FeedbackContextMarkup.h
@@ -1,12 +1,8 @@
 // Copyright 1998-2014 Epic Games, Inc. All Rights Reserved.
 
-/*=============================================================================================
-	FeedbackContextMarkup.h: Functions to parse markup from external processes and control
-	a feedback context.
-==============================================================================================*/
-
 #pragma once
 
+
 /**
  * A text-based markup language can be used to allow external processes to control the state 
  * machine, allowing command-line utilities to give graphical user feedback.
@@ -23,13 +19,15 @@
 class CORE_API FFeedbackContextMarkup
 {
 public:
-	// Markup stack manipulation
+
+	/** Markup stack manipulation. */
 	static bool ParseCommand(const FString& Line, FFeedbackContext* Warn);
 
-	// Utility functions for dealing with external processes
+	/** Utility functions for dealing with external processes. */
 	static bool PipeProcessOutput(const FText& Description, const FString& URL, const FString& Params, FFeedbackContext* Warn, int32* OutExitCode);
 
 private:
+
 	/** Try to read a single exact-match token from the input stream. Must be followed by whitespace or EOL. */
 	static bool ReadToken(const TCHAR *&Text, const TCHAR *Token);
 
diff --git a/Engine/Source/Runtime/Core/Public/Misc/FilterCollection.h b/Engine/Source/Runtime/Core/Public/Misc/FilterCollection.h
--- a/Engine/Source/Runtime/Core/Public/Misc/FilterCollection.h
+++ b/Engine/Source/Runtime/Core/Public/Misc/FilterCollection.h
@@ -1,11 +1,14 @@
 // Copyright 1998-2014 Epic Games, Inc. All Rights Reserved.
+
 #pragma once
 
+
 /**
- *	A simple collection of Filters, with additional Filter specific functionality
+ *	A simple collection of Filters, with additional Filter specific functionality.
  */
 template< typename ItemType >
-class TFilterCollection : public TSharedFromThis< TFilterCollection< ItemType > >
+class TFilterCollection
+	: public TSharedFromThis< TFilterCollection< ItemType > >
 {
 public:
 
@@ -25,7 +28,6 @@ public:
 		}
 	}
 
-
 	/**
 	 *	Adds the specified Filter to the collection
 	 *
@@ -52,7 +54,6 @@ public:
 		return Result;
 	}
 	
-
 	/** 
 	 *	Removes as many instances of the specified Filter as there are in the collection
 	 *	
@@ -83,14 +84,12 @@ public:
 		return ChildFilters[Index];
 	}
 
-
 	/** Returns the number of Filters in the collection */
 	FORCEINLINE int32 Num() const
 	{
 		return ChildFilters.Num();
 	}
 
-
 	/** 
 	 *	Returns whether the specified Item passes all of the filters in the collection
 	 *
@@ -114,7 +113,6 @@ public:
 	DECLARE_EVENT( TFilterCollection<ItemType>, FChangedEvent );
 	FChangedEvent& OnChanged() { return ChangedEvent; }
 
-
 protected:
 
 	/** 
@@ -131,4 +129,4 @@ protected:
 
 	/**	Fires whenever any filter in the collection changes */
 	FChangedEvent ChangedEvent;
-};
\ No newline at end of file
+};
diff --git a/Engine/Source/Runtime/Core/Public/Misc/IFilter.h b/Engine/Source/Runtime/Core/Public/Misc/IFilter.h
--- a/Engine/Source/Runtime/Core/Public/Misc/IFilter.h
+++ b/Engine/Source/Runtime/Core/Public/Misc/IFilter.h
@@ -1,8 +1,10 @@
 // Copyright 1998-2014 Epic Games, Inc. All Rights Reserved.
+
 #pragma once
 
+
 /**
- *	A generic interface that represents a Filter of ItemType
+ * A generic interface that represents a Filter of ItemType.
  */
 template< typename TItemType >
 class IFilter
@@ -18,4 +20,4 @@ public:
 	/** Broadcasts anytime the restrictions of the Filter changes */
 	DECLARE_EVENT( IFilter<TItemType>, FChangedEvent );
 	virtual FChangedEvent& OnChanged() = 0;
-};
\ No newline at end of file
+};
diff --git a/Engine/Source/Runtime/Core/Public/Misc/ITransaction.h b/Engine/Source/Runtime/Core/Public/Misc/ITransaction.h
--- a/Engine/Source/Runtime/Core/Public/Misc/ITransaction.h
+++ b/Engine/Source/Runtime/Core/Public/Misc/ITransaction.h
@@ -1,9 +1,5 @@
 // Copyright 1998-2014 Epic Games, Inc. All Rights Reserved.
 
-/*=============================================================================
-	ITransaction.h: Declares the ITransaction interface.
-=============================================================================*/
-
 #pragma once
 
 
@@ -23,6 +19,7 @@ typedef void(*STRUCT_DTOR)( void* TPtr );					// destruct
  */
 class ITransactionObjectAnnotation { };
 
+
 /**
  * Interface for transactions.
  *
@@ -33,9 +30,7 @@ class ITransaction
 {
 public:
 
-	/**
-	 * Applies the transaction.
-	 */
+	/** Applies the transaction. */
 	virtual void Apply( ) = 0;
 
 	/**
@@ -49,6 +44,7 @@ public:
 	 * @param ElementSize
 	 * @param Serializer
 	 * @param Destructor
+	 * @see SaveObject
 	 */
 	virtual void SaveArray( UObject* Object, class FScriptArray* Array, int32 Index, int32 Count, int32 Oper, int32 ElementSize, STRUCT_DC DefaultConstructor, STRUCT_AR Serializer, STRUCT_DTOR Destructor ) = 0;
 
@@ -56,6 +52,7 @@ public:
 	 * Saves an UObject to the transaction.
 	 *
 	 * @param Object The object to save.
+	 * @see SaveArray
 	 */
 	virtual void SaveObject( UObject* Object ) = 0;
 
diff --git a/Engine/Source/Runtime/Core/Public/Misc/LocalTimestampDirectoryVisitor.h b/Engine/Source/Runtime/Core/Public/Misc/LocalTimestampDirectoryVisitor.h
--- a/Engine/Source/Runtime/Core/Public/Misc/LocalTimestampDirectoryVisitor.h
+++ b/Engine/Source/Runtime/Core/Public/Misc/LocalTimestampDirectoryVisitor.h
@@ -1,21 +1,17 @@
 // Copyright 1998-2014 Epic Games, Inc. All Rights Reserved.
 
-/*=============================================================================
-	LocalTimestampDirectoryVisitor.h: Declares the FLocalTimestampDirectoryVisitor class.
-=============================================================================*/
-
 #pragma once
 
 
 /**
- * Visitor to gather local files with their timestamps
+ * Visitor to gather local files with their timestamps.
  */
 class CORE_API FLocalTimestampDirectoryVisitor
 	: public IPlatformFile::FDirectoryVisitor
 {
 public:
 
-	/** Relative paths to local files and their timestamps */
+	/** Relative paths to local files and their timestamps. */
 	TMap<FString, FDateTime> FileTimes;
 
 	/**
@@ -28,27 +24,23 @@ public:
 	 */
 	FLocalTimestampDirectoryVisitor( IPlatformFile& InFileInterface, const TArray<FString>& InDirectoriesToIgnore, const TArray<FString>& InDirectoriesToNotRecurse, bool bInCacheDirectories = false );
 
-
 public:
 
-	// Begin IPlatformFile::FDirectoryVisitor interface
+	// IPlatformFile::FDirectoryVisitor interface
 
 	virtual bool Visit(const TCHAR* FilenameOrDirectory, bool bIsDirectory);
 
-	// End IPlatformFile::FDirectoryVisitor interface
-
-
 private:
 
-	// true if we want directories in this list.
+	// true if we want directories in this list. */
 	bool bCacheDirectories;
 
-	// Holds a list of directories that we should not traverse into.
+	// Holds a list of directories that we should not traverse into. */
 	TArray<FString> DirectoriesToIgnore;
 
-	// Holds a list of directories that we should only go one level into.
+	// Holds a list of directories that we should only go one level into. */
 	TArray<FString> DirectoriesToNotRecurse;
 
-	// Holds the file interface to use for any file operations.
+	// Holds the file interface to use for any file operations. */
 	IPlatformFile& FileInterface;
 };
\ No newline at end of file
diff --git a/Engine/Source/Runtime/Core/Public/Misc/MemStack.h b/Engine/Source/Runtime/Core/Public/Misc/MemStack.h
--- a/Engine/Source/Runtime/Core/Public/Misc/MemStack.h
+++ b/Engine/Source/Runtime/Core/Public/Misc/MemStack.h
@@ -1,16 +1,22 @@
 // Copyright 1998-2014 Epic Games, Inc. All Rights Reserved.
 
-/*=============================================================================
-	MemStack.h: FMemStack class, ultra-fast temporary memory allocation
-=============================================================================*/
-
 #pragma once
+
 #include "LockFreeFixedSizeAllocator.h"
 
 
 // Enums for specifying memory allocation type.
-enum EMemZeroed {MEM_Zeroed=1};
-enum EMemOned   {MEM_Oned  =1};
+enum EMemZeroed
+{
+	MEM_Zeroed = 1
+};
+
+
+enum EMemOned
+{
+	MEM_Oned = 1
+};
+
 
 class CORE_API FPageAllocator
 {
@@ -59,6 +65,7 @@ private:
 	static TLockFreeFixedSizeAllocator<SmallPageSize, FThreadSafeCounter> TheSmallAllocator;
 };
 
+
 /**
  * Simple linear-allocation memory stack.
  * Items are allocated via PushBytes() or the specialized operator new()s.
@@ -182,10 +189,12 @@ private:
 	int32 MinMarksToAlloc;
 };
 
+
 class CORE_API FMemStack : public TThreadSingleton<FMemStack>, public FMemStackBase
 {
 };
 
+
 /*-----------------------------------------------------------------------------
 	FMemStack templates.
 -----------------------------------------------------------------------------*/
@@ -208,6 +217,7 @@ template <class T> inline T* NewOned(FMemStackBase& Mem, int32 Count = 1, int32
 	return (T*)Result;
 }
 
+
 /*-----------------------------------------------------------------------------
 	FMemStack operator new's.
 -----------------------------------------------------------------------------*/
@@ -252,6 +262,7 @@ inline void* operator new[](size_t Size, FMemStackBase& Mem, EMemOned Tag, int32
 	return Result;
 }
 
+
 /** A container allocator that allocates from a mem-stack. */
 template<uint32 Alignment = DEFAULT_ALIGNMENT>
 class TMemStackAllocator
@@ -314,6 +325,7 @@ public:
 	typedef ForElementType<FScriptContainerElement> ForAnyElementType;
 };
 
+
 /**
  * FMemMark marks a top-of-stack position in the memory stack.
  * When the marker is constructed or initialized with a particular memory 
@@ -370,6 +382,7 @@ public:
 	}
 
 private:
+
 	// Implementation variables.
 	FMemStackBase& Mem;
 	uint8* Top;
diff --git a/Engine/Source/Runtime/Core/Public/Misc/MonitoredProcess.h b/Engine/Source/Runtime/Core/Public/Misc/MonitoredProcess.h
--- a/Engine/Source/Runtime/Core/Public/Misc/MonitoredProcess.h
+++ b/Engine/Source/Runtime/Core/Public/Misc/MonitoredProcess.h
@@ -2,6 +2,7 @@
 
 #pragma once
 
+
 /**
  * Declares a delegate that is executed when a monitored process completed.
  *
@@ -34,10 +35,8 @@ public:
 	 */
 	FMonitoredProcess( const FString& InURL, const FString& InParams, bool InHidden );
 
-	/**
-	 * Destructor.
-	 */
-	~FMonitoredProcess( );
+	/** Destructor. */
+	~FMonitoredProcess();
 
 public:
 
@@ -57,23 +56,20 @@ public:
 	 *
 	 * @return Time duration.
 	 */
-	FTimespan GetDuration( ) const;
+	FTimespan GetDuration() const;
 
 	/**
 	 * Checks whether the process is still running.
 	 *
 	 * @return true if the process is running, false otherwise.
 	 */
-	bool IsRunning( ) const
+	bool IsRunning() const
 	{
-		return (Thread != NULL);
+		return (Thread != nullptr);
 	}
 
-	/**
-	 * Launches the process.
-	 */
-	bool Launch( );
-
+	/** Launches the process. */
+	bool Launch();
 
 public:
 
@@ -82,7 +78,7 @@ public:
 	 *
 	 * @return The delegate.
 	 */
-	FSimpleDelegate& OnCanceled( )
+	FSimpleDelegate& OnCanceled()
 	{
 		return CanceledDelegate;
 	}
@@ -92,7 +88,7 @@ public:
 	 *
 	 * @return The delegate.
 	 */
-	FOnMonitoredProcessCompleted& OnCompleted( )
+	FOnMonitoredProcessCompleted& OnCompleted()
 	{
 		return CompletedDelegate;
 	}
@@ -102,7 +98,7 @@ public:
 	 *
 	 * @return The delegate.
 	 */
-	FOnMonitoredProcessOutput& OnOutput( )
+	FOnMonitoredProcessOutput& OnOutput()
 	{
 		return OutputDelegate;
 	}
@@ -121,19 +117,19 @@ public:
 
 	// FRunnable interface
 
-	virtual bool Init( ) override
+	virtual bool Init() override
 	{
 		return true;
 	}
 
-	virtual uint32 Run( ) override;
+	virtual uint32 Run() override;
 
-	virtual void Stop( ) override
+	virtual void Stop() override
 	{
 		Cancel();
 	}
 
-	virtual void Exit( ) override { }
+	virtual void Exit() override { }
 
 protected:
 
@@ -146,50 +142,50 @@ protected:
 
 private:
 
-	// Whether the process is being canceled.
+	// Whether the process is being canceled. */
 	bool Canceling;
 
-	// Holds the time at which the process ended.
+	// Holds the time at which the process ended. */
 	FDateTime EndTime;
 
-	// Whether the window of the process should be hidden.
+	// Whether the window of the process should be hidden. */
 	bool Hidden;
 
-	// Whether to kill the entire process tree when cancelling this process.
+	// Whether to kill the entire process tree when cancelling this process. */
 	bool KillTree;
 
-	// Holds the command line parameters.
+	// Holds the command line parameters. */
 	FString Params;
 
-	// Holds the handle to the process.
+	// Holds the handle to the process. */
 	FProcHandle ProcessHandle;
 
-	// Holds the read pipe.
+	// Holds the read pipe. */
 	void* ReadPipe;
 
-	// Holds the time at which the process started.
+	// Holds the return code. */
+	int ReturnCode;
+
+	// Holds the time at which the process started. */
 	FDateTime StartTime;
 
-	// Holds the monitoring thread object.
+	// Holds the monitoring thread object. */
 	FRunnableThread* Thread;
 
-	// Holds the URL of the executable to launch.
+	// Holds the URL of the executable to launch. */
 	FString URL;
 
-	// Holds the write pipe.
+	// Holds the write pipe. */
 	void* WritePipe;
 
-	// Holds the return code
-	int ReturnCode;
-
 private:
 
-	// Holds a delegate that is executed when the process has been canceled.
+	// Holds a delegate that is executed when the process has been canceled. */
 	FSimpleDelegate CanceledDelegate;
 
-	// Holds a delegate that is executed when a monitored process completed.
+	// Holds a delegate that is executed when a monitored process completed. */
 	FOnMonitoredProcessCompleted CompletedDelegate;
 
-	// Holds a delegate that is executed when a monitored process produces output.
+	// Holds a delegate that is executed when a monitored process produces output. */
 	FOnMonitoredProcessOutput OutputDelegate;
 };
diff --git a/Engine/Source/Runtime/Core/Public/Misc/NetworkGuid.h b/Engine/Source/Runtime/Core/Public/Misc/NetworkGuid.h
--- a/Engine/Source/Runtime/Core/Public/Misc/NetworkGuid.h
+++ b/Engine/Source/Runtime/Core/Public/Misc/NetworkGuid.h
@@ -1,9 +1,5 @@
 // Copyright 1998-2014 Epic Games, Inc. All Rights Reserved.
 
-/*=============================================================================
-	NetworkGuid.h: Declares the FNetworkGUID class.
-=============================================================================*/
-
 #pragma once
 
 
@@ -19,10 +15,9 @@ public:
 
 	uint32 Value;
 
-
 public:
 
-	FNetworkGUID( )
+	FNetworkGUID()
 		: Value(0)
 	{ }
 
@@ -30,7 +25,6 @@ public:
 		: Value(V)
 	{ }
 		
-
 public:
 
 	friend bool operator==( const FNetworkGUID& X, const FNetworkGUID& Y )
@@ -48,7 +42,6 @@ public:
 		return Ar << G.Value;
 	}
 
-
 public:
 
 	void BuildFromNetIndex( int32 StaticNetIndex )
@@ -56,7 +49,7 @@ public:
 		Value = (StaticNetIndex << 1 | 1);
 	}
 
-	int32 ExtractNetIndex( )
+	int32 ExtractNetIndex()
 	{
 		if (Value & 1)
 		{
@@ -70,17 +63,17 @@ public:
 		return Guid.Value;
 	}
 
-	bool IsDynamic( ) const
+	bool IsDynamic() const
 	{
 		return Value > 0 && !(Value & 1);
 	}
 
-	bool IsStatic( ) const
+	bool IsStatic() const
 	{
 		return Value & 1;
 	}
 
-	bool IsValid( ) const
+	bool IsValid() const
 	{
 		return Value > 0;
 	}
@@ -103,12 +96,11 @@ public:
 		Value = 0;
 	}
 
-	FString ToString( ) const
+	FString ToString() const
 	{
 		return FString::Printf(TEXT("%d"), Value);
 	}
 
-
 public:
 
 	CORE_API static FNetworkGUID Make(int32 seed, bool bIsStatic)
diff --git a/Engine/Source/Runtime/Core/Public/Misc/ObjectThumbnail.h b/Engine/Source/Runtime/Core/Public/Misc/ObjectThumbnail.h
--- a/Engine/Source/Runtime/Core/Public/Misc/ObjectThumbnail.h
+++ b/Engine/Source/Runtime/Core/Public/Misc/ObjectThumbnail.h
@@ -1,15 +1,7 @@
 // Copyright 1998-2014 Epic Games, Inc. All Rights Reserved.
 
-/*=============================================================================
-	ObjectThumbnail.h: Stored thumbnail support for Unreal objects
-=============================================================================*/
-
 #pragma once
 
-#ifndef __ObjectThumbnail_h__
-#define __ObjectThumbnail_h__
-
-
 
 /**
  * Thumbnail compression interface for packages.  The engine registers a class that can compress and
@@ -17,7 +9,6 @@
  */
 class FThumbnailCompressionInterface
 {
-	
 public:
 
 	/**
@@ -27,12 +18,10 @@ public:
 	 * @param	InWidth				Width of the image
 	 * @param	InHeight			Height of the image
 	 * @param	OutCompressedData	[Out] Compressed image data
-	 *
 	 * @return	true if the image was compressed successfully, otherwise false if an error occurred
 	 */
 	virtual bool CompressImage( const TArray< uint8 >& InUncompressedData, const int32 InWidth, const int32 InHeight, TArray< uint8 >& OutCompressedData ) = 0;
 
-
 	/**
 	 * Decompresses an image
 	 *
@@ -40,130 +29,120 @@ public:
 	 * @param	InWidth				Width of the image
 	 * @param	InHeight			Height of the image
 	 * @param	OutUncompressedData	[Out] Uncompressed image data
-	 *
 	 * @return	true if the image was decompressed successfully, otherwise false if an error occurred
 	 */
 	virtual bool DecompressImage( const TArray< uint8 >& InCompressedData, const int32 InWidth, const int32 InHeight, TArray< uint8 >& OutUncompressedData ) = 0;
-
 };
 
 
-
-
-/** Thumbnail image data for an object */
+/**
+ * Thumbnail image data for an object.
+ */
 class CORE_API FObjectThumbnail
 {
-
 public:
 
 	/**
 	 * Static: Sets the thumbnail compressor to use when loading/saving packages.  The caller is
-	 * responsible for the object's lifespan
+	 * responsible for the object's lifespan.
 	 *
-	 * @param	InThumbnailCompressor	A class derived from FThumbnailCompressionInterface
+	 * @param	InThumbnailCompressor	A class derived from FThumbnailCompressionInterface.
 	 */
 	static void SetThumbnailCompressor( FThumbnailCompressionInterface* InThumbnailCompressor )
 	{
 		ThumbnailCompressor = InThumbnailCompressor;
 	}
 
-
 private:
 
-	/** Static: Thumbnail compressor */
+	/** Static: Thumbnail compressor. */
 	static FThumbnailCompressionInterface* ThumbnailCompressor;
 
-
 public:
 
-	/** Default constructor */
+	/** Default constructor. */
 	FObjectThumbnail();
 
-
-	/** Returns the width of the thumbnail */
+	/** Returns the width of the thumbnail. */
 	int32 GetImageWidth() const
 	{
 		return ImageWidth;
 	}
 
-
-	/** Returns the height of the thumbnail */
+	/** Returns the height of the thumbnail. */
 	int32 GetImageHeight() const
 	{
 		return ImageHeight;
 	}
 
-	/** @return		the number of bytes in this thumbnail's compressed image data */
+	/** @return		the number of bytes in this thumbnail's compressed image data. */
 	int32 GetCompressedDataSize() const
 	{
 		return CompressedImageData.Num();
 	}
 
-
-	/** Sets the image dimensions */
+	/** Sets the image dimensions. */
 	void SetImageSize( int32 InWidth, int32 InHeight )
 	{
 		ImageWidth = InWidth;
 		ImageHeight = InHeight;
 	}
 
-	/** Returns true if the thumbnail was loaded from disk and not dynamically generated */
+	/** Returns true if the thumbnail was loaded from disk and not dynamically generated. */
 	bool IsLoadedFromDisk(void) const { return bLoadedFromDisk; }
 
-	/** Returns true if the thumbnail was saved AFTER custom-thumbnails for shared thumbnail asset types was supported */
+	/** Returns true if the thumbnail was saved AFTER custom-thumbnails for shared thumbnail asset types was supported. */
 	bool IsCreatedAfterCustomThumbsEnabled (void) const { return bCreatedAfterCustomThumbForSharedTypesEnabled; }
-	/** For newly generated custom thumbnails, mark it as valid in the future */
+	
+	/** For newly generated custom thumbnails, mark it as valid in the future. */
 	void SetCreatedAfterCustomThumbsEnabled(void) { bCreatedAfterCustomThumbForSharedTypesEnabled = true; }
 
-
-	/** Returns true if the thumbnail is dirty and needs to be regenerated at some point */
+	/** Returns true if the thumbnail is dirty and needs to be regenerated at some point. */
 	bool IsDirty() const
 	{
 		return bIsDirty;
 	}
-
 	
-	/** Marks the thumbnail as dirty */
+	/** Marks the thumbnail as dirty. */
 	void MarkAsDirty()
 	{
 		bIsDirty = true;
 	}
 	
-	/** Access the image data in place (does not decompress) */
+	/** Access the image data in place (does not decompress). */
 	TArray< uint8 >& AccessImageData()
 	{
 		return ImageData;
 	}
 
-	/** Access the image data in place (does not decompress) const version */
+	/** Access the image data in place (does not decompress) const version. */
 	const TArray< uint8 >& AccessImageData() const
 	{
 		return ImageData;
 	}
 
-	/** Access the compressed image data */
+	/** Access the compressed image data. */
 	TArray< uint8 >& AccessCompressedImageData()
 	{
 		return CompressedImageData;
 	}
 
-	/** Returns true if this is an empty thumbnail */
+	/** Returns true if this is an empty thumbnail. */
 	bool IsEmpty() const
 	{
 		return ImageWidth == 0 || ImageHeight == 0;
 	}
 
-
-	/** Returns uncompressed image data, decompressing it on demand if needed */
+	/** Returns uncompressed image data, decompressing it on demand if needed. */
 	const TArray< uint8 >& GetUncompressedImageData() const;
 
 	/** Serializer */
 	void Serialize( FArchive& Ar );
 	
-	/** Compress image data */
+	/** Compress image data. */
 	void CompressImageData();
 
-	/** Decompress image data */
+	/** Decompress image data. */
 	void DecompressImageData();
 
 	/**
@@ -174,13 +153,14 @@ public:
 	void CountBytes( FArchive& Ar ) const;
 
 	/**
-	 * Calculates the amount of memory used by the compressed bytes array
+	 * Calculates the amount of memory used by the compressed bytes array.
 	 *
 	 * @param	Ar	the FArchiveCountMem (or similar) archive that will store the results of the memory usage calculation.
 	 */
 	void CountImageBytes_Compressed( FArchive& Ar ) const;
+
 	/**
-	 * Calculates the amount of memory used by the uncompressed bytes array
+	 * Calculates the amount of memory used by the uncompressed bytes array.
 	 *
 	 * @param	Ar	the FArchiveCountMem (or similar) archive that will store the results of the memory usage calculation.
 	 */
@@ -199,6 +179,7 @@ public:
 		}
 		return Ar;
 	}
+
 	friend FArchive& operator<<( FArchive& Ar, const FObjectThumbnail& Thumb )
 	{
 		Thumb.CountBytes(Ar);
@@ -213,6 +194,7 @@ public:
 			&&	bIsDirty			== Other.bIsDirty
 			&&	CompressedImageData	== Other.CompressedImageData;
 	}
+
 	bool operator !=( const FObjectThumbnail& Other ) const
 	{
 		return	ImageWidth			!= Other.ImageWidth
@@ -240,9 +222,9 @@ private:
 
 	/** Whether the thumbnail has a backup on disk*/
 	bool bLoadedFromDisk;
+
 	/** Whether this was saved AFTER custom-thumbnails for shared thumbnail asset types was supported */
 	bool bCreatedAfterCustomThumbForSharedTypesEnabled;
-
 };
 
 
@@ -262,22 +244,19 @@ struct CORE_API FObjectFullNameAndThumbnail
 	/** Offset in the file where the data is stored */
 	int32 FileOffset;
 
-
 	/** Constructor */
 	FObjectFullNameAndThumbnail()
 		: ObjectFullName(),
-		  ObjectThumbnail( NULL ),
+		  ObjectThumbnail( nullptr ),
 		  FileOffset( 0 )
-	{
-	}
+	{ }
 
 	/** Constructor */
 	FObjectFullNameAndThumbnail( const FName InFullName, const FObjectThumbnail* InThumbnail )
 		: ObjectFullName( InFullName ),
 		  ObjectThumbnail( InThumbnail ),
 		  FileOffset( 0 )
-	{
-	}
+	{ }
 
 	/**
 	 * Calculates the memory usage of this FObjectFullNameAndThumbnail.
@@ -299,13 +278,10 @@ struct CORE_API FObjectFullNameAndThumbnail
 		}
 		return Ar;
 	}
+
 	friend FArchive& operator<<( FArchive& Ar, const FObjectFullNameAndThumbnail& NameThumbPair )
 	{
 		NameThumbPair.CountBytes(Ar);
 		return Ar;
 	}
 };
-
-
-
-#endif	// __ObjectThumbnail_h__
\ No newline at end of file
diff --git a/Engine/Source/Runtime/Core/Public/Misc/OutputDevice.h b/Engine/Source/Runtime/Core/Public/Misc/OutputDevice.h
--- a/Engine/Source/Runtime/Core/Public/Misc/OutputDevice.h
+++ b/Engine/Source/Runtime/Core/Public/Misc/OutputDevice.h
@@ -1,11 +1,8 @@
 // Copyright 1998-2014 Epic Games, Inc. All Rights Reserved.
 
-/*=============================================================================
-	OutputDevice.h: FOutputDevice and OutputDebugString type functions
-=============================================================================*/
-
 #pragma once
 
+
 class FText;
 
 // Globals.
diff --git a/Engine/Source/Runtime/Core/Public/Misc/OutputDeviceConsole.h b/Engine/Source/Runtime/Core/Public/Misc/OutputDeviceConsole.h
--- a/Engine/Source/Runtime/Core/Public/Misc/OutputDeviceConsole.h
+++ b/Engine/Source/Runtime/Core/Public/Misc/OutputDeviceConsole.h
@@ -1,22 +1,16 @@
 // Copyright 1998-2014 Epic Games, Inc. All Rights Reserved.
 
-/*=============================================================================
-	ConsoleOutputDevice.h: FOutputDeviceConsole definition.
-=============================================================================*/
-
 #pragma once
 
+
 /**
  * This class servers as the base class for console window output.
  */
-class CORE_API FOutputDeviceConsole : public FOutputDevice
+class CORE_API FOutputDeviceConsole
+	: public FOutputDevice
 {
-protected:
-
-	/** Ini file name to write console settings to. */
-	FString IniFilename;
-
 public:
+
 	/**
 	 * Shows or hides the console window. 
 	 *
@@ -25,16 +19,16 @@ public:
 	virtual void Show( bool ShowWindow )=0;
 
 	/** 
-	 * Returns whether console is currently shown or not
+	 * Returns whether console is currently shown or not.
 	 *
-	 * @return true if console is shown, false otherwise
+	 * @return true if console is shown, false otherwise.
 	 */
 	virtual bool IsShown()=0;
 
 	/** 
-	 * Returns whether the application is already attached to a console window
+	 * Returns whether the application is already attached to a console window.
 	 *
-	 * @return true if console is attahced, false otherwise
+	 * @return true if console is attached, false otherwise.
 	 */
 	virtual bool IsAttached()
 	{
@@ -42,10 +36,17 @@ public:
 	}
 
 	/**
-	 * Sets ini file name to write console settings to.
+	 * Sets the INI file name to write console settings to.
+	 *
+	 * @param InFilename The INI file name to set.
 	 */
 	void SetIniFilename(const TCHAR* InFilename) 
 	{
 		IniFilename = InFilename;
 	}
+
+protected:
+
+	/** Ini file name to write console settings to. */
+	FString IniFilename;
 };
diff --git a/Engine/Source/Runtime/Core/Public/Misc/Paths.h b/Engine/Source/Runtime/Core/Public/Misc/Paths.h
--- a/Engine/Source/Runtime/Core/Public/Misc/Paths.h
+++ b/Engine/Source/Runtime/Core/Public/Misc/Paths.h
@@ -1,9 +1,11 @@
 // Copyright 1998-2014 Epic Games, Inc. All Rights Reserved.
+
 #pragma once
 
-/*-----------------------------------------------------------------------------
-	Path helpers for retrieving game dir, engine dir, etc.
------------------------------------------------------------------------------*/
+
+/**
+ * Path helpers for retrieving game dir, engine dir, etc.
+ */
 class CORE_API FPaths
 {
 	/**
@@ -12,6 +14,7 @@ class CORE_API FPaths
 	static bool ShouldSaveToUserDir();
 
 public:
+
 	/** 
 	 * Returns the base directory of the "core" engine that can be shared across
 	 * several games or across games & mods. Shaders and base localization files
@@ -193,7 +196,6 @@ public:
 	 */
 	static FString GameLogDir();
 
-
 	/**
 	 * @return The directory for automation save files
 	 */
@@ -359,43 +361,43 @@ public:
 	static void MakePlatformFilename(FString& InPath);
 
 	/** 
-	 * Assuming both paths (or filenames) are relative to the base dir, find the relative path to the InPath 
+	 * Assuming both paths (or filenames) are relative to the base dir, find the relative path to the InPath.
 	 *
-	 * @Param InPath - Path to make this path relative to.
-	 * @return Path relative to InPath
-	 **/
+	 * @Param InPath Path to make this path relative to.
+	 * @return Path relative to InPath.
+	 */
 	static bool MakePathRelativeTo( FString& InPath, const TCHAR* InRelativeTo );
 
 	/**
-	 * Converts a relative path name to a fully qualified name relative to the process BaseDir()
+	 * Converts a relative path name to a fully qualified name relative to the process BaseDir().
 	 */
 	static FString ConvertRelativePathToFull(const FString& InPath);
 
 	/**
-	 * Converts a relative path name to a fully qualified name relative to the specified BasePath
+	 * Converts a relative path name to a fully qualified name relative to the specified BasePath.
 	 */
 	static FString ConvertRelativePathToFull(const FString& BasePath, const FString& InPath);
 
 	/**
-	 * Converts a normal path to a sandbox path (in Saved/Sandboxes)
+	 * Converts a normal path to a sandbox path (in Saved/Sandboxes).
 	 *
-	 * @param InSandboxName - name of the sandbox
+	 * @param InSandboxName The name of the sandbox.
 	 */
 	static FString ConvertToSandboxPath( const FString& InPath, const TCHAR* InSandboxName );
 
 	/**
-	 * Converts a sandbox (in Saved/Sandboxes) path to a normal path
+	 * Converts a sandbox (in Saved/Sandboxes) path to a normal path.
 	 *
-	 * @param InSandboxName - name of the sandbox
+	 * @param InSandboxName The name of the sandbox.
 	 */
 	static FString ConvertFromSandboxPath( const FString& InPath, const TCHAR* InSandboxName );
 
 	/** 
-	 * Creates a temporary filename with the specified prefix 
+	 * Creates a temporary filename with the specified prefix.
 	 *
-	 * @param Path - file pathname
-	 * @param Prefix - file prefix
-	 * @param Extension - file extension ('.' required!)
+	 * @param Path The file pathname.
+	 * @param Prefix The file prefix.
+	 * @param Extension File extension ('.' required).
 	 */
 	static FString CreateTempFilename( const TCHAR* Path, const TCHAR* Prefix = TEXT(""), const TCHAR* Extension = TEXT(".tmp") );
 
@@ -471,5 +473,4 @@ private:
 		static FCriticalSection Lock;
 		return &Lock; 
 	}
-
 };
diff --git a/Engine/Source/Runtime/Core/Public/Misc/RemoteConfigIni.h b/Engine/Source/Runtime/Core/Public/Misc/RemoteConfigIni.h
--- a/Engine/Source/Runtime/Core/Public/Misc/RemoteConfigIni.h
+++ b/Engine/Source/Runtime/Core/Public/Misc/RemoteConfigIni.h
@@ -2,6 +2,7 @@
 
 #pragma once
 
+
 /**
  * Stores info relating to remote config files
  */
@@ -43,38 +44,40 @@ struct FRemoteConfigAsyncIOInfo
 class FRemoteConfigAsyncWorker
 {
 public:
-	/** Constructor */
+
+	/** Constructor. */
 	FRemoteConfigAsyncWorker(const TCHAR* InFilename, FRemoteConfigAsyncIOInfo& InIOInfo, FString* InContents, bool bInIsRead);
 
-	/** Performs the actual IO operations */
+	/** Performs the actual IO operations. */
 	void DoWork();
 
-	/** Returns true if the read IO operation succeeded */
+	/** Returns true if the read IO operation succeeded. */
 	bool IsReadSuccess() const;
 
 	/** Returns the local IO info object */
 	FRemoteConfigAsyncIOInfo& GetIOInfo();
 
-	/** Give the name for external event viewers */
+	/** Give the name for external event viewers. */
 	static const TCHAR* Name();
 
-	/** Indicates to the thread pool that this task is abandonable */
+	/** Indicates to the thread pool that this task is abandonable. */
 	bool CanAbandon();
 
 	/** Abandon routine */
 	void Abandon();
 
 private:
-	/** Name of the remote config file */
+
+	/** Name of the remote config file. */
 	TCHAR Filename[1024];
 
-	/** Local copy of the IO info object, for thread safety */
+	/** Local copy of the IO info object, for thread safety. */
 	FRemoteConfigAsyncIOInfo IOInfo;
 
-	/** Contents to write out */
+	/** Contents to write out. */
 	FString Contents;
 
-	/** true if read operation, false if write */
+	/** true if read operation, false if write. */
 	bool bIsRead;
 };
 
@@ -102,114 +105,114 @@ public:
 class FRemoteConfigAsyncTaskManager
 {
 public:
-	/** Returns a reference to the global FRemoteConfigAsyncTaskManager object */
+
+	/** Returns a reference to the global FRemoteConfigAsyncTaskManager object. */
 	CORE_API static FRemoteConfigAsyncTaskManager* Get();
 
-	/** Handles cached write tasks */
+	/** Handles cached write tasks. */
 	CORE_API void Tick();
 
-	/** Add an async IO task to the queue and kick it off */
+	/** Add an async IO task to the queue and kick it off. */
 	bool StartTask(const TCHAR* InFilename, const TCHAR* RemotePath, FRemoteConfigAsyncIOInfo& InIOInfo, FString* InContents, bool bInIsRead);
 
 	/** Returns true if the task has completed */
 	bool IsFinished(const TCHAR* InFilename);
 
-	/** Returns true if the all tasks in the queue have completed (or, if the queue is empty) */
+	/** Returns true if the all tasks in the queue have completed (or, if the queue is empty). */
 	bool AreAllTasksFinished(bool bDoRemoval);
 
-	/** Safely retrieve the read data from the completed async task */
+	/** Safely retrieve the read data from the completed async task. */
 	bool GetReadData(const TCHAR* InFilename, FRemoteConfigAsyncIOInfo& OutIOInfo);
 
 private:
-	/** Returns true if file exists in the cached task list */
+
+	/** Returns true if file exists in the cached task list. */
 	bool FindCachedWriteTask(const TCHAR* InFilename, bool bCompareContents, const TCHAR* InContents);
 
 	/** List of pending async IO operations */
 	TMap<FString, FAsyncTask<FRemoteConfigAsyncWorker>* > PendingTasks;
 
-	/** List of cached off write tasks waiting for currently pending tasks to complete */
+	/** List of cached off write tasks waiting for currently pending tasks to complete. */
 	TArray<FRemoteConfigAsyncCachedWriteTask> CachedWriteTasks;
 
-	/** Object used for synchronization via a scoped lock **/
+	/** Object used for synchronization via a scoped lock. **/
 	FCriticalSection SynchronizationObject;
 };
 
 
 /**
- * Manages remote config files
+ * Manages remote config files.
  */
 class FRemoteConfig
 {
 public:
+
 	/** Constructor */
 	FRemoteConfig();
 
-	/** Returns a reference to the global FRemoteConfig object */
+	/** Returns a reference to the global FRemoteConfig object. */
 	static FRemoteConfig* Get();
 
-	/** Returns true if the specified config file has been flagged as being remote **/
+	/** Returns true if the specified config file has been flagged as being remote. */
 	bool IsRemoteFile(const TCHAR* Filename);
 
-	/** Returns true if the specified file is remote and still needs to be read **/
+	/** Returns true if the specified file is remote and still needs to be read. */
 	bool ShouldReadRemoteFile(const TCHAR* Filename);
 
-	/** Simple accessor function **/
+	/** Simple accessor function. */
 	FRemoteConfigAsyncIOInfo* FindConfig(const TCHAR* Filename);
 
-	/** Returns true if the task has completed */
+	/** Returns true if the task has completed. */
 	bool IsFinished(const TCHAR* InFilename);
 
-	/** Queues up a new async task for reading a remote config file **/
+	/** Queues up a new async task for reading a remote config file. */
 	bool Read(const TCHAR* GeneratedIniFile, const TCHAR* DefaultIniFile);
 
-	/** Queues up a new async task for writing a remote config file **/
+	/** Queues up a new async task for writing a remote config file. */
 	bool Write(const TCHAR* Filename, FString& Contents);
 
-	/** Waits on the async read if it hasn't finished yet... times out if the operation has taken too long **/
+	/** Waits on the async read if it hasn't finished yet... times out if the operation has taken too long. */
 	void FinishRead(const TCHAR* Filename);
 
-	/** Finishes all pending async IO tasks **/
+	/** Finishes all pending async IO tasks. */
 	CORE_API static void Flush();
 
-	/* Replaces chars used by the ini parser with "special chars" */
+	/* Replaces chars used by the ini parser with "special chars". */
 	CORE_API static FString ReplaceIniCharWithSpecialChar(const FString& Str);
 
-	/* Replaces "special chars" that have been inserted to avoid problems with the ini parser with equivalent regular chars */
+	/* Replaces "special chars" that have been inserted to avoid problems with the ini parser with equivalent regular chars. */
 	CORE_API static FString ReplaceIniSpecialCharWithChar(const FString& Str);
 	
 private:
-	/** Creates the remote path string for the specified file */
+
+	/** Creates the remote path string for the specified file. */
 	FString GenerateRemotePath(const TCHAR* Filename);
 
-	/** List of remote config file info */
+	/** List of remote config file info. */
 	TMap<FString, FRemoteConfigAsyncIOInfo> ConfigBuffers;
 
-	/** Async IO operation timeout */
+	/** Async IO operation timeout. */
 	float Timeout;
 
-	/** If true, remote ops are allowed */
+	/** If true, remote ops are allowed. */
 	bool bIsEnabled;
 
-	/** true if the list of desired remote files has been filled out */
+	/** true if the list of desired remote files has been filled out. */
 	bool bHasCachedFilenames;
 
-	/** List of desired remote files */
+	/** List of desired remote files. */
 	TArray<FString> CachedFileNames;
 };
 
 
-/** Returns true if no remote version of this config file exists and/or isn't being used */
+/** Returns true if no remote version of this config file exists and/or isn't being used. */
 bool IsUsingLocalIniFile(const TCHAR* FilenameToLoad, const TCHAR* IniFileName);
 
-/** Contains the logic for processing config files, local or remote */
+/** Contains the logic for processing config files, local or remote. */
 void ProcessIniContents(const TCHAR* FilenameToLoad, const TCHAR* IniFileName, FConfigFile* Config, bool bDoEmptyConfig, bool bDoCombine, bool bDoWrite);
 
-/** Returns the timestamp of the appropriate config file */
+/** Returns the timestamp of the appropriate config file. */
 FDateTime GetIniTimeStamp(const TCHAR* FilenameToLoad, const TCHAR* IniFileName);
 
-/** Before overwriting the local file with the contents from the remote file, we save off a copy of the local file (if it exists) */
+/** Before overwriting the local file with the contents from the remote file, we save off a copy of the local file (if it exists). */
 void MakeLocalCopy(const TCHAR* Filename);
-
-
-
-
diff --git a/Engine/Source/Runtime/Core/Public/Misc/RocketSupport.h b/Engine/Source/Runtime/Core/Public/Misc/RocketSupport.h
--- a/Engine/Source/Runtime/Core/Public/Misc/RocketSupport.h
+++ b/Engine/Source/Runtime/Core/Public/Misc/RocketSupport.h
@@ -1,23 +1,20 @@
 // Copyright 1998-2014 Epic Games, Inc. All Rights Reserved.
 
-/*=============================================================================
-	RocketSupport.h: Declares the FRocketSupport class.
-=============================================================================*/
-
 #pragma once
 
+
 /**
  * Core support for launching the engine in "Rocket" mode.
  */
 class CORE_API FRocketSupport
 {
 public:
+
 	/**
 	 * Checks if we are running in 'Rocket' mode (Rocket game or Rocket editor).
 	 *
-	 * @param CmdLine	The optional command line to parse for the rocket switch
-	 * 
-	 * @return	True if we are in rocket mode, otherwise false
+	 * @param CmdLine The optional command line to parse for the rocket switch.
+	 * @return True if we are in rocket mode, otherwise false.
 	 */
 	static bool IsRocket( const TCHAR* CmdLine = FCommandLine::Get() );
 };
diff --git a/Engine/Source/Runtime/Core/Public/Misc/SecureHash.h b/Engine/Source/Runtime/Core/Public/Misc/SecureHash.h
--- a/Engine/Source/Runtime/Core/Public/Misc/SecureHash.h
+++ b/Engine/Source/Runtime/Core/Public/Misc/SecureHash.h
@@ -1,11 +1,8 @@
 // Copyright 1998-2014 Epic Games, Inc. All Rights Reserved.
 
-/*=============================================================================
-	SecureHash.h: SHA and MD5 has functions
-=============================================================================*/
-
 #pragma once
 
+
 /*-----------------------------------------------------------------------------
 	MD5 functions.
 -----------------------------------------------------------------------------*/
diff --git a/Engine/Source/Runtime/Core/Public/Misc/UProjectInfo.h b/Engine/Source/Runtime/Core/Public/Misc/UProjectInfo.h
--- a/Engine/Source/Runtime/Core/Public/Misc/UProjectInfo.h
+++ b/Engine/Source/Runtime/Core/Public/Misc/UProjectInfo.h
@@ -2,26 +2,31 @@
 
 #pragma once
 
-/** Dictionary of all the non-foreign projects for an engine installation, found by parsing .uprojectdirs files for source directories. */
+
+/**
+ * Dictionary of all the non-foreign projects for an engine installation, found by parsing .uprojectdirs files for source directories.
+ */
 class CORE_API FUProjectDictionary
 {
 public:
-	/** Scans the engine root directory for all the known projects */
+
+	/** Scans the engine root directory for all the known projects. */
 	FUProjectDictionary(const FString& InRootDir);
 	
-	/** Determines whether a project is a foreign project or not */
+	/** Determines whether a project is a foreign project or not. */
 	bool IsForeignProject(const FString& ProjectFileName) const;
 
 	/** Gets the project filename for the given game. Empty if not found. */
 	FString GetRelativeProjectPathForGame(const TCHAR* GameName, const FString& BaseDir) const;
 
-	/** Gets a list of all the known projects */
+	/** Gets a list of all the known projects. */
 	TArray<FString> GetProjectPaths() const;
 
-	/** Gets the project dictionary for the active engine installation */
+	/** Gets the project dictionary for the active engine installation. */
 	static FUProjectDictionary& GetDefault();
 
 private:
-	/** Map of short game names to full project paths */
+
+	/** Map of short game names to full project paths. */
 	TMap<FString, FString> ShortProjectNameDictionary;
 };
diff --git a/Engine/Source/Runtime/Core/Public/Misc/WildcardString.h b/Engine/Source/Runtime/Core/Public/Misc/WildcardString.h
--- a/Engine/Source/Runtime/Core/Public/Misc/WildcardString.h
+++ b/Engine/Source/Runtime/Core/Public/Misc/WildcardString.h
@@ -1,9 +1,5 @@
 // Copyright 1998-2014 Epic Games, Inc. All Rights Reserved.
 
-/*=============================================================================
-	WildcardString.h: Declares the FWildcardString class.
-=============================================================================*/
-
 #pragma once
 
 
@@ -15,9 +11,7 @@ class FWildcardString
 {
 public:
 
-	/**
-	 * Default constructor.
-	 */
+	/** Default constructor. */
 	FWildcardString( )
 		: FString()
 	{ }
