Core: Replaced EngineVersion namespaced enums with enum classes

diff --git a/Engine/Source/Runtime/Core/Private/Misc/EngineVersion.cpp b/Engine/Source/Runtime/Core/Private/Misc/EngineVersion.cpp
--- a/Engine/Source/Runtime/Core/Private/Misc/EngineVersion.cpp
+++ b/Engine/Source/Runtime/Core/Private/Misc/EngineVersion.cpp
@@ -60,8 +60,8 @@ bool FEngineVersion::IsCompatibleWith(const FEngineVersion &Other) const
 	}
 
 	// Otherwise compare the versions
-	EVersionComponent::Type Component;
-	EVersionComparison::Type Comparison = FEngineVersion::GetNewest(*this, Other, &Component);
+	EVersionComponent Component;
+	EVersionComparison Comparison = FEngineVersion::GetNewest(*this, Other, &Component);
 
 	// If this engine version is the same or newer, it's definitely compatible
 	if(Comparison == EVersionComparison::Neither || Comparison == EVersionComparison::First)
@@ -80,7 +80,7 @@ bool FEngineVersion::IsCompatibleWith(const FEngineVersion &Other) const
 	}
 }
 
-FString FEngineVersion::ToString(EVersionComponent::Type LastComponent) const
+FString FEngineVersion::ToString(EVersionComponent LastComponent) const
 {
 	FString Result = FString::Printf(TEXT("%u"), Major);
 	if(LastComponent >= EVersionComponent::Minor)
@@ -139,9 +139,9 @@ bool FEngineVersion::Parse(const FString &Text, FEngineVersion &OutVersion)
 	return true;
 }
 
-EVersionComparison::Type FEngineVersion::GetNewest(const FEngineVersion &First, const FEngineVersion &Second, EVersionComponent::Type *OutComponent)
+EVersionComparison FEngineVersion::GetNewest(const FEngineVersion &First, const FEngineVersion &Second, EVersionComponent *OutComponent)
 {
-	EVersionComponent::Type LocalComponent = EVersionComponent::Minor;
+	EVersionComponent LocalComponent = EVersionComponent::Minor;
 	auto& Component = OutComponent ? *OutComponent : LocalComponent;
 
 	// Compare major versions
diff --git a/Engine/Source/Runtime/Core/Public/Misc/EngineVersion.h b/Engine/Source/Runtime/Core/Public/Misc/EngineVersion.h
--- a/Engine/Source/Runtime/Core/Public/Misc/EngineVersion.h
+++ b/Engine/Source/Runtime/Core/Public/Misc/EngineVersion.h
@@ -4,29 +4,23 @@
 
 
 /** Enum for the components of a version string. */
-namespace EVersionComponent
+enum class EVersionComponent
 {
-	enum Type
-	{
-		Major,					///< Major version increments introduce breaking API changes.
-		Minor,					///< Minor version increments add additional functionality without breaking existing APIs.
-		Patch,					///< Patch version increments fix existing functionality without changing the API.
-		Changelist,				///< The pre-release field adds additional versioning through a series of comparable dotted strings or numbers.
-		Branch,					///<
-	};
-}
+	Major,					///< Major version increments introduce breaking API changes.
+	Minor,					///< Minor version increments add additional functionality without breaking existing APIs.
+	Patch,					///< Patch version increments fix existing functionality without changing the API.
+	Changelist,				///< The pre-release field adds additional versioning through a series of comparable dotted strings or numbers.
+	Branch,					///<
+};
 
 
 /** Components of a version string. */
-namespace EVersionComparison
+enum class EVersionComparison
 {
-	enum Type
-	{
-		Neither,
-		First,
-		Second,
-	};
-}
+	Neither,
+	First,
+	Second,
+};
 
 
 /** Utility functions. */
@@ -62,13 +56,13 @@ public:
 	bool IsCompatibleWith(const FEngineVersion &Other) const;
 
 	/** Generates a version string */
-	FString ToString(EVersionComponent::Type LastComponent = EVersionComponent::Branch) const;
+	FString ToString(EVersionComponent LastComponent = EVersionComponent::Branch) const;
 
 	/** Parses a version object from a string. Returns true on success. */
 	static bool Parse(const FString &Text, FEngineVersion &OutVersion);
 	
 	/** Returns the newest of two versions, and the component at which they differ */
-	static EVersionComparison::Type GetNewest(const FEngineVersion &First, const FEngineVersion &Second, EVersionComponent::Type *OutComponent);
+	static EVersionComparison GetNewest(const FEngineVersion &First, const FEngineVersion &Second, EVersionComponent *OutComponent);
 
 	/** Serialization function */
 	friend CORE_API void operator<<(FArchive &Ar, FEngineVersion &Version);
