Make Core module public headers individually compilable. Part 4 of n.

diff --git a/Engine/Source/Runtime/Core/Public/Android/AndroidIncludes.h b/Engine/Source/Runtime/Core/Public/Android/AndroidIncludes.h
--- a/Engine/Source/Runtime/Core/Public/Android/AndroidIncludes.h
+++ b/Engine/Source/Runtime/Core/Public/Android/AndroidIncludes.h
@@ -5,46 +5,7 @@
 ==================================================================================*/
 
 #pragma once
-
-#include <stdio.h>
-#include <stdlib.h>
-#include <stdint.h>
-#include <ctype.h>
-#include <wctype.h>
-#include <pthread.h>
-#include <limits.h>
-#include <fcntl.h>
-#include <sys/stat.h>
-#include <utime.h>
-
-// Set up compiler pragmas, etc
-#include "Android/AndroidCompilerSetup.h"
-
-//@todo android: verify malloc alignment or change?
-#define _aligned_malloc(Size,Align) malloc(Size)
-#define _aligned_realloc(Ptr,Size,Align) realloc(Ptr,Size)
-#define _aligned_free(Ptr) free(Ptr)
-
-int vswprintf( TCHAR *dst, int count, const TCHAR *fmt, va_list arg );
-
-/*
-// Define TEXT early	//@todo android:
-#define TEXT(x) L##x
-
-// map the Windows functions (that UE4 unfortunately uses be default) to normal functions
-#define _alloca alloca
-*/
-
-struct RECT
-{
-	int32 left;
-	int32 top;
-	int32 right;
-	int32 bottom;
-};
-
-#define OUT
-#define IN
+#include "Android/AndroidSystemIncludes.h"
 
 // include platform implementations
 #include "Android/AndroidMemory.h"
diff --git a/Engine/Source/Runtime/Core/Public/Android/AndroidMisc.h b/Engine/Source/Runtime/Core/Public/Android/AndroidMisc.h
--- a/Engine/Source/Runtime/Core/Public/Android/AndroidMisc.h
+++ b/Engine/Source/Runtime/Core/Public/Android/AndroidMisc.h
@@ -6,7 +6,8 @@
 ==============================================================================================*/
 
 #pragma once
-
+#include "Android/AndroidSystemIncludes.h"
+#include "GenericPlatform/GenericPlatformMisc.h"
 //@todo android: this entire file
 
 /**
diff --git a/Engine/Source/Runtime/Core/Public/Android/AndroidSystemIncludes.h b/Engine/Source/Runtime/Core/Public/Android/AndroidSystemIncludes.h
new file mode 100644
--- /dev/null
+++ b/Engine/Source/Runtime/Core/Public/Android/AndroidSystemIncludes.h
@@ -0,0 +1,42 @@
+// Copyright 1998-2014 Epic Games, Inc. All Rights Reserved.
+#pragma once
+
+#include <stdio.h>
+#include <stdlib.h>
+#include <stdint.h>
+#include <ctype.h>
+#include <wctype.h>
+#include <pthread.h>
+#include <limits.h>
+#include <fcntl.h>
+#include <sys/stat.h>
+#include <utime.h>
+
+// Set up compiler pragmas, etc
+#include "Android/AndroidCompilerSetup.h"
+
+//@todo android: verify malloc alignment or change?
+#define _aligned_malloc(Size,Align) malloc(Size)
+#define _aligned_realloc(Ptr,Size,Align) realloc(Ptr,Size)
+#define _aligned_free(Ptr) free(Ptr)
+
+int vswprintf(TCHAR *dst, int count, const TCHAR *fmt, va_list arg);
+
+/*
+// Define TEXT early	//@todo android:
+#define TEXT(x) L##x
+
+// map the Windows functions (that UE4 unfortunately uses be default) to normal functions
+#define _alloca alloca
+*/
+
+struct RECT
+{
+	int32 left;
+	int32 top;
+	int32 right;
+	int32 bottom;
+};
+
+#define OUT
+#define IN
\ No newline at end of file
diff --git a/Engine/Source/Runtime/Core/Public/CoreGlobals.h b/Engine/Source/Runtime/Core/Public/CoreGlobals.h
new file mode 100644
--- /dev/null
+++ b/Engine/Source/Runtime/Core/Public/CoreGlobals.h
@@ -0,0 +1,278 @@
+// Copyright 1998-2014 Epic Games, Inc. All Rights Reserved.
+#pragma once
+#include "Containers/ContainersFwd.h"
+#include "HAL/Platform.h"
+#include "Internationalization/Text.h"
+#include "Misc/CoreMiscDefines.h"
+#include "Misc/OutputDevice.h"
+
+class FConfigCacheIni;
+class FExec;
+class FName;
+class FOutputDeviceConsole;
+class FString;
+class ITransaction;
+
+struct FScriptTraceStackNode;
+
+extern CORE_API FConfigCacheIni* GConfig;
+extern CORE_API ITransaction* GUndo;
+extern CORE_API FOutputDeviceConsole* GLogConsole;
+
+extern CORE_API TCHAR GErrorHist[16384];
+
+// @TODO yrx 2014-08-19 Combine into one.
+extern CORE_API TCHAR GErrorExceptionDescription[4096];
+extern CORE_API TCHAR GErrorMessage[4096];
+
+extern CORE_API const FText GTrue, GFalse, GYes, GNo, GNone;
+
+/** If true, this executable is able to run all games (which are loaded as DLL's). */
+extern CORE_API bool GIsGameAgnosticExe;
+
+/** When saving out of the game, this override allows the game to load editor only properties. */
+extern CORE_API bool GForceLoadEditorOnly;
+
+/** Disable loading of objects not contained within script files; used during script compilation */
+extern CORE_API bool GVerifyObjectReferencesOnly;
+
+/** when constructing objects, use the fast path on consoles... */
+extern CORE_API bool GFastPathUniqueNameGeneration;
+
+/** allow AActor object to execute script in the editor from specific entry points, such as when running a construction script */
+extern CORE_API bool GAllowActorScriptExecutionInEditor;
+
+/** Forces use of template names for newly instanced components in a CDO. */
+extern CORE_API bool GCompilingBlueprint;
+
+/** True if we're reconstructing blueprint instances. Should never be true on cooked builds */
+extern CORE_API bool GIsReconstructingBlueprintInstances;
+
+/** Force blueprints to not compile on load */
+extern CORE_API bool GForceDisableBlueprintCompileOnLoad;
+
+
+/** Helper function to flush resource streaming. */
+extern CORE_API void(*GFlushStreamingFunc)(void);
+
+extern CORE_API float GVolumeMultiplier;
+
+
+#if WITH_EDITORONLY_DATA
+
+/**
+*	True if we are in the editor.
+*	Note that this is still true when using Play In Editor. You may want to use GWorld->HasBegunPlay in that case.
+*/
+extern CORE_API bool GIsEditor;
+extern CORE_API bool GIsImportingT3D;
+extern CORE_API bool PRIVATE_GIsRunningCommandlet;
+extern CORE_API bool GIsUCCMakeStandaloneHeaderGenerator;
+extern CORE_API bool GIsTransacting;
+
+/** Indicates that the game thread is currently paused deep in a call stack,
+while some subset of the editor user interface is pumped.  No game
+thread work can be done until the UI pumping loop returns naturally. */
+extern CORE_API bool			GIntraFrameDebuggingGameThread;
+/** True if this is the first time through the UI message pumping loop. */
+extern CORE_API bool			GFirstFrameIntraFrameDebugging;
+
+#else
+
+#define GIsEditor								false
+#define PRIVATE_GIsRunningCommandlet			false
+#define GIsUCCMakeStandaloneHeaderGenerator		false
+#define GIntraFrameDebuggingGameThread			false
+#define GFirstFrameIntraFrameDebugging			false
+
+#endif // WITH_EDITORONLY_DATA
+
+
+/**
+* Check to see if this executable is running a commandlet (custom command-line processing code in an editor-like environment)
+*/
+FORCEINLINE bool IsRunningCommandlet()
+{
+#if WITH_EDITORONLY_DATA
+	return PRIVATE_GIsRunningCommandlet;
+#else
+	return false;
+#endif
+}
+
+extern CORE_API bool GEdSelectionLock;
+extern CORE_API bool GIsClient;
+extern CORE_API bool GIsServer;
+extern CORE_API bool GIsCriticalError;
+extern CORE_API bool GIsRunning;
+extern CORE_API bool GIsGarbageCollecting;
+extern CORE_API bool GIsDuplicatingClassForReinstancing;
+
+/**
+* These are set when the engine first starts up.
+*/
+
+/**
+* This specifies whether the engine was launched as a build machine process.
+*/
+extern CORE_API bool GIsBuildMachine;
+
+/**
+* This determines if we should output any log text.  If Yes then no log text should be emitted.
+*/
+extern CORE_API bool GIsSilent;
+
+extern CORE_API bool GIsSlowTask;
+extern CORE_API bool GSlowTaskOccurred;
+extern CORE_API bool GIsGuarded;
+extern CORE_API bool GIsRequestingExit;
+
+/** Archive for serializing arbitrary data to and from memory						*/
+extern CORE_API class FReloadObjectArc* GMemoryArchive;
+
+/**
+*	Global value indicating on-screen warnings/message should be displayed.
+*	Disabled via console command "DISABLEALLSCREENMESSAGES"
+*	Enabled via console command "ENABLEALLSCREENMESSAGES"
+*	Toggled via console command "TOGGLEALLSCREENMESSAGES"
+*/
+extern CORE_API bool GAreScreenMessagesEnabled;
+extern CORE_API bool GScreenMessagesRestoreState;
+
+/* Whether we are dumping screen shots */
+extern CORE_API int32 GIsDumpingMovie;
+extern CORE_API bool GIsHighResScreenshot;
+extern CORE_API uint32 GScreenshotResolutionX;
+extern CORE_API uint32 GScreenshotResolutionY;
+extern CORE_API uint64 GMakeCacheIDIndex;
+extern CORE_API FString GEditorKeyBindingsIni;
+extern CORE_API FString GEngineIni;
+extern CORE_API FString GEditorIni;
+extern CORE_API FString GEditorLayoutIni;
+extern CORE_API FString GEditorUserSettingsIni;
+extern CORE_API FString GEditorGameAgnosticIni;
+extern CORE_API FString GCompatIni;
+extern CORE_API FString GLightmassIni;
+extern CORE_API FString GScalabilityIni;
+extern CORE_API FString GInputIni;
+extern CORE_API FString GGameIni;
+extern CORE_API FString GGameUserSettingsIni;
+
+extern CORE_API float GNearClippingPlane;
+
+extern CORE_API bool GExitPurge;
+extern CORE_API TCHAR GGameName[64];
+extern CORE_API const TCHAR* GForeignEngineDir;
+
+/** Exec handler for game debugging tool, allowing commands like "editactor" */
+extern CORE_API FExec* GDebugToolExec;
+
+/** Whether we're currently in the async loading code path or not */
+extern CORE_API bool GIsAsyncLoading;
+
+/** Whether the editor is currently loading a package or not */
+extern CORE_API bool GIsEditorLoadingPackage;
+
+/** Whether GWorld points to the play in editor world */
+extern CORE_API bool GIsPlayInEditorWorld;
+
+extern CORE_API int32 GPlayInEditorID;
+
+/** Whether or not PIE was attempting to play from PlayerStart */
+extern CORE_API bool GIsPIEUsingPlayerStart;
+
+/** Proxy class that allows verification on FApp::IsGame() accesses. */
+extern CORE_API bool IsInGameThread();
+
+/** true if the runtime needs textures to be powers of two */
+extern CORE_API bool GPlatformNeedsPowerOfTwoTextures;
+
+/** Time at which FPlatformTime::Seconds() was first initialized (very early on) */
+extern CORE_API double GStartTime;
+
+/** System time at engine init. */
+extern CORE_API FString GSystemStartTime;
+
+/** Whether we are still in the initial loading process. */
+extern CORE_API bool GIsInitialLoad;
+
+/** true when we are routing ConditionalPostLoad/PostLoad to objects */
+extern CORE_API bool GIsRoutingPostLoad;
+
+/** Steadily increasing frame counter. */
+extern CORE_API uint64 GFrameCounter;
+
+/** Incremented once per frame before the scene is being rendered. In split screen mode this is incremented once for all views (not for each view). */
+extern CORE_API uint32 GFrameNumber;
+
+/** Render Thread copy of the frame number. */
+extern CORE_API uint32 GFrameNumberRenderThread;
+
+#if !(UE_BUILD_SHIPPING && WITH_EDITOR)
+// We cannot count on this variable to be accurate in a shipped game, so make sure no code tries to use it
+/** Whether we are the first instance of the game running. */
+extern CORE_API bool GIsFirstInstance;
+#endif
+
+/** Threshold for a frame to be considered a hitch (in seconds. */
+extern CORE_API float GHitchThreshold;
+
+/** Size to break up data into when saving compressed data */
+extern CORE_API int32 GSavingCompressionChunkSize;
+
+/** Whether we are using the seek-free/ cooked loading code path. */
+extern CORE_API bool GUseSeekFreeLoading;
+
+/** Thread ID of the main/game thread */
+extern CORE_API uint32 GGameThreadId;
+
+/** Thread ID of the render thread, if any */
+extern CORE_API uint32 GRenderThreadId;
+
+/** Thread ID of the slate thread, if any */
+extern CORE_API uint32 GSlateLoadingThreadId;
+
+/** Has GGameThreadId been set yet? */
+extern CORE_API bool GIsGameThreadIdInitialized;
+
+/** Whether to emit begin/ end draw events. */
+extern CORE_API bool GEmitDrawEvents;
+
+/** Whether we want the rendering thread to be suspended, used e.g. for tracing. */
+extern CORE_API bool GShouldSuspendRenderingThread;
+
+/** Determines what kind of trace should occur, NAME_None for none. */
+extern CORE_API FName GCurrentTraceName;
+
+/** How to print the time in log output. */
+extern CORE_API ELogTimes::Type GPrintLogTimes;
+
+/** Global screen shot index to avoid overwriting ScreenShots. */
+extern CORE_API int32 GScreenshotBitmapIndex;
+
+/** Whether stats should emit named events for e.g. PIX. */
+extern CORE_API int32 GCycleStatsShouldEmitNamedEvents;
+
+/** Disables some warnings and minor features that would interrupt a demo presentation*/
+extern CORE_API bool GIsDemoMode;
+
+/** Name of the core package. */
+//@Package name transition, remove the double checks 
+extern CORE_API FName GLongCorePackageName;
+//@Package name transition, remove the double checks 
+extern CORE_API FName GLongCoreUObjectPackageName;
+
+/** Whether or not a unit test is currently being run. */
+extern CORE_API bool GIsAutomationTesting;
+
+/** Constrain bandwidth if wanted. Value is in MByte/ sec. */
+extern CORE_API float GAsyncIOBandwidthLimit;
+
+/** Whether or not messages are being pumped outside of main loop */
+extern CORE_API bool GPumpingMessagesOutsideOfMainLoop;
+
+/** Total blueprint compile time. */
+extern CORE_API double GBlueprintCompileTime;
+
+/** Stack names from the VM to be unrolled when we assert */
+extern CORE_API TArray<FScriptTraceStackNode> GScriptStack;
\ No newline at end of file
diff --git a/Engine/Source/Runtime/Core/Public/HAL/PlatformMisc.h b/Engine/Source/Runtime/Core/Public/HAL/PlatformMisc.h
new file mode 100644
--- /dev/null
+++ b/Engine/Source/Runtime/Core/Public/HAL/PlatformMisc.h
@@ -0,0 +1,23 @@
+// Copyright 1998-2014 Epic Games, Inc. All Rights Reserved.
+#pragma once
+#include "HAL/Platform.h"
+
+#if PLATFORM_WINDOWS
+#include "Windows/WindowsPlatformMisc.h"
+#elif PLATFORM_PS4
+#include "PS4/PS4Misc.h"
+#elif PLATFORM_XBOXONE
+#include "XboxOne/XboxOneMisc.h"
+#elif PLATFORM_MAC
+#include "Mac/MacPlatformMisc.h"
+#elif PLATFORM_IOS
+#include "IOS/IOSPlatformMisc.h"
+#elif PLATFORM_ANDROID
+#include "Android/AndroidMisc.h"
+#elif PLATFORM_WINRT
+#include "WinRT/WinRTMisc.h"
+#elif PLATFORM_HTML5
+#include "HTML5/HTML5PlatformMisc.h"
+#elif PLATFORM_LINUX
+#include "Linux/LinuxPlatformMisc.h"
+#endif
diff --git a/Engine/Source/Runtime/Core/Public/HTML5/HTML5PlatformCompilerSetup.h b/Engine/Source/Runtime/Core/Public/HTML5/HTML5PlatformCompilerSetup.h
--- a/Engine/Source/Runtime/Core/Public/HTML5/HTML5PlatformCompilerSetup.h
+++ b/Engine/Source/Runtime/Core/Public/HTML5/HTML5PlatformCompilerSetup.h
@@ -6,6 +6,7 @@
 ==============================================================================================*/
 
 #pragma once
+#include "HAL/Platform.h"
 
 #if PLATFORM_HTML5_WIN32
 #pragma warning(disable : 4481) // nonstandard extension used: override specifier 'override'
diff --git a/Engine/Source/Runtime/Core/Public/HTML5/HTML5PlatformIncludes.h b/Engine/Source/Runtime/Core/Public/HTML5/HTML5PlatformIncludes.h
--- a/Engine/Source/Runtime/Core/Public/HTML5/HTML5PlatformIncludes.h
+++ b/Engine/Source/Runtime/Core/Public/HTML5/HTML5PlatformIncludes.h
@@ -5,43 +5,7 @@
 ==================================================================================*/
 
 #pragma once
-
-#include <stdio.h>
-#include <stdlib.h>
-#include <stdint.h>
-#include <ctype.h>
-#include <wctype.h>
-#include <limits.h>
-
-// Set up compiler pragmas, etc
-#include "HTML5/HTML5PlatformCompilerSetup.h"
-
-// Define TEXT early
-#define TEXT_PASTE(x) L ## x
-#define TEXT_HELPER(a,b)	a ## b
-#define TEXT(s)				TEXT_HELPER(L, s)
-
-// map the Windows functions (that UE4 unfortunately uses be default) to normal functions
-#if PLATFORM_HTML5_WIN32
-
-#include <intrin.h>
-
-#else
-
-#define _alloca alloca
-
-#endif
-
-struct RECT
-{
-	int32 left;
-	int32 top;
-	int32 right;
-	int32 bottom;
-};
-
-#define OUT
-#define IN
+#include "HTML5/HTML5SystemIncludes.h"
 
 // include platform implementations
 #include "HTML5/HTML5PlatformMemory.h"
diff --git a/Engine/Source/Runtime/Core/Public/HTML5/HTML5PlatformMisc.h b/Engine/Source/Runtime/Core/Public/HTML5/HTML5PlatformMisc.h
--- a/Engine/Source/Runtime/Core/Public/HTML5/HTML5PlatformMisc.h
+++ b/Engine/Source/Runtime/Core/Public/HTML5/HTML5PlatformMisc.h
@@ -6,8 +6,10 @@
 ==============================================================================================*/
 
 #pragma once
-
-#include "HTML5DebugLogging.h"
+#include "GenericPlatform/GenericPlatformMisc.h"
+#include "HAL/Platform.h"
+#include "HTML5/HTML5DebugLogging.h"
+#include "HTML5/HTML5SystemIncludes.h"
 
 /**
  * HTML5 implementation of the misc OS functions
diff --git a/Engine/Source/Runtime/Core/Public/HTML5/HTML5SystemIncludes.h b/Engine/Source/Runtime/Core/Public/HTML5/HTML5SystemIncludes.h
new file mode 100644
--- /dev/null
+++ b/Engine/Source/Runtime/Core/Public/HTML5/HTML5SystemIncludes.h
@@ -0,0 +1,39 @@
+// Copyright 1998-2014 Epic Games, Inc. All Rights Reserved.
+#pragma once
+
+#include <stdio.h>
+#include <stdlib.h>
+#include <stdint.h>
+#include <ctype.h>
+#include <wctype.h>
+#include <limits.h>
+
+// Set up compiler pragmas, etc
+#include "HTML5/HTML5PlatformCompilerSetup.h"
+
+// Define TEXT early
+#define TEXT_PASTE(x) L ## x
+#define TEXT_HELPER(a,b)	a ## b
+#define TEXT(s)				TEXT_HELPER(L, s)
+
+// map the Windows functions (that UE4 unfortunately uses be default) to normal functions
+#if PLATFORM_HTML5_WIN32
+
+#include <intrin.h>
+
+#else
+
+#define _alloca alloca
+
+#endif
+
+struct RECT
+{
+	int32 left;
+	int32 top;
+	int32 right;
+	int32 bottom;
+};
+
+#define OUT
+#define IN
\ No newline at end of file
diff --git a/Engine/Source/Runtime/Core/Public/IOS/IOSPlatformIncludes.h b/Engine/Source/Runtime/Core/Public/IOS/IOSPlatformIncludes.h
--- a/Engine/Source/Runtime/Core/Public/IOS/IOSPlatformIncludes.h
+++ b/Engine/Source/Runtime/Core/Public/IOS/IOSPlatformIncludes.h
@@ -5,81 +5,7 @@
 ==================================================================================*/
 
 #pragma once
-
-// Set up compiler pragmas, etc
-#include "IOS/IOSPlatformCompilerSetup.h"
-
-#define FVector FVectorWorkaround
-#ifdef __OBJC__
-#import <UIKit/UIKit.h>
-#import <CoreData/CoreData.h>
-#else //@todo - zombie - iOS should always have Obj-C to function. This definition guard may be unnecessary. -astarkey
-#error "Objective C is undefined."
-#endif
-#undef check
-#undef verify
-#undef FVector
-
-#undef TYPE_BOOL
-
-#include <CoreFoundation/CoreFoundation.h>
-
-#include <string.h>
-#include <alloca.h>
-#include <stdlib.h>
-#include <stdio.h>
-#include <float.h>
-#include <ctype.h>
-#include <stdarg.h>
-#include <stdint.h>
-#include <limits.h>
-#include <sys/time.h>
-#include <math.h>
-#include <mach/mach_time.h>
-#include <wchar.h>
-#include <wctype.h>
-#include <errno.h>
-#include <unistd.h>
-#include <pthread.h>
-#include <libkern/OSAtomic.h>
-#include <sys/types.h>
-#include <sys/stat.h>
-#include <pwd.h>
-#include <fcntl.h>
-#include <dirent.h>
-#include <dlfcn.h>
-#include <copyfile.h>
-#include <utime.h>
-#include <mach/mach.h>
-#include <mach/mach_host.h>
-#include <execinfo.h>
-#include <signal.h>
-#include <sys/sysctl.h>
-
-// SIMD intrinsics
-#if WITH_SIMULATOR
-#include <xmmintrin.h>
-#else
-#include <arm_neon.h>
-#endif
-
-struct RECT
-{
-	int32 left;
-	int32 top;
-	int32 right;
-	int32 bottom;
-};
-
-#define OUT
-#define IN
-
-/*----------------------------------------------------------------------------
-	Memory. On Mac OS X malloc allocates memory aligned to 16 bytes.
-----------------------------------------------------------------------------*/
-#define _aligned_malloc(Size,Align) malloc(Size)
-#define _aligned_realloc(Ptr,Size,Align) realloc(Ptr,Size)
-#define _aligned_free(Ptr) free(Ptr)
+#include "IOS/IOSSystemIncludes.h"
 
 // include platform implementations
 #include "IOS/IOSPlatformMemory.h"
diff --git a/Engine/Source/Runtime/Core/Public/IOS/IOSPlatformMisc.h b/Engine/Source/Runtime/Core/Public/IOS/IOSPlatformMisc.h
--- a/Engine/Source/Runtime/Core/Public/IOS/IOSPlatformMisc.h
+++ b/Engine/Source/Runtime/Core/Public/IOS/IOSPlatformMisc.h
@@ -5,6 +5,8 @@
 ==============================================================================================*/
 
 #pragma once
+#include "GenericPlatform/GenericPlatformMisc.h"
+#include "IOS/IOSSystemIncludes.h"
 
 /**
 * iOS implementation of the misc OS functions
diff --git a/Engine/Source/Runtime/Core/Public/IOS/IOSSystemIncludes.h b/Engine/Source/Runtime/Core/Public/IOS/IOSSystemIncludes.h
new file mode 100644
--- /dev/null
+++ b/Engine/Source/Runtime/Core/Public/IOS/IOSSystemIncludes.h
@@ -0,0 +1,77 @@
+// Copyright 1998-2014 Epic Games, Inc. All Rights Reserved.
+#pragma once
+
+// Set up compiler pragmas, etc
+#include "IOS/IOSPlatformCompilerSetup.h"
+
+#define FVector FVectorWorkaround
+#ifdef __OBJC__
+#import <UIKit/UIKit.h>
+#import <CoreData/CoreData.h>
+#else //@todo - zombie - iOS should always have Obj-C to function. This definition guard may be unnecessary. -astarkey
+#error "Objective C is undefined."
+#endif
+#undef check
+#undef verify
+#undef FVector
+
+#undef TYPE_BOOL
+
+#include <CoreFoundation/CoreFoundation.h>
+
+#include <string.h>
+#include <alloca.h>
+#include <stdlib.h>
+#include <stdio.h>
+#include <float.h>
+#include <ctype.h>
+#include <stdarg.h>
+#include <stdint.h>
+#include <limits.h>
+#include <sys/time.h>
+#include <math.h>
+#include <mach/mach_time.h>
+#include <wchar.h>
+#include <wctype.h>
+#include <errno.h>
+#include <unistd.h>
+#include <pthread.h>
+#include <libkern/OSAtomic.h>
+#include <sys/types.h>
+#include <sys/stat.h>
+#include <pwd.h>
+#include <fcntl.h>
+#include <dirent.h>
+#include <dlfcn.h>
+#include <copyfile.h>
+#include <utime.h>
+#include <mach/mach.h>
+#include <mach/mach_host.h>
+#include <execinfo.h>
+#include <signal.h>
+#include <sys/sysctl.h>
+
+// SIMD intrinsics
+#if WITH_SIMULATOR
+#include <xmmintrin.h>
+#else
+#include <arm_neon.h>
+#endif
+
+struct RECT
+{
+	int32 left;
+	int32 top;
+	int32 right;
+	int32 bottom;
+};
+
+#define OUT
+#define IN
+
+/*----------------------------------------------------------------------------
+Memory. On Mac OS X malloc allocates memory aligned to 16 bytes.
+----------------------------------------------------------------------------*/
+#define _aligned_malloc(Size,Align) malloc(Size)
+#define _aligned_realloc(Ptr,Size,Align) realloc(Ptr,Size)
+#define _aligned_free(Ptr) free(Ptr)
\ No newline at end of file
diff --git a/Engine/Source/Runtime/Core/Public/Linux/LinuxPlatformIncludes.h b/Engine/Source/Runtime/Core/Public/Linux/LinuxPlatformIncludes.h
--- a/Engine/Source/Runtime/Core/Public/Linux/LinuxPlatformIncludes.h
+++ b/Engine/Source/Runtime/Core/Public/Linux/LinuxPlatformIncludes.h
@@ -5,53 +5,11 @@
 ==================================================================================*/
 
 #pragma once
-
-// Set up compiler pragmas, etc
-#include "Linux/LinuxPlatformCompilerSetup.h"
-
-#include <string.h>
-#include <alloca.h>
-#include <stdlib.h>
-#include <stdio.h>
-#include <float.h>
-#include <ctype.h>
-#include <stdarg.h>
-#include <stdint.h>
-#include <limits.h>
-#include <sys/time.h>
-#include <math.h>
-#include <wchar.h>
-#include <wctype.h>
-#include <errno.h>
-#include <unistd.h>
-#include <pthread.h>
-#include <sys/types.h>
-#include <sys/stat.h>
-#include <pwd.h>
-#include <fcntl.h>
-#include <dirent.h>
-#include <dlfcn.h>
-#include <utime.h>
-#include <execinfo.h>
-#include <signal.h>
-#include <sys/sysctl.h>
-#ifdef __x86_64__
-	#include <xmmintrin.h>
-#endif // PLATFORM_RASPBERRY
-#include <sys/utsname.h>
-#include <libgen.h>
-
-struct RECT
-{
-	int32 left;
-	int32 top;
-	int32 right;
-	int32 bottom;
-};
-
-#define OUT
-#define IN
-
+#include "GenericPlatform/GenericPlatformAffinity.h"
+#include "GenericPlatform/GenericPlatformFramePacer.h"
+#include "Linux/LinuxSystemIncludes.h"
+#include "Misc/Build.h"
+#include "Misc/CoreMiscDefines.h"
 
 // include platform implementations
 #include "Linux/LinuxPlatformMemory.h"
diff --git a/Engine/Source/Runtime/Core/Public/Linux/LinuxPlatformMisc.h b/Engine/Source/Runtime/Core/Public/Linux/LinuxPlatformMisc.h
--- a/Engine/Source/Runtime/Core/Public/Linux/LinuxPlatformMisc.h
+++ b/Engine/Source/Runtime/Core/Public/Linux/LinuxPlatformMisc.h
@@ -6,7 +6,11 @@
 ==============================================================================================*/
 
 #pragma once
+#include "GenericPlatform/GenericPlatformMisc.h"
+#include "Linux/LinuxSystemIncludes.h"
+#include "Misc/Build.h"
 
+class FString;
 struct FGenericCrashContext;
 
 /**
diff --git a/Engine/Source/Runtime/Core/Public/Linux/LinuxSystemIncludes.h b/Engine/Source/Runtime/Core/Public/Linux/LinuxSystemIncludes.h
new file mode 100644
--- /dev/null
+++ b/Engine/Source/Runtime/Core/Public/Linux/LinuxSystemIncludes.h
@@ -0,0 +1,48 @@
+// Copyright 1998-2014 Epic Games, Inc. All Rights Reserved.
+#pragma once
+
+// Set up compiler pragmas, etc
+#include "Linux/LinuxPlatformCompilerSetup.h"
+
+#include <string.h>
+#include <alloca.h>
+#include <stdlib.h>
+#include <stdio.h>
+#include <float.h>
+#include <ctype.h>
+#include <stdarg.h>
+#include <stdint.h>
+#include <limits.h>
+#include <sys/time.h>
+#include <math.h>
+#include <wchar.h>
+#include <wctype.h>
+#include <errno.h>
+#include <unistd.h>
+#include <pthread.h>
+#include <sys/types.h>
+#include <sys/stat.h>
+#include <pwd.h>
+#include <fcntl.h>
+#include <dirent.h>
+#include <dlfcn.h>
+#include <utime.h>
+#include <execinfo.h>
+#include <signal.h>
+#include <sys/sysctl.h>
+#ifdef __x86_64__
+#include <xmmintrin.h>
+#endif // PLATFORM_RASPBERRY
+#include <sys/utsname.h>
+#include <libgen.h>
+
+struct RECT
+{
+	int32 left;
+	int32 top;
+	int32 right;
+	int32 bottom;
+};
+
+#define OUT
+#define IN
diff --git a/Engine/Source/Runtime/Core/Public/Mac/MacPlatformIncludes.h b/Engine/Source/Runtime/Core/Public/Mac/MacPlatformIncludes.h
--- a/Engine/Source/Runtime/Core/Public/Mac/MacPlatformIncludes.h
+++ b/Engine/Source/Runtime/Core/Public/Mac/MacPlatformIncludes.h
@@ -6,76 +6,10 @@
 
 #pragma once
 
-// Set up compiler pragmas, etc
-#include "Mac/MacPlatformCompilerSetup.h"
-
-#define FVector FVectorWorkaround
-#ifdef __OBJC__
-#import <Cocoa/Cocoa.h>
-#endif
-#include <Carbon/Carbon.h>
-#include <CoreServices/CoreServices.h>
-#undef check
-#undef verify
-#undef FVector
-
-#include <CoreFoundation/CoreFoundation.h>
-#undef TYPE_BOOL
-
-#include <string.h>
-#include <alloca.h>
-#include <stdlib.h>
-#include <stdio.h>
-#include <float.h>
-#include <ctype.h>
-#include <stdarg.h>
-#include <stdint.h>
-#include <limits.h>
-#include <sys/time.h>
-#include <math.h>
-#include <mach/mach_time.h>
-#include <wchar.h>
-#include <wctype.h>
-#include <errno.h>
-#include <unistd.h>
-#include <pthread.h>
-#include <libkern/OSAtomic.h>
-#include <sys/types.h>
-#include <sys/stat.h>
-#include <pwd.h>
-#include <fcntl.h>
-#include <dirent.h>
-#include <dlfcn.h>
-#include <copyfile.h>
-#include <utime.h>
-#include <mach/mach_host.h>
-#include <mach/task.h>
-#include <execinfo.h>
-#include <signal.h>
-#include <sys/sysctl.h>
-#include <malloc/malloc.h>
-
-// SIMD intrinsics
-#include <xmmintrin.h>
-
-struct RECT
-{
-	int32 left;
-	int32 top;
-	int32 right;
-	int32 bottom;
-};
-
-#define OUT
-#define IN
-
-/*----------------------------------------------------------------------------
-	Memory. On Mac OS X malloc allocates memory aligned to 16 bytes.
-----------------------------------------------------------------------------*/
-#define _aligned_malloc(Size,Align) malloc(Size)
-#define _aligned_realloc(Ptr,Size,Align) realloc(Ptr,Size)
-#define _aligned_free(Ptr) free(Ptr)
-#define _aligned_msize(Ptr,Align,Offset) malloc_size(Ptr)
+#include "GenericPlatform/GenericPlatformAffinity.h"
+#include "GenericPlatform/GenericPlatformFramePacer.h"
+#include "Mac/MacSystemIncludes.h"
+#include "Misc/CoreMiscDefines.h"
 
 // include platform implementations
 #include "Mac/MacPlatformMemory.h"
diff --git a/Engine/Source/Runtime/Core/Public/Mac/MacPlatformMisc.h b/Engine/Source/Runtime/Core/Public/Mac/MacPlatformMisc.h
--- a/Engine/Source/Runtime/Core/Public/Mac/MacPlatformMisc.h
+++ b/Engine/Source/Runtime/Core/Public/Mac/MacPlatformMisc.h
@@ -5,6 +5,8 @@
 ==============================================================================================*/
 
 #pragma once
+#include "GenericPlatform/GenericPlatformMisc.h"
+#include "Mac/MacSystemIncludes.h"
 
 typedef void (*UpdateCachedMacMenuStateProc)(void);
 
diff --git a/Engine/Source/Runtime/Core/Public/Mac/MacSystemIncludes.h b/Engine/Source/Runtime/Core/Public/Mac/MacSystemIncludes.h
new file mode 100644
--- /dev/null
+++ b/Engine/Source/Runtime/Core/Public/Mac/MacSystemIncludes.h
@@ -0,0 +1,73 @@
+// Copyright 1998-2014 Epic Games, Inc. All Rights Reserved.
+#pragma once
+
+// Set up compiler pragmas, etc
+#include "Mac/MacPlatformCompilerSetup.h"
+
+#define FVector FVectorWorkaround
+#ifdef __OBJC__
+#import <Cocoa/Cocoa.h>
+#endif
+#include <Carbon/Carbon.h>
+#include <CoreServices/CoreServices.h>
+#undef check
+#undef verify
+#undef FVector
+
+#include <CoreFoundation/CoreFoundation.h>
+#undef TYPE_BOOL
+
+#include <string.h>
+#include <alloca.h>
+#include <stdlib.h>
+#include <stdio.h>
+#include <float.h>
+#include <ctype.h>
+#include <stdarg.h>
+#include <stdint.h>
+#include <limits.h>
+#include <sys/time.h>
+#include <math.h>
+#include <mach/mach_time.h>
+#include <wchar.h>
+#include <wctype.h>
+#include <errno.h>
+#include <unistd.h>
+#include <pthread.h>
+#include <libkern/OSAtomic.h>
+#include <sys/types.h>
+#include <sys/stat.h>
+#include <pwd.h>
+#include <fcntl.h>
+#include <dirent.h>
+#include <dlfcn.h>
+#include <copyfile.h>
+#include <utime.h>
+#include <mach/mach_host.h>
+#include <mach/task.h>
+#include <execinfo.h>
+#include <signal.h>
+#include <sys/sysctl.h>
+#include <malloc/malloc.h>
+
+// SIMD intrinsics
+#include <xmmintrin.h>
+
+struct RECT
+{
+	int32 left;
+	int32 top;
+	int32 right;
+	int32 bottom;
+};
+
+#define OUT
+#define IN
+
+/*----------------------------------------------------------------------------
+Memory. On Mac OS X malloc allocates memory aligned to 16 bytes.
+----------------------------------------------------------------------------*/
+#define _aligned_malloc(Size,Align) malloc(Size)
+#define _aligned_realloc(Ptr,Size,Align) realloc(Ptr,Size)
+#define _aligned_free(Ptr) free(Ptr)
+#define _aligned_msize(Ptr,Align,Offset) malloc_size(Ptr)
\ No newline at end of file
diff --git a/Engine/Source/Runtime/Core/Public/WinRT/WinRTMisc.h b/Engine/Source/Runtime/Core/Public/WinRT/WinRTMisc.h
--- a/Engine/Source/Runtime/Core/Public/WinRT/WinRTMisc.h
+++ b/Engine/Source/Runtime/Core/Public/WinRT/WinRTMisc.h
@@ -5,6 +5,12 @@
 ==============================================================================================*/
 
 #pragma once
+#include "GenericPlatform/GenericPlatformMisc.h"
+#include "HAL/Platform.h"
+#include "Misc/Build.h"
+#include "WinRT/WinRTSystemIncludes.h"
+
+class FString;
 
 /**
  * WinRT implementation of the misc OS functions
diff --git a/Engine/Source/Runtime/Core/Public/WinRT/WinRTPlatformIncludes.h b/Engine/Source/Runtime/Core/Public/WinRT/WinRTPlatformIncludes.h
--- a/Engine/Source/Runtime/Core/Public/WinRT/WinRTPlatformIncludes.h
+++ b/Engine/Source/Runtime/Core/Public/WinRT/WinRTPlatformIncludes.h
@@ -5,30 +5,7 @@
 ==================================================================================*/
 
 #pragma once
-
-// Pre-system API include
-#include "WinRT/PreWinRTApi.h"
-
-#ifndef STRICT
-	#define STRICT
-#endif
-
-// System API include
-#include "WinRT/MinWinRTApi.h"
-
-// Post-system API include
-#include "WinRT/PostWinRTApi.h"
-
-// Set up compiler pragmas, etc
-#include "WinRT/WinRTCompilerSetup.h"
-
-// Macro for releasing COM objects
-#define SAFE_RELEASE(p) { if(p) { (p)->Release(); (p)=NULL; } }
-
-#include <tchar.h>
-
-// SIMD intrinsics
-#include <intrin.h>
+#include "WinRT/WinRTSystemIncludes.h"
 
 // include platform implementations
 #include "WinRT/WinRTIncludes.h"
diff --git a/Engine/Source/Runtime/Core/Public/WinRT/WinRTSystemIncludes.h b/Engine/Source/Runtime/Core/Public/WinRT/WinRTSystemIncludes.h
new file mode 100644
--- /dev/null
+++ b/Engine/Source/Runtime/Core/Public/WinRT/WinRTSystemIncludes.h
@@ -0,0 +1,26 @@
+// Copyright 1998-2014 Epic Games, Inc. All Rights Reserved.
+#pragma once
+
+// Pre-system API include
+#include "WinRT/PreWinRTApi.h"
+
+#ifndef STRICT
+#define STRICT
+#endif
+
+// System API include
+#include "WinRT/MinWinRTApi.h"
+
+// Post-system API include
+#include "WinRT/PostWinRTApi.h"
+
+// Set up compiler pragmas, etc
+#include "WinRT/WinRTCompilerSetup.h"
+
+// Macro for releasing COM objects
+#define SAFE_RELEASE(p) { if(p) { (p)->Release(); (p)=NULL; } }
+
+#include <tchar.h>
+
+// SIMD intrinsics
+#include <intrin.h>
\ No newline at end of file
diff --git a/Engine/Source/Runtime/Core/Public/Windows/WindowsPlatformIncludes.h b/Engine/Source/Runtime/Core/Public/Windows/WindowsPlatformIncludes.h
--- a/Engine/Source/Runtime/Core/Public/Windows/WindowsPlatformIncludes.h
+++ b/Engine/Source/Runtime/Core/Public/Windows/WindowsPlatformIncludes.h
@@ -4,50 +4,11 @@
 	WindowsPlatformIncludes.h: Includes the platform specific headers for windows
 ==================================================================================*/
 #pragma once
-
-// Set up compiler pragmas, etc
-#include "Windows/WindowsPlatformCompilerSetup.h"
-
-// include windows.h
-#include "Windows/PreWindowsApi.h"
-#ifndef STRICT
-#define STRICT
-#endif
-#include "Windows/MinWindows.h"
-#include "Windows/PostWindowsApi.h"
-
-// Macro for releasing COM objects
-#define SAFE_RELEASE(p) { if(p) { (p)->Release(); (p)=NULL; } }
-
-// Current instance
-extern "C" CORE_API HINSTANCE hInstance; 
-
-// SIMD intrinsics
-#include <intrin.h>
-
-#include <stdint.h>
-#include <tchar.h>
-
-// @todo clang: When compiling under Clang for Windows, these headers cause us particular problems.  We need to make sure they're included
-// before we pull in our 'DoNotUseOldUE4Type' namespace.  This is because these headers will redeclare various numeric typedefs, but under
-// the Clang compiler it is not allowed to define a typedef with a global scope operator in it (such as ::INT).  So we'll get these headers
-// included early on to avoid compiler errors with that.
-#if defined(__clang__)
-	#include <intsafe.h>
-	#include <strsafe.h>
-#endif
-
-#if USING_CODE_ANALYSIS
-// Source annotation support
-#include <CodeAnalysis/SourceAnnotations.h>
-
-// Allows for disabling code analysis warnings by defining ALL_CODE_ANALYSIS_WARNINGS
-#include <CodeAnalysis/Warnings.h>
-
-// We import the vc_attributes namespace so we can use annotations more easily.  It only defines
-// MSVC-specific attributes so there should never be collisions.
-using namespace vc_attributes;
-#endif
+#include "GenericPlatform/GenericPlatformAffinity.h"
+#include "GenericPlatform/GenericPlatformFramePacer.h"
+#include "Misc/Build.h"
+#include "Misc/CoreMiscDefines.h"
+#include "Windows/WindowsSystemIncludes.h"
 
 // include platform implementations
 #include "Windows/WindowsPlatformMemory.h"
diff --git a/Engine/Source/Runtime/Core/Public/Windows/WindowsPlatformMisc.h b/Engine/Source/Runtime/Core/Public/Windows/WindowsPlatformMisc.h
--- a/Engine/Source/Runtime/Core/Public/Windows/WindowsPlatformMisc.h
+++ b/Engine/Source/Runtime/Core/Public/Windows/WindowsPlatformMisc.h
@@ -6,6 +6,8 @@
 ==============================================================================================*/
 
 #pragma once
+#include "GenericPlatform/GenericPlatformMisc.h"
+#include "Windows/WindowsHWrapper.h"
 
 /** Helper struct used to get the string version of the Windows version. */
 struct FWindowsOSVersionHelper
diff --git a/Engine/Source/Runtime/Core/Public/Windows/WindowsSystemIncludes.h b/Engine/Source/Runtime/Core/Public/Windows/WindowsSystemIncludes.h
new file mode 100644
--- /dev/null
+++ b/Engine/Source/Runtime/Core/Public/Windows/WindowsSystemIncludes.h
@@ -0,0 +1,46 @@
+// Copyright 1998-2014 Epic Games, Inc. All Rights Reserved.
+#pragma once
+
+// Set up compiler pragmas, etc
+#include "Windows/WindowsPlatformCompilerSetup.h"
+
+// include windows.h
+#include "Windows/PreWindowsApi.h"
+#ifndef STRICT
+#define STRICT
+#endif
+#include "Windows/MinWindows.h"
+#include "Windows/PostWindowsApi.h"
+
+// Macro for releasing COM objects
+#define SAFE_RELEASE(p) { if(p) { (p)->Release(); (p)=NULL; } }
+
+// Current instance
+extern "C" CORE_API HINSTANCE hInstance;
+
+// SIMD intrinsics
+#include <intrin.h>
+
+#include <stdint.h>
+#include <tchar.h>
+
+// @todo clang: When compiling under Clang for Windows, these headers cause us particular problems.  We need to make sure they're included
+// before we pull in our 'DoNotUseOldUE4Type' namespace.  This is because these headers will redeclare various numeric typedefs, but under
+// the Clang compiler it is not allowed to define a typedef with a global scope operator in it (such as ::INT).  So we'll get these headers
+// included early on to avoid compiler errors with that.
+#if defined(__clang__)
+#include <intsafe.h>
+#include <strsafe.h>
+#endif
+
+#if USING_CODE_ANALYSIS
+// Source annotation support
+#include <CodeAnalysis/SourceAnnotations.h>
+
+// Allows for disabling code analysis warnings by defining ALL_CODE_ANALYSIS_WARNINGS
+#include <CodeAnalysis/Warnings.h>
+
+// We import the vc_attributes namespace so we can use annotations more easily.  It only defines
+// MSVC-specific attributes so there should never be collisions.
+using namespace vc_attributes;
+#endif
\ No newline at end of file
