Make Core module public headers individually compilable. Part 5 of n.

diff --git a/Engine/Source/Runtime/Core/Public/Containers/ContainerAllocationPolicies.h b/Engine/Source/Runtime/Core/Public/Containers/ContainerAllocationPolicies.h
--- a/Engine/Source/Runtime/Core/Public/Containers/ContainerAllocationPolicies.h
+++ b/Engine/Source/Runtime/Core/Public/Containers/ContainerAllocationPolicies.h
@@ -7,6 +7,7 @@
 #pragma once
 
 #include "MemoryBase.h"
+#include "Misc/OutputDevice.h"
 
 /** Used to determine the alignment of an element type. */
 template<typename ElementType>
diff --git a/Engine/Source/Runtime/Core/Public/HAL/FileManager.h b/Engine/Source/Runtime/Core/Public/HAL/FileManager.h
--- a/Engine/Source/Runtime/Core/Public/HAL/FileManager.h
+++ b/Engine/Source/Runtime/Core/Public/HAL/FileManager.h
@@ -1,7 +1,10 @@
 // Copyright 1998-2014 Epic Games, Inc. All Rights Reserved.
 
 #pragma once
+#include "GenericPlatform/GenericPlatformFile.h"
+#include "HAL/Platform.h"
 
+struct FDateTime;
 
 // Maximum length of any filename.  For now, we have no restriction. We would probably use shortening rules if we have to.
 #define MAX_UNREAL_FILENAME_LENGTH (PLATFORM_MAX_FILEPATH_LENGTH)
diff --git a/Engine/Source/Runtime/Core/Public/Misc/OutputDevice.h b/Engine/Source/Runtime/Core/Public/Misc/OutputDevice.h
--- a/Engine/Source/Runtime/Core/Public/Misc/OutputDevice.h
+++ b/Engine/Source/Runtime/Core/Public/Misc/OutputDevice.h
@@ -2,8 +2,12 @@
 
 #pragma once
 
+#include "HAL/Platform.h"
+#include "HAL/PlatformMisc.h"
+#include "Misc/CoreMiscDefines.h"
 
 class FText;
+class FString;
 
 // Globals.
 #define GLog FOutputDeviceRedirector::Get()
@@ -294,7 +298,7 @@ public:
 
 	// static helpers
 	static const TCHAR* VerbosityToString(ELogVerbosity::Type Verbosity);
-	static FString FormatLogLine(ELogVerbosity::Type Verbosity, const class FName& Category, const TCHAR* Message = NULL, ELogTimes::Type LogTime = ELogTimes::None);
+	static FString FormatLogLine(ELogVerbosity::Type Verbosity, const class FName& Category, const TCHAR* Message = nullptr, ELogTimes::Type LogTime = ELogTimes::None);
 
 
 	// FOutputDevice interface.
