Make Core module public headers individually compilable. Part 6 of n.

diff --git a/Engine/Source/Runtime/Core/Public/HAL/PlatformProperties.h b/Engine/Source/Runtime/Core/Public/HAL/PlatformProperties.h
new file mode 100644
--- /dev/null
+++ b/Engine/Source/Runtime/Core/Public/HAL/PlatformProperties.h
@@ -0,0 +1,34 @@
+// Copyright 1998-2014 Epic Games, Inc. All Rights Reserved.
+#pragma once
+#include "HAL/Platform.h"
+#include "Misc/Build.h"
+#include "Misc/CoreMiscDefines.h"
+
+#if PLATFORM_WINDOWS
+#include "Windows/WindowsPlatformProperties.h"
+typedef FWindowsPlatformProperties<WITH_EDITORONLY_DATA, UE_SERVER, !WITH_SERVER_CODE> FPlatformProperties;
+#elif PLATFORM_PS4
+#include "PS4/PS4Properties.h"
+typedef FPS4PlatformProperties FPlatformProperties;
+#elif PLATFORM_XBOXONE
+#include "XboxOne/XboxOneProperties.h"
+typedef FXboxOnePlatformProperties FPlatformProperties;
+#elif PLATFORM_MAC
+#include "Mac/MacPlatformProperties.h"
+typedef FMacPlatformProperties<WITH_EDITORONLY_DATA, UE_SERVER, !WITH_SERVER_CODE> FPlatformProperties;
+#elif PLATFORM_IOS
+#include "IOS/IOSPlatformProperties.h"
+typedef FIOSPlatformProperties FPlatformProperties;
+#elif PLATFORM_ANDROID
+#include "Android/AndroidProperties.h"
+typedef FAndroidPlatformProperties FPlatformProperties;
+#elif PLATFORM_WINRT
+#include "WinRT/WinRTProperties.h"
+typedef FWinRTPlatformPropertiesFPlatformProperties;
+#elif PLATFORM_HTML5
+#include "HTML5/HTML5PlatformProperties.h"
+typedef FHTML5PlatformProperties FPlatformProperties;
+#elif PLATFORM_LINUX
+#include "Linux/LinuxPlatformProperties.h"
+typedef FLinuxPlatformProperties<WITH_EDITORONLY_DATA, UE_SERVER, !WITH_SERVER_CODE> FPlatformProperties;
+#endif
\ No newline at end of file
diff --git a/Engine/Source/Runtime/Core/Public/WinRT/WinRTProperties.h b/Engine/Source/Runtime/Core/Public/WinRT/WinRTProperties.h
--- a/Engine/Source/Runtime/Core/Public/WinRT/WinRTProperties.h
+++ b/Engine/Source/Runtime/Core/Public/WinRT/WinRTProperties.h
@@ -8,8 +8,8 @@
 ==================================================================================*/
 
 #pragma once
-
-#include "GenericPlatformProperties.h"
+#include "GenericPlatform/GenericPlatformProperties.h"
+#include "HAL/Platform.h"
 
 struct FWinRTPlatformProperties : public FGenericPlatformProperties
 {
diff --git a/Engine/Source/Runtime/Core/Public/Windows/WindowsPlatformProperties.h b/Engine/Source/Runtime/Core/Public/Windows/WindowsPlatformProperties.h
--- a/Engine/Source/Runtime/Core/Public/Windows/WindowsPlatformProperties.h
+++ b/Engine/Source/Runtime/Core/Public/Windows/WindowsPlatformProperties.h
@@ -8,11 +8,6 @@
 ==================================================================================*/
 
 #pragma once
-
-//#include "Core.h"
-//#include "AssertionMacros.h"
-//#include "Internationalization.h"
-
 #include "GenericPlatformProperties.h"
 
 
