Make Core module public headers individually compilable. Part 7 of n.

diff --git a/Engine/Source/Runtime/Core/Public/Containers/Array.h b/Engine/Source/Runtime/Core/Public/Containers/Array.h
--- a/Engine/Source/Runtime/Core/Public/Containers/Array.h
+++ b/Engine/Source/Runtime/Core/Public/Containers/Array.h
@@ -5,11 +5,14 @@
 =============================================================================*/
 
 #pragma once
+#include "Containers/ContainerAllocationPolicies.h"
+#include "HAL/Platform.h"
+#include "Serialization/ArchiveBase.h"
+#include "Templates/EnableIf.h"
+#include "Templates/UnrealTemplate.h"
 
 #define DEBUG_HEAP 0
 
-#include "Core.h"
-
 /**
  * Generic iterator which can operate on types that expose the following:
  * - A type called ElementType representing the contained type.
diff --git a/Engine/Source/Runtime/Core/Public/Misc/Compression.h b/Engine/Source/Runtime/Core/Public/Misc/Compression.h
--- a/Engine/Source/Runtime/Core/Public/Misc/Compression.h
+++ b/Engine/Source/Runtime/Core/Public/Misc/Compression.h
@@ -2,6 +2,7 @@
 
 #pragma once
 
+#include "HAL/Platform.h"
 
 /**
  * Flags controlling [de]compression
diff --git a/Engine/Source/Runtime/Core/Public/Serialization/ArchiveBase.h b/Engine/Source/Runtime/Core/Public/Serialization/ArchiveBase.h
--- a/Engine/Source/Runtime/Core/Public/Serialization/ArchiveBase.h
+++ b/Engine/Source/Runtime/Core/Public/Serialization/ArchiveBase.h
@@ -5,7 +5,9 @@
 =============================================================================*/
 
 #pragma once
-
+#include "Containers/EnumAsByte.h"
+#include "HAL/PlatformProperties.h"
+#include "Misc/Compression.h"
 
 class FCustomVersionContainer;
 struct FUntypedBulkData;
