Make Core module public headers individually compilable. Part 8 of n.

diff --git a/Engine/Source/Runtime/Core/Public/Android/AndroidMath.h b/Engine/Source/Runtime/Core/Public/Android/AndroidMath.h
--- a/Engine/Source/Runtime/Core/Public/Android/AndroidMath.h
+++ b/Engine/Source/Runtime/Core/Public/Android/AndroidMath.h
@@ -5,6 +5,7 @@
 ==============================================================================================*/
 
 #pragma once
+#include "GenericPlatform/GenericPlatformMath.h"
 
 /**
  * Android implementation of the Math OS functions
diff --git a/Engine/Source/Runtime/Core/Public/HAL/PlatformMath.h b/Engine/Source/Runtime/Core/Public/HAL/PlatformMath.h
new file mode 100644
--- /dev/null
+++ b/Engine/Source/Runtime/Core/Public/HAL/PlatformMath.h
@@ -0,0 +1,24 @@
+// Copyright 1998-2014 Epic Games, Inc. All Rights Reserved.
+#pragma once
+
+#include "HAL/Platform.h"
+
+#if PLATFORM_WINDOWS
+#include "Windows/WindowsPlatformMath.h"
+#elif PLATFORM_PS4
+#include "PS4/PS4Math.h"
+#elif PLATFORM_XBOXONE
+#include "XboxOne/XboxOneMath.h"
+#elif PLATFORM_MAC
+#include "Mac/MacPlatformMath.h"
+#elif PLATFORM_IOS
+#include "IOS/IOSPlatformMath.h"
+#elif PLATFORM_ANDROID
+#include "Android/AndroidMath.h"
+#elif PLATFORM_WINRT
+#include "WinRT/WinRTMath.h"
+#elif PLATFORM_HTML5
+#include "HTML5/HTML5PlatformMath.h"
+#elif PLATFORM_LINUX
+#include "Linux/LinuxPlatformMath.h"
+#endif
\ No newline at end of file
diff --git a/Engine/Source/Runtime/Core/Public/HTML5/HTML5PlatformMath.h b/Engine/Source/Runtime/Core/Public/HTML5/HTML5PlatformMath.h
--- a/Engine/Source/Runtime/Core/Public/HTML5/HTML5PlatformMath.h
+++ b/Engine/Source/Runtime/Core/Public/HTML5/HTML5PlatformMath.h
@@ -6,6 +6,6 @@
 ==============================================================================================*/
 
 #pragma once
-
+#include "GenericPlatform/GenericPlatformMath.h"
 
 typedef FGenericPlatformMath FPlatformMath;
diff --git a/Engine/Source/Runtime/Core/Public/IOS/IOSPlatformMath.h b/Engine/Source/Runtime/Core/Public/IOS/IOSPlatformMath.h
--- a/Engine/Source/Runtime/Core/Public/IOS/IOSPlatformMath.h
+++ b/Engine/Source/Runtime/Core/Public/IOS/IOSPlatformMath.h
@@ -5,6 +5,7 @@
 ==============================================================================================*/
 
 #pragma once
+#include "GenericPlatform/GenericPlatformMath.h"
 
 /**
 * IOS implementation of the Math OS functions
diff --git a/Engine/Source/Runtime/Core/Public/Linux/LinuxPlatformMath.h b/Engine/Source/Runtime/Core/Public/Linux/LinuxPlatformMath.h
--- a/Engine/Source/Runtime/Core/Public/Linux/LinuxPlatformMath.h
+++ b/Engine/Source/Runtime/Core/Public/Linux/LinuxPlatformMath.h
@@ -5,6 +5,8 @@
 ==============================================================================================*/
 
 #pragma once
+#include "GenericPlatform/GenericPlatformMath.h"
+#include "Linux/LinuxSystemIncludes.h"
 
 /**
 * Linux implementation of the Math OS functions
diff --git a/Engine/Source/Runtime/Core/Public/Mac/MacPlatformMath.h b/Engine/Source/Runtime/Core/Public/Mac/MacPlatformMath.h
--- a/Engine/Source/Runtime/Core/Public/Mac/MacPlatformMath.h
+++ b/Engine/Source/Runtime/Core/Public/Mac/MacPlatformMath.h
@@ -5,6 +5,8 @@
 ==============================================================================================*/
 
 #pragma once
+#include "GenericPlatform/GenericPlatformMath.h"
+#include "Mac/MacSystemIncludes.h"
 
 /**
 * Mac implementation of the Math OS functions
diff --git a/Engine/Source/Runtime/Core/Public/WinRT/WinRTMath.h b/Engine/Source/Runtime/Core/Public/WinRT/WinRTMath.h
--- a/Engine/Source/Runtime/Core/Public/WinRT/WinRTMath.h
+++ b/Engine/Source/Runtime/Core/Public/WinRT/WinRTMath.h
@@ -5,5 +5,6 @@
 ==============================================================================================*/
 
 #pragma once
+#include "GenericPlatform/GenericPlatformMath.h"
 
 typedef FGenericPlatformMath FPlatformMath;
diff --git a/Engine/Source/Runtime/Core/Public/Windows/WindowsPlatformMath.h b/Engine/Source/Runtime/Core/Public/Windows/WindowsPlatformMath.h
--- a/Engine/Source/Runtime/Core/Public/Windows/WindowsPlatformMath.h
+++ b/Engine/Source/Runtime/Core/Public/Windows/WindowsPlatformMath.h
@@ -6,6 +6,8 @@
 ==============================================================================================*/
 
 #pragma once
+#include "GenericPlatform/GenericPlatformMath.h"
+#include "Windows/WindowsSystemIncludes.h"
 
 /**
 * Windows implementation of the Math OS functions
