Fixing CIS: backing out changelist 2327127

diff --git a/Engine/Source/Runtime/Core/Public/Android/AndroidMath.h b/Engine/Source/Runtime/Core/Public/Android/AndroidMath.h
--- a/Engine/Source/Runtime/Core/Public/Android/AndroidMath.h
+++ b/Engine/Source/Runtime/Core/Public/Android/AndroidMath.h
@@ -5,7 +5,6 @@
 ==============================================================================================*/
 
 #pragma once
-#include "GenericPlatform/GenericPlatformMath.h"
 
 /**
  * Android implementation of the Math OS functions
diff --git a/Engine/Source/Runtime/Core/Public/HAL/PlatformMath.h b/Engine/Source/Runtime/Core/Public/HAL/PlatformMath.h
deleted file mode 100644
diff --git a/Engine/Source/Runtime/Core/Public/HTML5/HTML5PlatformMath.h b/Engine/Source/Runtime/Core/Public/HTML5/HTML5PlatformMath.h
--- a/Engine/Source/Runtime/Core/Public/HTML5/HTML5PlatformMath.h
+++ b/Engine/Source/Runtime/Core/Public/HTML5/HTML5PlatformMath.h
@@ -6,6 +6,6 @@
 ==============================================================================================*/
 
 #pragma once
-#include "GenericPlatform/GenericPlatformMath.h"
+
 
 typedef FGenericPlatformMath FPlatformMath;
diff --git a/Engine/Source/Runtime/Core/Public/IOS/IOSPlatformMath.h b/Engine/Source/Runtime/Core/Public/IOS/IOSPlatformMath.h
--- a/Engine/Source/Runtime/Core/Public/IOS/IOSPlatformMath.h
+++ b/Engine/Source/Runtime/Core/Public/IOS/IOSPlatformMath.h
@@ -5,7 +5,6 @@
 ==============================================================================================*/
 
 #pragma once
-#include "GenericPlatform/GenericPlatformMath.h"
 
 /**
 * IOS implementation of the Math OS functions
diff --git a/Engine/Source/Runtime/Core/Public/Linux/LinuxPlatformMath.h b/Engine/Source/Runtime/Core/Public/Linux/LinuxPlatformMath.h
--- a/Engine/Source/Runtime/Core/Public/Linux/LinuxPlatformMath.h
+++ b/Engine/Source/Runtime/Core/Public/Linux/LinuxPlatformMath.h
@@ -5,8 +5,6 @@
 ==============================================================================================*/
 
 #pragma once
-#include "GenericPlatform/GenericPlatformMath.h"
-#include "Linux/LinuxSystemIncludes.h"
 
 /**
 * Linux implementation of the Math OS functions
diff --git a/Engine/Source/Runtime/Core/Public/Mac/MacPlatformMath.h b/Engine/Source/Runtime/Core/Public/Mac/MacPlatformMath.h
--- a/Engine/Source/Runtime/Core/Public/Mac/MacPlatformMath.h
+++ b/Engine/Source/Runtime/Core/Public/Mac/MacPlatformMath.h
@@ -5,8 +5,6 @@
 ==============================================================================================*/
 
 #pragma once
-#include "GenericPlatform/GenericPlatformMath.h"
-#include "Mac/MacSystemIncludes.h"
 
 /**
 * Mac implementation of the Math OS functions
diff --git a/Engine/Source/Runtime/Core/Public/WinRT/WinRTMath.h b/Engine/Source/Runtime/Core/Public/WinRT/WinRTMath.h
--- a/Engine/Source/Runtime/Core/Public/WinRT/WinRTMath.h
+++ b/Engine/Source/Runtime/Core/Public/WinRT/WinRTMath.h
@@ -5,6 +5,5 @@
 ==============================================================================================*/
 
 #pragma once
-#include "GenericPlatform/GenericPlatformMath.h"
 
 typedef FGenericPlatformMath FPlatformMath;
diff --git a/Engine/Source/Runtime/Core/Public/Windows/WindowsPlatformMath.h b/Engine/Source/Runtime/Core/Public/Windows/WindowsPlatformMath.h
--- a/Engine/Source/Runtime/Core/Public/Windows/WindowsPlatformMath.h
+++ b/Engine/Source/Runtime/Core/Public/Windows/WindowsPlatformMath.h
@@ -6,8 +6,6 @@
 ==============================================================================================*/
 
 #pragma once
-#include "GenericPlatform/GenericPlatformMath.h"
-#include "Windows/WindowsSystemIncludes.h"
 
 /**
 * Windows implementation of the Math OS functions
