Core - TTP# BLOCKER: CRASH: MAIN: Client crashes when server times out when loading into a second zone (reworked a few parts of the stat system to avoid allocating large chunks of memory)

diff --git a/Engine/Source/Developer/Profiler/Private/ProfilerSession.cpp b/Engine/Source/Developer/Profiler/Private/ProfilerSession.cpp
--- a/Engine/Source/Developer/Profiler/Private/ProfilerSession.cpp
+++ b/Engine/Source/Developer/Profiler/Private/ProfilerSession.cpp
@@ -584,7 +584,7 @@ void FRawProfilerSession::PrepareLoading()
 			for( int32 PacketIndex = 0; PacketIndex < StatPacketArray.Packets.Num(); PacketIndex++ )
 			{
 				const FStatPacket& StatPacket = *StatPacketArray.Packets[PacketIndex];
-				const TArray<FStatMessage>& Data = StatPacket.StatMessages;
+				const FStatPacket::TStatMessagesArray& Data = StatPacket.StatMessages;
 
 				for( int32 Index = 0; Index < Data.Num(); Index++ )
 				{
@@ -919,7 +919,7 @@ void FRawProfilerSession::ProcessStatPacketArray( const FStatPacketArray& StatPa
 		Stack.Add( ThreadNode );
 		FProfilerStackNode* Current = Stack.Last();
 
-		const TArray<FStatMessage>& Data = StatPacket.StatMessages;
+		const FStatPacket::TStatMessagesArray& Data = StatPacket.StatMessages;
 		for( int32 Index = 0; Index < Data.Num(); Index++ )
 		{
 			const FStatMessage& Item = Data[Index];
diff --git a/Engine/Source/Runtime/Core/Private/Stats/Stats2.cpp b/Engine/Source/Runtime/Core/Private/Stats/Stats2.cpp
--- a/Engine/Source/Runtime/Core/Private/Stats/Stats2.cpp
+++ b/Engine/Source/Runtime/Core/Private/Stats/Stats2.cpp
@@ -816,6 +816,9 @@ void FThreadStats::Flush(bool bHasBrokenCallstacks)
 		{
 			UE_LOG( LogStats, Verbose, TEXT( "StatMessage Packet has more than %i messages.  Ignoring for the presize history." ), (int32)PresizeMaxSize );
 		}
+
+		UE_CLOG( Packet.StatMessages.Num()>256, LogStats, Log, TEXT("Packet.StatMessages.Num() %i"), Packet.StatMessages.Num() );
+
 		FStatPacket* ToSend = new FStatPacket(Packet);
 		Exchange(ToSend->StatMessages, Packet.StatMessages);
 		ToSend->bBrokenCallstacks = bHasBrokenCallstacks;
diff --git a/Engine/Source/Runtime/Core/Private/Stats/StatsData.cpp b/Engine/Source/Runtime/Core/Private/Stats/StatsData.cpp
--- a/Engine/Source/Runtime/Core/Private/Stats/StatsData.cpp
+++ b/Engine/Source/Runtime/Core/Private/Stats/StatsData.cpp
@@ -523,7 +523,7 @@ int64 FStatsThreadState::GetLatestValidFrame() const
 	return Result;
 }
 
-void FStatsThreadState::ScanForAdvance(TArray<FStatMessage> const& Data)
+void FStatsThreadState::ScanForAdvance(FStatPacket::TStatMessagesArray const& Data)
 {
 	for (int32 Index = Data.Num() - 1; Index >= 0 ; Index--)
 	{
@@ -587,7 +587,7 @@ void FStatsThreadState::ScanForAdvance(FStatPacketArray& NewData)
 	{
 		int64 FrameNum = NewData.Packets[Index]->ThreadType == EThreadType::Renderer ? CurrentRenderFrame : CurrentGameFrame;
 		NewData.Packets[Index]->Frame = FrameNum;
-		TArray<FStatMessage> const& Data = NewData.Packets[Index]->StatMessages;
+		FStatPacket::TStatMessagesArray const& Data = NewData.Packets[Index]->StatMessages;
 		ScanForAdvance(Data);
 		Count += Data.Num();
 	}
@@ -652,7 +652,7 @@ void FStatsThreadState::ProcessMetaDataOnly(TArray<FStatMessage>& Data)
 	}
 }
 
-void FStatsThreadState::ProcessNonFrameStats(TArray<FStatMessage>& Data, TSet<FName>* NonFrameStatsFound)
+void FStatsThreadState::ProcessNonFrameStats(FStatPacket::TStatMessagesArray& Data, TSet<FName>* NonFrameStatsFound)
 {
 	check(!bWasLoaded);
 	for (int32 Index = 0; Index < Data.Num() ; Index++)
@@ -760,7 +760,7 @@ void FStatsThreadState::AddToHistoryAndEmpty(FStatPacketArray& NewData)
 
 				check(PacketToCopyForNonFrame);
 
-				TArray<FStatMessage>* NonFrameMessages = NULL;
+				FStatPacket::TStatMessagesArray* NonFrameMessages = NULL;
 
 				for (TMap<FName, FStatMessage>::TConstIterator It(NotClearedEveryFrame); It; ++It)
 				{
@@ -1018,7 +1018,7 @@ void FStatsThreadState::GetRawStackStats(int64 TargetFrame, FRawStatStackNode& R
 			Stack.Add(ThreadRoot);
 			FRawStatStackNode* Current = Stack.Last();
 
-			TArray<FStatMessage> const& Data = Packet.StatMessages;
+			FStatPacket::TStatMessagesArray const& Data = Packet.StatMessages;
 			for (int32 Index = 0; Index < Data.Num(); Index++)
 			{
 				FStatMessage const& Item = Data[Index];
@@ -1154,7 +1154,7 @@ int64 FStatsThreadState::GetFastThreadFrameTimeInternal( int64 TargetFrame, int3
 
 		if( Packet.ThreadId == ThreadID || Packet.ThreadType == Thread )
 		{
-			TArray<FStatMessage> const& Data = Packet.StatMessages;
+			FStatPacket::TStatMessagesArray const& Data = Packet.StatMessages;
 			for (int32 Index = 0; Index < Data.Num(); Index++)
 			{
 				FStatMessage const& Item = Data[Index];
diff --git a/Engine/Source/Runtime/Core/Public/Containers/ChunkedArray.h b/Engine/Source/Runtime/Core/Public/Containers/ChunkedArray.h
--- a/Engine/Source/Runtime/Core/Public/Containers/ChunkedArray.h
+++ b/Engine/Source/Runtime/Core/Public/Containers/ChunkedArray.h
@@ -134,6 +134,24 @@ public:
 		return this->NumElements - 1;
 	}
 
+	/**
+	 * Appends the specified array to this array.
+	 * Cannot append to self.
+	 *
+	 * @param Other The array to append.
+	 */
+	FORCEINLINE TChunkedArray& operator+=(const TArray<ElementType>& Other) 
+	{ 
+		if( (UPTRINT*)this != (UPTRINT*)&Other )
+		{
+			for( const auto& It : Other )
+			{
+				AddElement(It);
+			}
+		}
+		return *this; 
+	}
+
 	int32 Add( int32 Count=1 )
 	{
 		check(Count>=0);
@@ -157,12 +175,25 @@ public:
 		NumElements = 0;
 	}
 
+	/**
+	 * Reserves memory such that the array can contain at least Number elements.
+	 *
+	 * @param Number The number of elements that the array should be able to
+	 *               contain after allocation.
+	 */
+	void Reserve(int32 Number)
+	{
+		Chunks.Reserve(Number % NumElementsPerChunk + 1);
+	}
+
 	void Shrink()
 	{
 		Chunks.Shrink();
 	}
 
-private:
+protected:
+	friend struct TContainerTraits<TChunkedArray<ElementType, TargetBytesPerChunk>>;
+
 	enum { NumElementsPerChunk = TargetBytesPerChunk / sizeof(ElementType) };
 
 	/** A chunk of the array's elements. */
diff --git a/Engine/Source/Runtime/Core/Public/Stats/Stats2.h b/Engine/Source/Runtime/Core/Public/Stats/Stats2.h
--- a/Engine/Source/Runtime/Core/Public/Stats/Stats2.h
+++ b/Engine/Source/Runtime/Core/Public/Stats/Stats2.h
@@ -1,10 +1,9 @@
 // Copyright 1998-2014 Epic Games, Inc. All Rights Reserved.
 
-/*=============================================================================
-	UnStats.h: Performance stats framework.
-=============================================================================*/
 #pragma once
 
+#include "ChunkedArray.h"
+
 /**
 * This is thread-private information about the thread idle stats, which we always collect, even in final builds
 */
@@ -955,6 +954,13 @@ template<> struct TIsPODType<FComplexStatMessage> { enum { Value = true }; };
 */
 struct FStatPacket
 {
+	enum 
+	{
+		MESSAGES_CHUNK_SIZE = 64*1024,
+	};
+
+	typedef TChunkedArray<FStatMessage,MESSAGES_CHUNK_SIZE> TStatMessagesArray;
+
 	/** Assigned later, this is the frame number this packet is for **/
 	int64 Frame;
 	/** ThreadId this packet came from **/
@@ -964,7 +970,7 @@ struct FStatPacket
 	/** true if this packet has broken callstacks **/
 	bool bBrokenCallstacks;
 	/** messages in this packet **/
-	TArray<FStatMessage> StatMessages;
+	TStatMessagesArray StatMessages;
 	/** Size we presize the message buffer to, currently the max of what we have seen so far. **/
 	TArray<int32> StatMessagesPresize;
 
diff --git a/Engine/Source/Runtime/Core/Public/Stats/StatsData.h b/Engine/Source/Runtime/Core/Public/Stats/StatsData.h
--- a/Engine/Source/Runtime/Core/Public/Stats/StatsData.h
+++ b/Engine/Source/Runtime/Core/Public/Stats/StatsData.h
@@ -419,7 +419,7 @@ class CORE_API FStatsThreadState
 	friend class FStatsThread;
 
 	/** Internal method to scan the messages to update the current frame **/
-	void ScanForAdvance(TArray<FStatMessage> const& Data);
+	void ScanForAdvance(FStatPacket::TStatMessagesArray const& Data);
 
 	/** Internal method to scan the messages to update the current frame and accumulate any non-frame stats. **/
 	void ScanForAdvance(FStatPacketArray& NewData);
@@ -439,7 +439,7 @@ public:
 
 private:
 	/** Internal method to scan the messages to accumulate any non-frame stats. **/
-	void ProcessNonFrameStats(TArray<FStatMessage>& Data, TSet<FName>* NonFrameStatsFound);
+	void ProcessNonFrameStats(FStatPacket::TStatMessagesArray& Data, TSet<FName>* NonFrameStatsFound);
 
 	/** Internal method to place the data into the history, discard and broadcast any new frames to anyone who cares. **/
 	void AddToHistoryAndEmpty(FStatPacketArray& NewData);
