Core - Removed debug logs

diff --git a/Engine/Source/Runtime/Core/Private/Stats/Stats2.cpp b/Engine/Source/Runtime/Core/Private/Stats/Stats2.cpp
--- a/Engine/Source/Runtime/Core/Private/Stats/Stats2.cpp
+++ b/Engine/Source/Runtime/Core/Private/Stats/Stats2.cpp
@@ -816,9 +816,6 @@ void FThreadStats::Flush(bool bHasBrokenCallstacks)
 		{
 			UE_LOG( LogStats, Verbose, TEXT( "StatMessage Packet has more than %i messages.  Ignoring for the presize history." ), (int32)PresizeMaxSize );
 		}
-
-		UE_CLOG( Packet.StatMessages.Num()>256, LogStats, Log, TEXT("Packet.StatMessages.Num() %i"), Packet.StatMessages.Num() );
-
 		FStatPacket* ToSend = new FStatPacket(Packet);
 		Exchange(ToSend->StatMessages, Packet.StatMessages);
 		ToSend->bBrokenCallstacks = bHasBrokenCallstacks;
