Make Core module public headers individually compilable. Part 9 of n.

diff --git a/Engine/Source/Runtime/Core/Public/Containers/BitArray.h b/Engine/Source/Runtime/Core/Public/Containers/BitArray.h
--- a/Engine/Source/Runtime/Core/Public/Containers/BitArray.h
+++ b/Engine/Source/Runtime/Core/Public/Containers/BitArray.h
@@ -5,6 +5,7 @@
 =============================================================================*/
 
 #pragma once
+#include "Math/UnrealMathUtility.h"
 
 // Functions for manipulating bit sets.
 struct FBitSet
diff --git a/Engine/Source/Runtime/Core/Public/Containers/Map.h b/Engine/Source/Runtime/Core/Public/Containers/Map.h
--- a/Engine/Source/Runtime/Core/Public/Containers/Map.h
+++ b/Engine/Source/Runtime/Core/Public/Containers/Map.h
@@ -1,10 +1,8 @@
 // Copyright 1998-2014 Epic Games, Inc. All Rights Reserved.
 
 #pragma once
-
-#include "Core.h"
-#include "Array.h"
-#include "Set.h"
+#include "Containers/Array.h"
+#include "Containers/Set.h"
 
 #define ExchangeB(A,B) {bool T=A; A=B; B=T;}
 
diff --git a/Engine/Source/Runtime/Core/Public/Containers/Set.h b/Engine/Source/Runtime/Core/Public/Containers/Set.h
--- a/Engine/Source/Runtime/Core/Public/Containers/Set.h
+++ b/Engine/Source/Runtime/Core/Public/Containers/Set.h
@@ -2,8 +2,8 @@
 
 #pragma once
 
-#include "SparseArray.h"
-
+#include "Containers/SparseArray.h"
+#include "Templates/Sorting.h"
 
 /**
  * The base KeyFuncs type with some useful definitions for all KeyFuncs; meant to be derived from instead of used directly.
diff --git a/Engine/Source/Runtime/Core/Public/Containers/SparseArray.h b/Engine/Source/Runtime/Core/Public/Containers/SparseArray.h
--- a/Engine/Source/Runtime/Core/Public/Containers/SparseArray.h
+++ b/Engine/Source/Runtime/Core/Public/Containers/SparseArray.h
@@ -2,8 +2,8 @@
 
 #pragma once
 
-#include "Array.h"
-#include "BitArray.h"
+#include "Containers/Array.h"
+#include "Containers/BitArray.h"
 
 
 // Forward declarations.
diff --git a/Engine/Source/Runtime/Core/Public/Math/UnrealMathUtility.h b/Engine/Source/Runtime/Core/Public/Math/UnrealMathUtility.h
--- a/Engine/Source/Runtime/Core/Public/Math/UnrealMathUtility.h
+++ b/Engine/Source/Runtime/Core/Public/Math/UnrealMathUtility.h
@@ -5,6 +5,8 @@
 =============================================================================*/
 
 #pragma once
+#include "HAL/Platform.h"
+#include "HAL/PlatformMath.h"
 
 //#define IMPLEMENT_ASSIGNMENT_OPERATOR_MANUALLY
 
diff --git a/Engine/Source/Runtime/Core/Public/Windows/WindowsSystemIncludes.h b/Engine/Source/Runtime/Core/Public/Windows/WindowsSystemIncludes.h
--- a/Engine/Source/Runtime/Core/Public/Windows/WindowsSystemIncludes.h
+++ b/Engine/Source/Runtime/Core/Public/Windows/WindowsSystemIncludes.h
@@ -5,12 +5,7 @@
 #include "Windows/WindowsPlatformCompilerSetup.h"
 
 // include windows.h
-#include "Windows/PreWindowsApi.h"
-#ifndef STRICT
-#define STRICT
-#endif
-#include "Windows/MinWindows.h"
-#include "Windows/PostWindowsApi.h"
+#include "Windows/WindowsHWrapper.h"
 
 // Macro for releasing COM objects
 #define SAFE_RELEASE(p) { if(p) { (p)->Release(); (p)=NULL; } }
