Fixed a Mac FCocoaWindow problem with deferred setFrameOrigin not caching the position

diff --git a/Engine/Source/Runtime/Core/Private/Mac/CocoaWindow.cpp b/Engine/Source/Runtime/Core/Private/Mac/CocoaWindow.cpp
--- a/Engine/Source/Runtime/Core/Private/Mac/CocoaWindow.cpp
+++ b/Engine/Source/Runtime/Core/Private/Mac/CocoaWindow.cpp
@@ -336,7 +336,7 @@ NSString* NSPerformDragOperation = @"NSPerformDragOperation";
 	else
 	{
 		self.bDeferSetOrigin = true;
-		self.DeferFrame.origin = Point;
+		self.DeferFrame = NSMakeRect(Point.x, Point.y, self.DeferFrame.size.width, self.DeferFrame.size.height);
 		NSNotification* Notification = [NSNotification notificationWithName:NSWindowDidMoveNotification object:self];
 		FMacEvent::SendToGameRunLoop(Notification, self, EMacEventSendMethod::Async, @[ NSDefaultRunLoopMode, UE4ResizeEventMode, UE4ShowEventMode ]);
 	}
