Fixed issue on iOS where the OS default locale was not properly being retrieved for I18N initialization.

diff --git a/Engine/Source/Runtime/Core/Private/Internationalization/ICUInternationalization.cpp b/Engine/Source/Runtime/Core/Private/Internationalization/ICUInternationalization.cpp
--- a/Engine/Source/Runtime/Core/Private/Internationalization/ICUInternationalization.cpp
+++ b/Engine/Source/Runtime/Core/Private/Internationalization/ICUInternationalization.cpp
@@ -117,7 +117,7 @@ bool FICUInternationalization::Initialize()
 	FICUBreakIteratorManager::Create();
 
 	I18N->InvariantCulture = FindOrMakeCulture(TEXT("en-US-POSIX"));
-	I18N->DefaultCulture = FindOrMakeCulture(icu::Locale::getDefault().getName());
+	I18N->DefaultCulture = FindOrMakeCulture(FPlatformMisc::GetDefaultLocale());
 	SetCurrentCulture( I18N->GetDefaultCulture()->GetName() );
 	return U_SUCCESS(ICUStatus) ? true : false;
 }
