Make Core module public headers individually compilable. Part 10 of n.

diff --git a/Engine/Source/Runtime/Core/Public/Android/AndroidString.h b/Engine/Source/Runtime/Core/Public/Android/AndroidString.h
--- a/Engine/Source/Runtime/Core/Public/Android/AndroidString.h
+++ b/Engine/Source/Runtime/Core/Public/Android/AndroidString.h
@@ -10,9 +10,9 @@
 // to just be stubs to the non-wide versions.  This doesn't work for obvious reasons.
 
 #pragma once
-
-#include "Char.h"
-#include "GenericPlatformStricmp.h"
+#include "Misc/Char.h"
+#include "GenericPlatform/GenericPlatformStricmp.h"
+#include "GenericPlatform/GenericPlatformString.h"
 
 /**
  * Android string implementation
diff --git a/Engine/Source/Runtime/Core/Public/Apple/ApplePlatformString.h b/Engine/Source/Runtime/Core/Public/Apple/ApplePlatformString.h
--- a/Engine/Source/Runtime/Core/Public/Apple/ApplePlatformString.h
+++ b/Engine/Source/Runtime/Core/Public/Apple/ApplePlatformString.h
@@ -9,6 +9,13 @@
 #include "StandardPlatformString.h"
 #include "AssertionMacros.h"
 #include "OutputDevice.h"
+#include "HAL/Platform.h"
+
+#if PLATFORM_MAC
+#include "Mac/MacSystemIncludes.h"
+#elif PLATFORM_IOS
+#include "IOS/IOSSystemIncludes.h"
+#endif
 
 /**
 * Mac string implementation
diff --git a/Engine/Source/Runtime/Core/Public/HAL/PlatformString.h b/Engine/Source/Runtime/Core/Public/HAL/PlatformString.h
new file mode 100644
--- /dev/null
+++ b/Engine/Source/Runtime/Core/Public/HAL/PlatformString.h
@@ -0,0 +1,24 @@
+// Copyright 1998-2014 Epic Games, Inc. All Rights Reserved.
+#pragma once
+
+#include "HAL/Platform.h"
+
+#if PLATFORM_WINDOWS
+#include "Windows/WindowsPlatformString.h"
+#elif PLATFORM_PS4
+#include "PS4/PS4String.h"
+#elif PLATFORM_XBOXONE
+#include "XboxOne/XboxOneString.h"
+#elif PLATFORM_MAC
+#include "Apple/ApplePlatformString.h"
+#elif PLATFORM_IOS
+#include "IOS/IOSPlatformString.h"
+#elif PLATFORM_ANDROID
+#include "Android/AndroidString.h"
+#elif PLATFORM_WINRT
+#include "WinRT/WinRTString.h"
+#elif PLATFORM_HTML5
+#include "HTML5/HTML5PlatformString.h"
+#elif PLATFORM_LINUX
+#include "Linux/LinuxPlatformString.h"
+#endif
\ No newline at end of file
diff --git a/Engine/Source/Runtime/Core/Public/HTML5/HTML5PlatformString.h b/Engine/Source/Runtime/Core/Public/HTML5/HTML5PlatformString.h
--- a/Engine/Source/Runtime/Core/Public/HTML5/HTML5PlatformString.h
+++ b/Engine/Source/Runtime/Core/Public/HTML5/HTML5PlatformString.h
@@ -6,8 +6,8 @@
 ==============================================================================================*/
 
 #pragma once
-
-#include "StandardPlatformString.h"
+#include "GenericPlatform/StandardPlatformString.h"
+#include "HAL/Platform.h"
 
 struct FHTML5PlatformString : public FStandardPlatformString
 {
diff --git a/Engine/Source/Runtime/Core/Public/IOS/IOSPlatformString.h b/Engine/Source/Runtime/Core/Public/IOS/IOSPlatformString.h
--- a/Engine/Source/Runtime/Core/Public/IOS/IOSPlatformString.h
+++ b/Engine/Source/Runtime/Core/Public/IOS/IOSPlatformString.h
@@ -5,7 +5,7 @@
 ==================================================================================*/
 
 #pragma once
-
+#include "HAL/Platform.h"
 #import <Foundation/NSString.h>
 
 class FString;
diff --git a/Engine/Source/Runtime/Core/Public/Linux/LinuxPlatformString.h b/Engine/Source/Runtime/Core/Public/Linux/LinuxPlatformString.h
--- a/Engine/Source/Runtime/Core/Public/Linux/LinuxPlatformString.h
+++ b/Engine/Source/Runtime/Core/Public/Linux/LinuxPlatformString.h
@@ -5,9 +5,8 @@
 ==============================================================================================*/
 
 #pragma once
-
-#include "StandardPlatformString.h"
-
+#include "GenericPlatform/StandardPlatformString.h"
+#include "Linux/LinuxSystemIncludes.h"
 /**
 * Linux string implementation
 **/
diff --git a/Engine/Source/Runtime/Core/Public/WinRT/WinRTString.h b/Engine/Source/Runtime/Core/Public/WinRT/WinRTString.h
--- a/Engine/Source/Runtime/Core/Public/WinRT/WinRTString.h
+++ b/Engine/Source/Runtime/Core/Public/WinRT/WinRTString.h
@@ -5,8 +5,8 @@
 ==============================================================================================*/
 
 #pragma once
-
-#include "MicrosoftPlatformString.h"
+#include "GenericPlatform/MicrosoftPlatformString.h"
+#include "HAL/Platform.h"
 
 /**
  * WinRT string implementation
diff --git a/Engine/Source/Runtime/Core/Public/Windows/WindowsPlatformString.h b/Engine/Source/Runtime/Core/Public/Windows/WindowsPlatformString.h
--- a/Engine/Source/Runtime/Core/Public/Windows/WindowsPlatformString.h
+++ b/Engine/Source/Runtime/Core/Public/Windows/WindowsPlatformString.h
@@ -7,7 +7,7 @@
 
 #pragma once
 
-#include "MicrosoftPlatformString.h"
+#include "GenericPlatform/MicrosoftPlatformString.h"
 
 /**
 * Windows string implementation
