Make Core module public headers individually compilable. Part 11 of n.

diff --git a/Engine/Source/Runtime/Core/Public/Containers/StringConv.h b/Engine/Source/Runtime/Core/Public/Containers/StringConv.h
--- a/Engine/Source/Runtime/Core/Public/Containers/StringConv.h
+++ b/Engine/Source/Runtime/Core/Public/Containers/StringConv.h
@@ -7,6 +7,9 @@
  */
 
 #pragma once
+#include "Containers/ContainerAllocationPolicies.h"
+#include "HAL/Platform.h"
+#include "HAL/PlatformString.h"
 
 #define DEFAULT_STRING_CONVERSION_SIZE 128u
 #define UNICODE_BOGUS_CHAR_CODEPOINT   '?'
diff --git a/Engine/Source/Runtime/Core/Public/Containers/UnrealString.h b/Engine/Source/Runtime/Core/Public/Containers/UnrealString.h
--- a/Engine/Source/Runtime/Core/Public/Containers/UnrealString.h
+++ b/Engine/Source/Runtime/Core/Public/Containers/UnrealString.h
@@ -7,9 +7,12 @@
 =============================================================================*/
 
 #pragma once
-
-#include "Array.h"
-
+#include "Containers/Array.h"
+#include "Math/NumericLimits.h"
+#include "Math/UnrealMathUtility.h"
+#include "Misc/Crc.h"
+#include "Misc/CString.h"
+#include "Templates/MemoryOps.h"
 
 /** Determines case sensitivity options for string comparisons. */
 namespace ESearchCase
diff --git a/Engine/Source/Runtime/Core/Public/Misc/CString.h b/Engine/Source/Runtime/Core/Public/Misc/CString.h
--- a/Engine/Source/Runtime/Core/Public/Misc/CString.h
+++ b/Engine/Source/Runtime/Core/Public/Misc/CString.h
@@ -1,7 +1,8 @@
 // Copyright 1998-2014 Epic Games, Inc. All Rights Reserved.
 
 #pragma once
-
+#include "HAL/PlatformString.h"
+#include "Misc/Char.h"
 
 #define MAX_SPRINTF 1024
 
