Make Core module public headers individually compilable. Part 12 of n.

diff --git a/Engine/Source/Runtime/Core/Public/Android/AndroidProcess.h b/Engine/Source/Runtime/Core/Public/Android/AndroidProcess.h
--- a/Engine/Source/Runtime/Core/Public/Android/AndroidProcess.h
+++ b/Engine/Source/Runtime/Core/Public/Android/AndroidProcess.h
@@ -6,6 +6,7 @@
 ==============================================================================================*/
 
 #pragma once
+#include "GenericPlatform/GenericPlatformProcess.h"
 
 /** Dummy process handle for platforms that use generic implementation. */
 struct FProcHandle : public TProcHandle<void*, nullptr>
diff --git a/Engine/Source/Runtime/Core/Public/HAL/PlatformProcess.h b/Engine/Source/Runtime/Core/Public/HAL/PlatformProcess.h
new file mode 100644
--- /dev/null
+++ b/Engine/Source/Runtime/Core/Public/HAL/PlatformProcess.h
@@ -0,0 +1,24 @@
+// Copyright 1998-2014 Epic Games, Inc. All Rights Reserved.
+
+#pragma once
+#include "HAL/Platform.h"
+
+#if PLATFORM_WINDOWS
+#include "Windows/WindowsPlatformProcess.h"
+#elif PLATFORM_PS4
+#include "PS4/PS4Process.h"
+#elif PLATFORM_XBOXONE
+#include "XboxOne/XboxOneProcess.h"
+#elif PLATFORM_MAC
+#include "Mac/MacPlatformProcess.h"
+#elif PLATFORM_IOS
+#include "IOS/IOSPlatformProcess.h"
+#elif PLATFORM_ANDROID
+#include "Android/AndroidProcess.h"
+#elif PLATFORM_WINRT
+#include "WinRT/WinRTProcess.h"
+#elif PLATFORM_HTML5
+#include "HTML5/HTML5PlatformProcess.h"
+#elif PLATFORM_LINUX
+#include "Linux/LinuxPlatformProcess.h"
+#endif
\ No newline at end of file
diff --git a/Engine/Source/Runtime/Core/Public/HTML5/HTML5PlatformProcess.h b/Engine/Source/Runtime/Core/Public/HTML5/HTML5PlatformProcess.h
--- a/Engine/Source/Runtime/Core/Public/HTML5/HTML5PlatformProcess.h
+++ b/Engine/Source/Runtime/Core/Public/HTML5/HTML5PlatformProcess.h
@@ -6,6 +6,7 @@
 ==============================================================================================*/
 
 #pragma once
+#include "GenericPlatform/GenericPlatformProcess.h"
 
 /** Dummy process handle for platforms that use generic implementation. */
 struct FProcHandle : public TProcHandle<void*, nullptr>
diff --git a/Engine/Source/Runtime/Core/Public/IOS/IOSPlatformProcess.h b/Engine/Source/Runtime/Core/Public/IOS/IOSPlatformProcess.h
--- a/Engine/Source/Runtime/Core/Public/IOS/IOSPlatformProcess.h
+++ b/Engine/Source/Runtime/Core/Public/IOS/IOSPlatformProcess.h
@@ -5,6 +5,7 @@
 ==============================================================================================*/
 
 #pragma once
+#include "GenericPlatform/GenericPlatformProcess.h"
 
 /** Dummy process handle for platforms that use generic implementation. */
 struct FProcHandle : public TProcHandle<void*, nullptr>
diff --git a/Engine/Source/Runtime/Core/Public/Linux/LinuxPlatformProcess.h b/Engine/Source/Runtime/Core/Public/Linux/LinuxPlatformProcess.h
--- a/Engine/Source/Runtime/Core/Public/Linux/LinuxPlatformProcess.h
+++ b/Engine/Source/Runtime/Core/Public/Linux/LinuxPlatformProcess.h
@@ -6,6 +6,9 @@
 ==============================================================================================*/
 
 #pragma once
+#include "GenericPlatform/GenericPlatformProcess.h"
+#include "HAL/Platform.h"
+#include "Linux/LinuxSystemIncludes.h"
 
 /** Wrapper around Linux pid_t. */
 struct FProcHandle : public TProcHandle<pid_t, -1>
diff --git a/Engine/Source/Runtime/Core/Public/Mac/MacPlatformProcess.h b/Engine/Source/Runtime/Core/Public/Mac/MacPlatformProcess.h
--- a/Engine/Source/Runtime/Core/Public/Mac/MacPlatformProcess.h
+++ b/Engine/Source/Runtime/Core/Public/Mac/MacPlatformProcess.h
@@ -5,6 +5,7 @@
 ==============================================================================================*/
 
 #pragma once
+#include "GenericPlatform/GenericPlatformProcess.h"
 
 
 /** Mac implementation of the process handle. */
diff --git a/Engine/Source/Runtime/Core/Public/WinRT/WinRTProcess.h b/Engine/Source/Runtime/Core/Public/WinRT/WinRTProcess.h
--- a/Engine/Source/Runtime/Core/Public/WinRT/WinRTProcess.h
+++ b/Engine/Source/Runtime/Core/Public/WinRT/WinRTProcess.h
@@ -5,6 +5,7 @@
 ==============================================================================================*/
 
 #pragma once
+#include "GenericPlatform/GenericPlatformProcess.h"
 
 /** Dummy process handle for platforms that use generic implementation. */
 struct FProcHandle : public TProcHandle<void*, nullptr>
diff --git a/Engine/Source/Runtime/Core/Public/Windows/WindowsPlatformProcess.h b/Engine/Source/Runtime/Core/Public/Windows/WindowsPlatformProcess.h
--- a/Engine/Source/Runtime/Core/Public/Windows/WindowsPlatformProcess.h
+++ b/Engine/Source/Runtime/Core/Public/Windows/WindowsPlatformProcess.h
@@ -6,6 +6,8 @@
 ==============================================================================================*/
 
 #pragma once
+#include "GenericPlatform/GenericPlatformProcess.h"
+#include "Windows/WindowsSystemIncludes.h"
 
 /** Windows implementation of the process handle. */
 #if PLATFORM_WINDOWS && defined(__clang__)	// @todo clang: Clang on Windows has a bug with non-type template arguments
