Make Core module public headers individually compilable. Part 13 of n.

diff --git a/Engine/Source/Runtime/Core/Private/UObject/UnrealNames.cpp b/Engine/Source/Runtime/Core/Private/UObject/UnrealNames.cpp
--- a/Engine/Source/Runtime/Core/Private/UObject/UnrealNames.cpp
+++ b/Engine/Source/Runtime/Core/Private/UObject/UnrealNames.cpp
@@ -754,7 +754,14 @@ void FName::StaticInit()
 	{
 		// Register all hardcoded names.
 		#define REGISTER_NAME(num,namestr) FName Temp_##namestr(EName(num), TEXT(#namestr));
+		#ifdef _UNREAL_NAMES_H_
+		#undef _UNREAL_NAMES_H_
+		#define _RECOVER_UNREAL_NAMES_H_
+		#endif
 		#include "UObject/UnrealNames.h"
+		#ifdef _RECOVER_UNREAL_NAMES_H_
+		#define _UNREAL_NAMES_H_
+		#endif
 	}
 
 #if DO_CHECK
diff --git a/Engine/Source/Runtime/Core/Public/Async/TaskGraphInterfaces.h b/Engine/Source/Runtime/Core/Public/Async/TaskGraphInterfaces.h
--- a/Engine/Source/Runtime/Core/Public/Async/TaskGraphInterfaces.h
+++ b/Engine/Source/Runtime/Core/Public/Async/TaskGraphInterfaces.h
@@ -5,6 +5,8 @@
 =============================================================================*/
 
 #pragma once
+#include "HAL/Platform.h"
+#include "Misc/Build.h"
 
 #if !defined(STATS)
 #error "STATS must be defined as either zero or one."
diff --git a/Engine/Source/Runtime/Core/Public/Delegates/DelegateInstancesImpl_Variadics.inl b/Engine/Source/Runtime/Core/Public/Delegates/DelegateInstancesImpl_Variadics.inl
--- a/Engine/Source/Runtime/Core/Public/Delegates/DelegateInstancesImpl_Variadics.inl
+++ b/Engine/Source/Runtime/Core/Public/Delegates/DelegateInstancesImpl_Variadics.inl
@@ -11,7 +11,10 @@
 ================================================================================*/
 
 #pragma once
+#include "Delegates/DelegateInstanceInterface_Variadics.h"
+#include "UObject/NameTypes.h"
 
+class UFunction;
 
 /* Macros for function parameter and delegate payload lists
  *****************************************************************************/
diff --git a/Engine/Source/Runtime/Core/Public/Delegates/DelegateSignatureImpl_Variadics.inl b/Engine/Source/Runtime/Core/Public/Delegates/DelegateSignatureImpl_Variadics.inl
--- a/Engine/Source/Runtime/Core/Public/Delegates/DelegateSignatureImpl_Variadics.inl
+++ b/Engine/Source/Runtime/Core/Public/Delegates/DelegateSignatureImpl_Variadics.inl
@@ -6,6 +6,7 @@
 #endif
 
 #pragma once
+#include "UObject/ScriptDelegates.h"
 
 /**
  * Unicast delegate base object.
diff --git a/Engine/Source/Runtime/Core/Public/GenericPlatform/GenericApplicationMessageHandler.h b/Engine/Source/Runtime/Core/Public/GenericPlatform/GenericApplicationMessageHandler.h
--- a/Engine/Source/Runtime/Core/Public/GenericPlatform/GenericApplicationMessageHandler.h
+++ b/Engine/Source/Runtime/Core/Public/GenericPlatform/GenericApplicationMessageHandler.h
@@ -1,6 +1,10 @@
 // Copyright 1998-2014 Epic Games, Inc. All Rights Reserved.
 
 #pragma once
+#include "HAL/Platform.h"
+#include "Templates/SharedPointer.h"
+
+class FGenericWindow;
 
 namespace EMouseButtons
 {
diff --git a/Engine/Source/Runtime/Core/Public/GenericPlatform/GenericWindow.h b/Engine/Source/Runtime/Core/Public/GenericPlatform/GenericWindow.h
--- a/Engine/Source/Runtime/Core/Public/GenericPlatform/GenericWindow.h
+++ b/Engine/Source/Runtime/Core/Public/GenericPlatform/GenericWindow.h
@@ -1,8 +1,10 @@
 // Copyright 1998-2014 Epic Games, Inc. All Rights Reserved.
 
 #pragma once
+#include "HAL/Platform.h"
+#include "Templates/SharedPointer.h"
 
-#include "Core.h"
+struct FGenericWindowDefinition;
 
 /**
  * Modes that an FGenericWindow can be in
diff --git a/Engine/Source/Runtime/Core/Public/GenericPlatform/GenericWindowDefinition.h b/Engine/Source/Runtime/Core/Public/GenericPlatform/GenericWindowDefinition.h
--- a/Engine/Source/Runtime/Core/Public/GenericPlatform/GenericWindowDefinition.h
+++ b/Engine/Source/Runtime/Core/Public/GenericPlatform/GenericWindowDefinition.h
@@ -1,8 +1,8 @@
 // Copyright 1998-2014 Epic Games, Inc. All Rights Reserved.
 
 #pragma once
-
-#include "Core.h"
+#include "Containers/UnrealString.h"
+#include "HAL/Platform.h"
 
 struct CORE_API FGenericWindowDefinition
 {
diff --git a/Engine/Source/Runtime/Core/Public/Math/Color.h b/Engine/Source/Runtime/Core/Public/Math/Color.h
--- a/Engine/Source/Runtime/Core/Public/Math/Color.h
+++ b/Engine/Source/Runtime/Core/Public/Math/Color.h
@@ -5,6 +5,13 @@
 =============================================================================*/
 
 #pragma once
+#include "Math/UnrealMathUtility.h"
+#include "Misc/Parse.h"
+#include "Serialization/ArchiveBase.h"
+
+class FColor;
+class FVector;
+class FFloat16Color;
 
 /**
  * A linear, 32-bit/component floating point RGBA color.
@@ -24,9 +31,9 @@ struct FLinearColor
 	: R(0), G(0), B(0), A(0)
 	{}
 	FORCEINLINE FLinearColor(float InR,float InG,float InB,float InA = 1.0f): R(InR), G(InG), B(InB), A(InA) {}
-	CORE_API FLinearColor(const class FColor& C);
-	CORE_API FLinearColor(const class FVector& Vector);
-	CORE_API explicit FLinearColor(const class FFloat16Color& C);
+	CORE_API FLinearColor(const FColor& C);
+	CORE_API FLinearColor(const FVector& Vector);
+	CORE_API explicit FLinearColor(const FFloat16Color& C);
 
 	// Serializer.
 
diff --git a/Engine/Source/Runtime/Core/Public/Math/IntPoint.h b/Engine/Source/Runtime/Core/Public/Math/IntPoint.h
--- a/Engine/Source/Runtime/Core/Public/Math/IntPoint.h
+++ b/Engine/Source/Runtime/Core/Public/Math/IntPoint.h
@@ -5,7 +5,8 @@
 =============================================================================*/
 
 #pragma once
-
+#include "Containers/UnrealString.h"
+#include "HAL/Platform.h"
 
 /**
  * Structure for integer points in 2-d space.
diff --git a/Engine/Source/Runtime/Core/Public/Serialization/ArchiveBase.h b/Engine/Source/Runtime/Core/Public/Serialization/ArchiveBase.h
--- a/Engine/Source/Runtime/Core/Public/Serialization/ArchiveBase.h
+++ b/Engine/Source/Runtime/Core/Public/Serialization/ArchiveBase.h
@@ -9,10 +9,12 @@
 #include "HAL/PlatformProperties.h"
 #include "Misc/Compression.h"
 
+class FAssetPtr;
 class FCustomVersionContainer;
+class FLazyObjectPtr;
+struct FStringAssetReference;
 struct FUntypedBulkData;
 
-
 /**
  * Base class for archives that can be used for loading, saving, and garbage
  * collecting in a byte order neutral way.
diff --git a/Engine/Source/Runtime/Core/Public/Templates/ScopedCallback.h b/Engine/Source/Runtime/Core/Public/Templates/ScopedCallback.h
--- a/Engine/Source/Runtime/Core/Public/Templates/ScopedCallback.h
+++ b/Engine/Source/Runtime/Core/Public/Templates/ScopedCallback.h
@@ -2,6 +2,7 @@
 
 
 #pragma once
+#include "HAL/Platform.h"
 
 /**
  * Helper object for batching callback requests and firing on destruction of the FScopedCallback object.
diff --git a/Engine/Source/Runtime/Core/Public/Templates/SharedPointer.h b/Engine/Source/Runtime/Core/Public/Templates/SharedPointer.h
--- a/Engine/Source/Runtime/Core/Public/Templates/SharedPointer.h
+++ b/Engine/Source/Runtime/Core/Public/Templates/SharedPointer.h
@@ -5,7 +5,9 @@
 =============================================================================*/
 
 #pragma once
-
+#include "Containers/ContainersFwd.h"
+#include "Containers/Map.h"
+#include "Templates/UnrealTemplate.h"
 
 /**
  *	SharedPointer - Unreal smart pointer library
