Make Core module public headers individually compilable. Part 14 of n.

diff --git a/Engine/Source/Runtime/Core/Public/Android/AndroidMemory.h b/Engine/Source/Runtime/Core/Public/Android/AndroidMemory.h
--- a/Engine/Source/Runtime/Core/Public/Android/AndroidMemory.h
+++ b/Engine/Source/Runtime/Core/Public/Android/AndroidMemory.h
@@ -6,6 +6,7 @@
 ==============================================================================================*/
 
 #pragma once
+#include "GenericPlatform/GenericPlatformMemory.h"
 
 /**
  *	Android implementation of the FGenericPlatformMemoryStats.
diff --git a/Engine/Source/Runtime/Core/Public/GenericPlatform/GenericPlatformMemory.h b/Engine/Source/Runtime/Core/Public/GenericPlatform/GenericPlatformMemory.h
--- a/Engine/Source/Runtime/Core/Public/GenericPlatform/GenericPlatformMemory.h
+++ b/Engine/Source/Runtime/Core/Public/GenericPlatform/GenericPlatformMemory.h
@@ -11,6 +11,8 @@
 #include <wchar.h>
 #include <string.h>
 
+class FMalloc;
+class FOutputDevice;
 class FString;
 
 /** Holds generic memory stats, internally implemented as a map. */
@@ -178,7 +180,7 @@ struct CORE_API FGenericPlatformMemory
 	/**
 	 * @return the default allocator.
 	 */
-	static class FMalloc* BaseAllocator();
+	static FMalloc* BaseAllocator();
 
 	/**
 	 * @return platform specific current memory statistics.
@@ -220,10 +222,10 @@ struct CORE_API FGenericPlatformMemory
 	static void BinnedFreeToOS( void* Ptr );
 
 	/** Dumps basic platform memory statistics into the specified output device. */
-	static void DumpStats( class FOutputDevice& Ar );
+	static void DumpStats( FOutputDevice& Ar );
 
 	/** Dumps basic platform memory statistics and allocator specific statistics into the specified output device. */
-	static void DumpPlatformAndAllocatorStats( class FOutputDevice& Ar );
+	static void DumpPlatformAndAllocatorStats( FOutputDevice& Ar );
 
 	/** @name Memory functions */
 
diff --git a/Engine/Source/Runtime/Core/Public/HAL/PlatformMemory.h b/Engine/Source/Runtime/Core/Public/HAL/PlatformMemory.h
new file mode 100644
--- /dev/null
+++ b/Engine/Source/Runtime/Core/Public/HAL/PlatformMemory.h
@@ -0,0 +1,24 @@
+// Copyright 1998-2014 Epic Games, Inc. All Rights Reserved.
+#pragma once
+
+#include "HAL/Platform.h"
+
+#if PLATFORM_WINDOWS
+#include "Windows/WindowsPlatformMemory.h"
+#elif PLATFORM_PS4
+#include "PS4/PS4Memory.h"
+#elif PLATFORM_XBOXONE
+#include "XboxOne/XboxOneMemory.h"
+#elif PLATFORM_MAC
+#include "Mac/MacPlatformMemory.h"
+#elif PLATFORM_IOS
+#include "IOS/IOSPlatformMemory.h"
+#elif PLATFORM_ANDROID
+#include "Android/AndroidMemory.h"
+#elif PLATFORM_WINRT
+#include "WinRT/WinRTMemory.h"
+#elif PLATFORM_HTML5
+#include "HTML5/HTML5PlatformMemory.h"
+#elif PLATFORM_LINUX
+#include "Linux/LinuxPlatformMemory.h"
+#endif
diff --git a/Engine/Source/Runtime/Core/Public/HTML5/HTML5PlatformMemory.h b/Engine/Source/Runtime/Core/Public/HTML5/HTML5PlatformMemory.h
--- a/Engine/Source/Runtime/Core/Public/HTML5/HTML5PlatformMemory.h
+++ b/Engine/Source/Runtime/Core/Public/HTML5/HTML5PlatformMemory.h
@@ -6,7 +6,7 @@
 ==============================================================================================*/
 
 #pragma once
-
+#include "GenericPlatform/GenericPlatformMemory.h"
 /**
  *	HTML5 implementation of the FGenericPlatformMemoryStats.
  */
diff --git a/Engine/Source/Runtime/Core/Public/IOS/IOSPlatformMemory.h b/Engine/Source/Runtime/Core/Public/IOS/IOSPlatformMemory.h
--- a/Engine/Source/Runtime/Core/Public/IOS/IOSPlatformMemory.h
+++ b/Engine/Source/Runtime/Core/Public/IOS/IOSPlatformMemory.h
@@ -6,6 +6,7 @@
 ==============================================================================================*/
 
 #pragma once
+#include "GenericPlatform/GenericPlatformMemory.h"
 
 /**
  *	IOS implementation of the FGenericPlatformMemoryStats.
diff --git a/Engine/Source/Runtime/Core/Public/Linux/LinuxPlatformMemory.h b/Engine/Source/Runtime/Core/Public/Linux/LinuxPlatformMemory.h
--- a/Engine/Source/Runtime/Core/Public/Linux/LinuxPlatformMemory.h
+++ b/Engine/Source/Runtime/Core/Public/Linux/LinuxPlatformMemory.h
@@ -6,7 +6,7 @@
 ==============================================================================================*/
 
 #pragma once
-
+#include "GenericPlatform/GenericPlatformMemory.h"
 /**
  *	Linux implementation of the FGenericPlatformMemoryStats.
  */
diff --git a/Engine/Source/Runtime/Core/Public/Mac/MacPlatformMemory.h b/Engine/Source/Runtime/Core/Public/Mac/MacPlatformMemory.h
--- a/Engine/Source/Runtime/Core/Public/Mac/MacPlatformMemory.h
+++ b/Engine/Source/Runtime/Core/Public/Mac/MacPlatformMemory.h
@@ -6,6 +6,7 @@
 ==============================================================================================*/
 
 #pragma once
+#include "GenericPlatform/GenericPlatformMemory.h"
 
 /**
  *	Max implementation of the FGenericPlatformMemoryStats.
diff --git a/Engine/Source/Runtime/Core/Public/WinRT/WinRTMemory.h b/Engine/Source/Runtime/Core/Public/WinRT/WinRTMemory.h
--- a/Engine/Source/Runtime/Core/Public/WinRT/WinRTMemory.h
+++ b/Engine/Source/Runtime/Core/Public/WinRT/WinRTMemory.h
@@ -6,6 +6,7 @@
 ==============================================================================================*/
 
 #pragma once
+#include "GenericPlatform/GenericPlatformMemory.h"
 
 /**
  *	WinRT implementation of the FGenericPlatformMemoryStats.
diff --git a/Engine/Source/Runtime/Core/Public/Windows/WindowsPlatformMemory.h b/Engine/Source/Runtime/Core/Public/Windows/WindowsPlatformMemory.h
--- a/Engine/Source/Runtime/Core/Public/Windows/WindowsPlatformMemory.h
+++ b/Engine/Source/Runtime/Core/Public/Windows/WindowsPlatformMemory.h
@@ -6,6 +6,8 @@
 ==============================================================================================*/
 
 #pragma once
+#include "GenericPlatform/GenericPlatformMemory.h"
+#include "Windows/WindowsSystemIncludes.h"
 
 /**
  *	Windows implementation of the FGenericPlatformMemoryStats.
