Make Core module public headers individually compilable. Part 15 of n.

diff --git a/Engine/Source/Runtime/Core/Public/Android/AndroidAffinity.h b/Engine/Source/Runtime/Core/Public/Android/AndroidAffinity.h
--- a/Engine/Source/Runtime/Core/Public/Android/AndroidAffinity.h
+++ b/Engine/Source/Runtime/Core/Public/Android/AndroidAffinity.h
@@ -6,7 +6,7 @@ AndroidAffinity.h: Android affinity profile masks definitions.
 
 #pragma once
 
-#include "GenericPlatformAffinity.h"
+#include "GenericPlatform/GenericPlatformAffinity.h"
 
 class FAndroidAffinity : public FGenericPlatformAffinity
 {
diff --git a/Engine/Source/Runtime/Core/Public/HAL/PlatformAffinity.h b/Engine/Source/Runtime/Core/Public/HAL/PlatformAffinity.h
new file mode 100644
--- /dev/null
+++ b/Engine/Source/Runtime/Core/Public/HAL/PlatformAffinity.h
@@ -0,0 +1,13 @@
+// Copyright 1998-2014 Epic Games, Inc. All Rights Reserved.
+
+#pragma once
+#include "HAL/Platform.h"
+
+#if PLATFORM_XBOXONE
+#include "XboxOne/XBoxOneAffinity.h"
+#elif PLATFORM_ANDROID
+#include "Android/AndroidAffinity.h"
+#else
+#include "GenericPlatform/GenericPlatformAffinity.h"
+typedef FGenericPlatformAffinity FPlatformAffinity;
+#endif
\ No newline at end of file
