Make Core module public headers individually compilable. Part 16 of n.

diff --git a/Engine/Source/Runtime/Core/Public/Android/AndroidAtomics.h b/Engine/Source/Runtime/Core/Public/Android/AndroidAtomics.h
--- a/Engine/Source/Runtime/Core/Public/Android/AndroidAtomics.h
+++ b/Engine/Source/Runtime/Core/Public/Android/AndroidAtomics.h
@@ -6,7 +6,8 @@
 ==============================================================================================*/
 
 #pragma once
-
+#include "GenericPlatform/GenericPlatformAtomics.h"
+#include "Android/AndroidSystemIncludes.h"
 
 /**
  * Android implementation of the Atomics OS functions.
diff --git a/Engine/Source/Runtime/Core/Public/Apple/ApplePlatformAtomics.h b/Engine/Source/Runtime/Core/Public/Apple/ApplePlatformAtomics.h
--- a/Engine/Source/Runtime/Core/Public/Apple/ApplePlatformAtomics.h
+++ b/Engine/Source/Runtime/Core/Public/Apple/ApplePlatformAtomics.h
@@ -5,7 +5,13 @@
 ==============================================================================================*/
 
 #pragma once
-
+#include "GenericPlatform/GenericPlatformAtomics.h"
+#include "HAL/Platform.h"
+#if PLATFORM_MAC
+#include "Mac/MacSystemIncludes.h"
+#elif PLATFORM_IOS
+#include "IOS/IOSSystemIncludes.h"
+#endif
 
 /**
  * Apple implementation of the Atomics OS functions
diff --git a/Engine/Source/Runtime/Core/Public/HAL/PlatformAtomics.h b/Engine/Source/Runtime/Core/Public/HAL/PlatformAtomics.h
new file mode 100644
--- /dev/null
+++ b/Engine/Source/Runtime/Core/Public/HAL/PlatformAtomics.h
@@ -0,0 +1,24 @@
+// Copyright 1998-2014 Epic Games, Inc. All Rights Reserved.
+
+#pragma once
+#include "HAL/Platform.h"
+
+#if PLATFORM_WINDOWS
+#include "Windows/WindowsPlatformAtomics.h"
+#elif PLATFORM_PS4
+#include "PS4/PS4Atomics.h"
+#elif PLATFORM_XBOXONE
+#include "XboxOne/XboxOneAtomics.h"
+#elif PLATFORM_MAC
+#include "Apple/ApplePlatformAtomics.h"
+#elif PLATFORM_IOS
+#include "Apple/ApplePlatformAtomics.h"
+#elif PLATFORM_ANDROID
+#include "Android/AndroidAtomics.h"
+#elif PLATFORM_WINRT
+#include "WinRT/WinRTAtomics.h"
+#elif PLATFORM_HTML5
+#include "HTML5/HTML5PlatformAtomics.h"
+#elif PLATFORM_LINUX
+#include "Linux/LinuxPlatformAtomics.h"
+#endif
\ No newline at end of file
diff --git a/Engine/Source/Runtime/Core/Public/HTML5/HTML5PlatformAtomics.h b/Engine/Source/Runtime/Core/Public/HTML5/HTML5PlatformAtomics.h
--- a/Engine/Source/Runtime/Core/Public/HTML5/HTML5PlatformAtomics.h
+++ b/Engine/Source/Runtime/Core/Public/HTML5/HTML5PlatformAtomics.h
@@ -6,7 +6,7 @@
 ==============================================================================================*/
 
 #pragma once
-
+#include "GenericPlatform/GenericPlatformAtomics.h"
 
 /**
  * HTML5 implementation of the Atomics OS functions (no actual atomics, as HTML5 has no threads)
diff --git a/Engine/Source/Runtime/Core/Public/Linux/LinuxPlatformAtomics.h b/Engine/Source/Runtime/Core/Public/Linux/LinuxPlatformAtomics.h
--- a/Engine/Source/Runtime/Core/Public/Linux/LinuxPlatformAtomics.h
+++ b/Engine/Source/Runtime/Core/Public/Linux/LinuxPlatformAtomics.h
@@ -6,7 +6,8 @@
 ==============================================================================================*/
 
 #pragma once
-
+#include "GenericPlatform/GenericPlatformAtomics.h"
+#include "Linux/LinuxSystemIncludes.h"
 
 /**
  * Android implementation of the Atomics OS functions.
diff --git a/Engine/Source/Runtime/Core/Public/WinRT/WinRTAtomics.h b/Engine/Source/Runtime/Core/Public/WinRT/WinRTAtomics.h
--- a/Engine/Source/Runtime/Core/Public/WinRT/WinRTAtomics.h
+++ b/Engine/Source/Runtime/Core/Public/WinRT/WinRTAtomics.h
@@ -5,6 +5,8 @@
 ==============================================================================================*/
 
 #pragma once
+#include "GenericPlatform/GenericPlatformAtomics.h"
+#include "WinRT/WinRTSystemIncludes.h"
 
 /**
  * WinRT implementation of the Atomics OS functions
diff --git a/Engine/Source/Runtime/Core/Public/Windows/WindowsPlatformAtomics.h b/Engine/Source/Runtime/Core/Public/Windows/WindowsPlatformAtomics.h
--- a/Engine/Source/Runtime/Core/Public/Windows/WindowsPlatformAtomics.h
+++ b/Engine/Source/Runtime/Core/Public/Windows/WindowsPlatformAtomics.h
@@ -7,6 +7,8 @@
 
 #pragma once
 
+#include "GenericPlatform/GenericPlatformAtomics.h"
+#include "Windows/WindowsSystemIncludes.h"
 
 /**
  * Windows implementation of the Atomics OS functions
