Make Core module public headers individually compilable. Part 17 of n.

diff --git a/Engine/Source/Runtime/Core/Public/Android/AndroidTime.h b/Engine/Source/Runtime/Core/Public/Android/AndroidTime.h
--- a/Engine/Source/Runtime/Core/Public/Android/AndroidTime.h
+++ b/Engine/Source/Runtime/Core/Public/Android/AndroidTime.h
@@ -6,7 +6,7 @@
 ==============================================================================================*/
 
 #pragma once
-
+#include "GenericPlatform/GenericPlatformTime.h"
 //@todo android: this entire file
 
 /**
diff --git a/Engine/Source/Runtime/Core/Public/Apple/ApplePlatformTime.h b/Engine/Source/Runtime/Core/Public/Apple/ApplePlatformTime.h
--- a/Engine/Source/Runtime/Core/Public/Apple/ApplePlatformTime.h
+++ b/Engine/Source/Runtime/Core/Public/Apple/ApplePlatformTime.h
@@ -5,6 +5,13 @@
 ==============================================================================================*/
 
 #pragma once
+#include "GenericPlatform/GenericPlatformTime.h"
+#include "HAL/Platform.h"
+#if PLATFORM_MAC
+#include "Mac/MacSystemIncludes.h"
+#elif PLATFORM_IOS
+#include "IOS/IOSSystemIncludes.h"
+#endif
 
 /**
  * Please see following UDN post about using rdtsc on processors that support
diff --git a/Engine/Source/Runtime/Core/Public/HAL/PlatformTime.h b/Engine/Source/Runtime/Core/Public/HAL/PlatformTime.h
new file mode 100644
--- /dev/null
+++ b/Engine/Source/Runtime/Core/Public/HAL/PlatformTime.h
@@ -0,0 +1,21 @@
+#include "HAL/Platform.h"
+
+#if PLATFORM_WINDOWS
+#include "Windows/WindowsPlatformTime.h"
+#elif PLATFORM_PS4
+#include "PS4/PS4Time.h"
+#elif PLATFORM_XBOXONE
+#include "XboxOne/XboxOneTime.h"
+#elif PLATFORM_MAC
+#include "Apple/ApplePlatformTime.h"
+#elif PLATFORM_IOS
+#include "Apple/ApplePlatformTime.h"
+#elif PLATFORM_ANDROID
+#include "Android/AndroidTime.h"
+#elif PLATFORM_WINRT
+#include "WinRT/WinRTTime.h"
+#elif PLATFORM_HTML5
+#include "HTML5/HTML5PlatformTime.h"
+#elif PLATFORM_LINUX
+#include "Linux/LinuxPlatformTime.h"
+#endif
\ No newline at end of file
diff --git a/Engine/Source/Runtime/Core/Public/HTML5/HTML5PlatformTime.h b/Engine/Source/Runtime/Core/Public/HTML5/HTML5PlatformTime.h
--- a/Engine/Source/Runtime/Core/Public/HTML5/HTML5PlatformTime.h
+++ b/Engine/Source/Runtime/Core/Public/HTML5/HTML5PlatformTime.h
@@ -6,6 +6,8 @@
 ==============================================================================================*/
 
 #pragma once
+#include "GenericPlatform/GenericPlatformTime.h"
+#include "HTML5/HTML5SystemIncludes.h"
 
 #if PLATFORM_HTML5_WIN32
 #include <time.h>
diff --git a/Engine/Source/Runtime/Core/Public/Linux/LinuxPlatformTime.h b/Engine/Source/Runtime/Core/Public/Linux/LinuxPlatformTime.h
--- a/Engine/Source/Runtime/Core/Public/Linux/LinuxPlatformTime.h
+++ b/Engine/Source/Runtime/Core/Public/Linux/LinuxPlatformTime.h
@@ -6,6 +6,8 @@
 ==============================================================================================*/
 
 #pragma once
+#include "GenericPlatform/GenericPlatformTime.h"
+#include "Linux/LinuxSystemIncludes.h"
 
 /**
  * Linux implementation of the Time OS functions
diff --git a/Engine/Source/Runtime/Core/Public/WinRT/WinRTTime.h b/Engine/Source/Runtime/Core/Public/WinRT/WinRTTime.h
--- a/Engine/Source/Runtime/Core/Public/WinRT/WinRTTime.h
+++ b/Engine/Source/Runtime/Core/Public/WinRT/WinRTTime.h
@@ -5,6 +5,8 @@
 ==============================================================================================*/
 
 #pragma once
+#include "GenericPlatform/GenericPlatformTime.h"
+#include "WinRT/WinRTSystemIncludes.h"
 
 /**
  * WinRT implementation of the Time OS functions
diff --git a/Engine/Source/Runtime/Core/Public/Windows/WindowsPlatformTime.h b/Engine/Source/Runtime/Core/Public/Windows/WindowsPlatformTime.h
--- a/Engine/Source/Runtime/Core/Public/Windows/WindowsPlatformTime.h
+++ b/Engine/Source/Runtime/Core/Public/Windows/WindowsPlatformTime.h
@@ -6,6 +6,8 @@
 ==============================================================================================*/
 
 #pragma once
+#include "GenericPlatform/GenericPlatformTime.h"
+#include "Windows/WindowsSystemIncludes.h"
 
 /**
  * Please see following UDN post about using rdtsc on processors that support
