Make Core module public headers individually compilable. Part 18 of n.

diff --git a/Engine/Source/Runtime/Core/Public/Android/AndroidTLS.h b/Engine/Source/Runtime/Core/Public/Android/AndroidTLS.h
--- a/Engine/Source/Runtime/Core/Public/Android/AndroidTLS.h
+++ b/Engine/Source/Runtime/Core/Public/Android/AndroidTLS.h
@@ -6,6 +6,8 @@
 ==============================================================================================*/
 
 #pragma once
+#include "GenericPlatform/GenericPlatformTLS.h"
+#include "Android/AndroidSystemIncludes.h"
 
 /**
  * Android implementation of the TLS OS functions
diff --git a/Engine/Source/Runtime/Core/Public/Apple/ApplePlatformTLS.h b/Engine/Source/Runtime/Core/Public/Apple/ApplePlatformTLS.h
--- a/Engine/Source/Runtime/Core/Public/Apple/ApplePlatformTLS.h
+++ b/Engine/Source/Runtime/Core/Public/Apple/ApplePlatformTLS.h
@@ -5,6 +5,13 @@
 ==============================================================================================*/
 
 #pragma once
+#include "GenericPlatform/GenericPlatformTLS.h"
+#include "HAL/Platform.h"
+#if PLATFORM_MAC
+#include "Mac/MacSystemIncludes.h"
+#elif PLATFORM_IOS
+#include "IOS/IOSSystemIncludes.h"
+#endif
 
 /**
 * Apple implementation of the TLS OS functions
diff --git a/Engine/Source/Runtime/Core/Public/HAL/PlatformTLS.h b/Engine/Source/Runtime/Core/Public/HAL/PlatformTLS.h
new file mode 100644
--- /dev/null
+++ b/Engine/Source/Runtime/Core/Public/HAL/PlatformTLS.h
@@ -0,0 +1,24 @@
+// Copyright 1998-2014 Epic Games, Inc. All Rights Reserved.
+
+#pragma once
+#include "HAL/Platform.h"
+
+#if PLATFORM_WINDOWS
+#include "Windows/WindowsPlatformTLS.h"
+#elif PLATFORM_PS4
+#include "PS4/PS4TLS.h"
+#elif PLATFORM_XBOXONE
+#include "XboxOne/XboxOneTLS.h"
+#elif PLATFORM_MAC
+#include "Apple/ApplePlatformTLS.h"
+#elif PLATFORM_IOS
+#include "Apple/ApplePlatformTLS.h"
+#elif PLATFORM_ANDROID
+#include "Android/AndroidTLS.h"
+#elif PLATFORM_WINRT
+#include "WinRT/WinRTTLS.h"
+#elif PLATFORM_HTML5
+#include "HTML5/HTML5PlatformTLS.h"
+#elif PLATFORM_LINUX
+#include "Linux/LinuxPlatformTLS.h"
+#endif
\ No newline at end of file
diff --git a/Engine/Source/Runtime/Core/Public/HTML5/HTML5PlatformTLS.h b/Engine/Source/Runtime/Core/Public/HTML5/HTML5PlatformTLS.h
--- a/Engine/Source/Runtime/Core/Public/HTML5/HTML5PlatformTLS.h
+++ b/Engine/Source/Runtime/Core/Public/HTML5/HTML5PlatformTLS.h
@@ -6,6 +6,7 @@
 ==============================================================================================*/
 
 #pragma once
+#include "GenericPlatform/GenericPlatformTLS.h"
 
 /**
  * HTML5 implementation of the TLS OS functions
diff --git a/Engine/Source/Runtime/Core/Public/Linux/LinuxPlatformTLS.h b/Engine/Source/Runtime/Core/Public/Linux/LinuxPlatformTLS.h
--- a/Engine/Source/Runtime/Core/Public/Linux/LinuxPlatformTLS.h
+++ b/Engine/Source/Runtime/Core/Public/Linux/LinuxPlatformTLS.h
@@ -6,6 +6,8 @@
 ==============================================================================================*/
 
 #pragma once
+#include "GenericPlatform/GenericPlatformTLS.h"
+#include "Linux/LinuxSystemIncludes.h"
 
 #if defined(_GNU_SOURCE)
 	#include <sys/syscall.h>	// SYS_gettid
diff --git a/Engine/Source/Runtime/Core/Public/WinRT/WinRTTLS.h b/Engine/Source/Runtime/Core/Public/WinRT/WinRTTLS.h
--- a/Engine/Source/Runtime/Core/Public/WinRT/WinRTTLS.h
+++ b/Engine/Source/Runtime/Core/Public/WinRT/WinRTTLS.h
@@ -5,7 +5,8 @@
 ==============================================================================================*/
 
 #pragma once
-
+#include "GenericPlatform/GenericPlatformTLS.h"
+#include "WinRT/WinRTSystemIncludes.h"
 [REDACTED_EMAIL]: Do NOT use the ThreadEmulation if this is ever checked in!!!!
 // Or if we do, it must have a TPS filed.
 // (Legal approved temporary usage for my Free Friday test)
diff --git a/Engine/Source/Runtime/Core/Public/Windows/WindowsPlatformTLS.h b/Engine/Source/Runtime/Core/Public/Windows/WindowsPlatformTLS.h
--- a/Engine/Source/Runtime/Core/Public/Windows/WindowsPlatformTLS.h
+++ b/Engine/Source/Runtime/Core/Public/Windows/WindowsPlatformTLS.h
@@ -6,6 +6,9 @@
 ==============================================================================================*/
 
 #pragma once
+#include "GenericPlatform/GenericPlatformTLS.h"
+#include "HAL/Platform.h"
+#include "Windows/WindowsSystemIncludes.h"
 
 /**
 * Windows implementation of the TLS OS functions
