Stats - Restored old type for the array of stat messages, working on proper fix

diff --git a/Engine/Source/Runtime/Core/Public/Stats/Stats2.h b/Engine/Source/Runtime/Core/Public/Stats/Stats2.h
--- a/Engine/Source/Runtime/Core/Public/Stats/Stats2.h
+++ b/Engine/Source/Runtime/Core/Public/Stats/Stats2.h
@@ -956,10 +956,11 @@ struct FStatPacket
 {
 	enum 
 	{
-		MESSAGES_CHUNK_SIZE = 64*1024,
+		MESSAGES_CHUNK_SIZE = 1024*1024,
 	};
 
-	typedef TChunkedArray<FStatMessage,MESSAGES_CHUNK_SIZE> TStatMessagesArray;
+	//typedef TChunkedArray<FStatMessage,MESSAGES_CHUNK_SIZE> TStatMessagesArray;
+	typedef TArray<FStatMessage> TStatMessagesArray;
 
 	/** Assigned later, this is the frame number this packet is for **/
 	int64 Frame;
