Removed code for handling modal windows in MacWindow and CocoaWindow. It's uneccessary now because we no longer allow OS X to activate windows on mouse down.

diff --git a/Engine/Source/Runtime/Core/Private/Mac/CocoaWindow.cpp b/Engine/Source/Runtime/Core/Private/Mac/CocoaWindow.cpp
--- a/Engine/Source/Runtime/Core/Private/Mac/CocoaWindow.cpp
+++ b/Engine/Source/Runtime/Core/Private/Mac/CocoaWindow.cpp
@@ -8,8 +8,6 @@
 #include "CocoaThread.h"
 #include "MacCursor.h"
 
-TArray< FCocoaWindow* > GRunningModalWindows;
-
 NSString* NSWindowRedrawContents = @"NSWindowRedrawContents";
 NSString* NSDraggingExited = @"NSDraggingExited";
 NSString* NSDraggingUpdated = @"NSDraggingUpdated";
@@ -148,15 +146,11 @@ NSString* NSPerformDragOperation = @"NSPerformDragOperation";
 - (void)orderWindow:(NSWindowOrderingMode)OrderingMode relativeTo:(NSInteger)OtherWindowNumber
 {
 	SCOPED_AUTORELEASE_POOL;
-	bool bModal = FMacWindow::CurrentModalWindow() == nil || FMacWindow::CurrentModalWindow() == self || [self styleMask] == NSBorderlessWindowMask;
-	if(OrderingMode == NSWindowOut || bModal)
+	if([self alphaValue] > 0.0f)
 	{
-		if([self alphaValue] > 0.0f)
-		{
-			[self performDeferredSetFrame];
-		}
-		[super orderWindow:OrderingMode relativeTo:OtherWindowNumber];
+		[self performDeferredSetFrame];
 	}
+	[super orderWindow:OrderingMode relativeTo:OtherWindowNumber];
 }
 
 - (bool)roundedCorners
@@ -205,12 +199,8 @@ NSString* NSPerformDragOperation = @"NSPerformDragOperation";
 	SCOPED_AUTORELEASE_POOL;
 	if ([NSApp isHidden] == NO)
 	{
-		bool bBringToFront = FMacWindow::CurrentModalWindow() == nil || FMacWindow::CurrentModalWindow() == self || [self styleMask] == NSBorderlessWindowMask;
-		if (bBringToFront)
-		{
-			[self orderFront:nil];
-		}
-		
+		[self orderFront:nil];
+
 		if (bMain && [self canBecomeMainWindow] && self != [NSApp mainWindow])
 		{
 			[self makeMainWindow];
@@ -226,8 +216,7 @@ NSString* NSPerformDragOperation = @"NSPerformDragOperation";
 - (BOOL)canBecomeMainWindow
 {
 	SCOPED_AUTORELEASE_POOL;
-	bool bNoModalOrCurrent = FMacWindow::CurrentModalWindow() == nil || FMacWindow::CurrentModalWindow() == self;
-	return bAcceptsInput && ([self styleMask] != NSBorderlessWindowMask) && bNoModalOrCurrent;
+	return bAcceptsInput && ([self styleMask] != NSBorderlessWindowMask);
 }
 
 - (BOOL)canBecomeKeyWindow
@@ -399,16 +388,9 @@ NSString* NSPerformDragOperation = @"NSPerformDragOperation";
 	SCOPED_AUTORELEASE_POOL;
 	if([NSApp isHidden] == NO)
 	{
-		if(FMacWindow::CurrentModalWindow() == nil || FMacWindow::CurrentModalWindow() == self || [self styleMask] == NSBorderlessWindowMask)
-		{
-			[self orderFrontAndMakeMain:false andKey:false];
-		}
-		else
-		{
-			[FMacWindow::CurrentModalWindow() orderFrontAndMakeMain:true andKey:true];
-		}
+		[self orderFrontAndMakeMain:false andKey:false];
 	}
-	
+
 	if(self.bForwardEvents)
 	{
 		FMacEvent::SendToGameRunLoop(Notification, self, EMacEventSendMethod::Async, @[ NSDefaultRunLoopMode, UE4ShowEventMode, UE4CloseEventMode, UE4FullscreenEventMode ]);
diff --git a/Engine/Source/Runtime/Core/Private/Mac/MacWindow.cpp b/Engine/Source/Runtime/Core/Private/Mac/MacWindow.cpp
--- a/Engine/Source/Runtime/Core/Private/Mac/MacWindow.cpp
+++ b/Engine/Source/Runtime/Core/Private/Mac/MacWindow.cpp
@@ -8,8 +8,6 @@
 #include "CocoaTextView.h"
 #include "CocoaThread.h"
 
-TArray< FCocoaWindow* > FMacWindow::RunningModalWindows;
-
 FMacWindow::~FMacWindow()
 {
 	// While on Windows invalid HWNDs fail silently, accessing an invalid NSWindow is fatal.
@@ -309,13 +307,9 @@ void FMacWindow::Destroy()
 	{
 		SCOPED_AUTORELEASE_POOL;
 		FCocoaWindow* Window = WindowHandle;
-		if(Definition->IsModalWindow)
-		{
-			RemoveModalWindow(Window);
-		}
-		
+
 		bIsClosed = true;
-		
+
 		if( MacApplication->OnWindowDestroyed( Window ) )
 		{
 			// This FMacWindow may have been destructed by now & so the WindowHandle will probably be invalid memory.
@@ -380,12 +374,7 @@ void FMacWindow::Show()
 	SCOPED_AUTORELEASE_POOL;
 	if (!bIsClosed && !bIsVisible)
 	{
-		if(Definition->IsModalWindow)
-		{
-			AddModalWindow(WindowHandle);
-		}
-		
-		bool bMakeMainAndKey = ([WindowHandle canBecomeKeyWindow] && Definition->ActivateWhenFirstShown);
+		const bool bMakeMainAndKey = ([WindowHandle canBecomeKeyWindow] && Definition->ActivateWhenFirstShown);
 		
 		MainThreadCall(^{
 			SCOPED_AUTORELEASE_POOL;
@@ -402,10 +391,6 @@ void FMacWindow::Hide()
 	{
 		SCOPED_AUTORELEASE_POOL;
 		bIsVisible = false;
-		if(Definition->IsModalWindow)
-		{
-			RemoveModalWindow(WindowHandle);
-		}
 		MainThreadCall(^{
 			SCOPED_AUTORELEASE_POOL;
 			[WindowHandle orderOut:nil];
@@ -633,26 +618,3 @@ bool FMacWindow::OnIMKKeyDown(NSEvent* Event)
 		return false;
 	}
 }
-
-void FMacWindow::AddModalWindow(FCocoaWindow* Window)
-{
-	if(!RunningModalWindows.Contains(Window))
-	{
-		RunningModalWindows.Add(Window);
-	}
-}
-
-void FMacWindow::RemoveModalWindow(FCocoaWindow* Window)
-{
-	RunningModalWindows.Remove(Window);
-}
-
-FCocoaWindow* FMacWindow::CurrentModalWindow(void)
-{
-	FCocoaWindow* Window = nil;
-	if(RunningModalWindows.Num() > 0)
-	{
-		Window = RunningModalWindows.Last();
-	}
-	return Window;
-}
diff --git a/Engine/Source/Runtime/Core/Public/Mac/MacWindow.h b/Engine/Source/Runtime/Core/Public/Mac/MacWindow.h
--- a/Engine/Source/Runtime/Core/Public/Mac/MacWindow.h
+++ b/Engine/Source/Runtime/Core/Public/Mac/MacWindow.h
@@ -24,10 +24,6 @@ public:
 	void OnDisplayReconfiguration(CGDirectDisplayID Display, CGDisplayChangeSummaryFlags Flags);
 	
 	bool OnIMKKeyDown(NSEvent* Event);
-	
-	static void AddModalWindow(FCocoaWindow* Window);
-	static void RemoveModalWindow(FCocoaWindow* Window);
-	static FCocoaWindow* CurrentModalWindow(void);
 
 public:
 
@@ -101,8 +97,6 @@ private:
 
 
 private:
-	/** Currently open modal windows. */
-	static TArray< FCocoaWindow* > RunningModalWindows;
 
 	FMacApplication* OwningApplication;
 
