Added support for GameStop version of Xbox 360 controller on Mac

diff --git a/Engine/Source/Runtime/Core/Private/Mac/HIDInputInterface.cpp b/Engine/Source/Runtime/Core/Private/Mac/HIDInputInterface.cpp
--- a/Engine/Source/Runtime/Core/Private/Mac/HIDInputInterface.cpp
+++ b/Engine/Source/Runtime/Core/Private/Mac/HIDInputInterface.cpp
@@ -153,7 +153,8 @@ void HIDInputInterface::FHIDDeviceInfo::SetupMappings()
 		RightAnalogYMapping = kHIDUsage_GD_Rz;
 		RightTriggerAnalogMapping = kHIDUsage_GD_Ry;
 	}
-	else if (VendorID == 0x45e && (ProductID == 0x28e || ProductID == 0x719))
+	else if ((VendorID == 0x45e && (ProductID == 0x28e || ProductID == 0x719)) // Original Microsoft controller
+			 || (VendorID == 0xe6f && ProductID == 0x401)) // GameStop version
 	{
 		// Xbox 360 Controller
 		ButtonsMapping[1]	= 0;	// A
