#ue4 - Fixed FVector::Parallel() and added FVector::Orthogonal().

- Fix FVector::Parallel() reporting false for parallel vectors pointing in opposite directions. Use FVector::Coincident() if you care about the direction.
- FVector::Coplanar() takes a threshold for controlling the parallel test.
- Add FVector::Orthogonal() for determining if two normals are orthogonal / perpendicular within a tolerance.

- Add THRESH_NORMALS_ARE_ORTHOGONAL for orthogonality test.
- Rename THRESH_VECTORS_ARE_PARALLEL to THRESH_NORMALS_ARE_PARALLEL, and change the meaning to the cosine of the angle we compare to.
   - old tolerance was very fuzzy (11.4 degrees), new tolerance is about 1 degree.

- Fix up engine code using the above.

diff --git a/Engine/Source/Editor/UnrealEd/Private/EditorCsg.cpp b/Engine/Source/Editor/UnrealEd/Private/EditorCsg.cpp
--- a/Engine/Source/Editor/UnrealEd/Private/EditorCsg.cpp
+++ b/Engine/Source/Editor/UnrealEd/Private/EditorCsg.cpp
@@ -794,7 +794,9 @@ void UEditorEngine::polySelectCoplanars(UWorld* InWorld, UModel* InModel)
 					const FVector Base = Model->Points[Model->Verts[Node.iVertPool].pVertex];
 					const FVector Normal = Model->Vectors[Surf.vNormal];
 
-					if(FVector::Coplanar(SelectedBase,SelectedNormal,Base,Normal) && !(Surf.PolyFlags & PF_Selected))
+					const float ParallelDotThreshold = 0.98f; // roughly 11.4 degrees (!), but this is the long-standing behavior.
+					if (FVector::Coincident(SelectedNormal, Normal, ParallelDotThreshold) &&
+						FVector::Coplanar(SelectedBase, SelectedNormal, Base, Normal, ParallelDotThreshold) && !(Surf.PolyFlags & PF_Selected))
 					{
 						Flags2[Node.iSurf] = 1;
 					}
diff --git a/Engine/Source/Editor/UnrealEd/Private/Factories/ActorFactory.cpp b/Engine/Source/Editor/UnrealEd/Private/Factories/ActorFactory.cpp
--- a/Engine/Source/Editor/UnrealEd/Private/Factories/ActorFactory.cpp
+++ b/Engine/Source/Editor/UnrealEd/Private/Factories/ActorFactory.cpp
@@ -83,7 +83,8 @@ FQuat FindActorAlignmentRotation(const FQuat& InActorRotation, const FVector& In
 		// Transform the model axis with this new pitch rotation to see if there is any need for yaw
 		TransformedModelAxis = (InActorRotation * DeltaRotation).RotateVector(InModelAxis);
 
-		if (!FVector::Parallel(InWorldNormal, TransformedModelAxis))
+		const float ParallelDotThreshold = 0.98f; // roughly 11.4 degrees (!)
+		if (!FVector::Coincident(InWorldNormal, TransformedModelAxis, ParallelDotThreshold))
 		{
 			const float Yaw = FMath::Atan2(InWorldNormal.X, InWorldNormal.Y) - FMath::Atan2(TransformedModelAxis.X, TransformedModelAxis.Y);
 
diff --git a/Engine/Source/Runtime/Core/Public/Math/UnrealMathUtility.h b/Engine/Source/Runtime/Core/Public/Math/UnrealMathUtility.h
--- a/Engine/Source/Runtime/Core/Public/Math/UnrealMathUtility.h
+++ b/Engine/Source/Runtime/Core/Public/Math/UnrealMathUtility.h
@@ -74,7 +74,8 @@ struct FLinearColor;
 #define THRESH_SPLIT_POLY_WITH_PLANE	(0.25f)		/* A plane splits a polygon in half */
 #define THRESH_SPLIT_POLY_PRECISELY		(0.01f)		/* A plane exactly splits a polygon */
 #define THRESH_ZERO_NORM_SQUARED		(0.0001f)	/* Size of a unit normal that is considered "zero", squared */
-#define THRESH_VECTORS_ARE_PARALLEL		(0.02f)		/* Vectors are parallel if dot product varies less than this */
+#define THRESH_NORMALS_ARE_PARALLEL		(0.999845f)	/* Two unit vectors are parallel if abs(A dot B) is greater than or equal to this. This is roughly cosine(1.0 degrees). */
+#define THRESH_NORMALS_ARE_ORTHOGONAL	(0.017455f)	/* Two unit vectors are orthogonal (perpendicular) if abs(A dot B) is less than or equal this. This is roughly cosine(89.0 degrees). */
 
 #define THRESH_VECTOR_NORMALIZED		(0.01f)		/** Allowed error for a normalized vector (against squared magnitude) */
 #define THRESH_QUAT_NORMALIZED			(0.01f)		/** Allowed error for a normalized quaternion (against squared magnitude) */
diff --git a/Engine/Source/Runtime/Core/Public/Math/Vector.h b/Engine/Source/Runtime/Core/Public/Math/Vector.h
--- a/Engine/Source/Runtime/Core/Public/Math/Vector.h
+++ b/Engine/Source/Runtime/Core/Public/Math/Vector.h
@@ -790,35 +790,46 @@ public:
 	static FORCEINLINE float BoxPushOut( const FVector& Normal, const FVector& Size );
 
 	/**
-	 * See if two normal vectors (or plane normals) are nearly parallel.
+	 * See if two normal vectors are nearly parallel, meaning the angle between them is close to 0 degrees.
 	 *
-	 * @param Normal1 First normalized vector.
-	 * @param Normal1 Second normalized vector.
-	 * @param  DotTolerance Vectors are parallel if dot product varies less than this.
+	 * @param  Normal1 First normalized vector.
+	 * @param  Normal1 Second normalized vector.
+	 * @param  ParallelCosineThreshold Normals are parallel if absolute value of dot product (cosine of angle between them) is greater than or equal to this. For example: cos(1.0 degrees).
 	 * @return true if vectors are nearly parallel, false otherwise.
 	 */
-	static bool Parallel(const FVector& Normal1, const FVector& Normal2, float DotTolerance = THRESH_VECTORS_ARE_PARALLEL);
+	static bool Parallel(const FVector& Normal1, const FVector& Normal2, float ParallelCosineThreshold = THRESH_NORMALS_ARE_PARALLEL);
 
 	/**
-	 * See if two normal vectors (or plane normals) are coincident (nearly parallel and point in the same direction).
+	 * See if two normal vectors are coincident (nearly parallel and point in the same direction).
 	 * 
 	 * @param  Normal1 First normalized vector.
 	 * @param  Normal2 Second normalized vector.
-	 * @param  DotTolerance Vectors are parallel if dot product varies less than this.
-	 * @return true if vectors are coincident (nearly parallel and point in the same direction).
+	 * @param  ParallelCosineThreshold Normals are coincident if dot product (cosine of angle between them) is greater than or equal to this. For example: cos(1.0 degrees).
+	 * @return true if vectors are coincident (nearly parallel and point in the same direction), false otherwise.
 	 */
-	static bool Coincident(const FVector& Normal1, const FVector& Normal2, float DotTolerance = THRESH_VECTORS_ARE_PARALLEL);
+	static bool Coincident(const FVector& Normal1, const FVector& Normal2, float ParallelCosineThreshold = THRESH_NORMALS_ARE_PARALLEL);
 
 	/**
-	 * See if two planes are coplanar.
+	 * See if two normal vectors are nearly orthogonal (perpendicular), meaning the angle between them is close to 90 degrees.
+	 * 
+	 * @param  Normal1 First normalized vector.
+	 * @param  Normal2 Second normalized vector.
+	 * @param  OrthogonalCosineThreshold Normals are orthogonal if absolute value of dot product (cosine of angle between them) is less than or equal to this. For example: cos(89.0 degrees).
+	 * @return true if vectors are orthogonal (perpendicular), false otherwise.
+	 */
+	static bool Orthogonal(const FVector& Normal1, const FVector& Normal2, float OrthogonalCosineThreshold = THRESH_NORMALS_ARE_ORTHOGONAL);
+
+	/**
+	 * See if two planes are coplanar. They are coplanar if the normals are nearly parallel and the planes include the same set of points.
 	 *
 	 * @param Base1 The base point in the first plane.
 	 * @param Normal1 The normal of the first plane.
 	 * @param Base2 The base point in the second plane.
 	 * @param Normal2 The normal of the second plane.
+	 * @param ParallelCosineThreshold Normals are parallel if absolute value of dot product is greater than or equal to this.
 	 * @return true if the planes are coplanar, false otherwise.
 	 */
-	static bool Coplanar( const FVector &Base1, const FVector &Normal1, const FVector &Base2, const FVector &Normal2 );
+	static bool Coplanar(const FVector& Base1, const FVector& Normal1, const FVector& Base2, const FVector& Normal2, float ParallelCosineThreshold = THRESH_NORMALS_ARE_PARALLEL);
 
 	/**
 	 * Triple product of three vectors: X dot (Y cross Z).
@@ -1058,23 +1069,29 @@ inline FVector FVector::VectorPlaneProject(const FVector& V, const FVector& Plan
 	return V - V.ProjectOnToNormal(PlaneNormal);
 }
 
-inline bool FVector::Parallel(const FVector &Normal1, const FVector &Normal2, float DotTolerance)
+inline bool FVector::Parallel(const FVector& Normal1, const FVector& Normal2, float ParallelCosineThreshold)
+{
+	const float NormalDot = Normal1 | Normal2;
+	return FMath::Abs(NormalDot) >= ParallelCosineThreshold;
+}
+
+inline bool FVector::Coincident(const FVector& Normal1, const FVector& Normal2, float ParallelCosineThreshold)
 {
 	const float NormalDot = Normal1 | Normal2;
-	return (FMath::Abs(NormalDot - 1.f) <= DotTolerance);
+	return NormalDot >= ParallelCosineThreshold;
 }
 
-inline bool FVector::Coincident(const FVector &Normal1, const FVector &Normal2, float DotTolerance)
+inline bool FVector::Orthogonal(const FVector& Normal1, const FVector& Normal2, float OrthogonalCosineThreshold)
 {
 	const float NormalDot = Normal1 | Normal2;
-	return (1.f - NormalDot) <= DotTolerance;
+	return FMath::Abs(NormalDot) <= OrthogonalCosineThreshold;
 }
 
-inline bool FVector::Coplanar( const FVector &Base1, const FVector &Normal1, const FVector &Base2, const FVector &Normal2 )
+inline bool FVector::Coplanar(const FVector &Base1, const FVector &Normal1, const FVector &Base2, const FVector &Normal2, float ParallelCosineThreshold)
 {
-	if      (!FVector::Parallel(Normal1,Normal2)) return false;
+	if      (!FVector::Parallel(Normal1,Normal2,ParallelCosineThreshold)) return false;
 	else if (FVector::PointPlaneDist (Base2,Base1,Normal1) > THRESH_POINT_ON_PLANE) return false;
-	else    return true;
+	else return true;
 }
 
 inline float FVector::Triple( const FVector& X, const FVector& Y, const FVector& Z )
diff --git a/Engine/Source/Runtime/Engine/Private/ProjectileMovementComponent.cpp b/Engine/Source/Runtime/Engine/Private/ProjectileMovementComponent.cpp
--- a/Engine/Source/Runtime/Engine/Private/ProjectileMovementComponent.cpp
+++ b/Engine/Source/Runtime/Engine/Private/ProjectileMovementComponent.cpp
@@ -173,7 +173,7 @@ void UProjectileMovementComponent::TickComponent(float DeltaTime, enum ELevelTic
 			
 			// if velocity still into wall (after HandleHitWall() had a chance to adjust), slide along wall
 			const float DotTolerance = 0.01f;
-			const bool bSliding = (bMultiHit && FVector::Coincident(OldHitNormal, Hit.Normal, DotTolerance)) ||
+			const bool bSliding = (bMultiHit && FVector::Coincident(OldHitNormal, Hit.Normal)) ||
 								  ((Velocity.SafeNormal() | Hit.Normal) <= DotTolerance);
 			
 			if (bSliding)
