If an internationalization archive fails to load during locres generation, an error log will be printed appropriately.

diff --git a/Engine/Source/Runtime/Core/Private/Internationalization/TextLocalizationResourceGenerator.cpp b/Engine/Source/Runtime/Core/Private/Internationalization/TextLocalizationResourceGenerator.cpp
--- a/Engine/Source/Runtime/Core/Private/Internationalization/TextLocalizationResourceGenerator.cpp
+++ b/Engine/Source/Runtime/Core/Private/Internationalization/TextLocalizationResourceGenerator.cpp
@@ -185,7 +185,11 @@ bool FTextLocalizationResourceGenerator::Generate(const FString& SourcePath, con
 			continue;
 		}
 
-		ArchiveSerializer.DeserializeArchive(ArchiveContent, InternationalizationArchive);
+		if (!ArchiveSerializer.DeserializeArchive(ArchiveContent, InternationalizationArchive))
+		{
+			UE_LOG(LogTextLocalizationResourceGenerator, Error, TEXT("Failed to serialize archive from file %s."), *ArchiveFilePath);
+			continue;
+		}
 #endif
 
 		// Generate text localization resource from manifest and archive entries.
