HTML5 crash handling & stack walking

diff --git a/Engine/Build/HTML5/Game.html.template b/Engine/Build/HTML5/Game.html.template
--- a/Engine/Build/HTML5/Game.html.template
+++ b/Engine/Build/HTML5/Game.html.template
@@ -113,6 +113,13 @@ var UE4 = {
     UE4.resize_game = fn;
     return fn;
   }
+  ,
+  get on_fatal() {
+    var fn = Module.cwrap('on_fatal', null, ['string', 'string'])
+    delete UE4["on_fatal"];
+    UE4.on_fatal = fn;
+    return fn;
+  }
 };
 
 function resizeListener(width, height) {
@@ -241,17 +248,18 @@ if (!isBrowser64Bit()) {
 // setup global error handling. 
 // make exceptions more visible. 
 window.onerror  = function(msg, url, line, column, error) {
-   if ( msg.indexOf("memory") > -1 ) { 
+  UE4.on_fatal(msg, error);
+  if ( msg.indexOf("memory") > -1 ) { 
 
-     var message = !isBrowser64Bit() ? " We are running on a 32 bit browser, please use a 64 bit browser to avoid memory constraints "
-            : " Looks like the game needs more than the allocated " + TOTAL_GAME_MEMORY + " bytes, please edit HTML5Engine.ini and repackage "; 
+   var message = !isBrowser64Bit() ? " We are running on a 32 bit browser, please use a 64 bit browser to avoid memory constraints "
+          : " Looks like the game needs more than the allocated " + TOTAL_GAME_MEMORY + " bytes, please edit HTML5Engine.ini and repackage "; 
 
-     alert ( "We ran out of memory : " + message );
-   }
-   else
-   {
-     alert ( "Error: " + msg );
-   }
+   alert ( "We ran out of memory : " + message );
+  }
+  else
+  {
+   alert ( "Error: " + msg );
+  }
 };
 
 var Module;
diff --git a/Engine/Source/Programs/UnrealBuildTool/HTML5/HTML5ToolChain.cs b/Engine/Source/Programs/UnrealBuildTool/HTML5/HTML5ToolChain.cs
--- a/Engine/Source/Programs/UnrealBuildTool/HTML5/HTML5ToolChain.cs
+++ b/Engine/Source/Programs/UnrealBuildTool/HTML5/HTML5ToolChain.cs
@@ -133,7 +133,7 @@ namespace UnrealBuildTool
             // we want full ES2
             Result += " -s FULL_ES2=1 ";
             // export console command handler. Export main func too because default exports ( e.g Main ) are overridden if we use custom exported functions. 
-            Result += " -s EXPORTED_FUNCTIONS=\"['_main', '_resize_game']\" ";
+            Result += " -s EXPORTED_FUNCTIONS=\"['_main', '_resize_game', '_on_fatal']\" ";
 
             // NOTE: This may slow down the compiler's startup time!
             { 
diff --git a/Engine/Source/Runtime/Core/Private/HTML5/HTML5PlatformMisc.cpp b/Engine/Source/Runtime/Core/Private/HTML5/HTML5PlatformMisc.cpp
--- a/Engine/Source/Runtime/Core/Private/HTML5/HTML5PlatformMisc.cpp
+++ b/Engine/Source/Runtime/Core/Private/HTML5/HTML5PlatformMisc.cpp
@@ -170,3 +170,26 @@ EAppReturnType::Type FHTML5Misc::MessageBoxExt(EAppMsgType::Type MsgType, const
 #endif 
 
 }
+
+void (*GHTML5CrashHandler)(const FGenericCrashContext& Context) = nullptr;
+
+extern "C"
+{
+#if PLATFORM_HTML5_BROWSER
+	// callback from javascript. 
+	void on_fatal(const char* msg, const char* error)
+	{
+		// !!JM todo: pass msg & error to a crash context? Must be copied?
+		if (GHTML5CrashHandler)
+		{
+			FGenericCrashContext Ctx;
+			GHTML5CrashHandler(Ctx);
+		}
+	}
+#endif 
+}
+
+void FHTML5Misc::SetCrashHandler(void(* CrashHandler)(const FGenericCrashContext& Context))
+{
+	GHTML5CrashHandler = CrashHandler;
+}
diff --git a/Engine/Source/Runtime/Core/Private/HTML5/HTML5PlatformStackWalk.cpp b/Engine/Source/Runtime/Core/Private/HTML5/HTML5PlatformStackWalk.cpp
new file mode 100644
--- /dev/null
+++ b/Engine/Source/Runtime/Core/Private/HTML5/HTML5PlatformStackWalk.cpp
@@ -0,0 +1,136 @@
+// Copyright 1998-2014 Epic Games, Inc. All Rights Reserved.
+
+#include "Core.h"
+#if PLATFORM_HTML5_BROWSER
+#	include <emscripten.h>
+#endif
+#include <string.h>
+
+static char BacktraceLog[4096];
+
+static void ParseError(FProgramCounterSymbolInfo& out_SymbolInfo) 
+{
+	///Fail with some defaults
+	out_SymbolInfo.LineNumber = 0;
+	strcpy(out_SymbolInfo.Filename,"???");
+	strcpy(out_SymbolInfo.FunctionName,"???");
+}
+
+void FHTML5PlatformStackWalk::ProgramCounterToSymbolInfo(uint64 ProgramCounter,FProgramCounterSymbolInfo& out_SymbolInfo)
+{
+#if PLATFORM_HTML5_BROWSER
+	char* TP = (char*)ProgramCounter;
+	// No module name, SymbolDisplacement, OffsetInModule or PC support
+	out_SymbolInfo.ModuleName[0] = 0; 
+	out_SymbolInfo.SymbolDisplacement = 0;
+	out_SymbolInfo.OffsetInModule = 0;
+	out_SymbolInfo.ProgramCounter = 0;
+	for (int Index = 0; Index < FProgramCounterSymbolInfo::MAX_NAME_LENGHT && out_SymbolInfo.FunctionName[Index] != 0; ++Index) 
+	{
+		out_SymbolInfo.FunctionName[Index] = 0;
+	}
+	for(int Index = 0; Index < FProgramCounterSymbolInfo::MAX_NAME_LENGHT && out_SymbolInfo.Filename[Index] != 0; ++Index)
+	{
+		out_SymbolInfo.Filename[Index] = 0;
+	}
+	// Parse string to get this info out...*sad face* format is "   at  $FunctionName($FunctionParameters) ($Filepath:$LineNumber:$Column)"
+	char* Ptr = strstr(TP, "at");
+	if (!Ptr)
+	{
+		return ParseError(out_SymbolInfo);
+	}
+	//skip at and any whitespace.
+	Ptr+=2;
+	while (isspace(*Ptr))
+	{
+		if (*Ptr == 0)
+		{
+			return ParseError(out_SymbolInfo);
+		}
+		++Ptr;
+	}
+	//Now assume that any non-whitespace is function symbol name, this may include a mangled js name
+	int32 Out=0;
+	bool ParameterParse=false;
+	while ((!isspace(*Ptr) || ParameterParse))
+	{
+		if(*Ptr == 0 || Out >= FProgramCounterSymbolInfo::MAX_NAME_LENGHT)
+		{
+			return ParseError(out_SymbolInfo);
+		}
+		if (*Ptr == '(') 
+		{
+			ParameterParse = true;
+		}
+		if(*Ptr == ')')
+		{
+			ParameterParse = false;
+		}
+		out_SymbolInfo.FunctionName[Out++] = *Ptr;
+		++Ptr;
+	}
+	//skip whitespace
+	while(isspace(*Ptr))
+	{
+		if(*Ptr == 0)
+		{
+			return ParseError(out_SymbolInfo);
+		}
+		++Ptr;
+	}
+	// filename is surrounded by ()
+	char* Colon1 = nullptr;
+	char* Colon2 = nullptr;
+	Out=0;
+	while(*Ptr != '\n' && *Ptr != '\r')
+	{
+		if(*Ptr == 0 || Out >= FProgramCounterSymbolInfo::MAX_NAME_LENGHT)
+		{
+			return ParseError(out_SymbolInfo);
+		}
+		if (*Ptr == '(' || *Ptr == ')')
+		{
+			++Ptr;
+			continue;
+		}
+		if (*Ptr == ':')
+		{// It's possible for the file path to contain : characters (e.g. http://) we only care about the last two.
+			Colon1 = Colon2;
+			Colon2 = out_SymbolInfo.Filename+Out;
+		}
+		out_SymbolInfo.Filename[Out++] = *Ptr;
+		++Ptr;
+	}
+	if (Colon1 && Colon2) 
+	{
+		*Colon2 = 0;
+		out_SymbolInfo.LineNumber = atoi(Colon1+1);
+		*Colon1 = 0;
+	}
+#else
+	return ParseError(out_SymbolInfo);
+#endif
+}
+
+void FHTML5PlatformStackWalk::CaptureStackBackTrace(uint64* BackTrace,uint32 MaxDepth,void* Context)
+{
+#if PLATFORM_HTML5_BROWSER
+	if (MaxDepth < 1) 
+	{
+		return;
+	}
+
+	--MaxDepth;
+	emscripten_get_callstack(EM_LOG_C_STACK|EM_LOG_DEMANGLE, BacktraceLog, sizeof(BacktraceLog));
+
+	int32 SP = 0;
+	for (char* P = strchr(BacktraceLog, '\n'); P != nullptr && SP < MaxDepth; P = strchr(P+1, '\n'))
+	{
+		BackTrace[SP++] = (uint64)(P+1);
+	}
+
+	BackTrace[SP] = 0;
+#else
+	BackTrace[0] = 0;
+#endif
+}
\ No newline at end of file
diff --git a/Engine/Source/Runtime/Core/Public/HTML5/HTML5PlatformMisc.h b/Engine/Source/Runtime/Core/Public/HTML5/HTML5PlatformMisc.h
--- a/Engine/Source/Runtime/Core/Public/HTML5/HTML5PlatformMisc.h
+++ b/Engine/Source/Runtime/Core/Public/HTML5/HTML5PlatformMisc.h
@@ -21,6 +21,7 @@ struct CORE_API FHTML5Misc : public FGenericPlatformMisc
 	static uint32 GetKeyMap(uint16* KeyCodes, FString* KeyNames, uint32 MaxMappings);
 	static uint32 GetCharKeyMap(uint16* KeyCodes, FString* KeyNames, uint32 MaxMappings);
 	static FString GetDefaultLocale();
+	static void SetCrashHandler(void (* CrashHandler)(const FGenericCrashContext& Context));
 	static EAppReturnType::Type MessageBoxExt( EAppMsgType::Type MsgType, const TCHAR* Text, const TCHAR* Caption );
 	FORCEINLINE static int32 NumberOfCores()
 	{
diff --git a/Engine/Source/Runtime/Core/Public/HTML5/HTML5PlatformStackWalk.h b/Engine/Source/Runtime/Core/Public/HTML5/HTML5PlatformStackWalk.h
--- a/Engine/Source/Runtime/Core/Public/HTML5/HTML5PlatformStackWalk.h
+++ b/Engine/Source/Runtime/Core/Public/HTML5/HTML5PlatformStackWalk.h
@@ -7,4 +7,15 @@
 
 #pragma once
 
-typedef FGenericPlatformStackWalk FPlatformStackWalk;
+/**
+* Android platform stack walking
+*/
+struct CORE_API FHTML5PlatformStackWalk : public FGenericPlatformStackWalk
+{
+	typedef FGenericPlatformStackWalk Parent;
+
+	static void ProgramCounterToSymbolInfo(uint64 ProgramCounter,FProgramCounterSymbolInfo& out_SymbolInfo);
+	static void CaptureStackBackTrace(uint64* BackTrace,uint32 MaxDepth,void* Context = NULL);
+};
+
+typedef FHTML5PlatformStackWalk FPlatformStackWalk;
