Make Core module public headers individually compilable. Part 19 of n.

diff --git a/Engine/Source/Runtime/Core/Public/Android/AndroidPlatformStackWalk.h b/Engine/Source/Runtime/Core/Public/Android/AndroidPlatformStackWalk.h
--- a/Engine/Source/Runtime/Core/Public/Android/AndroidPlatformStackWalk.h
+++ b/Engine/Source/Runtime/Core/Public/Android/AndroidPlatformStackWalk.h
@@ -6,6 +6,7 @@
 ==============================================================================================*/
 
 #pragma once
+#include "GenericPlatform/GenericPlatformStackWalk.h"
 
 #undef PLATFORM_SUPPORTS_STACK_SYMBOLS
 #define PLATFORM_SUPPORTS_STACK_SYMBOLS 1
@@ -17,9 +18,9 @@ struct CORE_API FAndroidPlatformStackWalk : public FGenericPlatformStackWalk
 {
 	typedef FGenericPlatformStackWalk Parent;
 
-	static bool ProgramCounterToHumanReadableString(int32 CurrentCallDepth, uint64 ProgramCounter, ANSICHAR* HumanReadableString, SIZE_T HumanReadableStringSize, FGenericCrashContext* Context = NULL);
+	static bool ProgramCounterToHumanReadableString(int32 CurrentCallDepth, uint64 ProgramCounter, ANSICHAR* HumanReadableString, SIZE_T HumanReadableStringSize, FGenericCrashContext* Context = nullptr);
 	static void ProgramCounterToSymbolInfo(uint64 ProgramCounter, FProgramCounterSymbolInfo& out_SymbolInfo);
-	static void CaptureStackBackTrace(uint64* BackTrace, uint32 MaxDepth, void* Context = NULL);
+	static void CaptureStackBackTrace(uint64* BackTrace, uint32 MaxDepth, void* Context = nullptr);
 };
 
 typedef FAndroidPlatformStackWalk FPlatformStackWalk;
diff --git a/Engine/Source/Runtime/Core/Public/Apple/ApplePlatformStackWalk.h b/Engine/Source/Runtime/Core/Public/Apple/ApplePlatformStackWalk.h
--- a/Engine/Source/Runtime/Core/Public/Apple/ApplePlatformStackWalk.h
+++ b/Engine/Source/Runtime/Core/Public/Apple/ApplePlatformStackWalk.h
@@ -5,6 +5,7 @@
 ==============================================================================================*/
 
 #pragma once
+#include "GenericPlatform/GenericPlatformStackWalk.h"
 
 #undef PLATFORM_SUPPORTS_STACK_SYMBOLS
 #define PLATFORM_SUPPORTS_STACK_SYMBOLS 1
@@ -15,9 +16,9 @@
  */
 struct CORE_API FApplePlatformStackWalk : public FGenericPlatformStackWalk
 {
-	static bool ProgramCounterToHumanReadableString( int32 CurrentCallDepth, uint64 ProgramCounter, ANSICHAR* HumanReadableString, SIZE_T HumanReadableStringSize, FGenericCrashContext* Context = NULL );
+	static bool ProgramCounterToHumanReadableString( int32 CurrentCallDepth, uint64 ProgramCounter, ANSICHAR* HumanReadableString, SIZE_T HumanReadableStringSize, FGenericCrashContext* Context = nullptr );
 	static void ProgramCounterToSymbolInfo( uint64 ProgramCounter, FProgramCounterSymbolInfo&  out_SymbolInfo);
-	static void CaptureStackBackTrace( uint64* BackTrace, uint32 MaxDepth, void* Context = NULL );
+	static void CaptureStackBackTrace( uint64* BackTrace, uint32 MaxDepth, void* Context = nullptr );
 	static int32 GetProcessModuleCount();
 	static int32 GetProcessModuleSignatures(FStackWalkModuleInfo *ModuleSignatures, const int32 ModuleSignaturesSize);
 };
diff --git a/Engine/Source/Runtime/Core/Public/HAL/PlatformStackWalk.h b/Engine/Source/Runtime/Core/Public/HAL/PlatformStackWalk.h
new file mode 100644
--- /dev/null
+++ b/Engine/Source/Runtime/Core/Public/HAL/PlatformStackWalk.h
@@ -0,0 +1,24 @@
+// Copyright 1998-2014 Epic Games, Inc. All Rights Reserved.
+#pragma once
+
+#include "HAL/Platform.h"
+
+#if PLATFORM_WINDOWS
+#include "Windows/WindowsPlatformStackWalk.h"
+#elif PLATFORM_PS4
+#include "PS4/PS4StackWalk.h"
+#elif PLATFORM_XBOXONE
+#include "XboxOne/XboxOneStackWalk.h"
+#elif PLATFORM_MAC
+#include "Apple/ApplePlatformStackWalk.h"
+#elif PLATFORM_IOS
+#include "Apple/ApplePlatformStackWalk.h"
+#elif PLATFORM_ANDROID
+#include "Android/AndroidPlatformStackWalk.h"
+#elif PLATFORM_WINRT
+#include "WinRT/WinRTStackWalk.h"
+#elif PLATFORM_HTML5
+#include "HTML5/HTML5PlatformStackWalk.h"
+#elif PLATFORM_LINUX
+#include "Linux/LinuxPlatformStackWalk.h"
+#endif
diff --git a/Engine/Source/Runtime/Core/Public/HTML5/HTML5PlatformStackWalk.h b/Engine/Source/Runtime/Core/Public/HTML5/HTML5PlatformStackWalk.h
--- a/Engine/Source/Runtime/Core/Public/HTML5/HTML5PlatformStackWalk.h
+++ b/Engine/Source/Runtime/Core/Public/HTML5/HTML5PlatformStackWalk.h
@@ -6,6 +6,7 @@
 ==============================================================================================*/
 
 #pragma once
+#include "GenericPlatform/GenericPlatformStackWalk.h"
 
 /**
 * Android platform stack walking
@@ -15,7 +16,7 @@ struct CORE_API FHTML5PlatformStackWalk : public FGenericPlatformStackWalk
 	typedef FGenericPlatformStackWalk Parent;
 
 	static void ProgramCounterToSymbolInfo(uint64 ProgramCounter,FProgramCounterSymbolInfo& out_SymbolInfo);
-	static void CaptureStackBackTrace(uint64* BackTrace,uint32 MaxDepth,void* Context = NULL);
+	static void CaptureStackBackTrace(uint64* BackTrace,uint32 MaxDepth,void* Context = nullptr);
 };
 
 typedef FHTML5PlatformStackWalk FPlatformStackWalk;
diff --git a/Engine/Source/Runtime/Core/Public/Linux/LinuxPlatformStackWalk.h b/Engine/Source/Runtime/Core/Public/Linux/LinuxPlatformStackWalk.h
--- a/Engine/Source/Runtime/Core/Public/Linux/LinuxPlatformStackWalk.h
+++ b/Engine/Source/Runtime/Core/Public/Linux/LinuxPlatformStackWalk.h
@@ -6,12 +6,13 @@
 ==============================================================================================*/
 
 #pragma once
+#include "GenericPlatform/GenericPlatformStackWalk.h"
 
 struct CORE_API FLinuxPlatformStackWalk : public FGenericPlatformStackWalk
 {
 	static void ProgramCounterToSymbolInfo( uint64 ProgramCounter, FProgramCounterSymbolInfo& out_SymbolInfo );
-	static bool ProgramCounterToHumanReadableString( int32 CurrentCallDepth, uint64 ProgramCounter, ANSICHAR* HumanReadableString, SIZE_T HumanReadableStringSize, FGenericCrashContext* Context = NULL );
-	static void CaptureStackBackTrace( uint64* BackTrace, uint32 MaxDepth, void* Context = NULL );
+	static bool ProgramCounterToHumanReadableString( int32 CurrentCallDepth, uint64 ProgramCounter, ANSICHAR* HumanReadableString, SIZE_T HumanReadableStringSize, FGenericCrashContext* Context = nullptr );
+	static void CaptureStackBackTrace( uint64* BackTrace, uint32 MaxDepth, void* Context = nullptr );
 };
 
 typedef FLinuxPlatformStackWalk FPlatformStackWalk;
diff --git a/Engine/Source/Runtime/Core/Public/WinRT/WinRTStackWalk.h b/Engine/Source/Runtime/Core/Public/WinRT/WinRTStackWalk.h
--- a/Engine/Source/Runtime/Core/Public/WinRT/WinRTStackWalk.h
+++ b/Engine/Source/Runtime/Core/Public/WinRT/WinRTStackWalk.h
@@ -5,5 +5,6 @@
 ==============================================================================================*/
 
 #pragma once
+#include "GenericPlatform/GenericPlatformStackWalk.h"
 
 typedef FGenericPlatformStackWalk FPlatformStackWalk;
diff --git a/Engine/Source/Runtime/Core/Public/Windows/WindowsPlatformStackWalk.h b/Engine/Source/Runtime/Core/Public/Windows/WindowsPlatformStackWalk.h
--- a/Engine/Source/Runtime/Core/Public/Windows/WindowsPlatformStackWalk.h
+++ b/Engine/Source/Runtime/Core/Public/Windows/WindowsPlatformStackWalk.h
@@ -6,6 +6,7 @@
 ==============================================================================================*/
 
 #pragma once
+#include "GenericPlatform/GenericPlatformStackWalk.h"
 
 #undef PLATFORM_SUPPORTS_STACK_SYMBOLS
 #define PLATFORM_SUPPORTS_STACK_SYMBOLS 1
@@ -19,8 +20,8 @@ struct CORE_API FWindowsPlatformStackWalk : public FGenericPlatformStackWalk
 	static bool InitStackWalking();
 	
 	static void ProgramCounterToSymbolInfo( uint64 ProgramCounter, FProgramCounterSymbolInfo& out_SymbolInfo );
-	static void CaptureStackBackTrace( uint64* BackTrace, uint32 MaxDepth, void* Context = NULL );
-	static void StackWalkAndDump( ANSICHAR* HumanReadableString, SIZE_T HumanReadableStringSize, int32 IgnoreCount, void* Context = NULL );
+	static void CaptureStackBackTrace( uint64* BackTrace, uint32 MaxDepth, void* Context = nullptr );
+	static void StackWalkAndDump( ANSICHAR* HumanReadableString, SIZE_T HumanReadableStringSize, int32 IgnoreCount, void* Context = nullptr );
 
 	static int32 GetProcessModuleCount();
 	static int32 GetProcessModuleSignatures(FStackWalkModuleInfo *ModuleSignatures, const int32 ModuleSignaturesSize);
