Make Core module public headers individually compilable. Part 20 of n.

diff --git a/Engine/Source/Runtime/Core/Public/HAL/PlatformNamedPipe.h b/Engine/Source/Runtime/Core/Public/HAL/PlatformNamedPipe.h
new file mode 100644
--- /dev/null
+++ b/Engine/Source/Runtime/Core/Public/HAL/PlatformNamedPipe.h
@@ -0,0 +1,10 @@
+// Copyright 1998-2014 Epic Games, Inc. All Rights Reserved.
+#pragma once
+
+#include "HAL/Platform.h"
+
+#if PLATFORM_WINDOWS
+#include "Windows/WindowsPlatformNamedPipe.h"
+#else
+#include "GenericPlatform/GenericPlatformNamedPipe.h"
+#endif
\ No newline at end of file
diff --git a/Engine/Source/Runtime/Core/Public/Windows/WindowsPlatformNamedPipe.h b/Engine/Source/Runtime/Core/Public/Windows/WindowsPlatformNamedPipe.h
--- a/Engine/Source/Runtime/Core/Public/Windows/WindowsPlatformNamedPipe.h
+++ b/Engine/Source/Runtime/Core/Public/Windows/WindowsPlatformNamedPipe.h
@@ -6,8 +6,14 @@
 ==============================================================================================*/
 
 #pragma once
+#include "HAL/Platform.h"
+#include "GenericPlatform/GenericPlatformNamedPipe.h"
+#include "Windows/WindowsSystemIncludes.h"
+
 #if PLATFORM_SUPPORTS_NAMED_PIPES
 
+class FString;
+
 // Windows wrapper for named pipe communications
 class CORE_API FWindowsPlatformNamedPipe : public FGenericPlatformNamedPipe
 {
