Make Core module public headers individually compilable. Part 21 of n.

diff --git a/Engine/Source/Runtime/Core/Public/Android/AndroidOutputDevices.h b/Engine/Source/Runtime/Core/Public/Android/AndroidOutputDevices.h
--- a/Engine/Source/Runtime/Core/Public/Android/AndroidOutputDevices.h
+++ b/Engine/Source/Runtime/Core/Public/Android/AndroidOutputDevices.h
@@ -6,10 +6,11 @@
 ==============================================================================================*/
 
 #pragma once
+#include "GenericPlatform/GenericPlatformOutputDevices.h"
+
 struct CORE_API FAndroidOutputDevices : public FGenericPlatformOutputDevices
 {
-//	static class FOutputDevice*			GetLog();
-	static class FOutputDeviceError*	GetError();
+	static FOutputDeviceError*	GetError();
 };
 
 typedef FAndroidOutputDevices FPlatformOutputDevices;
diff --git a/Engine/Source/Runtime/Core/Public/HAL/PlatformOutputDevices.h b/Engine/Source/Runtime/Core/Public/HAL/PlatformOutputDevices.h
new file mode 100644
--- /dev/null
+++ b/Engine/Source/Runtime/Core/Public/HAL/PlatformOutputDevices.h
@@ -0,0 +1,24 @@
+// Copyright 1998-2014 Epic Games, Inc. All Rights Reserved.
+#pragma once
+
+#include "HAL/Platform.h"
+
+#if PLATFORM_WINDOWS
+#include "Windows/WindowsPlatformOutputDevices.h"
+#elif PLATFORM_PS4
+#include "PS4/PS4OutputDevices.h"
+#elif PLATFORM_XBOXONE
+#include "XboxOne/XboxOneOutputDevices.h"
+#elif PLATFORM_MAC
+#include "Mac/MacPlatformOutputDevices.h"
+#elif PLATFORM_IOS
+#include "IOS/IOSPlatformOutputDevices.h"
+#elif PLATFORM_ANDROID
+#include "Android/AndroidOutputDevices.h"
+#elif PLATFORM_WINRT
+#include "WinRT/WinRTOutputDevices.h"
+#elif PLATFORM_HTML5
+#include "HTML5/HTML5PlatformOutputDevices.h"
+#elif PLATFORM_LINUX
+#include "Linux/LinuxPlatformOutputDevices.h"
+#endif
diff --git a/Engine/Source/Runtime/Core/Public/HTML5/HTML5PlatformOutputDevices.h b/Engine/Source/Runtime/Core/Public/HTML5/HTML5PlatformOutputDevices.h
--- a/Engine/Source/Runtime/Core/Public/HTML5/HTML5PlatformOutputDevices.h
+++ b/Engine/Source/Runtime/Core/Public/HTML5/HTML5PlatformOutputDevices.h
@@ -6,6 +6,7 @@
 ==============================================================================================*/
 
 #pragma once
+#include "GenericPlatform/GenericPlatformOutputDevices.h"
 
 // generic version
 typedef FGenericPlatformOutputDevices FPlatformOutputDevices;
diff --git a/Engine/Source/Runtime/Core/Public/IOS/IOSPlatformOutputDevices.h b/Engine/Source/Runtime/Core/Public/IOS/IOSPlatformOutputDevices.h
--- a/Engine/Source/Runtime/Core/Public/IOS/IOSPlatformOutputDevices.h
+++ b/Engine/Source/Runtime/Core/Public/IOS/IOSPlatformOutputDevices.h
@@ -5,10 +5,11 @@
 ==============================================================================================*/
 
 #pragma once
+#include "GenericPlatform/GenericPlatformOutputDevices.h"
 
 struct CORE_API FIOSPlatformOutputDevices : public FGenericPlatformOutputDevices
 {
-	static class FOutputDeviceError*	GetError();
+	static FOutputDeviceError*	GetError();
 };
 
 typedef FIOSPlatformOutputDevices FPlatformOutputDevices;
diff --git a/Engine/Source/Runtime/Core/Public/Linux/LinuxPlatformOutputDevices.h b/Engine/Source/Runtime/Core/Public/Linux/LinuxPlatformOutputDevices.h
--- a/Engine/Source/Runtime/Core/Public/Linux/LinuxPlatformOutputDevices.h
+++ b/Engine/Source/Runtime/Core/Public/Linux/LinuxPlatformOutputDevices.h
@@ -6,15 +6,16 @@
 ==============================================================================================*/
 
 #pragma once
+#include "GenericPlatform/GenericPlatformOutputDevices.h"
 
 struct CORE_API FLinuxOutputDevices : public FGenericPlatformOutputDevices
 {
 	static void							SetupOutputDevices();
 
-	static class FOutputDevice*			GetEventLog();
-	static class FOutputDeviceConsole*	GetLogConsole();
-	static class FOutputDeviceError*	GetError();
-	static class FFeedbackContext*		GetWarn();
+	static FOutputDevice*			GetEventLog();
+	static FOutputDeviceConsole*	GetLogConsole();
+	static FOutputDeviceError*		GetError();
+	static FFeedbackContext*		GetWarn();
 };
 
 typedef FLinuxOutputDevices FPlatformOutputDevices;
diff --git a/Engine/Source/Runtime/Core/Public/Mac/MacPlatformOutputDevices.h b/Engine/Source/Runtime/Core/Public/Mac/MacPlatformOutputDevices.h
--- a/Engine/Source/Runtime/Core/Public/Mac/MacPlatformOutputDevices.h
+++ b/Engine/Source/Runtime/Core/Public/Mac/MacPlatformOutputDevices.h
@@ -5,13 +5,14 @@
 ==============================================================================================*/
 
 #pragma once
+#include "GenericPlatform/GenericPlatformOutputDevices.h"
 
 struct CORE_API FMacPlatformOutputDevices : public FGenericPlatformOutputDevices
 {
-	static class FOutputDevice*			GetEventLog();
-	static class FOutputDeviceConsole*	GetLogConsole();
-	static class FOutputDeviceError*	GetError();
-	static class FFeedbackContext*		GetWarn();
+	static FOutputDevice*			GetEventLog();
+	static FOutputDeviceConsole*	GetLogConsole();
+	static FOutputDeviceError*		GetError();
+	static FFeedbackContext*		GetWarn();
 };
 
 typedef FMacPlatformOutputDevices FPlatformOutputDevices;
diff --git a/Engine/Source/Runtime/Core/Public/WinRT/WinRTOutputDevices.h b/Engine/Source/Runtime/Core/Public/WinRT/WinRTOutputDevices.h
--- a/Engine/Source/Runtime/Core/Public/WinRT/WinRTOutputDevices.h
+++ b/Engine/Source/Runtime/Core/Public/WinRT/WinRTOutputDevices.h
@@ -5,13 +5,14 @@
 ==============================================================================================*/
 
 #pragma once
+#include "GenericPlatform/GenericPlatformOutputDevices.h"
 
 struct CORE_API FWinRTOutputDevices : public FGenericPlatformOutputDevices
 {
-	static class FOutputDevice*			GetEventLog();
-	static class FOutputDeviceConsole*	GetLogConsole();
-	static class FOutputDeviceError*	GetError();
-	static class FFeedbackContext*		GetWarn();
+	static FOutputDevice*			GetEventLog();
+	static FOutputDeviceConsole*	GetLogConsole();
+	static FOutputDeviceError*		GetError();
+	static FFeedbackContext*		GetWarn();
 };
 
 typedef FWinRTOutputDevices FPlatformOutputDevices;
diff --git a/Engine/Source/Runtime/Core/Public/Windows/WindowsPlatformOutputDevices.h b/Engine/Source/Runtime/Core/Public/Windows/WindowsPlatformOutputDevices.h
--- a/Engine/Source/Runtime/Core/Public/Windows/WindowsPlatformOutputDevices.h
+++ b/Engine/Source/Runtime/Core/Public/Windows/WindowsPlatformOutputDevices.h
@@ -6,13 +6,14 @@
 ==============================================================================================*/
 
 #pragma once
+#include "GenericPlatform/GenericPlatformOutputDevices.h"
 
 struct CORE_API FWindowsPlatformOutputDevices : public FGenericPlatformOutputDevices
 {
-	static class FOutputDevice*			GetEventLog();
-	static class FOutputDeviceConsole*	GetLogConsole();
-	static class FOutputDeviceError*	GetError();
-	static class FFeedbackContext*		GetWarn();
+	static FOutputDevice*			GetEventLog();
+	static FOutputDeviceConsole*	GetLogConsole();
+	static FOutputDeviceError*		GetError();
+	static FFeedbackContext*		GetWarn();
 };
 
 typedef FWindowsPlatformOutputDevices FPlatformOutputDevices;
