Make Core module public headers individually compilable. Part 22 of n.

diff --git a/Engine/Source/Runtime/Core/Public/Android/AndroidSplash.h b/Engine/Source/Runtime/Core/Public/Android/AndroidSplash.h
--- a/Engine/Source/Runtime/Core/Public/Android/AndroidSplash.h
+++ b/Engine/Source/Runtime/Core/Public/Android/AndroidSplash.h
@@ -5,6 +5,7 @@
 ==============================================================================================*/
 
 #pragma once
+#include "GenericPlatform/GenericPlatformSplash.h"
 
 /**
  * Android splash implementation
diff --git a/Engine/Source/Runtime/Core/Public/HAL/PlatformSplash.h b/Engine/Source/Runtime/Core/Public/HAL/PlatformSplash.h
new file mode 100644
--- /dev/null
+++ b/Engine/Source/Runtime/Core/Public/HAL/PlatformSplash.h
@@ -0,0 +1,24 @@
+// Copyright 1998-2014 Epic Games, Inc. All Rights Reserved.
+#pragma once
+
+#include "HAL/Platform.h"
+
+#if PLATFORM_WINDOWS
+#include "Windows/WindowsPlatformSplash.h"
+#elif PLATFORM_PS4
+#include "PS4/PS4Splash.h"
+#elif PLATFORM_XBOXONE
+#include "XboxOne/XboxOneSplash.h"
+#elif PLATFORM_MAC
+#include "Mac/MacPlatformSplash.h"
+#elif PLATFORM_IOS
+#include "IOS/IOSPlatformSplash.h"
+#elif PLATFORM_ANDROID
+#include "Android/AndroidSplash.h"
+#elif PLATFORM_WINRT
+#include "WinRT/WinRTSplash.h"
+#elif PLATFORM_HTML5
+#include "HTML5/HTML5PlatformSplash.h"
+#elif PLATFORM_LINUX
+#include "Linux/LinuxPlatformSplash.h"
+#endif
diff --git a/Engine/Source/Runtime/Core/Public/HTML5/HTML5PlatformSplash.h b/Engine/Source/Runtime/Core/Public/HTML5/HTML5PlatformSplash.h
--- a/Engine/Source/Runtime/Core/Public/HTML5/HTML5PlatformSplash.h
+++ b/Engine/Source/Runtime/Core/Public/HTML5/HTML5PlatformSplash.h
@@ -5,6 +5,7 @@
 ==============================================================================================*/
 
 #pragma once
+#include "GenericPlatform/GenericPlatformSplash.h"
 
 // default implementation for now
 typedef FGenericPlatformSplash FPlatformSplash;
diff --git a/Engine/Source/Runtime/Core/Public/IOS/IOSPlatformSplash.h b/Engine/Source/Runtime/Core/Public/IOS/IOSPlatformSplash.h
--- a/Engine/Source/Runtime/Core/Public/IOS/IOSPlatformSplash.h
+++ b/Engine/Source/Runtime/Core/Public/IOS/IOSPlatformSplash.h
@@ -5,6 +5,7 @@
 ==============================================================================================*/
 
 #pragma once
+#include "GenericPlatform/GenericPlatformSplash.h"
 
 /**
 * IOS splash implementation 
diff --git a/Engine/Source/Runtime/Core/Public/Linux/LinuxPlatformSplash.h b/Engine/Source/Runtime/Core/Public/Linux/LinuxPlatformSplash.h
--- a/Engine/Source/Runtime/Core/Public/Linux/LinuxPlatformSplash.h
+++ b/Engine/Source/Runtime/Core/Public/Linux/LinuxPlatformSplash.h
@@ -5,6 +5,7 @@
 ==============================================================================================*/
 
 #pragma once
+#include "GenericPlatform/GenericPlatformSplash.h"
 
 /**
  * Linux splash implementation
diff --git a/Engine/Source/Runtime/Core/Public/Mac/MacPlatformSplash.h b/Engine/Source/Runtime/Core/Public/Mac/MacPlatformSplash.h
--- a/Engine/Source/Runtime/Core/Public/Mac/MacPlatformSplash.h
+++ b/Engine/Source/Runtime/Core/Public/Mac/MacPlatformSplash.h
@@ -4,6 +4,7 @@
 ==============================================================================================*/
 
 #pragma once
+#include "GenericPlatform/GenericPlatformSplash.h"
 
 /**
 * Mac splash implementation 
diff --git a/Engine/Source/Runtime/Core/Public/WinRT/WinRTSplash.h b/Engine/Source/Runtime/Core/Public/WinRT/WinRTSplash.h
--- a/Engine/Source/Runtime/Core/Public/WinRT/WinRTSplash.h
+++ b/Engine/Source/Runtime/Core/Public/WinRT/WinRTSplash.h
@@ -1,6 +1,7 @@
 // Copyright 1998-2014 Epic Games, Inc. All Rights Reserved.
 
 #pragma once
+#include "GenericPlatform/GenericPlatformSplash.h"
 
 /**
  * WinRT splash implementation
diff --git a/Engine/Source/Runtime/Core/Public/Windows/WindowsPlatformSplash.h b/Engine/Source/Runtime/Core/Public/Windows/WindowsPlatformSplash.h
--- a/Engine/Source/Runtime/Core/Public/Windows/WindowsPlatformSplash.h
+++ b/Engine/Source/Runtime/Core/Public/Windows/WindowsPlatformSplash.h
@@ -4,6 +4,7 @@
 ==============================================================================================*/
 
 #pragma once
+#include "GenericPlatform/GenericPlatformSplash.h"
 
 /**
 * Windows splash implementation
