Make Core module public headers individually compilable. Part 23 of n.

diff --git a/Engine/Source/Runtime/Core/Public/Apple/ApplePlatformFile.h b/Engine/Source/Runtime/Core/Public/Apple/ApplePlatformFile.h
--- a/Engine/Source/Runtime/Core/Public/Apple/ApplePlatformFile.h
+++ b/Engine/Source/Runtime/Core/Public/Apple/ApplePlatformFile.h
@@ -5,6 +5,7 @@
 ==============================================================================================*/
 
 #pragma once
+#include "GenericPlatform/GenericPlatformFile.h"
 
 /**
  * Mac File I/O implementation
@@ -37,5 +38,5 @@ public:
 	virtual bool DeleteDirectory(const TCHAR* Directory) override;
 	bool IterateDirectory(const TCHAR* Directory, FDirectoryVisitor& Visitor);
 private:
-	int32 Stat(const TCHAR* Filename, struct stat* OutFileInfo, ANSICHAR* OutFixedPath = NULL, SIZE_T FixedPathSize = 0);
+	int32 Stat(const TCHAR* Filename, struct stat* OutFileInfo, ANSICHAR* OutFixedPath = nullptr, SIZE_T FixedPathSize = 0);
 };
diff --git a/Engine/Source/Runtime/Core/Public/HAL/PlatformFile.h b/Engine/Source/Runtime/Core/Public/HAL/PlatformFile.h
new file mode 100644
--- /dev/null
+++ b/Engine/Source/Runtime/Core/Public/HAL/PlatformFile.h
@@ -0,0 +1,24 @@
+// Copyright 1998-2014 Epic Games, Inc. All Rights Reserved.
+#pragma once
+
+#include "HAL/Platform.h"
+
+#if PLATFORM_WINDOWS
+#include "Windows/WindowsPlatformFile.h"
+#elif PLATFORM_PS4
+#include "PS4/PS4File.h"
+#elif PLATFORM_XBOXONE
+#include "XboxOne/XboxOneFile.h"
+#elif PLATFORM_MAC
+#include "Mac/MacPlatformFile.h"
+#elif PLATFORM_IOS
+#include "IOS/IOSPlatformFile.h"
+#elif PLATFORM_ANDROID
+#include "Android/AndroidFile.h"
+#elif PLATFORM_WINRT
+#include "WinRT/WinRTFile.h"
+#elif PLATFORM_HTML5
+//#include "HTML5PlatformFile.h"
+#elif PLATFORM_LINUX
+#include "Linux/LinuxPlatformFile.h"
+#endif
diff --git a/Engine/Source/Runtime/Core/Public/Linux/LinuxPlatformFile.h b/Engine/Source/Runtime/Core/Public/Linux/LinuxPlatformFile.h
--- a/Engine/Source/Runtime/Core/Public/Linux/LinuxPlatformFile.h
+++ b/Engine/Source/Runtime/Core/Public/Linux/LinuxPlatformFile.h
@@ -5,6 +5,7 @@
 ==============================================================================================*/
 
 #pragma once
+#include "GenericPlatform/GenericPlatformFile.h"
 
 /**
  * Linux File I/O implementation
