Make Core module public headers individually compilable. Part 24 of n.

diff --git a/Engine/Source/Runtime/Core/Public/Android/AndroidSurvey.h b/Engine/Source/Runtime/Core/Public/Android/AndroidSurvey.h
--- a/Engine/Source/Runtime/Core/Public/Android/AndroidSurvey.h
+++ b/Engine/Source/Runtime/Core/Public/Android/AndroidSurvey.h
@@ -6,6 +6,7 @@
 ==============================================================================================*/
 
 #pragma once
+#include "GenericPlatform/GenericPlatformSurvey.h"
 
 /**
  * Android implementation of FGenericPlatformSurvey
diff --git a/Engine/Source/Runtime/Core/Public/HAL/PlatformSurvey.h b/Engine/Source/Runtime/Core/Public/HAL/PlatformSurvey.h
new file mode 100644
--- /dev/null
+++ b/Engine/Source/Runtime/Core/Public/HAL/PlatformSurvey.h
@@ -0,0 +1,24 @@
+// Copyright 1998-2014 Epic Games, Inc. All Rights Reserved.
+#pragma once
+
+#include "HAL/Platform.h"
+
+#if PLATFORM_WINDOWS
+#include "Windows/WindowsPlatformSurvey.h"
+#elif PLATFORM_PS4
+#include "PS4/PS4Survey.h"
+#elif PLATFORM_XBOXONE
+#include "XboxOne/XboxOneSurvey.h"
+#elif PLATFORM_MAC
+#include "Mac/MacPlatformSurvey.h"
+#elif PLATFORM_IOS
+#include "IOS/IOSPlatformSurvey.h"
+#elif PLATFORM_ANDROID
+#include "Android/AndroidSurvey.h"
+#elif PLATFORM_WINRT
+#include "WinRT/WinRTSurvey.h"
+#elif PLATFORM_HTML5
+#include "HTML5/HTML5PlatformSurvey.h"
+#elif PLATFORM_LINUX
+#include "Linux/LinuxPlatformSurvey.h"
+#endif
diff --git a/Engine/Source/Runtime/Core/Public/HTML5/HTML5PlatformSurvey.h b/Engine/Source/Runtime/Core/Public/HTML5/HTML5PlatformSurvey.h
--- a/Engine/Source/Runtime/Core/Public/HTML5/HTML5PlatformSurvey.h
+++ b/Engine/Source/Runtime/Core/Public/HTML5/HTML5PlatformSurvey.h
@@ -6,6 +6,7 @@
 ==============================================================================================*/
 
 #pragma once
+#include "GenericPlatform/GenericPlatformSurvey.h"
 
 // generic implementation
 typedef FGenericPlatformSurvey FPlatformSurvey;
diff --git a/Engine/Source/Runtime/Core/Public/IOS/IOSPlatformSurvey.h b/Engine/Source/Runtime/Core/Public/IOS/IOSPlatformSurvey.h
--- a/Engine/Source/Runtime/Core/Public/IOS/IOSPlatformSurvey.h
+++ b/Engine/Source/Runtime/Core/Public/IOS/IOSPlatformSurvey.h
@@ -6,6 +6,7 @@
 ==============================================================================================*/
 
 #pragma once
+#include "GenericPlatform/GenericPlatformSurvey.h"
 
 /**
 * iOS implementation of FGenericPlatformSurvey
diff --git a/Engine/Source/Runtime/Core/Public/Linux/LinuxPlatformSurvey.h b/Engine/Source/Runtime/Core/Public/Linux/LinuxPlatformSurvey.h
--- a/Engine/Source/Runtime/Core/Public/Linux/LinuxPlatformSurvey.h
+++ b/Engine/Source/Runtime/Core/Public/Linux/LinuxPlatformSurvey.h
@@ -6,6 +6,7 @@
 ==============================================================================================*/
 
 #pragma once
+#include "GenericPlatform/GenericPlatformSurvey.h"
 
 /**
 * Android implementation of FGenericPlatformSurvey
diff --git a/Engine/Source/Runtime/Core/Public/Mac/MacPlatformSurvey.h b/Engine/Source/Runtime/Core/Public/Mac/MacPlatformSurvey.h
--- a/Engine/Source/Runtime/Core/Public/Mac/MacPlatformSurvey.h
+++ b/Engine/Source/Runtime/Core/Public/Mac/MacPlatformSurvey.h
@@ -6,6 +6,7 @@
 ==============================================================================================*/
 
 #pragma once
+#include "GenericPlatform/GenericPlatformSurvey.h"
 
 /**
 * Mac implementation of FGenericPlatformSurvey
diff --git a/Engine/Source/Runtime/Core/Public/WinRT/WinRTSurvey.h b/Engine/Source/Runtime/Core/Public/WinRT/WinRTSurvey.h
--- a/Engine/Source/Runtime/Core/Public/WinRT/WinRTSurvey.h
+++ b/Engine/Source/Runtime/Core/Public/WinRT/WinRTSurvey.h
@@ -6,6 +6,7 @@
 ==============================================================================================*/
 
 #pragma once
+#include "GenericPlatform/GenericPlatformSurvey.h"
 
 /**
 * WinRT implementation of FGenericPlatformSurvey
diff --git a/Engine/Source/Runtime/Core/Public/Windows/WindowsPlatformSurvey.h b/Engine/Source/Runtime/Core/Public/Windows/WindowsPlatformSurvey.h
--- a/Engine/Source/Runtime/Core/Public/Windows/WindowsPlatformSurvey.h
+++ b/Engine/Source/Runtime/Core/Public/Windows/WindowsPlatformSurvey.h
@@ -6,6 +6,7 @@
 ==============================================================================================*/
 
 #pragma once
+#include "GenericPlatform/GenericPlatformSurvey.h"
 
 /**
 * Windows implementation of FGenericPlatformSurvey
