Make Core module public headers individually compilable. Part 25 of n.

diff --git a/Engine/Source/Runtime/Core/Public/Android/AndroidHttp.h b/Engine/Source/Runtime/Core/Public/Android/AndroidHttp.h
--- a/Engine/Source/Runtime/Core/Public/Android/AndroidHttp.h
+++ b/Engine/Source/Runtime/Core/Public/Android/AndroidHttp.h
@@ -1,6 +1,7 @@
 // Copyright 1998-2014 Epic Games, Inc. All Rights Reserved.
 
 #pragma once
+#include "GenericPlatform/GenericPlatformHttp.h"
 
 /**
  * Platform specific Http implementations
@@ -19,7 +20,7 @@ public:
 	 *
 	 * @return NULL if default implementation is to be used
 	 */
-	static class FHttpManager * CreatePlatformHttpManager();
+	static FHttpManager* CreatePlatformHttpManager();
 
 	/**
 	 * Platform shutdown step
@@ -31,7 +32,7 @@ public:
 	 *
 	 * @return request object
 	 */
-	static class IHttpRequest* ConstructRequest();
+	static IHttpRequest* ConstructRequest();
 };
 
 
diff --git a/Engine/Source/Runtime/Core/Public/HAL/PlatformHttp.h b/Engine/Source/Runtime/Core/Public/HAL/PlatformHttp.h
new file mode 100644
--- /dev/null
+++ b/Engine/Source/Runtime/Core/Public/HAL/PlatformHttp.h
@@ -0,0 +1,24 @@
+// Copyright 1998-2014 Epic Games, Inc. All Rights Reserved.
+#pragma once
+
+#include "HAL/Platform.h"
+
+#if PLATFORM_WINDOWS
+#include "Windows/WindowsPlatformHttp.h"
+#elif PLATFORM_PS4
+#include "PS4/PS4PlatformHttp.h"
+#elif PLATFORM_XBOXONE
+#include "XboxOne/XboxOneHttp.h"
+#elif PLATFORM_MAC
+#include "Mac/MacPlatformHttp.h"
+#elif PLATFORM_IOS
+#include "IOS/IOSPlatformHttp.h"
+#elif PLATFORM_ANDROID
+#include "Android/AndroidHttp.h"
+#elif PLATFORM_WINRT
+#include "WinRT/WinRTHttp.h"
+#elif PLATFORM_HTML5
+#include "HTML5/HTML5PlatformHttp.h"
+#elif PLATFORM_LINUX
+#include "Linux/LinuxPlatformHttp.h"
+#endif
diff --git a/Engine/Source/Runtime/Core/Public/HTML5/HTML5PlatformHttp.h b/Engine/Source/Runtime/Core/Public/HTML5/HTML5PlatformHttp.h
--- a/Engine/Source/Runtime/Core/Public/HTML5/HTML5PlatformHttp.h
+++ b/Engine/Source/Runtime/Core/Public/HTML5/HTML5PlatformHttp.h
@@ -1,6 +1,7 @@
 // Copyright 1998-2014 Epic Games, Inc. All Rights Reserved.
 
 #pragma once
+#include "GenericPlatform/GenericPlatformHttp.h"
 
 /**
  * Platform specific Http implementations
@@ -19,9 +20,9 @@ public:
 	 *
 	 * @return NULL if default implementation is to be used
 	 */
-	static class FHttpManager * CreatePlatformHttpManager()
+	static FHttpManager* CreatePlatformHttpManager()
 	{
-		return NULL;
+		return nullptr;
 	}
 
 	/**
@@ -34,7 +35,7 @@ public:
 	 *
 	 * @return request object
 	 */
-	static class IHttpRequest* ConstructRequest();
+	static IHttpRequest* ConstructRequest();
 };
 
 
diff --git a/Engine/Source/Runtime/Core/Public/IOS/IOSPlatformHttp.h b/Engine/Source/Runtime/Core/Public/IOS/IOSPlatformHttp.h
--- a/Engine/Source/Runtime/Core/Public/IOS/IOSPlatformHttp.h
+++ b/Engine/Source/Runtime/Core/Public/IOS/IOSPlatformHttp.h
@@ -1,6 +1,7 @@
 // Copyright 1998-2014 Epic Games, Inc. All Rights Reserved.
 
 #pragma once
+#include "GenericPlatform/GenericPlatformHttp.h"
 
 /**
  * Platform specific Http implementations
@@ -19,9 +20,9 @@ public:
 	 *
 	 * @return NULL if default implementation is to be used
 	 */
-	static class FHttpManager * CreatePlatformHttpManager()
+	static FHttpManager* CreatePlatformHttpManager()
 	{
-		return NULL;
+		return nullptr;
 	}
 
 	/**
@@ -34,7 +35,7 @@ public:
 	 *
 	 * @return request object
 	 */
-	static class IHttpRequest* ConstructRequest();
+	static IHttpRequest* ConstructRequest();
 };
 
 
diff --git a/Engine/Source/Runtime/Core/Public/Linux/LinuxPlatformHttp.h b/Engine/Source/Runtime/Core/Public/Linux/LinuxPlatformHttp.h
--- a/Engine/Source/Runtime/Core/Public/Linux/LinuxPlatformHttp.h
+++ b/Engine/Source/Runtime/Core/Public/Linux/LinuxPlatformHttp.h
@@ -1,6 +1,7 @@
 // Copyright 1998-2014 Epic Games, Inc. All Rights Reserved.
 
 #pragma once
+#include "GenericPlatform/GenericPlatformHttp.h"
 
 /**
  * Platform specific Http implementations
@@ -19,7 +20,7 @@ public:
 	 *
 	 * @return NULL if default implementation is to be used
 	 */
-	static class FHttpManager * CreatePlatformHttpManager();
+	static FHttpManager* CreatePlatformHttpManager();
 
 	/**
 	 * Platform shutdown step
@@ -31,7 +32,7 @@ public:
 	 *
 	 * @return request object
 	 */
-	static class IHttpRequest* ConstructRequest();
+	static IHttpRequest* ConstructRequest();
 };
 
 
diff --git a/Engine/Source/Runtime/Core/Public/Mac/MacPlatformHttp.h b/Engine/Source/Runtime/Core/Public/Mac/MacPlatformHttp.h
--- a/Engine/Source/Runtime/Core/Public/Mac/MacPlatformHttp.h
+++ b/Engine/Source/Runtime/Core/Public/Mac/MacPlatformHttp.h
@@ -1,6 +1,7 @@
 // Copyright 1998-2014 Epic Games, Inc. All Rights Reserved.
 
 #pragma once
+#include "GenericPlatform/GenericPlatformHttp.h"
 
 /**
  * Platform specific Http implementations
@@ -19,9 +20,9 @@ public:
 	 *
 	 * @return NULL if default implementation is to be used
 	 */
-	static class FHttpManager * CreatePlatformHttpManager()
+	static FHttpManager* CreatePlatformHttpManager()
 	{
-		return NULL;
+		return nullptr;
 	}
 
 	/**
@@ -34,7 +35,7 @@ public:
 	 *
 	 * @return request object
 	 */
-	static class IHttpRequest* ConstructRequest();
+	static IHttpRequest* ConstructRequest();
 };
 
 
diff --git a/Engine/Source/Runtime/Core/Public/Mac/MacPlatformSurvey.h b/Engine/Source/Runtime/Core/Public/Mac/MacPlatformSurvey.h
--- a/Engine/Source/Runtime/Core/Public/Mac/MacPlatformSurvey.h
+++ b/Engine/Source/Runtime/Core/Public/Mac/MacPlatformSurvey.h
@@ -8,6 +8,8 @@
 #pragma once
 #include "GenericPlatform/GenericPlatformSurvey.h"
 
+class FString;
+
 /**
 * Mac implementation of FGenericPlatformSurvey
 **/
diff --git a/Engine/Source/Runtime/Core/Public/WinRT/WinRTHttp.h b/Engine/Source/Runtime/Core/Public/WinRT/WinRTHttp.h
--- a/Engine/Source/Runtime/Core/Public/WinRT/WinRTHttp.h
+++ b/Engine/Source/Runtime/Core/Public/WinRT/WinRTHttp.h
@@ -1,6 +1,7 @@
 // Copyright 1998-2014 Epic Games, Inc. All Rights Reserved.
 
 #pragma once
+#include "GenericPlatform/GenericPlatformHttp.h"
 
 /**
  * Platform specific Http implementations
@@ -19,9 +20,9 @@ public:
 	 *
 	 * @return NULL if default implementation is to be used
 	 */
-	static class FHttpManager * CreatePlatformHttpManager()
+	static FHttpManager* CreatePlatformHttpManager()
 	{
-		return NULL;
+		return nullptr;
 	}
 
 	/**
@@ -34,7 +35,7 @@ public:
 	 *
 	 * @return request object
 	 */
-	static class IHttpRequest* ConstructRequest();
+	static IHttpRequest* ConstructRequest();
 };
 
 
diff --git a/Engine/Source/Runtime/Core/Public/Windows/WindowsPlatformHttp.h b/Engine/Source/Runtime/Core/Public/Windows/WindowsPlatformHttp.h
--- a/Engine/Source/Runtime/Core/Public/Windows/WindowsPlatformHttp.h
+++ b/Engine/Source/Runtime/Core/Public/Windows/WindowsPlatformHttp.h
@@ -1,6 +1,7 @@
 // Copyright 1998-2014 Epic Games, Inc. All Rights Reserved.
 
 #pragma once
+#include "GenericPlatform/GenericPlatformHttp.h"
 
 /**
  * Platform specific Http implementations
@@ -19,7 +20,7 @@ public:
 	 *
 	 * @return NULL if default implementation is to be used
 	 */
-	static class FHttpManager * CreatePlatformHttpManager();
+	static FHttpManager* CreatePlatformHttpManager();
 
 	/**
 	 * Platform shutdown step
@@ -31,7 +32,7 @@ public:
 	 *
 	 * @return request object
 	 */
-	static class IHttpRequest* ConstructRequest();
+	static IHttpRequest* ConstructRequest();
 };
 
 
