Make Core module public headers individually compilable. Part 26 of n.

diff --git a/Engine/Source/Runtime/Core/Public/GenericPlatform/GenericPlatformFramePacer.h b/Engine/Source/Runtime/Core/Public/GenericPlatform/GenericPlatformFramePacer.h
--- a/Engine/Source/Runtime/Core/Public/GenericPlatform/GenericPlatformFramePacer.h
+++ b/Engine/Source/Runtime/Core/Public/GenericPlatform/GenericPlatformFramePacer.h
@@ -7,6 +7,7 @@
 
 
 #pragma once
+#include "HAL/Platform.h"
 
 class FEvent;
 
diff --git a/Engine/Source/Runtime/Core/Public/HAL/PlatformFramePacer.h b/Engine/Source/Runtime/Core/Public/HAL/PlatformFramePacer.h
new file mode 100644
--- /dev/null
+++ b/Engine/Source/Runtime/Core/Public/HAL/PlatformFramePacer.h
@@ -0,0 +1,10 @@
+// Copyright 1998-2014 Epic Games, Inc. All Rights Reserved.
+#pragma once
+#include "HAL/Platform.h"
+
+#if PLATFORM_IOS
+#include "IOS/IOSPlatformFramePacer.h"
+#else
+#include "GenericPlatform/GenericPlatformFramePacer.h"
+typedef FGenericPlatformRHIFramePacer FPlatformRHIFramePacer;
+#endif
\ No newline at end of file
diff --git a/Engine/Source/Runtime/Core/Public/IOS/IOSPlatformFramePacer.h b/Engine/Source/Runtime/Core/Public/IOS/IOSPlatformFramePacer.h
--- a/Engine/Source/Runtime/Core/Public/IOS/IOSPlatformFramePacer.h
+++ b/Engine/Source/Runtime/Core/Public/IOS/IOSPlatformFramePacer.h
@@ -7,6 +7,7 @@
 
 
 #pragma once
+#include "GenericPlatform/GenericPlatformFramePacer.h"
 
 // Forward declare the ios frame pacer class we will be using.
 @class FIOSFramePacer;
