New/delete nothrow overrides, needed for newer Clangs.

Note: the changes here are only for Mac - it is likely that similar changes will be needed for other Clang platforms when they get upgraded.

diff --git a/Engine/Source/Runtime/Core/Public/Android/AndroidPlatform.h b/Engine/Source/Runtime/Core/Public/Android/AndroidPlatform.h
--- a/Engine/Source/Runtime/Core/Public/Android/AndroidPlatform.h
+++ b/Engine/Source/Runtime/Core/Public/Android/AndroidPlatform.h
@@ -66,3 +66,5 @@ typedef FAndroidTypes FPlatformTypes;
 // Operator new/delete handling.
 #define OPERATOR_NEW_THROW_SPEC throw (std::bad_alloc)
 #define OPERATOR_DELETE_THROW_SPEC throw()
+#define OPERATOR_NEW_NOTHROW_SPEC throw ()
+#define OPERATOR_DELETE_NOTHROW_SPEC throw()
diff --git a/Engine/Source/Runtime/Core/Public/HAL/Platform.h b/Engine/Source/Runtime/Core/Public/HAL/Platform.h
--- a/Engine/Source/Runtime/Core/Public/HAL/Platform.h
+++ b/Engine/Source/Runtime/Core/Public/HAL/Platform.h
@@ -408,6 +408,12 @@
 #ifndef OPERATOR_DELETE_THROW_SPEC
 	#define OPERATOR_DELETE_THROW_SPEC
 #endif
+#ifndef OPERATOR_NEW_NOTHROW_SPEC
+	#define OPERATOR_NEW_NOTHROW_SPEC throw()
+#endif
+#ifndef OPERATOR_DELETE_NOTHROW_SPEC
+	#define OPERATOR_DELETE_NOTHROW_SPEC throw()
+#endif
 
 #ifndef checkAtCompileTime
 	#define checkAtCompileTime(expr, msg) \
diff --git a/Engine/Source/Runtime/Core/Public/HTML5/HTML5Platform.h b/Engine/Source/Runtime/Core/Public/HTML5/HTML5Platform.h
--- a/Engine/Source/Runtime/Core/Public/HTML5/HTML5Platform.h
+++ b/Engine/Source/Runtime/Core/Public/HTML5/HTML5Platform.h
@@ -93,6 +93,8 @@ typedef FHTML5Types FPlatformTypes;
 #if !PLATFORM_HTML5_WIN32
 #define OPERATOR_NEW_THROW_SPEC throw (std::bad_alloc)
 #define OPERATOR_DELETE_THROW_SPEC throw()
+#define OPERATOR_NEW_NOTHROW_SPEC throw ()
+#define OPERATOR_DELETE_NOTHROW_SPEC throw()
 #endif
 
 
diff --git a/Engine/Source/Runtime/Core/Public/IOS/IOSPlatform.h b/Engine/Source/Runtime/Core/Public/IOS/IOSPlatform.h
--- a/Engine/Source/Runtime/Core/Public/IOS/IOSPlatform.h
+++ b/Engine/Source/Runtime/Core/Public/IOS/IOSPlatform.h
@@ -70,6 +70,8 @@ typedef FIOSPlatformTypes FPlatformTypes;
 // new/delete operators
 #define OPERATOR_NEW_THROW_SPEC throw (std::bad_alloc)
 #define OPERATOR_DELETE_THROW_SPEC throw()
+#define OPERATOR_NEW_NOTHROW_SPEC throw ()
+#define OPERATOR_DELETE_NOTHROW_SPEC throw()
 
 // DLL export and import definitions
 #define DLLEXPORT
diff --git a/Engine/Source/Runtime/Core/Public/Linux/LinuxPlatform.h b/Engine/Source/Runtime/Core/Public/Linux/LinuxPlatform.h
--- a/Engine/Source/Runtime/Core/Public/Linux/LinuxPlatform.h
+++ b/Engine/Source/Runtime/Core/Public/Linux/LinuxPlatform.h
@@ -75,3 +75,5 @@ typedef FLinuxPlatformTypes FPlatformTypes;
 // Operator new/delete handling.
 #define OPERATOR_NEW_THROW_SPEC throw (std::bad_alloc)
 #define OPERATOR_DELETE_THROW_SPEC throw()
+#define OPERATOR_NEW_NOTHROW_SPEC throw ()
+#define OPERATOR_DELETE_NOTHROW_SPEC throw()
diff --git a/Engine/Source/Runtime/Core/Public/Mac/MacPlatform.h b/Engine/Source/Runtime/Core/Public/Mac/MacPlatform.h
--- a/Engine/Source/Runtime/Core/Public/Mac/MacPlatform.h
+++ b/Engine/Source/Runtime/Core/Public/Mac/MacPlatform.h
@@ -39,8 +39,6 @@ typedef FMacPlatformTypes FPlatformTypes;
 #define PLATFORM_MAX_FILEPATH_LENGTH				MAX_PATH
 #define PLATFORM_SUPPORTS_TBB						1
 
-#define FORCE_ANSI_ALLOCATOR 1
-
 // Function type macros.
 #define VARARGS														/* Functions with variable arguments */
 #define CDECL														/* Standard C function */
@@ -57,9 +55,16 @@ typedef FMacPlatformTypes FPlatformTypes;
 #define GCC_PACK(n) __attribute__((packed,aligned(n)))
 #define GCC_ALIGN(n) __attribute__((aligned(n)))
 
-// new/delete operators
+// operator new/delete operators
+// As of 10.9 we need to use _NOEXCEPT & cxx_noexcept compatible definitions
+#if __has_feature(cxx_noexcept)
+#define OPERATOR_NEW_THROW_SPEC
+#else
 #define OPERATOR_NEW_THROW_SPEC throw (std::bad_alloc)
-#define OPERATOR_DELETE_THROW_SPEC throw()
+#endif
+#define OPERATOR_DELETE_THROW_SPEC _NOEXCEPT
+#define OPERATOR_NEW_NOTHROW_SPEC  _NOEXCEPT
+#define OPERATOR_DELETE_NOTHROW_SPEC  _NOEXCEPT
 
 // DLL export and import definitions
 #define DLLEXPORT
diff --git a/Engine/Source/Runtime/Core/Public/Modules/Boilerplate/ModuleBoilerplate.h b/Engine/Source/Runtime/Core/Public/Modules/Boilerplate/ModuleBoilerplate.h
--- a/Engine/Source/Runtime/Core/Public/Modules/Boilerplate/ModuleBoilerplate.h
+++ b/Engine/Source/Runtime/Core/Public/Modules/Boilerplate/ModuleBoilerplate.h
@@ -19,10 +19,14 @@
 #endif
 
 #define REPLACEMENT_OPERATOR_NEW_AND_DELETE \
-	OPERATOR_NEW_MSVC_PRAGMA void* operator new  ( size_t Size ) OPERATOR_NEW_THROW_SPEC    { return FMemory::Malloc( Size ); } \
-	OPERATOR_NEW_MSVC_PRAGMA void* operator new[]( size_t Size ) OPERATOR_NEW_THROW_SPEC    { return FMemory::Malloc( Size ); } \
-	void operator delete  ( void* Ptr )                          OPERATOR_DELETE_THROW_SPEC { FMemory::Free( Ptr ); } \
-	void operator delete[]( void* Ptr )                          OPERATOR_DELETE_THROW_SPEC { FMemory::Free( Ptr ); }
+	OPERATOR_NEW_MSVC_PRAGMA void* operator new  ( size_t Size                  ) OPERATOR_NEW_THROW_SPEC      { return FMemory::Malloc( Size ); } \
+	OPERATOR_NEW_MSVC_PRAGMA void* operator new[]( size_t Size                  ) OPERATOR_NEW_THROW_SPEC      { return FMemory::Malloc( Size ); } \
+	OPERATOR_NEW_MSVC_PRAGMA void* operator new  ( size_t Size, std::nothrow_t& ) OPERATOR_NEW_NOTHROW_SPEC    { return FMemory::Malloc( Size ); } \
+	OPERATOR_NEW_MSVC_PRAGMA void* operator new[]( size_t Size, std::nothrow_t& ) OPERATOR_NEW_NOTHROW_SPEC    { return FMemory::Malloc( Size ); } \
+	void operator delete  ( void* Ptr )                                           OPERATOR_DELETE_THROW_SPEC   { FMemory::Free( Ptr ); } \
+	void operator delete[]( void* Ptr )                                           OPERATOR_DELETE_THROW_SPEC   { FMemory::Free( Ptr ); } \
+	void operator delete  ( void* Ptr, std::nothrow_t& )                          OPERATOR_DELETE_NOTHROW_SPEC { FMemory::Free( Ptr ); } \
+	void operator delete[]( void* Ptr, std::nothrow_t& )                          OPERATOR_DELETE_NOTHROW_SPEC { FMemory::Free( Ptr ); }
 
 // in DLL builds, these are done per-module, otherwise we just need one in the application
 // visual studio cannot find cross dll data for visualizers, so these provide access
