Fix variable shadowing

diff --git a/Engine/Source/Developer/ShaderCompilerCommon/Private/HlslLexer.cpp b/Engine/Source/Developer/ShaderCompilerCommon/Private/HlslLexer.cpp
--- a/Engine/Source/Developer/ShaderCompilerCommon/Private/HlslLexer.cpp
+++ b/Engine/Source/Developer/ShaderCompilerCommon/Private/HlslLexer.cpp
@@ -578,7 +578,7 @@ namespace CrossCompiler
 				OutIdentifier += Char;
 				do
 				{
-					auto Char = *Current;
+					Char = *Current;
 					if (!IsAlphaOrDigit(Char) && Char != '_')
 					{
 						break;
diff --git a/Engine/Source/Runtime/Core/Private/Internationalization/Internationalization.cpp b/Engine/Source/Runtime/Core/Private/Internationalization/Internationalization.cpp
--- a/Engine/Source/Runtime/Core/Private/Internationalization/Internationalization.cpp
+++ b/Engine/Source/Runtime/Core/Private/Internationalization/Internationalization.cpp
@@ -269,9 +269,9 @@ void FInternationalization::GetCulturesWithAvailableLocalization(const TArray<FS
 			{
 				TArray<FString> PrioritizedParentCultureNames = Culture->GetPrioritizedParentCultureNames();
 
-				for (const FString& CultureName : PrioritizedParentCultureNames)
+				for (const FString& PrioritizedParentCultureName : PrioritizedParentCultureNames)
 				{
-					if(AllLocalizationFolders.Contains(CultureName))
+					if(AllLocalizationFolders.Contains(PrioritizedParentCultureName))
 					{
 						OutAvailableCultures.AddUnique(Culture.ToSharedRef());
 						break;
diff --git a/Engine/Source/Runtime/Core/Private/Internationalization/TextHistory.cpp b/Engine/Source/Runtime/Core/Private/Internationalization/TextHistory.cpp
--- a/Engine/Source/Runtime/Core/Private/Internationalization/TextHistory.cpp
+++ b/Engine/Source/Runtime/Core/Private/Internationalization/TextHistory.cpp
@@ -409,9 +409,9 @@ void FTextHistory_AsPercent::Serialize( FArchive& Ar )
 ///////////////////////////////////////
 // FTextHistory_AsCurrency
 
-FTextHistory_AsCurrency::FTextHistory_AsCurrency(const FFormatArgumentValue& InSourceValue, const FString& CurrencyCode, const FNumberFormattingOptions* const InFormatOptions, const FCulturePtr InTargetCulture)
+FTextHistory_AsCurrency::FTextHistory_AsCurrency(const FFormatArgumentValue& InSourceValue, const FString& InCurrencyCode, const FNumberFormattingOptions* const InFormatOptions, const FCulturePtr InTargetCulture)
 : FTextHistory_FormatNumber(InSourceValue, InFormatOptions, InTargetCulture)
-, CurrencyCode(CurrencyCode)
+, CurrencyCode(InCurrencyCode)
 {
 }
 
diff --git a/Engine/Source/Runtime/Core/Private/Internationalization/TextLocalizationManager.cpp b/Engine/Source/Runtime/Core/Private/Internationalization/TextLocalizationManager.cpp
--- a/Engine/Source/Runtime/Core/Private/Internationalization/TextLocalizationManager.cpp
+++ b/Engine/Source/Runtime/Core/Private/Internationalization/TextLocalizationManager.cpp
@@ -271,25 +271,27 @@ void FTextLocalizationManager::LoadResources(const bool ShouldLoadEditor, const
 
 	FInternationalization& I18N = FInternationalization::Get();
 
-	const FString& CultureName = I18N.GetCurrentCulture()->GetName();
-
 #if ENABLE_LOC_TESTING
-	if(CultureName == TEXT("LEET"))
 	{
-		for(auto NamespaceIterator = LiveTable.NamespaceTable.CreateIterator(); NamespaceIterator; ++NamespaceIterator)
+		const FString& CultureName = I18N.GetCurrentCulture()->GetName();
+
+		if(CultureName == TEXT("LEET"))
 		{
-			const FString& Namespace = NamespaceIterator.Key();
-			FTextLookupTable::FKeyTable& LiveKeyTable = NamespaceIterator.Value();
-			for(auto KeyIterator = LiveKeyTable.CreateIterator(); KeyIterator; ++KeyIterator)
+			for(auto NamespaceIterator = LiveTable.NamespaceTable.CreateIterator(); NamespaceIterator; ++NamespaceIterator)
 			{
-				const FString& Key = KeyIterator.Key();
-				FStringEntry& LiveStringEntry = KeyIterator.Value();
-				LiveStringEntry.bIsLocalized = true;
-				FInternationalization::Leetify( *LiveStringEntry.String );
+				const FString& Namespace = NamespaceIterator.Key();
+				FTextLookupTable::FKeyTable& LiveKeyTable = NamespaceIterator.Value();
+				for(auto KeyIterator = LiveKeyTable.CreateIterator(); KeyIterator; ++KeyIterator)
+				{
+					const FString& Key = KeyIterator.Key();
+					FStringEntry& LiveStringEntry = KeyIterator.Value();
+					LiveStringEntry.bIsLocalized = true;
+					FInternationalization::Leetify( *LiveStringEntry.String );
+				}
 			}
-		}
 
-		return;
+			return;
+		}
 	}
 #endif
 
diff --git a/Engine/Source/Runtime/Core/Private/Mac/MacApplication.cpp b/Engine/Source/Runtime/Core/Private/Mac/MacApplication.cpp
--- a/Engine/Source/Runtime/Core/Private/Mac/MacApplication.cpp
+++ b/Engine/Source/Runtime/Core/Private/Mac/MacApplication.cpp
@@ -587,8 +587,7 @@ void FMacApplication::ProcessNSEvent(NSEvent* const Event, TSharedPtr< FMacWindo
 					while ((screen = [screenEnumerator nextObject]) && !NSMouseInRect(NSMakePoint(HighPrecisionMousePos.X, HighPrecisionMousePos.Y), screen.frame, NO))
 						;
 					
-					FMacCursor* MacCursor = static_cast<FMacCursor*>( Cursor.Get() );
-					
+				
 					// Clamp to no more than the reported delta - a single event of no mouse movement won't be noticed
 					// but going in the wrong direction will.
 					const FVector2D FullDelta([Event deltaX], [Event deltaY]);
diff --git a/Engine/Source/Runtime/Core/Private/Math/TransformVectorized.cpp b/Engine/Source/Runtime/Core/Private/Math/TransformVectorized.cpp
--- a/Engine/Source/Runtime/Core/Private/Math/TransformVectorized.cpp
+++ b/Engine/Source/Runtime/Core/Private/Math/TransformVectorized.cpp
@@ -60,15 +60,15 @@ bool FTransform::InitFromString( const FString& Source )
 	}
 
 	// Translation
-	FVector Translation = FVector::ZeroVector;
-	if( !FDefaultValueHelper::ParseVector(ComponentStrings[0], Translation) )
+	FVector ParsedTranslation = FVector::ZeroVector;
+	if( !FDefaultValueHelper::ParseVector(ComponentStrings[0], ParsedTranslation) )
 	{
 		return false;
 	}
 
 	// Rotation
-	FRotator Rotation = FRotator::ZeroRotator;
-	if( !FDefaultValueHelper::ParseRotator(ComponentStrings[1], Rotation) )
+	FRotator ParsedRotation = FRotator::ZeroRotator;
+	if( !FDefaultValueHelper::ParseRotator(ComponentStrings[1], ParsedRotation) )
 	{
 		return false;
 	}
@@ -80,7 +80,7 @@ bool FTransform::InitFromString( const FString& Source )
 		return false;
 	}
 
-	SetComponents(FQuat(Rotation), Translation, Scale);
+	SetComponents(FQuat(ParsedRotation), ParsedTranslation, Scale);
 
 	return true;
 }
diff --git a/Engine/Source/Runtime/Navmesh/Private/DetourTileCache/DetourTileCacheBuilder.cpp b/Engine/Source/Runtime/Navmesh/Private/DetourTileCache/DetourTileCacheBuilder.cpp
--- a/Engine/Source/Runtime/Navmesh/Private/DetourTileCache/DetourTileCacheBuilder.cpp
+++ b/Engine/Source/Runtime/Navmesh/Private/DetourTileCache/DetourTileCacheBuilder.cpp
@@ -1792,10 +1792,10 @@ static dtStatus removeVertex(dtTileCachePolyMesh& mesh, const unsigned short rem
 	// Generate temp vertex array for triangulation.
 	for (int i = 0; i < nhole; ++i)
 	{
-		const unsigned short pi = hole[i];
-		tverts[i*4+0] = mesh.verts[pi*3+0];
-		tverts[i*4+1] = mesh.verts[pi*3+1];
-		tverts[i*4+2] = mesh.verts[pi*3+2];
+		const unsigned short hi = hole[i];
+		tverts[i*4+0] = mesh.verts[hi*3+0];
+		tverts[i*4+1] = mesh.verts[hi*3+1];
+		tverts[i*4+2] = mesh.verts[hi*3+2];
 		tverts[i*4+3] = 0;
 		tpoly[i] = (unsigned short)i;
 	}
