Core: Documentation cleanup pass

- also NULL to nullptr
- consistent line spacing

diff --git a/Engine/Source/Runtime/Core/Public/Serialization/Archive.h b/Engine/Source/Runtime/Core/Public/Serialization/Archive.h
--- a/Engine/Source/Runtime/Core/Public/Serialization/Archive.h
+++ b/Engine/Source/Runtime/Core/Public/Serialization/Archive.h
@@ -257,7 +257,7 @@ public:
 		if( bFreeOnClose )
 		{
 			FMemory::Free( ReaderData );
-			ReaderData = NULL;
+			ReaderData = nullptr;
 		}
 		return !ArIsError;
 	}
diff --git a/Engine/Source/Runtime/Core/Public/Serialization/ArchiveBase.h b/Engine/Source/Runtime/Core/Public/Serialization/ArchiveBase.h
--- a/Engine/Source/Runtime/Core/Public/Serialization/ArchiveBase.h
+++ b/Engine/Source/Runtime/Core/Public/Serialization/ArchiveBase.h
@@ -1,20 +1,19 @@
 // Copyright 1998-2014 Epic Games, Inc. All Rights Reserved.
 
-/*=============================================================================
-	ArchiveBase.h: Implements the FArchive class and related types.
-=============================================================================*/
-
 #pragma once
+
 #include "Containers/EnumAsByte.h"
 #include "HAL/PlatformProperties.h"
 #include "Misc/Compression.h"
 
+
 class FAssetPtr;
 class FCustomVersionContainer;
 class FLazyObjectPtr;
 struct FStringAssetReference;
 struct FUntypedBulkData;
 
+
 /**
  * Base class for archives that can be used for loading, saving, and garbage
  * collecting in a byte order neutral way.
@@ -23,27 +22,22 @@ class CORE_API FArchive
 {
 public:
 
-	/**
-	 * Default constructor.
-	 */
+	/** Default constructor. */
 	FArchive();
 
-	/**
-	 * Copy constructor.
-	 */
+	/** Copy constructor. */
 	FArchive(const FArchive&);
 
 	/**
 	 * Copy assignment operator.
+	 *
+	 * @param ArchiveToCopy The archive to copy from.
 	 */
-	FArchive& operator=(const FArchive&);
+	FArchive& operator=(const FArchive& ArchiveToCopy);
 
-	/**
-	 * Destructor.
-	 */
+	/** Destructor. */
 	virtual ~FArchive();
 
-
 public:
 
 	/**
@@ -51,8 +45,7 @@ public:
 	 *
 	 * This operator can be implemented by sub-classes that wish to serialize FName instances.
 	 *
-	 * @param Value - The value to serialize.
-	 *
+	 * @param Value The value to serialize.
 	 * @return This instance.
 	 */
 	virtual FArchive& operator<<(class FName& Value)
@@ -65,8 +58,7 @@ public:
 	 *
 	 * This operator can be implemented by sub-classes that wish to serialize UObject instances.
 	 *
-	 * @param Value - The value to serialize.
-	 *
+	 * @param Value The value to serialize.
 	 * @return This instance.
 	 */
 	virtual FArchive& operator<<(class UObject*& Value)
@@ -79,8 +71,7 @@ public:
 	 *
 	 * Most of the time, FLazyObjectPtrs are serialized as UObject*, but some archives need to override this.
 	 *
-	 * @param Value - The value to serialize.
-	 *
+	 * @param Value The value to serialize.
 	 * @return This instance.
 	 */
 	virtual FArchive& operator<<(class FLazyObjectPtr& Value);
@@ -90,8 +81,7 @@ public:
 	 *
 	 * Most of the time, FAssetPtrs are serialized as UObject *, but some archives need to override this.
 	 *
-	 * @param Value - The asset pointer to serialize.
-	 *
+	 * @param Value The asset pointer to serialize.
 	 * @return This instance.
 	 */
 	virtual FArchive& operator<<(class FAssetPtr& Value);
@@ -100,19 +90,17 @@ public:
 	 * Serializes string asset reference from or into this archive.
 	 *
 	 * @param Value String asset reference to serialize.
-	 *
 	 * @return This instance.
 	 */
 	virtual FArchive& operator<<(struct FStringAssetReference& Value);
 
-
 public:
 
 	/**
 	 * Serializes an ANSICHAR value from or into an archive.
 	 *
-	 * @param Ar - The archive to serialize from or to.
-	 * @param Value - The value to serialize.
+	 * @param Ar The archive to serialize from or to.
+	 * @param Value The value to serialize.
 	 */
 	FORCEINLINE friend FArchive& operator<<(FArchive& Ar, ANSICHAR& Value)
 	{
@@ -124,8 +112,8 @@ public:
 	/**
 	 * Serializes a WIDECHAR value from or into an archive.
 	 *
-	 * @param Ar - The archive to serialize from or to.
-	 * @param Value - The value to serialize.
+	 * @param Ar The archive to serialize from or to.
+	 * @param Value The value to serialize.
 	 */
 	FORCEINLINE friend FArchive& operator<<(FArchive& Ar, WIDECHAR& Value)
 	{
@@ -137,8 +125,8 @@ public:
 	/**
 	 * Serializes an unsigned 8-bit integer value from or into an archive.
 	 *
-	 * @param Ar - The archive to serialize from or to.
-	 * @param Value - The value to serialize.
+	 * @param Ar The archive to serialize from or to.
+	 * @param Value The value to serialize.
 	 */
 	FORCEINLINE friend FArchive& operator<<(FArchive& Ar, uint8& Value)
 	{
@@ -150,8 +138,8 @@ public:
 	/**
 	 * Serializes an enumeration value from or into an archive.
 	 *
-	 * @param Ar - The archive to serialize from or to.
-	 * @param Value - The value to serialize.
+	 * @param Ar The archive to serialize from or to.
+	 * @param Value The value to serialize.
 	 */
 	template<class TEnum>
 	FORCEINLINE friend FArchive& operator<<(FArchive& Ar, TEnumAsByte<TEnum>& Value)
@@ -164,8 +152,8 @@ public:
 	/**
 	 * Serializes a signed 8-bit integer value from or into an archive.
 	 *
-	 * @param Ar - The archive to serialize from or to.
-	 * @param Value - The value to serialize.
+	 * @param Ar The archive to serialize from or to.
+	 * @param Value The value to serialize.
 	 */
 	FORCEINLINE friend FArchive& operator<<(FArchive& Ar, int8& Value)
 	{
@@ -177,8 +165,8 @@ public:
 	/**
 	 * Serializes an unsigned 16-bit integer value from or into an archive.
 	 *
-	 * @param Ar - The archive to serialize from or to.
-	 * @param Value - The value to serialize.
+	 * @param Ar The archive to serialize from or to.
+	 * @param Value The value to serialize.
 	 */
 	FORCEINLINE friend FArchive& operator<<(FArchive& Ar, uint16& Value)
 	{
@@ -190,8 +178,8 @@ public:
 	/**
 	 * Serializes a signed 16-bit integer value from or into an archive.
 	 *
-	 * @param Ar - The archive to serialize from or to.
-	 * @param Value - The value to serialize.
+	 * @param Ar The archive to serialize from or to.
+	 * @param Value The value to serialize.
 	 */
 	FORCEINLINE friend FArchive& operator<<(FArchive& Ar, int16& Value)
 	{
@@ -203,8 +191,8 @@ public:
 	/**
 	 * Serializes an unsigned 32-bit integer value from or into an archive.
 	 *
-	 * @param Ar - The archive to serialize from or to.
-	 * @param Value - The value to serialize.
+	 * @param Ar The archive to serialize from or to.
+	 * @param Value The value to serialize.
 	 */
 	FORCEINLINE friend FArchive& operator<<(FArchive& Ar, uint32& Value)
 	{
@@ -216,8 +204,8 @@ public:
 	/**
 	 * Serializes a Boolean value from or into an archive.
 	 *
-	 * @param Ar - The archive to serialize from or to.
-	 * @param Value - The value to serialize.
+	 * @param Ar The archive to serialize from or to.
+	 * @param Value The value to serialize.
 	 */
 	FORCEINLINE friend FArchive& operator<<( FArchive& Ar, bool& D )
 	{
@@ -231,8 +219,8 @@ public:
 	/**
 	 * Serializes a signed 32-bit integer value from or into an archive.
 	 *
-	 * @param Ar - The archive to serialize from or to.
-	 * @param Value - The value to serialize.
+	 * @param Ar The archive to serialize from or to.
+	 * @param Value The value to serialize.
 	 */
 	FORCEINLINE friend FArchive& operator<<(FArchive& Ar, int32& Value)
 	{
@@ -245,8 +233,8 @@ public:
 	/**
 	 * Serializes a long integer value from or into an archive.
 	 *
-	 * @param Ar - The archive to serialize from or to.
-	 * @param Value - The value to serialize.
+	 * @param Ar The archive to serialize from or to.
+	 * @param Value The value to serialize.
 	 */
 	FORCEINLINE friend FArchive& operator<<(FArchive& Ar, long& Value)
 	{
@@ -259,8 +247,8 @@ public:
 	/**
 	 * Serializes a single precision floating point value from or into an archive.
 	 *
-	 * @param Ar - The archive to serialize from or to.
-	 * @param Value - The value to serialize.
+	 * @param Ar The archive to serialize from or to.
+	 * @param Value The value to serialize.
 	 */
 	FORCEINLINE friend FArchive& operator<<( FArchive& Ar, float& Value)
 	{
@@ -272,8 +260,8 @@ public:
 	/**
 	 * Serializes a double precision floating point value from or into an archive.
 	 *
-	 * @param Ar - The archive to serialize from or to.
-	 * @param Value - The value to serialize.
+	 * @param Ar The archive to serialize from or to.
+	 * @param Value The value to serialize.
 	 */
 	FORCEINLINE friend FArchive& operator<<(FArchive& Ar, double& Value)
 	{
@@ -285,8 +273,8 @@ public:
 	/**
 	 * Serializes a unsigned 64-bit integer value from or into an archive.
 	 *
-	 * @param Ar - The archive to serialize from or to.
-	 * @param Value - The value to serialize.
+	 * @param Ar The archive to serialize from or to.
+	 * @param Value The value to serialize.
 	 */
 	FORCEINLINE friend FArchive& operator<<(FArchive &Ar, uint64& Value)
 	{
@@ -298,8 +286,8 @@ public:
 	/**
 	 * Serializes a signed 64-bit integer value from or into an archive.
 	 *
-	 * @param Ar - The archive to serialize from or to.
-	 * @param Value - The value to serialize.
+	 * @param Ar The archive to serialize from or to.
+	 * @param Value The value to serialize.
 	 */
 	/*FORCEINLINE*/friend FArchive& operator<<(FArchive& Ar, int64& Value)
 	{
@@ -311,16 +299,16 @@ public:
 	/**
 	 * Serializes an FIntRect value from or into an archive.
 	 *
-	 * @param Ar - The archive to serialize from or to.
-	 * @param Value - The value to serialize.
+	 * @param Ar The archive to serialize from or to.
+	 * @param Value The value to serialize.
 	 */
 	friend FArchive& operator<<(FArchive& Ar, struct FIntRect& Value);
 
 	/**
 	 * Serializes an FString value from or into an archive.
 	 *
-	 * @param Ar - The archive to serialize from or to.
-	 * @param Value - The value to serialize.
+	 * @param Ar The archive to serialize from or to.
+	 * @param Value The value to serialize.
 	 */
 	friend CORE_API FArchive& operator<<(FArchive& Ar, FString& Value);
 
@@ -351,21 +339,22 @@ public:
 
 	virtual void CountBytes(SIZE_T InNum, SIZE_T InMax) { }
 
-
 	/**
   	 * Returns the name of the Archive.  Useful for getting the name of the package a struct or object
 	 * is in when a loading error occurs.
 	 *
 	 * This is overridden for the specific Archive Types
-	 **/
+	 */
 	virtual FString GetArchiveName() const;
 
 	/**
 	 * If this archive is a ULinkerLoad or ULinkerSave, returns a pointer to the ULinker portion.
+	 *
+	 * @return The linker, or nullptr if the archive is not a linker.
 	 */
 	virtual class ULinker* GetLinker()
 	{
-		return NULL;
+		return nullptr;
 	}
 
 	virtual int64 Tell()
@@ -497,41 +486,31 @@ public:
 		return *this;
 	}
 
-	/**
-	 * Sets a flag indicating that this archive contains code.
-	 */
+	/** Sets a flag indicating that this archive contains code. */
 	void ThisContainsCode()
 	{
 		ArContainsCode = true;
 	}
 
-	/**
-	 * Sets a flag indicating that this archive contains a ULevel or UWorld object.
-	 */
+	/** Sets a flag indicating that this archive contains a ULevel or UWorld object. */
 	void ThisContainsMap() 
 	{
 		ArContainsMap = true;
 	}
 
-	/**
-	 * Sets a flag indicating that this archive contains data required to be gathered for localization.
-	 */
+	/** Sets a flag indicating that this archive contains data required to be gathered for localization. */
 	void ThisRequiresLocalizationGather()
 	{
 		ArRequiresLocalizationGather = true;
 	}
 
-	/**
-	 * Sets a flag indicating that this archive is currently serializing class/struct defaults
-	 */
+	/** Sets a flag indicating that this archive is currently serializing class/struct defaults. */
 	void StartSerializingDefaults() 
 	{
 		ArSerializingDefaults++;
 	}
 
-	/**
-	 * Indicate that this archive is no longer serializing class/struct defaults
-	 */
+	/** Indicate that this archive is no longer serializing class/struct defaults. */
 	void StopSerializingDefaults() 
 	{
 		ArSerializingDefaults--;
@@ -577,15 +556,14 @@ public:
 	 * Registers the custom version to the archive.  This is used to inform the archive that custom version information is about to be stored.
 	 * There is no effect when the archive is being loaded from.
 	 *
-	 * @param Guid - The guid of the custom version.  This must have previously been registered with FCustomVersionRegistration.
+	 * @param Guid The guid of the custom version.  This must have previously been registered with FCustomVersionRegistration.
 	 */
 	void UsingCustomVersion(FGuid Guid);
 
 	/**
 	 * Queries a custom version from the archive.  If the archive is being used to write, the custom version must have already been registered.
 	 *
-	 * @param Key - The guid of the custom version to query.
-	 *
+	 * @param Key The guid of the custom version to query.
 	 * @return The version number, or 0 if the custom tag isn't stored in the archive.
 	 */
 	int32 CustomVer(FGuid Key) const;
@@ -779,9 +757,7 @@ public:
 	 */
 	void SetCustomVersions(const FCustomVersionContainer& CustomVersionContainer);
 
-	/**
-	 * Resets the custom version numbers for this archive.
-	 */
+	/** Resets the custom version numbers for this archive. */
 	void ResetCustomVersions();
 
 	/**
@@ -825,11 +801,10 @@ public:
 	}
 
 	/**
-	 * Returns if an async close operation has finished or not, as well as if there was an error
-	 *
-	 * @param bHasError true if there was an error
+	 * Returns if an async close operation has finished or not, as well as if there was an error.
 	 *
-	 * @return true if the close operation is complete (or if Close is not an async operation) 
+	 * @param bHasError true if there was an error.
+	 * @return true if the close operation is complete (or if Close is not an async operation).
 	 */
 	virtual bool IsCloseComplete(bool& bHasError)
 	{
@@ -851,7 +826,7 @@ public:
 	/**
 	 * Sets a flag indicating that this archive needs to filter editor-only content.
 	 *
-	 * @param InFilterEditorOnly - Whether to filter editor-only content.
+	 * @param InFilterEditorOnly Whether to filter editor-only content.
 	 */
 	virtual void SetFilterEditorOnly(bool InFilterEditorOnly)
 	{
@@ -893,7 +868,7 @@ public:
 	/**
 	 * Sets the cooking target platform.
 	 *
-	 * @param InCookingTarget - The target platform to set.
+	 * @param InCookingTarget The target platform to set.
 	 */
 	FORCEINLINE void SetCookingTarget(const class ITargetPlatform* InCookingTarget) 
 	{
@@ -913,20 +888,17 @@ public:
 
 protected:
 
-	/**
-	 * Resets all of the base archive members
-	 */
+	/** Resets all of the base archive members. */
 	void Reset();
 
-
 private:
+
 	/** Copies all of the members except CustomVersionContainer */
 	void CopyTrivialFArchiveStatusMembers(const FArchive& ArchiveStatusToCopy);
 
-
 protected:
 
-	// Status variables.
+	/** Status variables. */
 	int32 ArUE3Ver;
 
 	/** Holds the archive's network version. */
@@ -939,6 +911,7 @@ protected:
 	int32 ArLicenseeUE4Ver;
 
 private:
+
 	/**
 	 * All the custom versions stored in the archive.
 	 * Stored as a pointer to a heap-allocated object because of a 3-way dependency between TArray, FCustomVersionContainer and FArchive, which is too much work to change right now.
@@ -946,6 +919,7 @@ private:
 	FCustomVersionContainer* CustomVersionContainer;
 
 public:
+
 	/** Whether this archive is for loading data. */
 	bool ArIsLoading;
 	
@@ -1026,7 +1000,7 @@ public:
 	
 private:
 
-	// Holds the cooking target platform.
+	/** Holds the cooking target platform. */
 	const ITargetPlatform* CookingTargetPlatform;
 
 	/**
@@ -1060,7 +1034,7 @@ public:
 	/**
 	 * Creates and initializes a new instance of the archive proxy.
 	 *
-	 * @param InInnerArchive - The inner archive to proxy.
+	 * @param InInnerArchive The inner archive to proxy.
 	 */
 	CORE_API FArchiveProxy(FArchive& InInnerArchive);
 
@@ -1181,12 +1155,9 @@ public:
 		return InnerArchive.IsCloseComplete(bHasError);
 	}
 
-
 protected:
 
-	/**
-	 * Holds the archive that this archive is a proxy to.
-	 */
+	/** Holds the archive that this archive is a proxy to. */
 	FArchive& InnerArchive;
 };
 
@@ -1199,7 +1170,7 @@ struct FNameAsStringProxyArchive : public FArchiveProxy
 	/**
 	 * Creates and initializes a new instance.
 	 *
-	 * @param InInnerArchive - The inner archive to proxy.
+	 * @param InInnerArchive The inner archive to proxy.
 	 */
 	 FNameAsStringProxyArchive(FArchive& InInnerArchive)
 		 :	FArchiveProxy(InInnerArchive)
@@ -1208,6 +1179,7 @@ struct FNameAsStringProxyArchive : public FArchiveProxy
 	 CORE_API virtual FArchive& operator<<(class FName& N);
 };
 
+
 /**
  * Implements a helper structure for compression support
  *
@@ -1215,14 +1187,10 @@ struct FNameAsStringProxyArchive : public FArchiveProxy
  */
 struct FCompressedChunkInfo
 {
-	/**
-	 * Holds the data's compressed size.
-	 */
+	/** Holds the data's compressed size. */
 	int64 CompressedSize;
 
-	/**
-	 * Holds the data's uncompresses size.
-	 */
+	/** Holds the data's uncompresses size. */
 	int64 UncompressedSize;
 };
 
@@ -1230,7 +1198,7 @@ struct FCompressedChunkInfo
 /**
  * Serializes an FCompressedChunkInfo value from or into an archive.
  *
- * @param Ar - The archive to serialize from or to.
- * @param Value - The value to serialize.
+ * @param Ar The archive to serialize from or to.
+ * @param Value The value to serialize.
  */
 CORE_API FArchive& operator<<(FArchive& Ar, FCompressedChunkInfo& Value);
diff --git a/Engine/Source/Runtime/Core/Public/Serialization/Bitstreams.h b/Engine/Source/Runtime/Core/Public/Serialization/Bitstreams.h
--- a/Engine/Source/Runtime/Core/Public/Serialization/Bitstreams.h
+++ b/Engine/Source/Runtime/Core/Public/Serialization/Bitstreams.h
@@ -1,11 +1,8 @@
 // Copyright 1998-2014 Epic Games, Inc. All Rights Reserved.
 
-/*=============================================================================
-	Bitstreams.h: Unreal bitstream manipulation classes.
-=============================================================================*/
-
 #pragma once
 
+
 /*-----------------------------------------------------------------------------
 	FBitWriter.
 -----------------------------------------------------------------------------*/
@@ -16,17 +13,19 @@
 struct CORE_API FBitWriter : public FArchive
 {
 	friend struct FBitWriterMark;
+
 public:
-	/**
-	 * Default constructor. Zeros everything.
-	 */
+	/** Default constructor. Zeros everything. */
 	FBitWriter(void);
+
 	/**
 	 * Constructor using known size the buffer needs to be
 	 */
 	FBitWriter( int64 InMaxBits, bool AllowResize = false );
+
 	void SerializeBits( void* Src, int64 LengthBits );
 	void SerializeInt( uint32& Value, uint32 Max );
+
 	/** serializes the specified value, but does not bounds check against Max; instead, it will wrap around if the value exceeds Max
 	 * (this differs from SerializeInt(), which clamps)
 	 * @param Result - the value to serialize
@@ -37,23 +36,25 @@ public:
 	void Serialize( void* Src, int64 LengthBytes );
 
 	/**
-	 * Returns a pointer to our internal buffer
+	 * Returns a pointer to our internal buffer.
 	 */
 	FORCEINLINE uint8* GetData(void)
 	{
 		return Buffer.GetData();
 	}
+
 	FORCEINLINE const uint8* GetData(void) const
 	{
 		return Buffer.GetData();
 	}
+
 	FORCEINLINE const TArray<uint8>* GetBuffer(void) const
 	{
 		return &Buffer;
 	}
 
 	/**
-	 * Returns the number of bytes written
+	 * Returns the number of bytes written.
 	 */
 	FORCEINLINE int64 GetNumBytes(void) const
 	{
@@ -61,7 +62,7 @@ public:
 	}
 
 	/**
-	 * Returns the number of bits written
+	 * Returns the number of bits written.
 	 */
 	FORCEINLINE int64 GetNumBits(void) const
 	{
@@ -69,7 +70,7 @@ public:
 	}
 
 	/**
-	 * Returns the number of bits the buffer supports
+	 * Returns the number of bits the buffer supports.
 	 */
 	FORCEINLINE int64 GetMaxBits(void) const
 	{
@@ -77,7 +78,7 @@ public:
 	}
 
 	/**
-	 * Marks this bit writer as overflowed
+	 * Marks this bit writer as overflowed.
 	 */
 	FORCEINLINE void SetOverflowed(void)
 	{
@@ -125,45 +126,53 @@ public:
 	}
 
 private:
+
 	TArray<uint8> Buffer;
 	int64   Num;
 	int64   Max;
 	bool	AllowResize;
 };
 
+
 //
 // For pushing and popping FBitWriter positions.
 //
 struct CORE_API FBitWriterMark
 {
 public:
+
 	FBitWriterMark()
-	:	Overflowed	( false ),
-		Num			( 0 )
-	{}
+		: Overflowed(false)
+		, Num(0)
+	{ }
+
 	FBitWriterMark( FBitWriter& Writer )
 	{
 		Init(Writer);
 	}
+
 	int64 GetNumBits()
 	{
 		return Num;
 	}
+
 	void Init( FBitWriter& Writer)
 	{
 		Num = Writer.Num;
 		Overflowed = Writer.ArIsError;
 	}
+
 	void Pop( FBitWriter& Writer );
 	void Copy( FBitWriter& Writer, TArray<uint8> &Buffer );
-
 	void PopWithoutClear( FBitWriter& Writer );
 
 private:
+
 	bool			Overflowed;
 	int64			Num;
 };
 
+
 /*-----------------------------------------------------------------------------
 	FBitReader.
 -----------------------------------------------------------------------------*/
@@ -174,8 +183,10 @@ private:
 struct CORE_API FBitReader : public FArchive
 {
 	friend struct FBitReaderMark;
+
 public:
-	FBitReader( uint8* Src=NULL, int64 CountBits=0 );
+
+	FBitReader( uint8* Src = nullptr, int64 CountBits = 0 );
 	void SetData( FBitReader& Src, int64 CountBits );
 	void SerializeBits( void* Dest, int64 LengthBits );
 	void SerializeInt( uint32& Value, uint32 Max );
@@ -195,30 +206,38 @@ public:
 	void AppendDataFromChecked( uint8* Src, uint32 NumBits );
 	void AppendTo( TArray<uint8> &Buffer );
 	void EatByteAlign();
+
 protected:
+
 	TArray<uint8> Buffer;
-	int64   Num;
-	int64   Pos;
+	int64 Num;
+	int64 Pos;
 };
 
+
 //
 // For pushing and popping FBitWriter positions.
 //
 struct CORE_API FBitReaderMark
 {
 public:
+
 	FBitReaderMark()
-		:	Pos         ( 0 )
-	{}
+		: Pos(0)
+	{ }
+
 	FBitReaderMark( FBitReader& Reader )
-		:	Pos			( Reader.Pos )
-	{}
+		: Pos(Reader.Pos)
+	{ }
+
 	int64 GetPos()
 	{
 		return Pos;
 	}
+
 	void Copy( FBitReader& Reader, TArray<uint8> &Buffer );
 
 private:
-	int64			Pos;
-};
\ No newline at end of file
+
+	int64 Pos;
+};
diff --git a/Engine/Source/Runtime/Core/Public/Serialization/CustomVersion.h b/Engine/Source/Runtime/Core/Public/Serialization/CustomVersion.h
--- a/Engine/Source/Runtime/Core/Public/Serialization/CustomVersion.h
+++ b/Engine/Source/Runtime/Core/Public/Serialization/CustomVersion.h
@@ -4,9 +4,6 @@
 
 #include "Misc/Guid.h"
 
-/*=============================================================================
-	CustomVersion.h: Unreal custom versioning system.
-=============================================================================*/
 
 struct ECustomVersionSerializationFormat
 {
@@ -17,6 +14,7 @@ struct ECustomVersionSerializationFormat
 	};
 };
 
+
 /**
  * Structure to hold unique custom key with its version.
  */
@@ -63,8 +61,10 @@ struct CORE_API FCustomVersion
 	CORE_API friend FArchive& operator<<(FArchive& Ar, FCustomVersion& Version);
 };
 
+
 class CORE_API FCustomVersionRegistration;
 
+
 /**
  * Container for all available/serialized custom versions.
  */
@@ -83,7 +83,7 @@ public:
 	 * Gets a custom version from the container.
 	 *
 	 * @param CustomKey Custom key for which to retrieve the version.
-	 * @return The FCustomVersion for the specified custom key, or NULL if the key doesn't exist in the container.
+	 * @return The FCustomVersion for the specified custom key, or nullptr if the key doesn't exist in the container.
 	 */
 	const FCustomVersion* GetVersion(FGuid CustomKey) const;
 
@@ -115,6 +115,7 @@ public:
 	FString ToString(const FString& Indent) const;
 
 private:
+
 	/** Private implementation getter */
 	static FCustomVersionContainer& GetInstance();
 
@@ -122,6 +123,7 @@ private:
 	TArray<FCustomVersion> Versions;
 };
 
+
 /**
  * This class will cause the registration of a custom version number and key with the global
  * FCustomVersionContainer when instantiated, and unregister when destroyed.  It should be
@@ -130,9 +132,11 @@ private:
 class FCustomVersionRegistration : FNoncopyable
 {
 public:
+
 	FCustomVersionRegistration(FGuid InKey, int32 Version, const TCHAR*);
 	~FCustomVersionRegistration();
 
 private:
+
 	FGuid Key;
 };
diff --git a/Engine/Source/Runtime/Core/Public/Stats/Stats2.h b/Engine/Source/Runtime/Core/Public/Stats/Stats2.h
--- a/Engine/Source/Runtime/Core/Public/Stats/Stats2.h
+++ b/Engine/Source/Runtime/Core/Public/Stats/Stats2.h
@@ -1053,7 +1053,7 @@ public:
 		FThreadStats* Stats = (FThreadStats*)FPlatformTLS::GetTlsValue(TlsSlot);
 		if (Stats)
 		{
-			FPlatformTLS::SetTlsValue(TlsSlot, NULL);
+			FPlatformTLS::SetTlsValue(TlsSlot, nullptr);
 			delete Stats;
 		}
 	}
@@ -1123,7 +1123,7 @@ public:
 	}
 	static FORCEINLINE_STATS bool IsCollectingData(TStatId StatId)
 	{
-		// we don't test StatId for NULL here because we assume it is non-null. If it is NULL, that indicates a problem with higher level code.
+		// we don't test StatId for nullptr here because we assume it is non-null. If it is nullptr, that indicates a problem with higher level code.
 		return !StatId.IsNone() && IsCollectingData();
 	}
 
diff --git a/Engine/Source/Runtime/Core/Public/Stats/StatsData.h b/Engine/Source/Runtime/Core/Public/Stats/StatsData.h
--- a/Engine/Source/Runtime/Core/Public/Stats/StatsData.h
+++ b/Engine/Source/Runtime/Core/Public/Stats/StatsData.h
@@ -1,10 +1,8 @@
 // Copyright 1998-2014 Epic Games, Inc. All Rights Reserved.
 
-/*=============================================================================
-	StatsData.h: Performance stats framework.
-=============================================================================*/
 #pragma once
 
+
 #if STATS
 
 /**
@@ -264,7 +262,7 @@ struct CORE_API FRawStatStackNode
 	void AddSelf();
 
 	/** Print this tree to the log **/
-	void DebugPrint(TCHAR const* Filter = NULL, int32 MaxDepth = MAX_int32, int32 Depth = 0) const;
+	void DebugPrint(TCHAR const* Filter = nullptr, int32 MaxDepth = MAX_int32, int32 Depth = 0) const;
 
 	/** Condense this tree into a flat list using EStatOperation::ChildrenStart, EStatOperation::ChildrenEnd, EStatOperation::Leaf **/
 	void Encode(TArray<FStatMessage>& OutStats) const;
@@ -448,7 +446,7 @@ private:
 	void FindOrAddMetaData(FStatMessage const& Item);
 
 	/** Parse the raw frame history into a tree representation **/
-	void GetRawStackStats(int64 FrameNumber, FRawStatStackNode& Out, TArray<FStatMessage>* OutNonStackStats = NULL) const;
+	void GetRawStackStats(int64 FrameNumber, FRawStatStackNode& Out, TArray<FStatMessage>* OutNonStackStats = nullptr) const;
 
 	/** Parse the raw frame history into a flat-array-based tree representation **/
 	void Condense(int64 TargetFrame, TArray<FStatMessage>& OutStats) const;
@@ -550,13 +548,13 @@ public:
 	}
 
 	/** Gets the old-skool flat grouped inclusive stats. These ignore recursion, merge threads, etc and so generally the condensed callstack is less confusing. **/
-	void GetInclusiveAggregateStackStats(int64 TargetFrame, TArray<FStatMessage>& OutStats, IItemFiler* Filter = NULL, bool bAddNonStackStats = true) const;
+	void GetInclusiveAggregateStackStats(int64 TargetFrame, TArray<FStatMessage>& OutStats, IItemFiler* Filter = nullptr, bool bAddNonStackStats = true) const;
 
 	/** Gets the old-skool flat grouped exclusive stats. These merge threads, etc and so generally the condensed callstack is less confusing. **/
-	void GetExclusiveAggregateStackStats(int64 TargetFrame, TArray<FStatMessage>& OutStats, IItemFiler* Filter = NULL, bool bAddNonStackStats = true) const;
+	void GetExclusiveAggregateStackStats(int64 TargetFrame, TArray<FStatMessage>& OutStats, IItemFiler* Filter = nullptr, bool bAddNonStackStats = true) const;
 
 	/** Used to turn the condensed version of stack stats back into a tree for easier handling. **/
-	void UncondenseStackStats(int64 TargetFrame, FRawStatStackNode& Out, IItemFiler* Filter = NULL, TArray<FStatMessage>* OutNonStackStats = NULL) const;
+	void UncondenseStackStats(int64 TargetFrame, FRawStatStackNode& Out, IItemFiler* Filter = nullptr, TArray<FStatMessage>* OutNonStackStats = nullptr) const;
 
 	/** Adds missing stats to the group so it doesn't jitter. **/
 	void AddMissingStats(TArray<FStatMessage>& Dest, TSet<FName> const& EnabledItems) const;
diff --git a/Engine/Source/Runtime/Core/Public/Stats/StatsFile.h b/Engine/Source/Runtime/Core/Public/Stats/StatsFile.h
--- a/Engine/Source/Runtime/Core/Public/Stats/StatsFile.h
+++ b/Engine/Source/Runtime/Core/Public/Stats/StatsFile.h
@@ -1,15 +1,13 @@
 // Copyright 1998-2014 Epic Games, Inc. All Rights Reserved.
 
-/*=============================================================================
-	StatsFile.cpp: Declares stats file related functionality.
-=============================================================================*/
-
 #pragma once
 
+
 #if	STATS
 
 class FAsyncWriteWorker;
 
+
 /**
 * Magic numbers for stats streams, this is for the first version.
 */
diff --git a/Engine/Source/Runtime/Core/Public/Stats/StatsMisc.h b/Engine/Source/Runtime/Core/Public/Stats/StatsMisc.h
--- a/Engine/Source/Runtime/Core/Public/Stats/StatsMisc.h
+++ b/Engine/Source/Runtime/Core/Public/Stats/StatsMisc.h
@@ -1,8 +1,5 @@
 // Copyright 1998-2014 Epic Games, Inc. All Rights Reserved.
 
-/*=============================================================================
-	StatsMisc.h: misc timers and such
-=============================================================================*/
 #pragma once
 
 
diff --git a/Engine/Source/Runtime/Core/Public/Templates/EnableIf.h b/Engine/Source/Runtime/Core/Public/Templates/EnableIf.h
--- a/Engine/Source/Runtime/Core/Public/Templates/EnableIf.h
+++ b/Engine/Source/Runtime/Core/Public/Templates/EnableIf.h
@@ -1,11 +1,8 @@
 // Copyright 1998-2014 Epic Games, Inc. All Rights Reserved.
 
-/*=============================================================================
-	EnableIf.h: Unreal common template definitions.
-=============================================================================*/
-
 #pragma once
 
+
 /**
  * Includes a function in an overload set if the predicate is true.  It should be used similarly to this:
  *
@@ -30,8 +27,8 @@ public:
 
 template <typename Result>
 class TEnableIf<false, Result>
-{
-};
+{ };
+
 
 /**
  * This is a variant of the above that will determine the return type 'lazily', i.e. only if the function is enabled.
@@ -59,5 +56,4 @@ public:
 
 template <typename Func>
 class TLazyEnableIf<false, Func>
-{
-};
+{ };
diff --git a/Engine/Source/Runtime/Core/Public/Templates/MemoryOps.h b/Engine/Source/Runtime/Core/Public/Templates/MemoryOps.h
--- a/Engine/Source/Runtime/Core/Public/Templates/MemoryOps.h
+++ b/Engine/Source/Runtime/Core/Public/Templates/MemoryOps.h
@@ -1,9 +1,5 @@
 // Copyright 1998-2014 Epic Games, Inc. All Rights Reserved.
 
-/*=============================================================================
-	MemoryOps.h: Functions to do bulk operations on contiguous blocks of objects.
-=============================================================================*/
-
 #pragma once
 
 #include <new>
@@ -13,6 +9,7 @@
 #include "UnrealTemplate.h"
 #include "UnrealTypeTraits.h"
 
+
 /**
  * Default constructs a range of items in memory.
  *
@@ -31,12 +28,14 @@ FORCEINLINE typename TEnableIf<!TIsZeroConstructType<ElementType>::Value>::Type
 	}
 }
 
+
 template <typename ElementType>
 FORCEINLINE typename TEnableIf<TIsZeroConstructType<ElementType>::Value>::Type DefaultConstructItems(void* Elements, int32 Count)
 {
 	FMemory::Memset(Elements, 0, sizeof(ElementType) * Count);
 }
 
+
 /**
  * Destructs a range of items in memory.
  *
@@ -57,11 +56,13 @@ FORCEINLINE typename TEnableIf<TTypeTraits<ElementType>::NeedsDestructor>::Type
 	}
 }
 
+
 template <typename ElementType>
 FORCEINLINE typename TEnableIf<!TTypeTraits<ElementType>::NeedsDestructor>::Type DestructItems(ElementType* Elements, int32 Count)
 {
 }
 
+
 /**
  * Copy constructs a range of items into memory.
  *
@@ -81,12 +82,14 @@ FORCEINLINE typename TEnableIf<TTypeTraits<ElementType>::NeedsCopyConstructor>::
 	}
 }
 
+
 template <typename ElementType>
 FORCEINLINE typename TEnableIf<!TTypeTraits<ElementType>::NeedsCopyConstructor>::Type CopyConstructItems(void* Dest, const ElementType* Source, int32 Count)
 {
 	FMemory::Memcpy(Dest, Source, sizeof(ElementType) * Count);
 }
 
+
 /**
  * Copy assigns a range of items.
  *
@@ -106,12 +109,14 @@ FORCEINLINE typename TEnableIf<TTypeTraits<ElementType>::NeedsCopyAssignment>::T
 	}
 }
 
+
 template <typename ElementType>
 FORCEINLINE typename TEnableIf<!TTypeTraits<ElementType>::NeedsCopyAssignment>::Type CopyAssignItems(ElementType* Dest, const ElementType* Source, int32 Count)
 {
 	FMemory::Memcpy(Dest, Source, sizeof(ElementType) * Count);
 }
 
+
 /**
  * Relocates a range of items to a new memory location. This is a so-called 'destructive move' for which
  * there is no single operation in C++ but which can be implemented very efficiently in general.
@@ -134,6 +139,7 @@ FORCEINLINE void RelocateItems(void* Dest, const ElementType* Source, int32 Coun
 	FMemory::Memmove(Dest, Source, sizeof(ElementType) * Count);
 }
 
+
 #if PLATFORM_COMPILER_HAS_RVALUE_REFERENCES
 
 	/**
@@ -208,6 +214,7 @@ FORCEINLINE typename TEnableIf<TTypeTraits<ElementType>::IsBytewiseComparable, b
 	return !FMemory::Memcmp(A, B, sizeof(ElementType) * Count);
 }
 
+
 template <typename ElementType>
 FORCEINLINE typename TEnableIf<!TTypeTraits<ElementType>::IsBytewiseComparable, bool>::Type CompareItems(const ElementType* A, const ElementType* B, int32 Count)
 {
diff --git a/Engine/Source/Runtime/Core/Public/Templates/RefCounting.h b/Engine/Source/Runtime/Core/Public/Templates/RefCounting.h
--- a/Engine/Source/Runtime/Core/Public/Templates/RefCounting.h
+++ b/Engine/Source/Runtime/Core/Public/Templates/RefCounting.h
@@ -1,11 +1,8 @@
 // Copyright 1998-2014 Epic Games, Inc. All Rights Reserved.
 
-/*=============================================================================
-	RefCounting.h: Reference counting definitions.
-=============================================================================*/
-
 #pragma once
 
+
 /** A virtual interface for ref counted objects to implement. */
 class IRefCountedObject
 {
@@ -15,6 +12,7 @@ public:
 	virtual uint32 GetRefCount() const = 0;
 };
 
+
 /**
  * The base class of reference counted objects.
  */
@@ -44,6 +42,7 @@ private:
 	mutable int32 NumRefs;
 };
 
+
 /**
  * A smart pointer to an object which implements AddRef/Release.
  */
@@ -51,11 +50,12 @@ template<typename ReferencedType>
 class TRefCountPtr
 {
 	typedef ReferencedType* ReferenceType;
+
 public:
 
 	TRefCountPtr():
-		Reference(NULL)
-	{}
+		Reference(nullptr)
+	{ }
 
 	TRefCountPtr(ReferencedType* InReference,bool bAddRef = true)
 	{
@@ -121,7 +121,7 @@ public:
 
 	ReferencedType** GetInitReference()
 	{
-		*this = NULL;
+		*this = nullptr;
 		return &Reference;
 	}
 
@@ -132,12 +132,12 @@ public:
 
 	friend bool IsValidRef(const TRefCountPtr& InReference)
 	{
-		return InReference.Reference != NULL;
+		return InReference.Reference != nullptr;
 	}
 
 	void SafeRelease()
 	{
-		*this = NULL;
+		*this = nullptr;
 	}
 
 	uint32 GetRefCount()
@@ -172,5 +172,6 @@ public:
 	}
 
 private:
+
 	ReferencedType* Reference;
 };
diff --git a/Engine/Source/Runtime/Core/Public/Templates/ScopedCallback.h b/Engine/Source/Runtime/Core/Public/Templates/ScopedCallback.h
--- a/Engine/Source/Runtime/Core/Public/Templates/ScopedCallback.h
+++ b/Engine/Source/Runtime/Core/Public/Templates/ScopedCallback.h
@@ -1,9 +1,10 @@
 // Copyright 1998-2014 Epic Games, Inc. All Rights Reserved.
 
-
 #pragma once
+
 #include "HAL/Platform.h"
 
+
 /**
  * Helper object for batching callback requests and firing on destruction of the FScopedCallback object.
  * CallbackType is a class implementing a static method called FireCallback, which does the work.
@@ -12,39 +13,36 @@ template< class CallbackType >
 class TScopedCallback
 {
 public:
+
 	TScopedCallback()
-		:	Counter( 0 )
-	{}
+		:	Counter(0)
+	{ }
 
-	/**
-	 * Fires a callback if outstanding requests exist.
-	 */
+	/** Fires a callback if outstanding requests exist. */
 	~TScopedCallback()
 	{
-		if ( HasRequests() )
+		if (HasRequests())
 		{
 			CallbackType::FireCallback();
 		}
 	}
 
-	/**
-	 * Request a callback.
-	 */
+	/** Request a callback. */
 	void Request()
 	{
 		++Counter;
 	}
 
-	/**
-	 * Unrequest a callback.
-	 */
+	/** Unrequest a callback. */
 	void Unrequest()
 	{
 		--Counter;
 	}
 
 	/**
-	 * @return	true if there are outstanding requests, false otherwise.
+	 * Checks whether this callback has outstanding requests.
+	 *
+	 * @return true if there are outstanding requests, false otherwise.
 	 */
 	bool HasRequests() const
 	{
@@ -52,7 +50,7 @@ public:
 	}
 
 private:
+
 	/** Counts callback requests. */
 	int32	Counter;
 };
-
diff --git a/Engine/Source/Runtime/Core/Public/Templates/ScopedPointer.h b/Engine/Source/Runtime/Core/Public/Templates/ScopedPointer.h
--- a/Engine/Source/Runtime/Core/Public/Templates/ScopedPointer.h
+++ b/Engine/Source/Runtime/Core/Public/Templates/ScopedPointer.h
@@ -1,11 +1,8 @@
 // Copyright 1998-2014 Epic Games, Inc. All Rights Reserved.
 
-/*=============================================================================
-	ScopedPointer.h: Scoped pointer type definition.
-=============================================================================*/
-
 #pragma once
 
+
 /**
  * Wrapper around a raw pointer that destroys it automatically.
  * Calls operator delete on the object, so it must have been allocated with
@@ -25,24 +22,25 @@ private:
 public:
 
 	/** Initialization constructor. */
-	explicit TScopedPointer(ReferencedType* InReference = NULL)
-	:	Reference(InReference)
-	{}
+	explicit TScopedPointer(ReferencedType* InReference = nullptr)
+		: Reference(InReference)
+	{ }
 
 	/** Copy constructor. */
 	TScopedPointer(const TScopedPointer& InCopy)
 	{
 		Reference = InCopy.Reference ?
 			new ReferencedType(*InCopy.Reference) :
-			NULL;
+			nullptr;
 	}
 
 	/** Destructor. */
 	~TScopedPointer()
 	{
 		delete Reference;
-		Reference = NULL;
+		Reference = nullptr;
 	}
+
 	/** Assignment operator. */
 	TScopedPointer& operator=(const TScopedPointer& InCopy)
 	{
@@ -51,16 +49,18 @@ public:
 			delete Reference;
 			Reference = InCopy.Reference ?
 				new ReferencedType(*InCopy.Reference) :
-				NULL;
+				nullptr;
 		}
 		return *this;
 	}
+
 	/** Assignment operator. */
 	TScopedPointer& operator=(ReferencedType* InReference)
 	{
 		Reset(InReference);
 		return *this;
 	}
+
 	// Dereferencing operators.
 	ReferencedType& operator*() const
 	{
@@ -99,7 +99,7 @@ public:
 	}
 
 	/** Deletes the current pointer and sets it to a new value. */
-	void Reset(ReferencedType* NewReference = NULL)
+	void Reset(ReferencedType* NewReference = nullptr)
 	{
 		check(!Reference || Reference != NewReference);
 		delete Reference;
@@ -110,7 +110,7 @@ public:
 	ReferencedType* Release()
 	{
 		ReferencedType* Result = GetOwnedPointer();
-		Reference = NULL;
+		Reference = nullptr;
 		return Result;
 	}
 
@@ -127,7 +127,7 @@ public:
 			delete OldReference;
 		}
 
-		// Serialize the value.  The caller of this serializer is responsible to only serialize for saving non-NULL pointers. */
+		// Serialize the value.  The caller of this serializer is responsible to only serialize for saving non-nullptr pointers. */
 		check(P.Reference);
 		Ar << *P.Reference;
 
@@ -135,6 +135,7 @@ public:
 	}
 };
 
+
 /** specialize container traits */
 template<typename ReferencedType>
 struct TTypeTraits<TScopedPointer<ReferencedType> > : public TTypeTraitsBase<TScopedPointer<ReferencedType> >
@@ -143,12 +144,14 @@ struct TTypeTraits<TScopedPointer<ReferencedType> > : public TTypeTraitsBase<TSc
 	typedef ReferencedType* ConstPointerType;
 };
 
+
 /** Implement movement of a scoped pointer to avoid copying the referenced value. */
 template<typename ReferencedType> void Move(TScopedPointer<ReferencedType>& A,ReferencedType* B)
 {
 	A.Reset(B);
 }
 
+
 /**
  * Wrapper around a raw pointer that destroys it automatically.
  * Calls operator delete on the object, so it must have been allocated with
@@ -167,7 +170,7 @@ private:
 public:
 
 	/** Initialization constructor. */
-	explicit TAutoPtr(ReferencedType* InReference = NULL)
+	explicit TAutoPtr(ReferencedType* InReference = nullptr)
 		:	Reference(InReference)
 	{}
 
@@ -175,7 +178,7 @@ public:
 	~TAutoPtr()
 	{
 		delete Reference;
-		Reference = NULL;
+		Reference = nullptr;
 	}
 
 	/** Assignment operator. */
@@ -222,7 +225,7 @@ public:
 	}
 
 	/** Deletes the current pointer and sets it to a new value. */
-	void Reset(ReferencedType* NewReference = NULL)
+	void Reset(ReferencedType* NewReference = nullptr)
 	{
 		check(!Reference || Reference != NewReference);
 		delete Reference;
@@ -231,6 +234,7 @@ public:
 
 };
 
+
 /** specialize container traits */
 template<typename ReferencedType>
 struct TTypeTraits<TAutoPtr<ReferencedType> > : public TTypeTraitsBase<TAutoPtr<ReferencedType> >
@@ -239,10 +243,9 @@ struct TTypeTraits<TAutoPtr<ReferencedType> > : public TTypeTraitsBase<TAutoPtr<
 	typedef ReferencedType* ConstPointerType;
 };
 
+
 /** Implement movement of a scoped pointer to avoid copying the referenced value. */
 template<typename ReferencedType> void Move(TAutoPtr<ReferencedType>& A,ReferencedType* B)
 {
 	A.Reset(B);
 }
-
-
diff --git a/Engine/Source/Runtime/Core/Public/Templates/SharedPointer.h b/Engine/Source/Runtime/Core/Public/Templates/SharedPointer.h
--- a/Engine/Source/Runtime/Core/Public/Templates/SharedPointer.h
+++ b/Engine/Source/Runtime/Core/Public/Templates/SharedPointer.h
@@ -1,14 +1,12 @@
 // Copyright 1998-2014 Epic Games, Inc. All Rights Reserved.
 
-/*=============================================================================
-	SharedPointer.h: Smart pointer library
-=============================================================================*/
-
 #pragma once
+
 #include "Containers/ContainersFwd.h"
 #include "Containers/Map.h"
 #include "Templates/UnrealTemplate.h"
 
+
 /**
  *	SharedPointer - Unreal smart pointer library
  *
@@ -112,7 +110,7 @@
 
 
 /**
- * Casts a shared reference of one type to another type. (static_cast)  Useful for downcasting.
+ * Casts a shared reference of one type to another type. (static_cast)  Useful for down-casting.
  *
  * @param  InSharedRef  The shared reference to cast
  */
diff --git a/Engine/Source/Runtime/Core/Public/Templates/SharedPointerInternals.h b/Engine/Source/Runtime/Core/Public/Templates/SharedPointerInternals.h
--- a/Engine/Source/Runtime/Core/Public/Templates/SharedPointerInternals.h
+++ b/Engine/Source/Runtime/Core/Public/Templates/SharedPointerInternals.h
@@ -1,14 +1,12 @@
 // Copyright 1998-2014 Epic Games, Inc. All Rights Reserved.
 
-/*=============================================================================
-	SharedPointerInternals.h: Smart pointer library internal implementation
-=============================================================================*/
-
 #pragma once
 
+
 /** Default behavior. */
 #define	FORCE_THREADSAFE_SHAREDPTRS 0
 
+
 /**
  * ESPMode is used select between either 'fast' or 'thread safe' shared pointer types.  This is
  * only used by templates at compile time to generate one code path or another.
@@ -42,7 +40,6 @@ template< class ObjectType, ESPMode::Type Mode = ESPMode::Fast > class TWeakPtr;
 template< class ObjectType, ESPMode::Type Mode = ESPMode::Fast > class TSharedFromThis;
 
 
-
 /**
  * SharedPointerInternals contains internal workings of shared and weak pointers.  You should
  * hopefully never have to use anything inside this namespace directly.
@@ -61,7 +58,6 @@ namespace SharedPointerInternals
 	struct FNullTag {};
 
 	
-
 	/** Proxy structure for implicitly converting raw pointers to shared/weak pointers */
 	// NOTE: The following is an Unreal extension to standard shared_ptr behavior
 	template< class ObjectType >
@@ -77,16 +73,16 @@ namespace SharedPointerInternals
 		}
 	};
 
+
 	struct FReferenceControllerState
 	{
 		/** Constructor */
 		FORCEINLINE FReferenceControllerState(void* InObject, void (*InDestroyObject)(void*))
-			: SharedReferenceCount(1),
-			  WeakReferenceCount  (1),
-			  Object              (InObject),
-			  DestroyObject       (InDestroyObject)
-		{
-		}
+			: SharedReferenceCount(1)
+			, WeakReferenceCount(1)
+			, Object(InObject)
+			, DestroyObject(InDestroyObject)
+		{ }
 
 		// NOTE: The primary reason these reference counters are 32-bit values (and not 16-bit to save
 		//       memory), is that atomic operations require at least 32-bit objects.
@@ -106,10 +102,12 @@ namespace SharedPointerInternals
 		void (*DestroyObject)(void*);
 
 	private:
+
 		FReferenceControllerState( FReferenceControllerState const& );
 		FReferenceControllerState& operator=( FReferenceControllerState const& );
 	};
 
+
 	/**
 	 * FReferenceController is a standalone heap-allocated object that tracks the number of references
 	 * to an object referenced by TSharedRef, TSharedPtr or TWeakPtr objects.
@@ -166,7 +164,6 @@ namespace SharedPointerInternals
 			}
 		}
 
-
 		/** Releases a shared reference to this counter */
 		static FORCEINLINE void ReleaseSharedReference(FReferenceControllerState* State)
 		{
@@ -203,6 +200,7 @@ namespace SharedPointerInternals
 		}
 	};
 
+
 	template<>
 	struct FReferenceControllerOps<ESPMode::NotThreadSafe>
 	{
@@ -235,7 +233,6 @@ namespace SharedPointerInternals
 			return true;
 		}
 
-
 		/** Releases a shared reference to this counter */
 		static FORCEINLINE void ReleaseSharedReference(FReferenceControllerState* State)
 		{
@@ -252,7 +249,6 @@ namespace SharedPointerInternals
 			}
 		}
 
-
 		/** Adds a weak reference to this counter */
 		static FORCEINLINE void AddWeakReference(FReferenceControllerState* State)
 		{
@@ -272,6 +268,7 @@ namespace SharedPointerInternals
 		}
 	};
 
+
 	/** Type-erasing function deleting objects of a given type via a common pointer type (void*) */
 	template <typename Type>
 	void DestroyObject(void* Object)
@@ -279,6 +276,7 @@ namespace SharedPointerInternals
 		delete (Type*)Object;
 	}
 
+
 	/**
 	 * FSharedReferencer is a wrapper around a pointer to a reference controller that is used by either a
 	 * TSharedRef or a TSharedPtr to keep track of a referenced object's lifetime
@@ -292,18 +290,14 @@ namespace SharedPointerInternals
 
 		/** Constructor for an empty shared referencer object */
 		FORCEINLINE FSharedReferencer()
-			: ReferenceController( NULL )
-		{
-		}
-
+			: ReferenceController( nullptr )
+		{ }
 		
 		/** Constructor that counts a single reference to the specified object */
 		template< class OtherType > 
 		inline explicit FSharedReferencer( OtherType* InObject )
 			: ReferenceController( new FReferenceControllerState( InObject, &DestroyObject<OtherType> ) )
-		{
-		}
-
+		{ }
 		
 		/** Copy constructor creates a new reference to the existing object */
 		FORCEINLINE FSharedReferencer( FSharedReferencer const& InSharedReference )
@@ -311,7 +305,7 @@ namespace SharedPointerInternals
 		{
 			// If the incoming reference had an object associated with it, then go ahead and increment the
 			// shared reference count
-			if( ReferenceController != NULL )
+			if( ReferenceController != nullptr )
 			{
 				TOps::AddSharedReference(ReferenceController);
 			}
@@ -321,10 +315,9 @@ namespace SharedPointerInternals
 		FORCEINLINE FSharedReferencer( FSharedReferencer&& InSharedReference )
 			: ReferenceController( InSharedReference.ReferenceController )
 		{
-			InSharedReference.ReferenceController = NULL;
+			InSharedReference.ReferenceController = nullptr;
 		}
 
-
 		/** Creates a shared referencer object from a weak referencer object.  This will only result
 		    in a valid object reference if the object already has at least one other shared referencer. */
 		FSharedReferencer( FWeakReferencer< Mode > const& InWeakReference )
@@ -332,23 +325,22 @@ namespace SharedPointerInternals
 		{
 			// If the incoming reference had an object associated with it, then go ahead and increment the
 			// shared reference count
-			if( ReferenceController != NULL )
+			if( ReferenceController != nullptr )
 			{
 				// Attempt to elevate a weak reference to a shared one.  For this to work, the object this
 				// weak counter is associated with must already have at least one shared reference.  We'll
 				// never revive a pointer that has already expired!
 				if( !TOps::ConditionallyAddSharedReference(ReferenceController) )
 				{
-					ReferenceController = NULL;
+					ReferenceController = nullptr;
 				}
 			}
 		}
 
-
-		/** Destructor */
+		/** Destructor. */
 		FORCEINLINE ~FSharedReferencer()
 		{
-			if( ReferenceController != NULL )
+			if( ReferenceController != nullptr )
 			{
 				// Tell the reference counter object that we're no longer referencing the object with
 				// this shared pointer
@@ -365,13 +357,13 @@ namespace SharedPointerInternals
 			if( NewReferenceController != ReferenceController )
 			{
 				// First, add a shared reference to the new object
-				if( NewReferenceController != NULL )
+				if( NewReferenceController != nullptr )
 				{
 					TOps::AddSharedReference(NewReferenceController);
 				}
 
 				// Release shared reference to the old object
-				if( ReferenceController != NULL )
+				if( ReferenceController != nullptr )
 				{
 					TOps::ReleaseSharedReference(ReferenceController);
 				}
@@ -393,11 +385,11 @@ namespace SharedPointerInternals
 			if( NewReferenceController != OldReferenceController )
 			{
 				// Assume ownership of the assigned reference counter
-				InSharedReference.ReferenceController = NULL;
+				InSharedReference.ReferenceController = nullptr;
 				ReferenceController                   = NewReferenceController;
 
 				// Release shared reference to the old object
-				if( OldReferenceController != NULL )
+				if( OldReferenceController != nullptr )
 				{
 					TOps::ReleaseSharedReference(OldReferenceController);
 				}
@@ -406,7 +398,6 @@ namespace SharedPointerInternals
 			return *this;
 		}
 
-
 		/**
 		 * Tests to see whether or not this shared counter contains a valid reference
 		 *
@@ -414,10 +405,9 @@ namespace SharedPointerInternals
 		 */
 		FORCEINLINE const bool IsValid() const
 		{
-			return ReferenceController != NULL;
+			return ReferenceController != nullptr;
 		}
 
-
 		/**
 		 * Returns the number of shared references to this object (including this reference.)
 		 *
@@ -425,10 +415,9 @@ namespace SharedPointerInternals
 		 */
 		FORCEINLINE const int32 GetSharedReferenceCount() const
 		{
-			return ReferenceController != NULL ? TOps::GetSharedReferenceCount(ReferenceController) : 0;
+			return ReferenceController != nullptr ? TOps::GetSharedReferenceCount(ReferenceController) : 0;
 		}
 
-
 		/**
 		 * Returns true if this is the only shared reference to this object.  Note that there may be
 		 * outstanding weak references left.
@@ -440,13 +429,11 @@ namespace SharedPointerInternals
 			return GetSharedReferenceCount() == 1;
 		}
 
-
 	private:
 
  		// Expose access to ReferenceController to FWeakReferencer
 		template< ESPMode::Type OtherMode > friend class FWeakReferencer;
 
-
 	private:
 
 		/** Pointer to the reference controller for the object a shared reference/pointer is referencing */
@@ -454,7 +441,6 @@ namespace SharedPointerInternals
 	};
 
 
-
 	/**
 	 * FWeakReferencer is a wrapper around a pointer to a reference controller that is used
 	 * by a TWeakPtr to keep track of a referenced object's lifetime.
@@ -468,54 +454,48 @@ namespace SharedPointerInternals
 
 		/** Default constructor with empty counter */
 		FORCEINLINE FWeakReferencer()
-			: ReferenceController( NULL )
-		{
-		}
-
+			: ReferenceController( nullptr )
+		{ }
 
 		/** Construct a weak referencer object from another weak referencer */
 		FORCEINLINE FWeakReferencer( FWeakReferencer const& InWeakRefCountPointer )
 			: ReferenceController( InWeakRefCountPointer.ReferenceController )
 		{
 			// If the weak referencer has a valid controller, then go ahead and add a weak reference to it!
-			if( ReferenceController != NULL )
+			if( ReferenceController != nullptr )
 			{
 				TOps::AddWeakReference(ReferenceController);
 			}
 		}
 
-
 		/** Construct a weak referencer object from an rvalue weak referencer */
 		FORCEINLINE FWeakReferencer( FWeakReferencer&& InWeakRefCountPointer )
 			: ReferenceController( InWeakRefCountPointer.ReferenceController )
 		{
-			InWeakRefCountPointer.ReferenceController = NULL;
+			InWeakRefCountPointer.ReferenceController = nullptr;
 		}
 
-
 		/** Construct a weak referencer object from a shared referencer object */
 		FORCEINLINE FWeakReferencer( FSharedReferencer< Mode > const& InSharedRefCountPointer )
 			: ReferenceController( InSharedRefCountPointer.ReferenceController )
 		{
 			// If the shared referencer had a valid controller, then go ahead and add a weak reference to it!
-			if( ReferenceController != NULL )
+			if( ReferenceController != nullptr )
 			{
 				TOps::AddWeakReference(ReferenceController);
 			}
 		}
 
-
-		/** Destructor */
+		/** Destructor. */
 		FORCEINLINE ~FWeakReferencer()
 		{
-			if( ReferenceController != NULL )
+			if( ReferenceController != nullptr )
 			{
 				// Tell the reference counter object that we're no longer referencing the object with
 				// this weak pointer
 				TOps::ReleaseWeakReference(ReferenceController);
 			}
 		}
-
 		
 		/** Assignment operator from a weak referencer object.  If this counter was previously referencing an
 		    object, that reference will be released. */
@@ -526,15 +506,14 @@ namespace SharedPointerInternals
 			return *this;
 		}
 
-
 		/** Assignment operator from an rvalue weak referencer object.  If this counter was previously referencing an
 		    object, that reference will be released. */
 		FORCEINLINE FWeakReferencer& operator=( FWeakReferencer&& InWeakReference )
 		{
 			auto OldReferenceController         = ReferenceController;
 			ReferenceController                 = InWeakReference.ReferenceController;
-			InWeakReference.ReferenceController = NULL;
-			if( OldReferenceController != NULL )
+			InWeakReference.ReferenceController = nullptr;
+			if( OldReferenceController != nullptr )
 			{
 				TOps::ReleaseWeakReference(OldReferenceController);
 			}
@@ -542,7 +521,6 @@ namespace SharedPointerInternals
 			return *this;
 		}
 
-
 		/** Assignment operator from a shared reference counter.  If this counter was previously referencing an
 		   object, that reference will be released. */
 		FORCEINLINE FWeakReferencer& operator=( FSharedReferencer< Mode > const& InSharedReference )
@@ -552,7 +530,6 @@ namespace SharedPointerInternals
 			return *this;
 		}
 
-
 		/**
 		 * Tests to see whether or not this weak counter contains a valid reference
 		 *
@@ -560,10 +537,9 @@ namespace SharedPointerInternals
 		 */
 		FORCEINLINE const bool IsValid() const
 		{
-			return ReferenceController != NULL && TOps::GetSharedReferenceCount(ReferenceController) > 0;
+			return ReferenceController != nullptr && TOps::GetSharedReferenceCount(ReferenceController) > 0;
 		}
 
-
 	private:
 
 		/** Assigns a new reference controller to this counter object, first adding a reference to it, then
@@ -574,13 +550,13 @@ namespace SharedPointerInternals
 			if( NewReferenceController != ReferenceController )
 			{
 				// First, add a weak reference to the new object
-				if( NewReferenceController != NULL )
+				if( NewReferenceController != nullptr )
 				{
 					TOps::AddWeakReference(NewReferenceController);
 				}
 
 				// Release weak reference to the old object
-				if( ReferenceController != NULL )
+				if( ReferenceController != nullptr )
 				{
 					TOps::ReleaseWeakReference(ReferenceController);
 				}
@@ -590,13 +566,11 @@ namespace SharedPointerInternals
 			}
 		}
 
-
 	private:
 
- 		// Expose access to ReferenceController to FSharedReferencer
+ 		/** Expose access to ReferenceController to FSharedReferencer. */
 		template< ESPMode::Type OtherMode > friend class FSharedReferencer;
 
-
 	private:
 
 		/** Pointer to the reference controller for the object a TWeakPtr is referencing */
@@ -608,7 +582,7 @@ namespace SharedPointerInternals
 	template< class SharedPtrType, class ObjectType, class OtherType, ESPMode::Type Mode >
 	FORCEINLINE void EnableSharedFromThis( TSharedPtr< SharedPtrType, Mode > const* InSharedPtr, ObjectType const* InObject, TSharedFromThis< OtherType, Mode > const* InShareable )
 	{
-		if( InShareable != NULL )
+		if( InShareable != nullptr )
 		{
 			InShareable->UpdateWeakReferenceInternal( InSharedPtr, const_cast< ObjectType* >( InObject ) );
 		}
@@ -619,7 +593,7 @@ namespace SharedPointerInternals
 	template< class SharedPtrType, class ObjectType, class OtherType, ESPMode::Type Mode >
 	FORCEINLINE void EnableSharedFromThis( TSharedPtr< SharedPtrType, Mode >* InSharedPtr, ObjectType const* InObject, TSharedFromThis< OtherType, Mode > const* InShareable )
 	{
-		if( InShareable != NULL )
+		if( InShareable != nullptr )
 		{
 			InShareable->UpdateWeakReferenceInternal( InSharedPtr, const_cast< ObjectType* >( InObject ) );
 		}
@@ -630,7 +604,7 @@ namespace SharedPointerInternals
 	template< class SharedRefType, class ObjectType, class OtherType, ESPMode::Type Mode >
 	FORCEINLINE void EnableSharedFromThis( TSharedRef< SharedRefType, Mode > const* InSharedRef, ObjectType const* InObject, TSharedFromThis< OtherType, Mode > const* InShareable )
 	{
-		if( InShareable != NULL )
+		if( InShareable != nullptr )
 		{
 			InShareable->UpdateWeakReferenceInternal( InSharedRef, const_cast< ObjectType* >( InObject ) );
 		}
@@ -641,7 +615,7 @@ namespace SharedPointerInternals
 	template< class SharedRefType, class ObjectType, class OtherType, ESPMode::Type Mode >
 	FORCEINLINE void EnableSharedFromThis( TSharedRef< SharedRefType, Mode >* InSharedRef, ObjectType const* InObject, TSharedFromThis< OtherType, Mode > const* InShareable )
 	{
-		if( InShareable != NULL )
+		if( InShareable != nullptr )
 		{
 			InShareable->UpdateWeakReferenceInternal( InSharedRef, const_cast< ObjectType* >( InObject ) );
 		}
@@ -649,12 +623,5 @@ namespace SharedPointerInternals
 
 
 	/** Templated helper catch-all function, accomplice to the above helper functions */
-	FORCEINLINE void EnableSharedFromThis( ... )
-	{
-	}
-
-
-}	// namespace SharedPointerInternals
-
-
-
+	FORCEINLINE void EnableSharedFromThis( ... ) { }
+}
diff --git a/Engine/Source/Runtime/Core/Public/Templates/SharedPointerTesting.h b/Engine/Source/Runtime/Core/Public/Templates/SharedPointerTesting.h
--- a/Engine/Source/Runtime/Core/Public/Templates/SharedPointerTesting.h
+++ b/Engine/Source/Runtime/Core/Public/Templates/SharedPointerTesting.h
@@ -1,12 +1,7 @@
 // Copyright 1998-2014 Epic Games, Inc. All Rights Reserved.
 
-/*=============================================================================
-	SharedPointerTesting.h: Smart pointer testing suite
-=============================================================================*/
-
 #pragma once
 
-
 #include "SharedPointer.h"
 
 
@@ -16,7 +11,6 @@
 
 #if WITH_SHARED_POINTER_TESTS
 
-
 /**
  * Shared pointer testing suite
  */
@@ -37,11 +31,10 @@ namespace SharedPointerTesting
 			check( !MyEmptyBoolPtr.IsValid() );
 
 			// Test access to raw pointer
-			if( MyEmptyBoolPtr.Get() == NULL )
+			if( MyEmptyBoolPtr.Get() == nullptr )
 			{
 				// ...
 			}
-
 		}
 
 		{
@@ -107,7 +100,6 @@ namespace SharedPointerTesting
 			// NOTE: OtherSharedArrayReference goes out of scope here (object is destroyed)
 		}
 
-
 		{
 			// Test casting
 			class FBase
@@ -116,14 +108,12 @@ namespace SharedPointerTesting
 			};
 			class FDerived
 				: public FBase
-			{
-			};
+			{ };
 
 			{
 				// Explicit downcast to derived shared ptr
 				TSharedPtr< FBase, Mode > DerivedAsBasePtr( new FDerived() );
 				TSharedPtr< FDerived, Mode > DerivedPtr( StaticCastSharedPtr< FDerived >( DerivedAsBasePtr ) );
-
 			}
 
 			{
@@ -137,28 +127,22 @@ namespace SharedPointerTesting
 				TSharedPtr< FDerived, Mode > DerivedPtr( new FDerived() );
 				TSharedPtr< FBase, Mode > BasePtr = DerivedPtr;
 			}
-
-
 		}
 
-
-		// Create a shared pointer to NULL.  Consistent with std::shared_ptr, this results
+		// Create a shared pointer to nullptr.  Consistent with std::shared_ptr, this results
 		// in a non-empty TSharedPtr (heap allocated reference count)
 		{
-			bool* Foo = NULL;
+			bool* Foo = nullptr;
 			TSharedPtr< bool, Mode > NullPtr( Foo );
 			check( !NullPtr.IsValid() );
 		}
 
-
 		// Simple validity test syntax
 		{
 			TSharedPtr< bool, Mode > BoolPtr( new bool( true ) );
 			check( BoolPtr.IsValid() );
 		}
 
-		
-
 		// Empty weak pointer
 		{
 			TWeakPtr< bool, Mode > EmptyBoolWeakPtr;
@@ -176,7 +160,6 @@ namespace SharedPointerTesting
 			check( WeakInt.Pin().IsValid() );
 		}
 
-
 		// Create weak pointer from shared pointer (assignment)
 		{
 			TSharedPtr< int32, Mode > SharedInt( new int32( 64 ) );
@@ -190,7 +173,6 @@ namespace SharedPointerTesting
 			check( !WeakInt.Pin().IsValid() );
 		}
 
-
 		// Test weak pointer becoming invalid
 		{
 			TSharedPtr< int32, Mode > SharedInt( new int32( 64 ) );
@@ -199,7 +181,6 @@ namespace SharedPointerTesting
 			check( !WeakInt.Pin().IsValid() );
 		}
 
-
 		// Compare shared pointers
 		{
 			TSharedPtr< int32, Mode > SharedA( new int32( 64 ) );
@@ -211,7 +192,6 @@ namespace SharedPointerTesting
 			check( SharedB == SharedC );
 		}
 
-
 		// Compare weak pointers
 		{
 			TSharedPtr< int32, Mode > SharedA( new int32( 64 ) );
@@ -237,7 +217,6 @@ namespace SharedPointerTesting
 #endif
 		}
 
-
 		// Test 'const'
 		{
 			TSharedPtr< const int32, Mode > IntPtr( new int32( 10 ) );
@@ -266,7 +245,6 @@ namespace SharedPointerTesting
 			// Test conversion from mutable to const (this is OK!)
 			FloatPtrA = MutableFloat;
 
-
 			if( FloatPtrB.IsValid() )
 			{
 #if 0			// Won't compile as value is const
@@ -289,7 +267,6 @@ namespace SharedPointerTesting
 			WeakFloat = ConstCastSharedPtr< float >( FloatPtrB );
 
 			*WeakFloat.Pin() = 20.0f;
-
 		}
 
 		// Test forward declaring a smart pointer to an incomplete type
@@ -302,22 +279,20 @@ namespace SharedPointerTesting
 			VecPtr = TSharedPtr< FBarFoo, Mode >( new FBarFoo() );
 			VecPtr->Val = 20;
 		}
-
 		
 		// Test Unreal non-standard extensions (expanded syntax)
 		{
-			// Initialize with NULL
-			TSharedPtr< bool, Mode > EmptyPtr( NULL );
-			TSharedPtr< float, Mode > FloatPtr = NULL;
+			// Initialize with nullptr
+			TSharedPtr< bool, Mode > EmptyPtr( nullptr );
+			TSharedPtr< float, Mode > FloatPtr = nullptr;
 
-			// Initialize with NULL
-			TWeakPtr< bool, Mode > EmptyWeakPtr( NULL );
-			TWeakPtr< float, Mode > FloatWeakPtr = NULL;
+			// Initialize with nullptr
+			TWeakPtr< bool, Mode > EmptyWeakPtr( nullptr );
+			TWeakPtr< float, Mode > FloatWeakPtr = nullptr;
 
-			// Reassign to NULL directly (instead of calling Reset)
+			// Reassign to nullptr directly (instead of calling Reset)
 			FloatPtr = TSharedPtr< float, Mode >( new float( 0.1f ) );
-			FloatPtr = NULL;
-
+			FloatPtr = nullptr;
 
 			// Test implicit construction helper (MakeShareable)
 			FloatPtr = MakeShareable( new float( 30.0f ) );
@@ -340,10 +315,8 @@ namespace SharedPointerTesting
 					return MakeShareable( new float( 123.0f ) );
 				}
 			};
-
 		}
 
-
 		// Test TSharedRef
 		{
 			// Empty TSharedRef is not allowed
@@ -352,8 +325,8 @@ namespace SharedPointerTesting
 				TSharedRef< bool, Mode > EmptyRef;
 #endif
 
-#if 0			// Won't compile as TSharedRef has no implicit constructor that a pointer (including NULL)
-				TSharedRef< bool, Mode > NullRef = NULL;
+#if 0			// Won't compile as TSharedRef has no implicit constructor that a pointer (including nullptr)
+				TSharedRef< bool, Mode > NullRef = nullptr;
 #endif
 
 				// TSharedRef initialized in constructor
@@ -361,7 +334,6 @@ namespace SharedPointerTesting
 					TSharedRef< float, Mode > FloatRef( new float( 123.0f ) );
 				}
 
-
 				// Test reference access
 				{
 					TSharedRef< float, Mode > FloatRef( new float( 123.0f ) );
@@ -373,7 +345,6 @@ namespace SharedPointerTesting
 					const float& MyFloat2 = FloatRef.Get();
 				}
 
-
 				// Implicit conversion to TSharedRef is never allowed (use MakeShareable!)
 				{
 #if 0				// Won't compile as constructor is explicit
@@ -386,23 +357,21 @@ namespace SharedPointerTesting
 					TSharedRef< float, Mode > FloatRef = MakeShareable( new float( 123.0f ) );
 				}
 
-				// New shared ref with NULL object pointer is not allowed!  The following code will compile
+				// New shared ref with nullptr object pointer is not allowed!  The following code will compile
 				// correctly, but will trigger a runtime assertion.
 				if( 0 )
 				{
 					// NOTE: The following code is unsafe and will trigger an assert
-					int32* NullInt = NULL;
+					int32* NullInt = nullptr;
 					TSharedRef< int32, Mode > RefWithNullObject( NullInt );
 				}
 
-
 				// Implicit conversion from a TSharedRef to a TSharedPtr (always OK)
 				{
 					TSharedRef< int32, Mode > MySharedRef( new int32( 1 ) );
 					TSharedPtr< int32, Mode > MySharedPtr( MySharedRef );
 				}
 
-
 				// Explicit TSharedRef construction from a TSharedPtr is not allowed
 				{
 					TSharedPtr< int32, Mode > MySharedPtr( new int32( 1 ) );
@@ -428,22 +397,22 @@ namespace SharedPointerTesting
 				// Conversion from an invalid TSharedPtr to a TSharedRef.  Never safe, will assert at runtime.
 				if( 0 )		// Will trigger a runtime assert as MySharedPtr.IsValid() == false
 				{
-					int32* NullInt = NULL;
+					int32* NullInt = nullptr;
 					TSharedPtr< int32, Mode > MySharedPtr( NullInt );
 					TSharedRef< int32, Mode > MySharedRef( MySharedPtr.ToSharedRef() );
 				}
 
-				// TSharedRef reassignment.  Safe as long as new object is not NULL.
+				// TSharedRef reassignment.  Safe as long as new object is not nullptr.
 				{
 					TSharedRef< int32, Mode > IntRef( new int32( 10 ) );
 					IntRef = TSharedRef< int32, Mode >( new int32( 20 ) );
 				}
 
-				// TSharedRef reassignment to NULL object is not allowed (asserts at runtime)
+				// TSharedRef reassignment to nullptr object is not allowed (asserts at runtime)
 				if( 0 )
 				{
 					TSharedRef< int32, Mode > IntRef( new int32( 10 ) );
-					int32* NullInt = NULL;
+					int32* NullInt = nullptr;
 					IntRef = TSharedRef< int32, Mode >( NullInt );
 				}
 
@@ -499,7 +468,6 @@ namespace SharedPointerTesting
 			}
 		}
 
-
 		// SharedFromThis
 		{
 			class FMyClass
@@ -527,7 +495,6 @@ namespace SharedPointerTesting
 			}
 		}
 
-
 		// Hash tables
 		{
 			TSharedRef< int32, Mode > FooRef( new int32( 1 ) );
@@ -546,7 +513,7 @@ namespace SharedPointerTesting
 			TMap< int32, TSharedRef< int32, Mode > > SharedRefValueHash;
 			SharedRefValueHash.Add( 10, FooRef );
 			const int32* FoundKey = SharedRefValueHash.FindKey( FooRef );
-			check( FoundKey != NULL && *FoundKey == 10 );
+			check( FoundKey != nullptr && *FoundKey == 10 );
 
 			// Maps and 'const'
 			{
@@ -554,14 +521,11 @@ namespace SharedPointerTesting
 				TMap< int32, TSharedRef< const int32, Mode > > ConstSharedRefValueHash;
 				ConstSharedRefValueHash.Add( 10, ConstFooRef );
 				const int32* FoundKey = ConstSharedRefValueHash.FindKey( ConstFooRef );
-				check( FoundKey != NULL && *FoundKey == 10 );
+				check( FoundKey != nullptr && *FoundKey == 10 );
 			}
 		}
 	}
-
 }
 
 
 #endif		// WITH_SHARED_POINTER_TESTS
-
-
diff --git a/Engine/Source/Runtime/Core/Public/Templates/Sorting.h b/Engine/Source/Runtime/Core/Public/Templates/Sorting.h
--- a/Engine/Source/Runtime/Core/Public/Templates/Sorting.h
+++ b/Engine/Source/Runtime/Core/Public/Templates/Sorting.h
@@ -1,14 +1,7 @@
 // Copyright 1998-2014 Epic Games, Inc. All Rights Reserved.
 
-/*=============================================================================
-	Sorting.h: Generic sorting definitions.
-=============================================================================*/
-
 #pragma once
 
-/*----------------------------------------------------------------------------
-	Sorting template.
-----------------------------------------------------------------------------*/
 
 /**
  * Default predicate class. Assumes < operator is defined for the template type.
diff --git a/Engine/Source/Runtime/Core/Public/Templates/TypeHash.h b/Engine/Source/Runtime/Core/Public/Templates/TypeHash.h
--- a/Engine/Source/Runtime/Core/Public/Templates/TypeHash.h
+++ b/Engine/Source/Runtime/Core/Public/Templates/TypeHash.h
@@ -1,11 +1,8 @@
 // Copyright 1998-2014 Epic Games, Inc. All Rights Reserved.
 
-/*=============================================================================
-	UnrealTypeHash.h: Unreal type hash functions.
-=============================================================================*/
-
 #pragma once
 
+
 /**
  * Combines two hash values to get a third.
  * Note - this function is not commutative.
@@ -28,6 +25,7 @@ inline uint32 HashCombine(uint32 A, uint32 C)
 	return C;
 }
 
+
 inline uint32 PointerHash(const void* Key,uint32 C = 0)
 {
 	// Avoid LHS stalls on PS3 and Xbox 360
@@ -42,6 +40,7 @@ inline uint32 PointerHash(const void* Key,uint32 C = 0)
 	return HashCombine(PtrInt, C);
 }
 
+
 //
 // Hash functions for common types.
 //
@@ -50,38 +49,47 @@ inline uint32 GetTypeHash( const uint8 A )
 {
 	return A;
 }
+
 inline uint32 GetTypeHash( const int8 A )
 {
 	return A;
 }
+
 inline uint32 GetTypeHash( const uint16 A )
 {
 	return A;
 }
+
 inline uint32 GetTypeHash( const int16 A )
 {
 	return A;
 }
+
 inline uint32 GetTypeHash( const int32 A )
 {
 	return A;
 }
+
 inline uint32 GetTypeHash( const uint32 A )
 {
 	return A;
 }
+
 inline uint32 GetTypeHash( const uint64 A )
 {
 	return (uint32)A+((uint32)(A>>32) * 23);
 }
+
 inline uint32 GetTypeHash( const int64 A )
 {
 	return (uint32)A+((uint32)(A>>32) * 23);
 }
+
 inline uint32 GetTypeHash( const float& Value )
 {
 	return *(uint32*)&Value;
 }
+
 inline uint32 GetTypeHash( const TCHAR* S )
 {
 	return FCrc::Strihash_DEPRECATED(S);
diff --git a/Engine/Source/Runtime/Core/Public/Templates/UniquePtr.h b/Engine/Source/Runtime/Core/Public/Templates/UniquePtr.h
--- a/Engine/Source/Runtime/Core/Public/Templates/UniquePtr.h
+++ b/Engine/Source/Runtime/Core/Public/Templates/UniquePtr.h
@@ -11,7 +11,7 @@
 //
 // TUniquePtr<MyClass> Ptr1(new MyClass);    // The MyClass object is owned by Ptr1.
 // TUniquePtr<MyClass> Ptr2(Ptr1);           // Error - TUniquePtr is not copyable
-// TUniquePtr<MyClass> Ptr3(MoveTemp(Ptr1)); // Ptr3 now owns the MyClass object - Ptr1 is now NULL.
+// TUniquePtr<MyClass> Ptr3(MoveTemp(Ptr1)); // Ptr3 now owns the MyClass object - Ptr1 is now nullptr.
 
 template <typename T>
 class TUniquePtr
@@ -24,7 +24,7 @@ public:
 	 * Default constructor - initializes the TUniquePtr to null.
 	 */
 	FORCEINLINE TUniquePtr()
-		: Ptr(NULL)
+		: Ptr(nullptr)
 	{
 	}
 
@@ -44,7 +44,7 @@ public:
 	FORCEINLINE TUniquePtr(TUniquePtr&& Other)
 		: Ptr(Other.Ptr)
 	{
-		Other.Ptr = NULL;
+		Other.Ptr = nullptr;
 	}
 
 	/**
@@ -54,7 +54,7 @@ public:
 	FORCEINLINE TUniquePtr(TUniquePtr<OtherT>&& Other)
 		: Ptr(Other.Ptr)
 	{
-		Other.Ptr = NULL;
+		Other.Ptr = nullptr;
 	}
 
 	/**
@@ -67,7 +67,7 @@ public:
 			// We delete last, because we don't want odd side effects if the destructor of T relies on the state of this or Other
 			T* OldPtr = Ptr;
 			Ptr = Other.Ptr;
-			Other.Ptr = NULL;
+			Other.Ptr = nullptr;
 			delete OldPtr;
 		}
 
@@ -85,7 +85,7 @@ public:
 			// We delete last, because we don't want odd side effects if the destructor of T relies on the state of this or Other
 			T* OldPtr = Ptr;
 			Ptr = Other.Ptr;
-			Other.Ptr = NULL;
+			Other.Ptr = nullptr;
 			delete OldPtr;
 		}
 
@@ -143,7 +143,7 @@ public:
 	/**
 	 * Returns a pointer to the owned object without relinquishing ownership.
 	 *
-	 * @return A copy of the pointer to the object owned by the TUniquePtr, or NULL if no object is being owned.
+	 * @return A copy of the pointer to the object owned by the TUniquePtr, or nullptr if no object is being owned.
 	 */
 	FORCEINLINE T* Get() const
 	{
@@ -153,12 +153,12 @@ public:
 	/**
 	 * Relinquishes control of the owned object to the caller and nulls the TUniquePtr.
 	 *
-	 * @return The pointer to the object that was owned by the TUniquePtr, or NULL if no object was being owned.
+	 * @return The pointer to the object that was owned by the TUniquePtr, or nullptr if no object was being owned.
 	 */
 	FORCEINLINE T* Release()
 	{
 		T* Result = Ptr;
-		Ptr = NULL;
+		Ptr = nullptr;
 		return Result;
 	}
 
@@ -167,7 +167,7 @@ public:
 	 *
 	 * @param InPtr A pointer to the object to take ownership of.
 	 */
-	FORCEINLINE void Reset(T* InPtr = NULL)
+	FORCEINLINE void Reset(T* InPtr = nullptr)
 	{
 		// We delete last, because we don't want odd side effects if the destructor of T relies on the state of this or Other
 		T* OldPtr = Ptr;
diff --git a/Engine/Source/Runtime/Core/Public/Templates/UnrealTemplate.h b/Engine/Source/Runtime/Core/Public/Templates/UnrealTemplate.h
--- a/Engine/Source/Runtime/Core/Public/Templates/UnrealTemplate.h
+++ b/Engine/Source/Runtime/Core/Public/Templates/UnrealTemplate.h
@@ -1,14 +1,12 @@
 // Copyright 1998-2014 Epic Games, Inc. All Rights Reserved.
 
-/*=============================================================================
-	UnrealTemplate.h: Unreal common template definitions.
-=============================================================================*/
-
 #pragma once
+
 #include "EnableIf.h"
 #include "HAL/Platform.h"
 #include "Templates/UnrealTypeTraits.h"
 
+
 /*-----------------------------------------------------------------------------
 	Standard templates.
 -----------------------------------------------------------------------------*/
@@ -88,8 +86,8 @@ template< class T > inline T AlignArbitrary( const T Ptr, uint32 Alignment )
 }
 
 /**
- * Chooses between the two parameters based on whether the first is NULL or not.
- * @return If the first parameter provided is non-NULL, it is returned; otherwise the second parameter is returned.
+ * Chooses between the two parameters based on whether the first is nullptr or not.
+ * @return If the first parameter provided is non-nullptr, it is returned; otherwise the second parameter is returned.
  */
 template<typename ReferencedType>
 FORCEINLINE ReferencedType* IfAThenAElseB(ReferencedType* A,ReferencedType* B)
diff --git a/Engine/Source/Runtime/Core/Public/UObject/AutoPointer.h b/Engine/Source/Runtime/Core/Public/UObject/AutoPointer.h
--- a/Engine/Source/Runtime/Core/Public/UObject/AutoPointer.h
+++ b/Engine/Source/Runtime/Core/Public/UObject/AutoPointer.h
@@ -1,11 +1,8 @@
 // Copyright 1998-2014 Epic Games, Inc. All Rights Reserved.
 
-/*=============================================================================
-	AutoPointer.h: Wrapper to convert various smart pointers into "auto" pointers
-=============================================================================*/
-
 #pragma once
 
+
 template<class T, class TBASE>
 class TAutoPointer : private TBASE
 {
@@ -157,5 +154,5 @@ public:
 
 };
 
-template<class T, class TBASE> struct TIsPODType<TAutoPointer<T, TBASE> > { enum { Value = TIsPODType<TBASE>::Value }; }; // pod-ness is the same as the podness of the base pointer type
 
+template<class T, class TBASE> struct TIsPODType<TAutoPointer<T, TBASE> > { enum { Value = TIsPODType<TBASE>::Value }; }; // pod-ness is the same as the podness of the base pointer type
diff --git a/Engine/Source/Runtime/Core/Public/UObject/NameTypes.h b/Engine/Source/Runtime/Core/Public/UObject/NameTypes.h
--- a/Engine/Source/Runtime/Core/Public/UObject/NameTypes.h
+++ b/Engine/Source/Runtime/Core/Public/UObject/NameTypes.h
@@ -282,7 +282,7 @@ class TStaticIndirectArrayThreadSafeRead
 			ElementType*** Chunk = &Chunks[ChunkIndex];
 			ElementType** NewChunk = (ElementType**)FMemory::Malloc(sizeof(ElementType*) * ElementsPerChunk);
 			FMemory::Memzero(NewChunk, sizeof(ElementType*) * ElementsPerChunk);
-			if (FPlatformAtomics::InterlockedCompareExchangePointer((void**)Chunk, NewChunk, NULL))
+			if (FPlatformAtomics::InterlockedCompareExchangePointer((void**)Chunk, NewChunk, nullptr))
 			{
 				// someone else beat us to the add, we don't support multiple concurrent adds
 				check(0)
@@ -339,7 +339,7 @@ public:
 	 * Return a reference to an element
 	 * @param	Index	Index to return
 	 * @return	a reference to the pointer to the element
-	 * Thread safe, if it is valid now, it is valid forever. This might return NULL, but by then, some other thread might have made it non-NULL.
+	 * Thread safe, if it is valid now, it is valid forever. This might return nullptr, but by then, some other thread might have made it non-nullptr.
 	**/
 	ElementType const* const& operator[](int32 Index) const
 	{
@@ -557,8 +557,8 @@ public:
 	FORCEINLINE bool IsValid() const
 	{
 		TNameEntryArray& Names = GetNames();
-		return GetComparisonIndexFast()>=0 && GetComparisonIndexFast()<Names.Num() && Names[GetComparisonIndexFast()]!=NULL
-			&& GetDisplayIndexFast()>=0 && GetDisplayIndexFast()<Names.Num() && Names[GetDisplayIndexFast()]!=NULL;
+		return GetComparisonIndexFast() >= 0 && GetComparisonIndexFast() < Names.Num() && Names[GetComparisonIndexFast()] != nullptr
+			&& GetDisplayIndexFast() >= 0 && GetDisplayIndexFast() < Names.Num() && Names[GetDisplayIndexFast()] != nullptr;
 	}
 
 	/**
@@ -578,7 +578,7 @@ public:
 	 *
 	 * @return	true if the name is valid
 	 */
-	bool IsValidXName( FString InvalidChars=INVALID_NAME_CHARACTERS, class FText* Reason=NULL ) const;
+	bool IsValidXName( FString InvalidChars=INVALID_NAME_CHARACTERS, class FText* Reason = nullptr ) const;
 
 	/**
 	 * Takes an FName and checks to see that it follows the rules that Unreal requires.
diff --git a/Engine/Source/Runtime/Core/Public/UObject/ObjectVersion.h b/Engine/Source/Runtime/Core/Public/UObject/ObjectVersion.h
--- a/Engine/Source/Runtime/Core/Public/UObject/ObjectVersion.h
+++ b/Engine/Source/Runtime/Core/Public/UObject/ObjectVersion.h
@@ -1,11 +1,8 @@
 // Copyright 1998-2014 Epic Games, Inc. All Rights Reserved.
 
-/*=============================================================================
-	ObjectVersion.h: Unreal object version.
-=============================================================================*/
-
 #pragma once
 
+
 // Prevents incorrect files from being loaded.
 
 #define PACKAGE_FILE_TAG			0x9E2A83C1
@@ -21,6 +18,7 @@
 // this is forever the UE3 engine version in UE4
 #define VER_LAST_ENGINE_UE3	864
 
+
 enum EUnrealEngineObjectUE4Version
 {
 	// Added array support to blueprints
diff --git a/Engine/Source/Runtime/Core/Public/UObject/ScriptDelegates.h b/Engine/Source/Runtime/Core/Public/UObject/ScriptDelegates.h
--- a/Engine/Source/Runtime/Core/Public/UObject/ScriptDelegates.h
+++ b/Engine/Source/Runtime/Core/Public/UObject/ScriptDelegates.h
@@ -1,23 +1,19 @@
 // Copyright 1998-2014 Epic Games, Inc. All Rights Reserved.
 
-/*=============================================================================
-	ScriptDelegate.h: UObject Delegate Implementation
-=============================================================================*/
-
 #pragma once
 
 #include "WeakObjectPtrTemplates.h"
 
 
 /**
- * Script delegate base class
+ * Script delegate base class.
  */
 template <typename TWeakPtr = FWeakObjectPtr>
 class TScriptDelegate
 {
 public:
 
-	/** Constructor */
+	/** Default constructor. */
 	TScriptDelegate() 
 		: Object( nullptr ),
 		  FunctionName( NAME_None )
diff --git a/Engine/Source/Runtime/Core/Public/UObject/UnrealNames.h b/Engine/Source/Runtime/Core/Public/UObject/UnrealNames.h
--- a/Engine/Source/Runtime/Core/Public/UObject/UnrealNames.h
+++ b/Engine/Source/Runtime/Core/Public/UObject/UnrealNames.h
@@ -1,9 +1,5 @@
 // Copyright 1998-2014 Epic Games, Inc. All Rights Reserved.
 
-/*=============================================================================
-	UnrealNames.h: Header file registering global hardcoded Unreal names.
-=============================================================================*/
-
 //
 // Macros.
 //
diff --git a/Engine/Source/Runtime/Core/Public/UObject/WeakObjectPtrTemplates.h b/Engine/Source/Runtime/Core/Public/UObject/WeakObjectPtrTemplates.h
--- a/Engine/Source/Runtime/Core/Public/UObject/WeakObjectPtrTemplates.h
+++ b/Engine/Source/Runtime/Core/Public/UObject/WeakObjectPtrTemplates.h
@@ -1,17 +1,14 @@
 // Copyright 1998-2014 Epic Games, Inc. All Rights Reserved.
 
-/*=============================================================================
-	WeakObjectPtrTemplates.h: FWeakObjectPtr template specializations.
-=============================================================================*/
-
 #pragma once
 
 #include "AutoPointer.h"
 
+
 /***
  * 
  * FWeakObjectPtr is a weak pointer to a UObject. 
- * It can return NULL later if the object is garbage collected.
+ * It can return nullptr later if the object is garbage collected.
  * It has no impact on if the object is garbage collected or not.
  * It can't be directly used across a network.
  *
@@ -20,6 +17,7 @@
 struct FWeakObjectPtr;
 struct FIndexToObject;
 
+
 /***
 * 
 * TWeakObjectPtr is templatized version of the generic FWeakObjectPtr
@@ -29,10 +27,10 @@ template<class T=UObject, class TWeakObjectPtrBase=FWeakObjectPtr, class TUObjec
 struct TWeakObjectPtr : private TWeakObjectPtrBase
 {
 public:
-	/** NULL constructor **/
-	FORCEINLINE TWeakObjectPtr()
-	{
-	}
+
+	/** Default constructor (no initialization). **/
+	FORCEINLINE TWeakObjectPtr() { }
+
 	/**  
 	 * Construct from an object pointer
 	 * @param Object object to create a weak pointer to
@@ -51,8 +49,7 @@ public:
 	**/
 	FORCEINLINE TWeakObjectPtr(const TWeakObjectPtr<T> &Other) :
 		TWeakObjectPtrBase(Other)
-	{
-	}
+	{ }
 
 	/**  
 	 * Construct from another weak pointer of another type, intended for derived-to-base conversions
@@ -332,5 +329,3 @@ void CopyFromWeakArray(T& Dest,const U& Src)
 		}
 	}
 }
-
-
