Core: Converted ESPMode from namespaced enum to enum class

#UpgradeNotes: Replace any existing usages of 'ESPMode::Type' with  'ESPMode'

diff --git a/Engine/Source/Editor/SceneOutliner/Public/SceneOutlinerFilters.h b/Engine/Source/Editor/SceneOutliner/Public/SceneOutlinerFilters.h
--- a/Engine/Source/Editor/SceneOutliner/Public/SceneOutlinerFilters.h
+++ b/Engine/Source/Editor/SceneOutliner/Public/SceneOutlinerFilters.h
@@ -91,7 +91,7 @@ namespace SceneOutliner
 
 		/** Meta template used to extract the filter type from a pointer or raw ptr proxy (to support MakeShareable) */
 		template <class> struct TGetFilterType;
-		template <class X, template<class, ESPMode::Type> class TPtr, ESPMode::Type E> struct TGetFilterType<TPtr<X, E>> { typedef typename X::ItemType Type; };
+		template <class X, template<class, ESPMode> class TPtr, ESPMode E> struct TGetFilterType<TPtr<X, E>> { typedef typename X::ItemType Type; };
 		template <class X, template<class> class TPtr> struct TGetFilterType<TPtr<X>> { typedef typename X::ItemType Type; };
 
 		/** Add a custom filter to this collection.
diff --git a/Engine/Source/Runtime/Core/Public/Delegates/DelegateInstancesImpl.inl b/Engine/Source/Runtime/Core/Public/Delegates/DelegateInstancesImpl.inl
--- a/Engine/Source/Runtime/Core/Public/Delegates/DelegateInstancesImpl.inl
+++ b/Engine/Source/Runtime/Core/Public/Delegates/DelegateInstancesImpl.inl
@@ -75,7 +75,7 @@ template<class UserClass, FUNC_PAYLOAD_TEMPLATE_DECL_TYPENAME> class RAW_METHOD_
 /**
  * Implements a delegate binding for shared pointer member functions.
  */
-template<class UserClass, FUNC_PAYLOAD_TEMPLATE_DECL_TYPENAME, ESPMode::Type SPMode>
+template<class UserClass, FUNC_PAYLOAD_TEMPLATE_DECL_TYPENAME, ESPMode SPMode>
 class SP_METHOD_DELEGATE_INSTANCE_CLASS
 	: public DELEGATE_INSTANCE_INTERFACE_CLASS<FUNC_TEMPLATE_ARGS>
 {
@@ -243,7 +243,7 @@ protected:
 private:
 
 	// Declare ourselves as a friend so we can access other template permutations in IsSameFunction().
-	template<class UserClassNoShadow, FUNC_PAYLOAD_TEMPLATE_DECL_NO_SHADOW, ESPMode::Type SPModeNoShadow> friend class SP_METHOD_DELEGATE_INSTANCE_CLASS;
+	template<class UserClassNoShadow, FUNC_PAYLOAD_TEMPLATE_DECL_NO_SHADOW, ESPMode SPModeNoShadow> friend class SP_METHOD_DELEGATE_INSTANCE_CLASS;
 
 	// Declare other pointer-based delegates as a friend so IsSameFunction() can compare members.
 	template<class UserClassNoShadow, FUNC_PAYLOAD_TEMPLATE_DECL_NO_SHADOW> friend class UOBJECT_METHOD_DELEGATE_INSTANCE_CLASS;
@@ -410,7 +410,7 @@ protected:
 private:
 
 	// Declare other pointer-based delegates as a friend so IsSameFunction() can compare members
-	template<class UserClassNoShadow, FUNC_PAYLOAD_TEMPLATE_DECL_NO_SHADOW, ESPMode::Type SPModeNoShadow> friend class SP_METHOD_DELEGATE_INSTANCE_CLASS;
+	template<class UserClassNoShadow, FUNC_PAYLOAD_TEMPLATE_DECL_NO_SHADOW, ESPMode SPModeNoShadow> friend class SP_METHOD_DELEGATE_INSTANCE_CLASS;
 	template<class UserClassNoShadow, FUNC_PAYLOAD_TEMPLATE_DECL_NO_SHADOW> friend class UOBJECT_METHOD_DELEGATE_INSTANCE_CLASS;
 
 	// Pointer to the user's class which contains a method we would like to call.
@@ -583,7 +583,7 @@ protected:
 private:
 
 	// Declare other pointer-based delegates as a friend so IsSameFunction() can compare members
-	template<class UserClassNoShadow, FUNC_PAYLOAD_TEMPLATE_DECL_NO_SHADOW,ESPMode::Type SPModeNoShadow> friend class SP_METHOD_DELEGATE_INSTANCE_CLASS;
+	template<class UserClassNoShadow, FUNC_PAYLOAD_TEMPLATE_DECL_NO_SHADOW,ESPMode SPModeNoShadow> friend class SP_METHOD_DELEGATE_INSTANCE_CLASS;
 	template<class UserClassNoShadow, FUNC_PAYLOAD_TEMPLATE_DECL_NO_SHADOW> friend class RAW_METHOD_DELEGATE_INSTANCE_CLASS;
 
 	// Pointer to the user's class which contains a method we would like to call.
diff --git a/Engine/Source/Runtime/Core/Public/Delegates/DelegateInstancesImpl_Variadics.inl b/Engine/Source/Runtime/Core/Public/Delegates/DelegateInstancesImpl_Variadics.inl
--- a/Engine/Source/Runtime/Core/Public/Delegates/DelegateInstancesImpl_Variadics.inl
+++ b/Engine/Source/Runtime/Core/Public/Delegates/DelegateInstancesImpl_Variadics.inl
@@ -200,10 +200,10 @@ public:
 /**
  * Implements a delegate binding for shared pointer member functions.
  */
-template <bool bConst, class UserClass, ESPMode::Type SPMode, typename FuncType, typename... VarTypes>
+template <bool bConst, class UserClass, ESPMode SPMode, typename FuncType, typename... VarTypes>
 class TBaseSPMethodDelegateInstance;
 
-template <bool bConst, class UserClass, ESPMode::Type SPMode, typename WrappedRetValType, typename... ParamTypes, typename... VarTypes>
+template <bool bConst, class UserClass, ESPMode SPMode, typename WrappedRetValType, typename... ParamTypes, typename... VarTypes>
 class TBaseSPMethodDelegateInstance<bConst, UserClass, SPMode, WrappedRetValType (ParamTypes...), VarTypes...> : public IBaseDelegateInstance<typename TUnwrapType<WrappedRetValType>::Type (ParamTypes...)>
 {
 	typedef IBaseDelegateInstance<typename TUnwrapType<WrappedRetValType>::Type (ParamTypes...)> Super;
@@ -355,7 +355,7 @@ protected:
 	TTuple<VarTypes...> Payload;
 };
 
-template <bool bConst, class UserClass, ESPMode::Type SPMode, typename... ParamTypes, typename... VarTypes>
+template <bool bConst, class UserClass, ESPMode SPMode, typename... ParamTypes, typename... VarTypes>
 class TBaseSPMethodDelegateInstance<bConst, UserClass, SPMode, void (ParamTypes...), VarTypes...> : public TBaseSPMethodDelegateInstance<bConst, UserClass, SPMode, TTypeWrapper<void> (ParamTypes...), VarTypes...>
 {
 	typedef TBaseSPMethodDelegateInstance<bConst, UserClass, SPMode, TTypeWrapper<void> (ParamTypes...), VarTypes...> Super;
diff --git a/Engine/Source/Runtime/Core/Public/Templates/SharedPointer.h b/Engine/Source/Runtime/Core/Public/Templates/SharedPointer.h
--- a/Engine/Source/Runtime/Core/Public/Templates/SharedPointer.h
+++ b/Engine/Source/Runtime/Core/Public/Templates/SharedPointer.h
@@ -114,7 +114,7 @@
  *
  * @param  InSharedRef  The shared reference to cast
  */
-template< class CastToType, class CastFromType, ESPMode::Type Mode >
+template< class CastToType, class CastFromType, ESPMode Mode >
 FORCEINLINE TSharedRef< CastToType, Mode > StaticCastSharedRef( TSharedRef< CastFromType, Mode > const& InSharedRef )
 {
 	return TSharedRef< CastToType, Mode >( InSharedRef, SharedPointerInternals::FStaticCastTag() );
@@ -129,7 +129,7 @@ class UObjectBase;
  * This shared reference will be conditionally thread-safe when the optional Mode template argument is set to ThreadSafe.
  */
 // NOTE: TSharedRef is an Unreal extension to standard smart pointer feature set
-template< class ObjectType, ESPMode::Type Mode >
+template< class ObjectType, ESPMode Mode >
 class TSharedRef
 {
 	// TSharedRefs with UObjects are illegal.
@@ -400,11 +400,11 @@ private:
 	}
 
 	// We declare ourselves as a friend (templated using OtherType) so we can access members as needed
-    template< class OtherType, ESPMode::Type OtherMode > friend class TSharedRef;
+    template< class OtherType, ESPMode OtherMode > friend class TSharedRef;
 
 	// Declare other smart pointer types as friends as needed
-    template< class OtherType, ESPMode::Type OtherMode > friend class TSharedPtr;
-    template< class OtherType, ESPMode::Type OtherMode > friend class TWeakPtr;
+    template< class OtherType, ESPMode OtherMode > friend class TSharedPtr;
+    template< class OtherType, ESPMode OtherMode > friend class TWeakPtr;
 
 private:
 
@@ -421,7 +421,7 @@ private:
  * TSharedPtr is a non-intrusive reference-counted authoritative object pointer.  This shared pointer
  * will be conditionally thread-safe when the optional Mode template argument is set to ThreadSafe.
  */
-template< class ObjectType, ESPMode::Type Mode >
+template< class ObjectType, ESPMode Mode >
 class TSharedPtr
 {
 	// TSharedPtrs with UObjects are illegal.
@@ -728,12 +728,12 @@ private:
 	}
 
 	// We declare ourselves as a friend (templated using OtherType) so we can access members as needed
-    template< class OtherType, ESPMode::Type OtherMode > friend class TSharedPtr;
+    template< class OtherType, ESPMode OtherMode > friend class TSharedPtr;
 
 	// Declare other smart pointer types as friends as needed
-    template< class OtherType, ESPMode::Type OtherMode > friend class TSharedRef;
-    template< class OtherType, ESPMode::Type OtherMode > friend class TWeakPtr;
-    template< class OtherType, ESPMode::Type OtherMode > friend class TSharedFromThis;
+    template< class OtherType, ESPMode OtherMode > friend class TSharedRef;
+    template< class OtherType, ESPMode OtherMode > friend class TWeakPtr;
+    template< class OtherType, ESPMode OtherMode > friend class TSharedFromThis;
 
 private:
 
@@ -747,14 +747,14 @@ private:
 };
 
 
-template<class ObjectType, ESPMode::Type Mode> struct TIsZeroConstructType<TSharedPtr<ObjectType, Mode>> { enum { Value = true }; };
+template<class ObjectType, ESPMode Mode> struct TIsZeroConstructType<TSharedPtr<ObjectType, Mode>> { enum { Value = true }; };
 
 
 /**
  * TWeakPtr is a non-intrusive reference-counted weak object pointer.  This weak pointer will be
  * conditionally thread-safe when the optional Mode template argument is set to ThreadSafe.
  */
-template< class ObjectType, ESPMode::Type Mode >
+template< class ObjectType, ESPMode Mode >
 class TWeakPtr
 {
 public:
@@ -959,10 +959,10 @@ private:
 	}
 
 	// We declare ourselves as a friend (templated using OtherType) so we can access members as needed
-    template< class OtherType, ESPMode::Type OtherMode > friend class TWeakPtr;
+    template< class OtherType, ESPMode OtherMode > friend class TWeakPtr;
 
 	// Declare ourselves as a friend of TSharedPtr so we can access members as needed
-    template< class OtherType, ESPMode::Type OtherMode > friend class TSharedPtr;
+    template< class OtherType, ESPMode OtherMode > friend class TSharedPtr;
 
 private:
 
@@ -976,15 +976,15 @@ private:
 };
 
 
-template<class T, ESPMode::Type Mode> struct TIsWeakPointerType<TWeakPtr<T, Mode> > { enum { Value = true }; };
-template<class T, ESPMode::Type Mode> struct TIsZeroConstructType<TWeakPtr<T, Mode> > { enum { Value = true }; };
+template<class T, ESPMode Mode> struct TIsWeakPointerType<TWeakPtr<T, Mode> > { enum { Value = true }; };
+template<class T, ESPMode Mode> struct TIsZeroConstructType<TWeakPtr<T, Mode> > { enum { Value = true }; };
 
 
 /**
  * Derive your class from TSharedFromThis to enable access to a TSharedRef directly from an object
  * instance that's already been allocated.  Use the optional Mode template argument for thread-safety.
  */
-template< class ObjectType, ESPMode::Type Mode >
+template< class ObjectType, ESPMode Mode >
 class TSharedFromThis
 {
 public:
@@ -1142,7 +1142,7 @@ private:
  *
  * @return  True if the two shared references are equal
  */
-template< class ObjectTypeA, class ObjectTypeB, ESPMode::Type Mode >
+template< class ObjectTypeA, class ObjectTypeB, ESPMode Mode >
 FORCEINLINE bool operator==( TSharedRef< ObjectTypeA, Mode > const& InSharedRefA, TSharedRef< ObjectTypeB, Mode > const& InSharedRefB )
 {
 	return &( InSharedRefA.Get() ) == &( InSharedRefB.Get() );
@@ -1154,7 +1154,7 @@ FORCEINLINE bool operator==( TSharedRef< ObjectTypeA, Mode > const& InSharedRefA
  *
  * @return  True if the two shared references are not equal
  */
-template< class ObjectTypeA, class ObjectTypeB, ESPMode::Type Mode >
+template< class ObjectTypeA, class ObjectTypeB, ESPMode Mode >
 FORCEINLINE bool operator!=( TSharedRef< ObjectTypeA, Mode > const& InSharedRefA, TSharedRef< ObjectTypeB, Mode > const& InSharedRefB )
 {
 	return &( InSharedRefA.Get() ) != &( InSharedRefB.Get() );
@@ -1166,7 +1166,7 @@ FORCEINLINE bool operator!=( TSharedRef< ObjectTypeA, Mode > const& InSharedRefA
  *
  * @return  True if the two shared pointers are equal
  */
-template< class ObjectTypeA, class ObjectTypeB, ESPMode::Type Mode >
+template< class ObjectTypeA, class ObjectTypeB, ESPMode Mode >
 FORCEINLINE bool operator==( TSharedPtr< ObjectTypeA, Mode > const& InSharedPtrA, TSharedPtr< ObjectTypeB, Mode > const& InSharedPtrB )
 {
 	return InSharedPtrA.Get() == InSharedPtrB.Get();
@@ -1178,7 +1178,7 @@ FORCEINLINE bool operator==( TSharedPtr< ObjectTypeA, Mode > const& InSharedPtrA
  *
  * @return  True if the two shared pointers are not equal
  */
-template< class ObjectTypeA, class ObjectTypeB, ESPMode::Type Mode >
+template< class ObjectTypeA, class ObjectTypeB, ESPMode Mode >
 FORCEINLINE bool operator!=( TSharedPtr< ObjectTypeA, Mode > const& InSharedPtrA, TSharedPtr< ObjectTypeB, Mode > const& InSharedPtrB )
 {
 	return InSharedPtrA.Get() != InSharedPtrB.Get();
@@ -1190,7 +1190,7 @@ FORCEINLINE bool operator!=( TSharedPtr< ObjectTypeA, Mode > const& InSharedPtrA
  *
  * @return  True if the shared reference and shared pointer are "equal"
  */
-template< class ObjectTypeA, class ObjectTypeB, ESPMode::Type Mode >
+template< class ObjectTypeA, class ObjectTypeB, ESPMode Mode >
 FORCEINLINE bool operator==( TSharedRef< ObjectTypeA, Mode > const& InSharedRef, TSharedPtr< ObjectTypeB, Mode > const& InSharedPtr )
 {
 	return InSharedPtr.IsValid() && InSharedPtr.Get() == &( InSharedRef.Get() );
@@ -1202,7 +1202,7 @@ FORCEINLINE bool operator==( TSharedRef< ObjectTypeA, Mode > const& InSharedRef,
  *
  * @return  True if the shared reference and shared pointer are not "equal"
  */
-template< class ObjectTypeA, class ObjectTypeB, ESPMode::Type Mode >
+template< class ObjectTypeA, class ObjectTypeB, ESPMode Mode >
 FORCEINLINE bool operator!=( TSharedRef< ObjectTypeA, Mode > const& InSharedRef, TSharedPtr< ObjectTypeB, Mode > const& InSharedPtr )
 {
 	return !InSharedPtr.IsValid() || ( InSharedPtr.Get() != &( InSharedRef.Get() ) );
@@ -1214,7 +1214,7 @@ FORCEINLINE bool operator!=( TSharedRef< ObjectTypeA, Mode > const& InSharedRef,
  *
  * @return  True if the shared reference and shared pointer are "equal"
  */
-template< class ObjectTypeA, class ObjectTypeB, ESPMode::Type Mode >
+template< class ObjectTypeA, class ObjectTypeB, ESPMode Mode >
 FORCEINLINE bool operator==( TSharedPtr< ObjectTypeB, Mode > const& InSharedPtr, TSharedRef< ObjectTypeA, Mode > const& InSharedRef )
 {
 	return InSharedRef == InSharedPtr;
@@ -1226,7 +1226,7 @@ FORCEINLINE bool operator==( TSharedPtr< ObjectTypeB, Mode > const& InSharedPtr,
  *
  * @return  True if the shared reference and shared pointer are not "equal"
  */
-template< class ObjectTypeA, class ObjectTypeB, ESPMode::Type Mode >
+template< class ObjectTypeA, class ObjectTypeB, ESPMode Mode >
 FORCEINLINE bool operator!=( TSharedPtr< ObjectTypeB, Mode > const& InSharedPtr, TSharedRef< ObjectTypeA, Mode > const& InSharedRef )
 {
 	return InSharedRef != InSharedPtr;
@@ -1238,7 +1238,7 @@ FORCEINLINE bool operator!=( TSharedPtr< ObjectTypeB, Mode > const& InSharedPtr,
  *
  * @return  True if the two weak pointers are equal
  */
-template< class ObjectTypeA, class ObjectTypeB, ESPMode::Type Mode >
+template< class ObjectTypeA, class ObjectTypeB, ESPMode Mode >
 FORCEINLINE bool operator==( TWeakPtr< ObjectTypeA, Mode > const& InWeakPtrA, TWeakPtr< ObjectTypeB, Mode > const& InWeakPtrB )
 {
 	return InWeakPtrA.Pin().Get() == InWeakPtrB.Pin().Get();
@@ -1250,7 +1250,7 @@ FORCEINLINE bool operator==( TWeakPtr< ObjectTypeA, Mode > const& InWeakPtrA, TW
  *
  * @return  True if the weak pointer is null
  */
-template< class ObjectTypeA, ESPMode::Type Mode >
+template< class ObjectTypeA, ESPMode Mode >
 FORCEINLINE bool operator==( TWeakPtr< ObjectTypeA, Mode > const& InWeakPtrA, decltype(nullptr) )
 {
 	return !InWeakPtrA.IsValid();
@@ -1262,7 +1262,7 @@ FORCEINLINE bool operator==( TWeakPtr< ObjectTypeA, Mode > const& InWeakPtrA, de
  *
  * @return  True if the weak pointer is null
  */
-template< class ObjectTypeB, ESPMode::Type Mode >
+template< class ObjectTypeB, ESPMode Mode >
 FORCEINLINE bool operator==( decltype(nullptr), TWeakPtr< ObjectTypeB, Mode > const& InWeakPtrB )
 {
 	return !InWeakPtrB.IsValid();
@@ -1274,7 +1274,7 @@ FORCEINLINE bool operator==( decltype(nullptr), TWeakPtr< ObjectTypeB, Mode > co
  *
  * @return  True if the two weak pointers are not equal
  */
-template< class ObjectTypeA, class ObjectTypeB, ESPMode::Type Mode >
+template< class ObjectTypeA, class ObjectTypeB, ESPMode Mode >
 FORCEINLINE bool operator!=( TWeakPtr< ObjectTypeA, Mode > const& InWeakPtrA, TWeakPtr< ObjectTypeB, Mode > const& InWeakPtrB )
 {
 	return InWeakPtrA.Pin().Get() != InWeakPtrB.Pin().Get();
@@ -1286,7 +1286,7 @@ FORCEINLINE bool operator!=( TWeakPtr< ObjectTypeA, Mode > const& InWeakPtrA, TW
  *
  * @return  True if the weak pointer is not null
  */
-template< class ObjectTypeA, ESPMode::Type Mode >
+template< class ObjectTypeA, ESPMode Mode >
 FORCEINLINE bool operator!=( TWeakPtr< ObjectTypeA, Mode > const& InWeakPtrA, decltype(nullptr) )
 {
 	return InWeakPtrA.IsValid();
@@ -1298,7 +1298,7 @@ FORCEINLINE bool operator!=( TWeakPtr< ObjectTypeA, Mode > const& InWeakPtrA, de
  *
  * @return  True if the weak pointer is not null
  */
-template< class ObjectTypeB, ESPMode::Type Mode >
+template< class ObjectTypeB, ESPMode Mode >
 FORCEINLINE bool operator!=( decltype(nullptr), TWeakPtr< ObjectTypeB, Mode > const& InWeakPtrB )
 {
 	return InWeakPtrB.IsValid();
@@ -1306,11 +1306,11 @@ FORCEINLINE bool operator!=( decltype(nullptr), TWeakPtr< ObjectTypeB, Mode > co
 
 
 /**
- * Casts a shared pointer of one type to another type. (static_cast)  Useful for downcasting.
+ * Casts a shared pointer of one type to another type. (static_cast)  Useful for down-casting.
  *
  * @param  InSharedPtr  The shared pointer to cast
  */
-template< class CastToType, class CastFromType, ESPMode::Type Mode >
+template< class CastToType, class CastFromType, ESPMode Mode >
 FORCEINLINE TSharedPtr< CastToType, Mode > StaticCastSharedPtr( TSharedPtr< CastFromType, Mode > const& InSharedPtr )
 {
 	return TSharedPtr< CastToType, Mode >( InSharedPtr, SharedPointerInternals::FStaticCastTag() );
@@ -1322,7 +1322,7 @@ FORCEINLINE TSharedPtr< CastToType, Mode > StaticCastSharedPtr( TSharedPtr< Cast
  *
  * @param  InSharedRef  The shared reference to cast
  */
-template< class CastToType, class CastFromType, ESPMode::Type Mode >
+template< class CastToType, class CastFromType, ESPMode Mode >
 FORCEINLINE TSharedRef< CastToType, Mode > ConstCastSharedRef( TSharedRef< CastFromType, Mode > const& InSharedRef )
 {
 	return TSharedRef< CastToType, Mode >( InSharedRef, SharedPointerInternals::FConstCastTag() );
@@ -1334,7 +1334,7 @@ FORCEINLINE TSharedRef< CastToType, Mode > ConstCastSharedRef( TSharedRef< CastF
  *
  * @param  InSharedPtr  The shared pointer to cast
  */
-template< class CastToType, class CastFromType, ESPMode::Type Mode >
+template< class CastToType, class CastFromType, ESPMode Mode >
 FORCEINLINE TSharedPtr< CastToType, Mode > ConstCastSharedPtr( TSharedPtr< CastFromType, Mode > const& InSharedPtr )
 {
 	return TSharedPtr< CastToType, Mode >( InSharedPtr, SharedPointerInternals::FConstCastTag() );
diff --git a/Engine/Source/Runtime/Core/Public/Templates/SharedPointerInternals.h b/Engine/Source/Runtime/Core/Public/Templates/SharedPointerInternals.h
--- a/Engine/Source/Runtime/Core/Public/Templates/SharedPointerInternals.h
+++ b/Engine/Source/Runtime/Core/Public/Templates/SharedPointerInternals.h
@@ -8,36 +8,33 @@
 
 
 /**
- * ESPMode is used select between either 'fast' or 'thread safe' shared pointer types.  This is
- * only used by templates at compile time to generate one code path or another.
+ * ESPMode is used select between either 'fast' or 'thread safe' shared pointer types.
+ * This is only used by templates at compile time to generate one code path or another.
  */
-namespace ESPMode
+enum class ESPMode
 {
-	enum Type
-	{
-		/** Forced to be not thread-safe. */
-		NotThreadSafe = 0,
+	/** Forced to be not thread-safe. */
+	NotThreadSafe = 0,
 
-		/**
-		 *	Fast, doesn't ever use atomic interlocks.
-		 *	Some code requires that all shared pointers are thread-safe.
-		 *	It's better to change it here, instead of replacing ESPMode::Fast to ESPMode::ThreadSafe throughout the code.
-		 */
-		Fast = FORCE_THREADSAFE_SHAREDPTRS ? 1 : 0,
+	/**
+		*	Fast, doesn't ever use atomic interlocks.
+		*	Some code requires that all shared pointers are thread-safe.
+		*	It's better to change it here, instead of replacing ESPMode::Fast to ESPMode::ThreadSafe throughout the code.
+		*/
+	Fast = FORCE_THREADSAFE_SHAREDPTRS ? 1 : 0,
 
-		/** Conditionally thread-safe, never spin locks, but slower */
-		ThreadSafe = 1
-	};
-}
+	/** Conditionally thread-safe, never spin locks, but slower */
+	ThreadSafe = 1
+};
 
 
 // Forward declarations.  Note that in the interest of fast performance, thread safety
 // features are mostly turned off (Mode = ESPMode::Fast).  If you need to access your
 // object on multiple threads, you should use ESPMode::ThreadSafe!
-template< class ObjectType, ESPMode::Type Mode = ESPMode::Fast > class TSharedRef;
-template< class ObjectType, ESPMode::Type Mode = ESPMode::Fast > class TSharedPtr;
-template< class ObjectType, ESPMode::Type Mode = ESPMode::Fast > class TWeakPtr;
-template< class ObjectType, ESPMode::Type Mode = ESPMode::Fast > class TSharedFromThis;
+template< class ObjectType, ESPMode Mode = ESPMode::Fast > class TSharedRef;
+template< class ObjectType, ESPMode Mode = ESPMode::Fast > class TSharedPtr;
+template< class ObjectType, ESPMode Mode = ESPMode::Fast > class TWeakPtr;
+template< class ObjectType, ESPMode Mode = ESPMode::Fast > class TSharedFromThis;
 
 
 /**
@@ -48,9 +45,9 @@ namespace SharedPointerInternals
 {
 
 	// Forward declarations
-	template< ESPMode::Type Mode > class FWeakReferencer;
+	template< ESPMode Mode > class FWeakReferencer;
 
-	/** Dummy structs used internally as template arguments for typecasts */
+	/** Dummy structures used internally as template arguments for typecasts */
 	struct FStaticCastTag {};
 	struct FConstCastTag {};
 
@@ -114,7 +111,7 @@ namespace SharedPointerInternals
 	 *
 	 * It is specialized for different threading modes.
 	 */
-	template< ESPMode::Type Mode >
+	template< ESPMode Mode >
 	struct FReferenceControllerOps;
 
 	template<>
@@ -281,7 +278,7 @@ namespace SharedPointerInternals
 	 * FSharedReferencer is a wrapper around a pointer to a reference controller that is used by either a
 	 * TSharedRef or a TSharedPtr to keep track of a referenced object's lifetime
 	 */
-	template< ESPMode::Type Mode >
+	template< ESPMode Mode >
 	class FSharedReferencer
 	{
 		typedef FReferenceControllerOps<Mode> TOps;
@@ -432,7 +429,7 @@ namespace SharedPointerInternals
 	private:
 
  		// Expose access to ReferenceController to FWeakReferencer
-		template< ESPMode::Type OtherMode > friend class FWeakReferencer;
+		template< ESPMode OtherMode > friend class FWeakReferencer;
 
 	private:
 
@@ -445,7 +442,7 @@ namespace SharedPointerInternals
 	 * FWeakReferencer is a wrapper around a pointer to a reference controller that is used
 	 * by a TWeakPtr to keep track of a referenced object's lifetime.
 	 */
-	template< ESPMode::Type Mode >
+	template< ESPMode Mode >
 	class FWeakReferencer
 	{
 		typedef FReferenceControllerOps<Mode> TOps;
@@ -569,7 +566,7 @@ namespace SharedPointerInternals
 	private:
 
  		/** Expose access to ReferenceController to FSharedReferencer. */
-		template< ESPMode::Type OtherMode > friend class FSharedReferencer;
+		template< ESPMode OtherMode > friend class FSharedReferencer;
 
 	private:
 
@@ -579,7 +576,7 @@ namespace SharedPointerInternals
 
 
 	/** Templated helper function (const) that creates a shared pointer from an object instance */
-	template< class SharedPtrType, class ObjectType, class OtherType, ESPMode::Type Mode >
+	template< class SharedPtrType, class ObjectType, class OtherType, ESPMode Mode >
 	FORCEINLINE void EnableSharedFromThis( TSharedPtr< SharedPtrType, Mode > const* InSharedPtr, ObjectType const* InObject, TSharedFromThis< OtherType, Mode > const* InShareable )
 	{
 		if( InShareable != nullptr )
@@ -590,7 +587,7 @@ namespace SharedPointerInternals
 
 
 	/** Templated helper function that creates a shared pointer from an object instance */
-	template< class SharedPtrType, class ObjectType, class OtherType, ESPMode::Type Mode >
+	template< class SharedPtrType, class ObjectType, class OtherType, ESPMode Mode >
 	FORCEINLINE void EnableSharedFromThis( TSharedPtr< SharedPtrType, Mode >* InSharedPtr, ObjectType const* InObject, TSharedFromThis< OtherType, Mode > const* InShareable )
 	{
 		if( InShareable != nullptr )
@@ -601,7 +598,7 @@ namespace SharedPointerInternals
 
 
 	/** Templated helper function (const) that creates a shared reference from an object instance */
-	template< class SharedRefType, class ObjectType, class OtherType, ESPMode::Type Mode >
+	template< class SharedRefType, class ObjectType, class OtherType, ESPMode Mode >
 	FORCEINLINE void EnableSharedFromThis( TSharedRef< SharedRefType, Mode > const* InSharedRef, ObjectType const* InObject, TSharedFromThis< OtherType, Mode > const* InShareable )
 	{
 		if( InShareable != nullptr )
@@ -612,7 +609,7 @@ namespace SharedPointerInternals
 
 
 	/** Templated helper function that creates a shared reference from an object instance */
-	template< class SharedRefType, class ObjectType, class OtherType, ESPMode::Type Mode >
+	template< class SharedRefType, class ObjectType, class OtherType, ESPMode Mode >
 	FORCEINLINE void EnableSharedFromThis( TSharedRef< SharedRefType, Mode >* InSharedRef, ObjectType const* InObject, TSharedFromThis< OtherType, Mode > const* InShareable )
 	{
 		if( InShareable != nullptr )
diff --git a/Engine/Source/Runtime/Core/Public/Templates/SharedPointerTesting.h b/Engine/Source/Runtime/Core/Public/Templates/SharedPointerTesting.h
--- a/Engine/Source/Runtime/Core/Public/Templates/SharedPointerTesting.h
+++ b/Engine/Source/Runtime/Core/Public/Templates/SharedPointerTesting.h
@@ -20,7 +20,7 @@ namespace SharedPointerTesting
 	 * Executes various shared pointer tests.  Note that some tests may require the programmer to enable
 	 * some #ifdef statements to verify that code that is not expected to compile in fact does not.
 	 */
-	template< ESPMode::Type Mode >		// Used to test *with* and *without* thread-safety features
+	template< ESPMode Mode >		// Used to test *with* and *without* thread-safety features
 	static void TestSharedPointer()
 	{
 		{
