Core.h improvements.

diff --git a/Engine/Source/Runtime/Core/Public/Containers/Array.h b/Engine/Source/Runtime/Core/Public/Containers/Array.h
--- a/Engine/Source/Runtime/Core/Public/Containers/Array.h
+++ b/Engine/Source/Runtime/Core/Public/Containers/Array.h
@@ -9,6 +9,7 @@
 #include "HAL/Platform.h"
 #include "Serialization/ArchiveBase.h"
 #include "Templates/EnableIf.h"
+#include "Templates/Sorting.h"
 #include "Templates/UnrealTemplate.h"
 
 #define DEBUG_HEAP 0
diff --git a/Engine/Source/Runtime/Core/Public/Containers/Map.h b/Engine/Source/Runtime/Core/Public/Containers/Map.h
--- a/Engine/Source/Runtime/Core/Public/Containers/Map.h
+++ b/Engine/Source/Runtime/Core/Public/Containers/Map.h
@@ -3,6 +3,7 @@
 #pragma once
 #include "Containers/Array.h"
 #include "Containers/Set.h"
+#include "Containers/UnrealString.h"
 
 #define ExchangeB(A,B) {bool T=A; A=B; B=T;}
 
diff --git a/Engine/Source/Runtime/Core/Public/Core.h b/Engine/Source/Runtime/Core/Public/Core.h
--- a/Engine/Source/Runtime/Core/Public/Core.h
+++ b/Engine/Source/Runtime/Core/Public/Core.h
@@ -10,8 +10,6 @@
 #include "Build.h"								// Build options.
 #include "ProfilingDebugging/UMemoryDefines.h"	// Memory build options.
 #include "CoreMiscDefines.h"					// Misc defines and enumerations
-class FString;									// @todo Figure out better include order in Core.h
-inline uint32 GetTypeHash (const int64 A);		// @todo Figure out better include order in Core.h
 #include "Containers/ContainersFwd.h"
 #include "Timespan.h"							// Time span definition
 #include "DateTime.h"							// Date and time handling
@@ -26,367 +24,14 @@ inline uint32 GetTypeHash (const int64 A);		// @todo Figure out better include o
 #include "UnrealTemplate.h"						// Common template definitions.
 #include "MemoryOps.h"							// Functions for efficient handling of object arrays.
 
-#ifndef STUBBED
-#define STUBBED(x)   \
-    do {                                                                  \
-        static bool AlreadySeenThisStubbedSection = false;                 \
-        if (!AlreadySeenThisStubbedSection)                               \
-        {                                                                 \
-            AlreadySeenThisStubbedSection = true;                         \
-            fprintf(stderr, "STUBBED: %s at %s:%d (%s)\n", x, __FILE__, __LINE__, __FUNCTION__); \
-        }                                                                 \
-    } while (0)
-#endif
-
-/*----------------------------------------------------------------------------
-	Forward declarations.
-----------------------------------------------------------------------------*/
+#include "Misc/CoreDefines.h"
 
 // Container forward declarations
 #include "Containers/ContainerAllocationPolicies.h"
 
-// Objects
-class UObjectBase;
-class UObjectBaseUtility;
-class	UObject;
-class		UField;
-class			UEnum;
-class			UProperty;
-class				UByteProperty;
-class				UUInt16Property;
-class				UUInt32Property;
-class				UUInt64Property;
-class				UInt8Property;
-class				UInt16Property;
-class				UIntProperty;
-class				UInt64Property;
-class				UBoolProperty;
-class				UFloatProperty;
-class				UDoubleProperty;
-class				UObjectPropertyBase;
-class				UObjectProperty;
-class					UClassProperty;
-class					UInterfaceProperty;
-class					UWeakObjectProperty;
-class					ULazyObjectProperty;
-class					UAssetObjectProperty;
-class				UNameProperty;
-class				UStructProperty;
-class               UStrProperty;
-class               UTextProperty;
-class               UArrayProperty;
-class				UDelegateProperty;
-class				UMulticastDelegateProperty;
-class			UStruct;
-class				UFunction;
-class				UClass;
-class				UScriptStruct;
-class		ULinker;
-class			ULinkerLoad;
-class			ULinkerSave;
-class		UPackage;
-class		USystem;
-class		UTextBuffer;
-class		UPackageMap;
-class		UObjectRedirector;
-
-
-// Structures
-class FName;
-class FArchive;
-class FCompactIndex;
-class FExec;
-class FGuid;
-class FMemStack;
-class ITransaction;
-class FUnknown;
-class FRepLink;
-class FText;
-class FMalloc;
-class FConfigCacheIni;
-class FConfigFile;
-class IFileManager;
-class FOutputDeviceRedirector;
-class FOutputDeviceConsole;
-
-// Core globals
-extern CORE_API FConfigCacheIni* GConfig;
-extern CORE_API ITransaction* GUndo;
-extern CORE_API FOutputDeviceConsole* GLogConsole;
-
-extern CORE_API TCHAR GErrorHist[16384];
-
-// @TODO yrx 2014-08-19 Combine into one.
-extern CORE_API TCHAR GErrorExceptionDescription[4096];
-extern CORE_API TCHAR GErrorMessage[4096];
-
-extern CORE_API const FText GTrue, GFalse, GYes, GNo, GNone;
-
-/** If true, this executable is able to run all games (which are loaded as DLL's). */
-extern CORE_API bool GIsGameAgnosticExe;	
-
-/** When saving out of the game, this override allows the game to load editor only properties. */
-extern CORE_API bool GForceLoadEditorOnly;
-
-/** Disable loading of objects not contained within script files; used during script compilation */
-extern CORE_API bool GVerifyObjectReferencesOnly;
-
-/** when constructing objects, use the fast path on consoles... */
-extern CORE_API bool GFastPathUniqueNameGeneration;
-
-/** allow AActor object to execute script in the editor from specific entry points, such as when running a construction script */
-extern CORE_API bool GAllowActorScriptExecutionInEditor;
-
-/** Forces use of template names for newly instanced components in a CDO. */
-extern CORE_API bool GCompilingBlueprint;
-
-/** True if we're reconstructing blueprint instances. Should never be true on cooked builds */
-extern CORE_API bool GIsReconstructingBlueprintInstances;
-
-/** Force blueprints to not compile on load */
-extern CORE_API bool GForceDisableBlueprintCompileOnLoad;
-
-
-/** Helper function to flush resource streaming. */
-extern CORE_API void (*GFlushStreamingFunc)(void);
-
-extern CORE_API float GVolumeMultiplier;
-
- 
-#if WITH_EDITORONLY_DATA
-
-	/** 
-	 *	True if we are in the editor. 
-	 *	Note that this is still true when using Play In Editor. You may want to use GWorld->HasBegunPlay in that case.
-	 */
-	extern CORE_API bool GIsEditor;
-	extern CORE_API bool GIsImportingT3D;
-	extern CORE_API bool PRIVATE_GIsRunningCommandlet;
-	extern CORE_API bool GIsUCCMakeStandaloneHeaderGenerator;
-	extern CORE_API bool GIsTransacting;
-
-	/** Indicates that the game thread is currently paused deep in a call stack,
-		while some subset of the editor user interface is pumped.  No game
-		thread work can be done until the UI pumping loop returns naturally. */
-	extern CORE_API bool			GIntraFrameDebuggingGameThread;
-	/** True if this is the first time through the UI message pumping loop. */
-	extern CORE_API bool			GFirstFrameIntraFrameDebugging;
-
-#else
-
-	#define GIsEditor								false
-	#define PRIVATE_GIsRunningCommandlet			false
-	#define GIsUCCMakeStandaloneHeaderGenerator		false
-	#define GIntraFrameDebuggingGameThread			false
-	#define GFirstFrameIntraFrameDebugging			false
-
-#endif // WITH_EDITORONLY_DATA
-
-
-/**
- * Check to see if this executable is running a commandlet (custom command-line processing code in an editor-like environment)
- */
-FORCEINLINE bool IsRunningCommandlet()
-{
-#if WITH_EDITORONLY_DATA
-	return PRIVATE_GIsRunningCommandlet;
-#else
-	return false;
-#endif
-}
-
-extern CORE_API bool GEdSelectionLock;
-extern CORE_API bool GIsClient;
-extern CORE_API bool GIsServer;
-extern CORE_API bool GIsCriticalError;
-extern CORE_API bool GIsRunning;
-extern CORE_API bool GIsGarbageCollecting;
-extern CORE_API bool GIsDuplicatingClassForReinstancing;
-
-/**
- * These are set when the engine first starts up.
- */
-
-/**
- * This specifies whether the engine was launched as a build machine process.
- */
-extern CORE_API bool GIsBuildMachine;
-
-/**
- * This determines if we should output any log text.  If Yes then no log text should be emitted.
- */
-extern CORE_API bool GIsSilent;
-
-extern CORE_API bool GIsSlowTask;
-extern CORE_API bool GSlowTaskOccurred;
-extern CORE_API bool GIsGuarded;
-extern CORE_API bool GIsRequestingExit;
-
-/** Archive for serializing arbitrary data to and from memory						*/
-extern CORE_API class FReloadObjectArc* GMemoryArchive;
-
-/**
- *	Global value indicating on-screen warnings/message should be displayed.
- *	Disabled via console command "DISABLEALLSCREENMESSAGES"
- *	Enabled via console command "ENABLEALLSCREENMESSAGES"
- *	Toggled via console command "TOGGLEALLSCREENMESSAGES"
- */
-extern CORE_API bool GAreScreenMessagesEnabled;
-extern CORE_API bool GScreenMessagesRestoreState;
-
-/* Whether we are dumping screen shots */
-extern CORE_API int32 GIsDumpingMovie;
-extern CORE_API bool GIsHighResScreenshot;
-extern CORE_API uint32 GScreenshotResolutionX;
-extern CORE_API uint32 GScreenshotResolutionY;
-extern CORE_API uint64 GMakeCacheIDIndex;
-extern CORE_API FString GEditorKeyBindingsIni;
-extern CORE_API FString GEngineIni;
-extern CORE_API FString GEditorIni;
-extern CORE_API FString GEditorLayoutIni;
-extern CORE_API FString GEditorUserSettingsIni;
-extern CORE_API FString GEditorGameAgnosticIni;
-extern CORE_API FString GCompatIni;
-extern CORE_API FString GLightmassIni;
-extern CORE_API FString GScalabilityIni;
-extern CORE_API FString GInputIni;
-extern CORE_API FString GGameIni;
-extern CORE_API FString GGameUserSettingsIni;
-
-extern CORE_API float GNearClippingPlane;
-
-extern CORE_API bool GExitPurge;
-extern CORE_API TCHAR GInternalGameName[64];
-extern CORE_API const TCHAR* GForeignEngineDir;
-
-/** Exec handler for game debugging tool, allowing commands like "editactor" */
-extern CORE_API FExec* GDebugToolExec;
-
-/** Whether we're currently in the async loading code path or not */
-extern CORE_API bool GIsAsyncLoading;
-
-/** Whether the editor is currently loading a package or not */
-extern CORE_API bool GIsEditorLoadingPackage;
-
-/** Whether GWorld points to the play in editor world */
-extern CORE_API bool GIsPlayInEditorWorld;
-
-extern CORE_API int32 GPlayInEditorID;
-
-/** Whether or not PIE was attempting to play from PlayerStart */
-extern CORE_API bool GIsPIEUsingPlayerStart;
-
-/** Proxy class that allows verification on FApp::IsGame() accesses. */
-extern CORE_API bool IsInGameThread();
-
-/** true if the runtime needs textures to be powers of two */
-extern CORE_API bool GPlatformNeedsPowerOfTwoTextures;
-
-/** Time at which FPlatformTime::Seconds() was first initialized (very early on) */
-extern CORE_API double GStartTime;
-
-/** System time at engine init. */
-extern CORE_API FString GSystemStartTime;
-
-/** Whether we are still in the initial loading process. */
-extern CORE_API bool GIsInitialLoad;
-
-/** true when we are routing ConditionalPostLoad/PostLoad to objects */
-extern CORE_API bool GIsRoutingPostLoad;
-
-/** Steadily increasing frame counter. */
-extern CORE_API uint64 GFrameCounter;
-
-/** Incremented once per frame before the scene is being rendered. In split screen mode this is incremented once for all views (not for each view). */
-extern CORE_API uint32 GFrameNumber;
-
-/** Render Thread copy of the frame number. */
-extern CORE_API uint32 GFrameNumberRenderThread;
-
-#if !(UE_BUILD_SHIPPING && WITH_EDITOR)
-	// We cannot count on this variable to be accurate in a shipped game, so make sure no code tries to use it
-	/** Whether we are the first instance of the game running. */
-	extern CORE_API bool GIsFirstInstance;
-#endif
-
-/** Threshold for a frame to be considered a hitch (in seconds. */
-extern CORE_API float GHitchThreshold;
-
-/** Size to break up data into when saving compressed data */
-extern CORE_API int32 GSavingCompressionChunkSize;
-
-/** Whether we are using the seek-free/ cooked loading code path. */
-extern CORE_API bool GUseSeekFreeLoading;
-
-/** Thread ID of the main/game thread */
-extern CORE_API uint32 GGameThreadId;
-
-/** Thread ID of the render thread, if any */
-extern CORE_API uint32 GRenderThreadId;
-
-/** Thread ID of the slate thread, if any */
-extern CORE_API uint32 GSlateLoadingThreadId;
-
-/** Has GGameThreadId been set yet? */
-extern CORE_API bool GIsGameThreadIdInitialized;
-
-/** Whether to emit begin/ end draw events. */
-extern CORE_API bool GEmitDrawEvents;
-
-/** Whether we want the rendering thread to be suspended, used e.g. for tracing. */
-extern CORE_API bool GShouldSuspendRenderingThread;
-
-/** Determines what kind of trace should occur, NAME_None for none. */
-extern CORE_API FName GCurrentTraceName;
-
-/** How to print the time in log output. */
-extern CORE_API ELogTimes::Type GPrintLogTimes;
-
-/** Global screen shot index to avoid overwriting ScreenShots. */
-extern CORE_API int32 GScreenshotBitmapIndex;
-
-/** Whether stats should emit named events for e.g. PIX. */
-extern CORE_API int32 GCycleStatsShouldEmitNamedEvents;
-
-/** Disables some warnings and minor features that would interrupt a demo presentation*/
-extern CORE_API bool GIsDemoMode;
-
-/** Name of the core package. */
-//@Package name transition, remove the double checks 
-extern CORE_API FName GLongCorePackageName;
-//@Package name transition, remove the double checks 
-extern CORE_API FName GLongCoreUObjectPackageName;
-
-/** Whether or not a unit test is currently being run. */
-extern CORE_API bool GIsAutomationTesting;
-
-/** Constrain bandwidth if wanted. Value is in MByte/ sec. */
-extern CORE_API float GAsyncIOBandwidthLimit;
-
-/** Whether or not messages are being pumped outside of main loop */
-extern CORE_API bool GPumpingMessagesOutsideOfMainLoop;
-
-/** Total blueprint compile time. */
-extern CORE_API double GBlueprintCompileTime;
-
-/** Stack names from the VM to be unrolled when we assert */
-extern CORE_API TArray<struct FScriptTraceStackNode> GScriptStack;
-
-/*----------------------------------------------------------------------------
-	Metadata macros.
-----------------------------------------------------------------------------*/
-
-#define CPP       1
-#define STRUCTCPP 1
-#define DEFAULTS  0
-
-
-/*-----------------------------------------------------------------------------
-	Seek-free defines.
------------------------------------------------------------------------------*/
-
-#define STANDALONE_SEEKFREE_SUFFIX	TEXT("_SF")
+#include "UObject/UObjectHierarchyFwd.h"
 
+#include "CoreGlobals.h"
 
 /*----------------------------------------------------------------------------
 	Includes.
@@ -469,50 +114,4 @@ extern CORE_API TArray<struct FScriptTraceStackNode> GScriptStack;
 #include "Attribute.h"
 #include "Optional.h"
 #include "EnumClassFlags.h"
-
-#ifdef TRUE
-	#undef TRUE
-#endif
-
-#ifdef FALSE
-	#undef FALSE
-#endif
-
-
-/// Trick to prevent people from using old UE4 types (which are still defined in Windef.h on PC).
-namespace DoNotUseOldUE4Type
-{
-	/// Used to cause compile errors through typedefs of unportable types. If unportable types 
-	/// are really necessary please use the global scope operator to access it (i.e. ::&nbsp;INT). 
-	/// Windows header file includes can be wrapped in #include "AllowWindowsPlatformTypes.h" 
-	/// and #include "HideWindowsPlatformTypes.h"
-	class FUnusableType
-	{
-		FUnusableType();
-		~FUnusableType();
-	};
-
-	/// This type is aliased to FUnusableType - use ::int32 instead.
-	typedef FUnusableType INT;
-
-	/// This type is aliased to FUnusableType - use ::uint32 instead.
-	typedef FUnusableType UINT;
-
-	/// This type is aliased to FUnusableType - use ::uint32 instead.
-	typedef FUnusableType DWORD;
-
-	/// This type is aliased to FUnusableType - use float instead.
-	typedef FUnusableType FLOAT;
-
-	/// This value is aliased to FUnusableType - use true instead.
-	typedef FUnusableType TRUE;
-
-	/// This value is aliased to FUnusableType - use false instead.
-	typedef FUnusableType FALSE;
-}
-
-
-
-using namespace DoNotUseOldUE4Type;
-
-
+#include "Misc/DisableOldUETypes.h"
\ No newline at end of file
diff --git a/Engine/Source/Runtime/Core/Public/CoreGlobals.h b/Engine/Source/Runtime/Core/Public/CoreGlobals.h
--- a/Engine/Source/Runtime/Core/Public/CoreGlobals.h
+++ b/Engine/Source/Runtime/Core/Public/CoreGlobals.h
@@ -2,7 +2,6 @@
 #pragma once
 #include "Containers/ContainersFwd.h"
 #include "HAL/Platform.h"
-#include "Internationalization/Text.h"
 #include "Misc/CoreMiscDefines.h"
 #include "Misc/OutputDevice.h"
 
@@ -10,7 +9,9 @@ class FConfigCacheIni;
 class FExec;
 class FName;
 class FOutputDeviceConsole;
+class FReloadObjectArc;
 class FString;
+class FText;
 class ITransaction;
 
 struct FScriptTraceStackNode;
@@ -128,7 +129,7 @@ extern CORE_API bool GIsGuarded;
 extern CORE_API bool GIsRequestingExit;
 
 /** Archive for serializing arbitrary data to and from memory						*/
-extern CORE_API class FReloadObjectArc* GMemoryArchive;
+extern CORE_API FReloadObjectArc* GMemoryArchive;
 
 /**
 *	Global value indicating on-screen warnings/message should be displayed.
diff --git a/Engine/Source/Runtime/Core/Public/Internationalization/Text.h b/Engine/Source/Runtime/Core/Public/Internationalization/Text.h
--- a/Engine/Source/Runtime/Core/Public/Internationalization/Text.h
+++ b/Engine/Source/Runtime/Core/Public/Internationalization/Text.h
@@ -1,8 +1,12 @@
 ﻿// Copyright 1998-2014 Epic Games, Inc. All Rights Reserved.
 #pragma once
+#include "Containers/UnrealString.h"
+#include "Internationalization/CulturePointer.h"
+#include "Templates/SharedPointer.h"
+#include "UObject/NameTypes.h"
 
-#include "SharedPointer.h"
-#include "CulturePointer.h"
+class FTimespan;
+struct FDateTime;
 
 #define ENABLE_TEXT_ERROR_CHECKING_RESULTS (UE_BUILD_DEBUG | UE_BUILD_DEVELOPMENT | UE_BUILD_TEST )
 
diff --git a/Engine/Source/Runtime/Core/Public/Misc/CoreDefines.h b/Engine/Source/Runtime/Core/Public/Misc/CoreDefines.h
new file mode 100644
--- /dev/null
+++ b/Engine/Source/Runtime/Core/Public/Misc/CoreDefines.h
@@ -0,0 +1,30 @@
+// Copyright 1998-2014 Epic Games, Inc. All Rights Reserved.
+#pragma once
+
+#ifndef STUBBED
+#define STUBBED(x)	\
+	do																								\
+	{																								\
+		static bool AlreadySeenThisStubbedSection = false;											\
+		if (!AlreadySeenThisStubbedSection)															\
+		{																							\
+			AlreadySeenThisStubbedSection = true;													\
+			fprintf(stderr, "STUBBED: %s at %s:%d (%s)\n", x, __FILE__, __LINE__, __FUNCTION__);	\
+		}																							\
+	} while (0)
+#endif
+
+/*----------------------------------------------------------------------------
+Metadata macros.
+----------------------------------------------------------------------------*/
+
+#define CPP       1
+#define STRUCTCPP 1
+#define DEFAULTS  0
+
+
+/*-----------------------------------------------------------------------------
+Seek-free defines.
+-----------------------------------------------------------------------------*/
+
+#define STANDALONE_SEEKFREE_SUFFIX	TEXT("_SF")
diff --git a/Engine/Source/Runtime/Core/Public/Misc/DisableOldUETypes.h b/Engine/Source/Runtime/Core/Public/Misc/DisableOldUETypes.h
new file mode 100644
--- /dev/null
+++ b/Engine/Source/Runtime/Core/Public/Misc/DisableOldUETypes.h
@@ -0,0 +1,50 @@
+// Copyright 1998-2014 Epic Games, Inc. All Rights Reserved.
+#pragma once
+
+
+#ifdef TRUE
+	#undef TRUE
+#endif
+
+#ifdef FALSE
+	#undef FALSE
+#endif
+
+
+/// Trick to prevent people from using old UE4 types (which are still defined in Windef.h on PC).
+namespace DoNotUseOldUE4Type
+{
+	/// Used to cause compile errors through typedefs of unportable types. If unportable types 
+	/// are really necessary please use the global scope operator to access it (i.e. ::&nbsp;INT). 
+	/// Windows header file includes can be wrapped in #include "AllowWindowsPlatformTypes.h" 
+	/// and #include "HideWindowsPlatformTypes.h"
+	class FUnusableType
+	{
+		FUnusableType();
+		~FUnusableType();
+	};
+
+	/// This type is aliased to FUnusableType - use ::int32 instead.
+	typedef FUnusableType INT;
+
+	/// This type is aliased to FUnusableType - use ::uint32 instead.
+	typedef FUnusableType UINT;
+
+	/// This type is aliased to FUnusableType - use ::uint32 instead.
+	typedef FUnusableType DWORD;
+
+	/// This type is aliased to FUnusableType - use float instead.
+	typedef FUnusableType FLOAT;
+
+	/// This value is aliased to FUnusableType - use true instead.
+	typedef FUnusableType TRUE;
+
+	/// This value is aliased to FUnusableType - use false instead.
+	typedef FUnusableType FALSE;
+}
+
+
+
+using namespace DoNotUseOldUE4Type;
+
+
diff --git a/Engine/Source/Runtime/Core/Public/UObject/NameTypes.h b/Engine/Source/Runtime/Core/Public/UObject/NameTypes.h
--- a/Engine/Source/Runtime/Core/Public/UObject/NameTypes.h
+++ b/Engine/Source/Runtime/Core/Public/UObject/NameTypes.h
@@ -1,6 +1,11 @@
 // Copyright 1998-2014 Epic Games, Inc. All Rights Reserved.
 
 #pragma once
+#include "Containers/StringConv.h"
+#include "Containers/UnrealString.h"
+#include "HAL/Platform.h"
+#include "HAL/PlatformProcess.h"
+#include "UObject/UnrealNames.h"
 
 /*----------------------------------------------------------------------------
 	Definitions.
diff --git a/Engine/Source/Runtime/Core/Public/UObject/UObjectHierarchyFwd.h b/Engine/Source/Runtime/Core/Public/UObject/UObjectHierarchyFwd.h
new file mode 100644
--- /dev/null
+++ b/Engine/Source/Runtime/Core/Public/UObject/UObjectHierarchyFwd.h
@@ -0,0 +1,48 @@
+// Copyright 1998-2014 Epic Games, Inc. All Rights Reserved.
+#pragma once
+
+
+// Objects
+class UObjectBase;
+class UObjectBaseUtility;
+class	UObject;
+class		UField;
+class			UEnum;
+class			UProperty;
+class				UByteProperty;
+class				UUInt16Property;
+class				UUInt32Property;
+class				UUInt64Property;
+class				UInt8Property;
+class				UInt16Property;
+class				UIntProperty;
+class				UInt64Property;
+class				UBoolProperty;
+class				UFloatProperty;
+class				UDoubleProperty;
+class				UObjectPropertyBase;
+class				UObjectProperty;
+class					UClassProperty;
+class					UInterfaceProperty;
+class					UWeakObjectProperty;
+class					ULazyObjectProperty;
+class					UAssetObjectProperty;
+class				UNameProperty;
+class				UStructProperty;
+class				UStrProperty;
+class				UTextProperty;
+class				UArrayProperty;
+class				UDelegateProperty;
+class				UMulticastDelegateProperty;
+class			UStruct;
+class				UFunction;
+class				UClass;
+class				UScriptStruct;
+class		ULinker;
+class			ULinkerLoad;
+class			ULinkerSave;
+class		UPackage;
+class		USystem;
+class		UTextBuffer;
+class		UPackageMap;
+class		UObjectRedirector;
\ No newline at end of file
diff --git a/Engine/Source/Runtime/Core/Public/UObject/UnrealNames.h b/Engine/Source/Runtime/Core/Public/UObject/UnrealNames.h
--- a/Engine/Source/Runtime/Core/Public/UObject/UnrealNames.h
+++ b/Engine/Source/Runtime/Core/Public/UObject/UnrealNames.h
@@ -1,4 +1,6 @@
 // Copyright 1998-2014 Epic Games, Inc. All Rights Reserved.
+#ifndef _UNREAL_NAMES_H_
+#define _UNREAL_NAMES_H_
 
 //
 // Macros.
@@ -205,3 +207,5 @@ REGISTER_NAME(500,PerformanceCapture)
 	#undef REGISTER_NAME
 	#undef REGISTERING_ENUM
 #endif
+
+#endif // _UNREAL_NAMES_H_
\ No newline at end of file
