Stats - Disabled stats due to the memory issue with the current solution

diff --git a/Engine/Source/Runtime/Core/Private/Internationalization/ICUInternationalization.cpp b/Engine/Source/Runtime/Core/Private/Internationalization/ICUInternationalization.cpp
--- a/Engine/Source/Runtime/Core/Private/Internationalization/ICUInternationalization.cpp
+++ b/Engine/Source/Runtime/Core/Private/Internationalization/ICUInternationalization.cpp
@@ -27,7 +27,13 @@ namespace
 			BytesInUseCount += FMemory::GetAllocSize(Result);
 			if(FThreadStats::IsThreadingReady() && CachedBytesInUseCount != BytesInUseCount)
 			{
-				SET_MEMORY_STAT(STAT_MemoryICUTotalAllocationSize, BytesInUseCount);
+				// The amount of startup stats messages for STAT_MemoryICUTotalAllocationSize is about 700k
+				// It needs to be replaced with something like this
+				//  FEngineLoop.Tick.{ 
+				//		FPlatformMemory::UpdateStats();
+				//		
+				//	And called once per frame.
+				//SET_MEMORY_STAT(STAT_MemoryICUTotalAllocationSize, BytesInUseCount);
 				CachedBytesInUseCount = BytesInUseCount;
 			}
 #endif
@@ -45,7 +51,7 @@ namespace
 			BytesInUseCount -= FMemory::GetAllocSize(mem);
 			if(FThreadStats::IsThreadingReady() && CachedBytesInUseCount != BytesInUseCount)
 			{
-				SET_MEMORY_STAT(STAT_MemoryICUTotalAllocationSize, BytesInUseCount);
+				//SET_MEMORY_STAT(STAT_MemoryICUTotalAllocationSize, BytesInUseCount);
 				CachedBytesInUseCount = BytesInUseCount;
 			}
 #endif
