Make CorePrivatePCH.h Core module private PCH instead of Core.h. (Submitted for Mikolaj Sieluzycki due to ObjectVersion.cpp change)

diff --git a/Engine/Source/Runtime/Core/Private/Android/AndroidApplication.cpp b/Engine/Source/Runtime/Core/Private/Android/AndroidApplication.cpp
--- a/Engine/Source/Runtime/Core/Private/Android/AndroidApplication.cpp
+++ b/Engine/Source/Runtime/Core/Private/Android/AndroidApplication.cpp
@@ -1,6 +1,6 @@
 // Copyright 1998-2014 Epic Games, Inc. All Rights Reserved.
 
-#include "Core.h"
+#include "CorePrivatePCH.h"
 #include "AndroidApplication.h"
 #include "AndroidInputInterface.h"
 #include "AndroidWindow.h"
diff --git a/Engine/Source/Runtime/Core/Private/Android/AndroidFile.cpp b/Engine/Source/Runtime/Core/Private/Android/AndroidFile.cpp
--- a/Engine/Source/Runtime/Core/Private/Android/AndroidFile.cpp
+++ b/Engine/Source/Runtime/Core/Private/Android/AndroidFile.cpp
@@ -4,7 +4,7 @@
 	AndroidFile.cpp: Android platform implementations of File functions
 =============================================================================*/
 
-#include "Core.h"
+#include "CorePrivatePCH.h"
 
 #include <dirent.h>
 #include <jni.h>
diff --git a/Engine/Source/Runtime/Core/Private/CorePrivatePCH.h b/Engine/Source/Runtime/Core/Private/CorePrivatePCH.h
new file mode 100644
--- /dev/null
+++ b/Engine/Source/Runtime/Core/Private/CorePrivatePCH.h
@@ -0,0 +1,117 @@
+// Copyright 1998-2014 Epic Games, Inc. All Rights Reserved.
+
+#pragma once
+
+/*----------------------------------------------------------------------------
+Low level includes.
+----------------------------------------------------------------------------*/
+
+#include "Platform.h"							// Set up base types, etc for the platform
+#include "Build.h"								// Build options.
+#include "ProfilingDebugging/UMemoryDefines.h"	// Memory build options.
+#include "CoreMiscDefines.h"					// Misc defines and enumerations
+#include "Containers/ContainersFwd.h"
+#include "Timespan.h"							// Time span definition
+#include "DateTime.h"							// Date and time handling
+#include "PlatformIncludes.h"					// Include the main and misc platform headers
+#include "PlatformFilemanager.h"				// Platform file manager.
+#include "AssertionMacros.h"					// Various assertion macros
+#include "UObject/UnrealNames.h"				// EName definition.
+#include "OutputDevice.h"						// Output devices, logf, debugf, etc
+#include "NumericLimits.h"						// Numeric limits
+#include "UnrealMathUtility.h"					// FMath
+#include "UnrealTypeTraits.h"					// Type Traits
+#include "UnrealTemplate.h"						// Common template definitions.
+#include "MemoryOps.h"							// Functions for efficient handling of object arrays.
+
+#include "Misc/CoreDefines.h"
+
+// Container forward declarations
+#include "Containers/ContainerAllocationPolicies.h"
+
+#include "UObject/UObjectHierarchyFwd.h"
+
+#include "CoreGlobals.h"
+
+/*----------------------------------------------------------------------------
+Includes.
+----------------------------------------------------------------------------*/
+
+#include "FileManager.h"				// File manager interface
+#include "ScopedDebugInfo.h"			// Scoped debug info.
+#include "ExternalProfiler.h"			// External profiler integration.
+#include "MemoryBase.h"					// Base memory allocation
+#include "ByteSwap.h"					// ByteSwapping utility code.
+#include "ITransaction.h"				// Transaction handling
+#include "Compression.h"				// Low level compression routines
+#include "StringUtility.h"				// Low level string utility code.
+#include "Parse.h"						// String parsing functions
+#include "StringConv.h"					// Low level string conversion
+#include "Crc.h"						// Crc functions
+#include "ObjectVersion.h"				// Object version info.
+#include "TypeHash.h"					// Type hashing functions
+#include "EnumAsByte.h"					// Common template definitions.
+#include "ArchiveBase.h"				// Archive class.
+#include "ScopedPointer.h"				// Scoped pointer definitions.
+#include "Sorting.h"					// Sorting definitions.
+#include "Array.h"						// Dynamic array definitions.
+#include "ArrayBuilder.h"				// Builder template for arrays.
+#include "BitArray.h"					// Bit array definition.
+#include "Bitstreams.h"					// Bit stream archiver.
+#include "SparseArray.h"				// Sparse array definitions.
+#include "UnrealString.h"				// Dynamic string definitions.
+#include "NameTypes.h"					// Global name subsystem.
+#include "CoreMisc.h"					// Low level utility code.
+#include "Paths.h"						// Path helper functions
+#include "StaticArray.h"                // Static array definition.
+#include "StaticBitArray.h"             // Static bit array definition.
+#include "Set.h"						// Set definitions.
+#include "Map.h"						// Dynamic map definitions.
+#include "MapBuilder.h"					// Builder template for maps.
+#include "List.h"						// Dynamic list definitions.
+#include "ResourceArray.h"				// Resource array definitions.
+#include "RefCounting.h"				// Reference counting definitions.
+#include "ScriptDelegates.h"
+#include "Delegate.h"					// C++ delegate system
+#include "ThreadingBase.h"				// Non-platform specific multi-threaded support.
+#include "Internationalization/Internationalization.h"
+#include "Internationalization/Culture.h"
+#include "Guid.h"						// FGuid class
+#include "NetworkGuid.h"				// FNetworkGuid class
+#include "UnrealMath.h"					// Vector math functions.
+#include "SHMath.h"						// SH math functions.
+#include "RandomStream.h"				// Random stream definitions.
+#include "OutputDevices.h"				// Output devices
+#include "CoreStats.h"
+#include "MemStack.h"					// Stack based memory management.
+#include "AsyncWork.h"					// Async threaded work
+#include "Archive.h"					// Utility archive classes
+#include "IOBase.h"						// base IO declarations, FIOManager, FIOSystem
+#include "Variant.h"
+#include "WildcardString.h"
+#include "CircularBuffer.h"
+#include "CircularQueue.h"
+#include "Queue.h"
+#include "Ticker.h"						// Efficient scheduled delegate manager
+#include "RocketSupport.h"				// Core support for launching in "Rocket" mode
+#include "ProfilingHelpers.h"			// Misc profiling helpers.
+#include "ConfigCacheIni.h"				// The configuration cache declarations
+#include "IConsoleManager.h"
+#include "FeedbackContext.h"
+#include "AutomationTest.h"				// Automation testing functionality for game, editor, etc.
+#include "CallbackDevice.h"				// Base class for callback devices.
+#include "ObjectThumbnail.h"			// Object thumbnails
+#include "LocalTimestampDirectoryVisitor.h"
+#include "GenericWindowDefinition.h"
+#include "GenericWindow.h"
+#include "GenericApplicationMessageHandler.h"
+#include "GenericApplication.h"
+#include "ICursor.h"
+#include "CustomVersion.h"				// Custom versioning system.
+#include "App.h"
+#include "OutputDeviceConsole.h"
+#include "MonitoredProcess.h"
+#include "Attribute.h"
+#include "Optional.h"
+#include "EnumClassFlags.h"
+#include "Misc/DisableOldUETypes.h"
