Handle the potential exceptions when launching a process on Mac OS X & return the error information if possible rather than crashing with an unhandled exception.

diff --git a/Engine/Source/Runtime/Core/Private/Mac/MacPlatformProcess.cpp b/Engine/Source/Runtime/Core/Private/Mac/MacPlatformProcess.cpp
--- a/Engine/Source/Runtime/Core/Private/Mac/MacPlatformProcess.cpp
+++ b/Engine/Source/Runtime/Core/Private/Mac/MacPlatformProcess.cpp
@@ -274,37 +274,46 @@ bool FMacPlatformProcess::ExecProcess( const TCHAR* URL, const TCHAR* Params, in
 		
 		[ProcessHandle setStandardError: (id)StdErrPipe];
 		
-		[ProcessHandle launch];
-		
-		[ProcessHandle waitUntilExit];
-		
-		if(OutReturnCode)
+		@try
 		{
-			*OutReturnCode = [ProcessHandle terminationStatus];
-		}
-		
-		if(OutStdOut)
-		{
-			NSFileHandle* StdOutFile = [StdOutPipe fileHandleForReading];
-			if(StdOutFile)
+			[ProcessHandle launch];
+			
+			[ProcessHandle waitUntilExit];
+			
+			if(OutReturnCode)
 			{
-				NSData* StdOutData = [StdOutFile readDataToEndOfFile];
-				NSString* StdOutString = (NSString*)[[[NSString alloc] initWithData:StdOutData encoding:NSUTF8StringEncoding] autorelease];
-				*OutStdOut = FString(StdOutString);
+				*OutReturnCode = [ProcessHandle terminationStatus];
 			}
-		}
-		
-		if(OutStdErr)
-		{
-			NSFileHandle* StdErrFile = [StdErrPipe fileHandleForReading];
-			if(StdErrFile)
+			
+			if(OutStdOut)
 			{
-				NSData* StdErrData = [StdErrFile readDataToEndOfFile];
-				NSString* StdErrString = (NSString*)[[[NSString alloc] initWithData:StdErrData encoding:NSUTF8StringEncoding] autorelease];
-				*OutStdErr = FString(StdErrString);
+				NSFileHandle* StdOutFile = [StdOutPipe fileHandleForReading];
+				if(StdOutFile)
+				{
+					NSData* StdOutData = [StdOutFile readDataToEndOfFile];
+					NSString* StdOutString = (NSString*)[[[NSString alloc] initWithData:StdOutData encoding:NSUTF8StringEncoding] autorelease];
+					*OutStdOut = FString(StdOutString);
+				}
 			}
+			
+			if(OutStdErr)
+			{
+				NSFileHandle* StdErrFile = [StdErrPipe fileHandleForReading];
+				if(StdErrFile)
+				{
+					NSData* StdErrData = [StdErrFile readDataToEndOfFile];
+					NSString* StdErrString = (NSString*)[[[NSString alloc] initWithData:StdErrData encoding:NSUTF8StringEncoding] autorelease];
+					*OutStdErr = FString(StdErrString);
+				}
+			}
+			return true;
+		}
+		@catch (NSException* Exc)
+		{
+			*OutReturnCode = ENOENT;
+			*OutStdErr = FString([Exc reason]);
+			return false;
 		}
-		return true;
 	}
 	return false;
 }
@@ -436,14 +445,23 @@ FProcHandle FMacPlatformProcess::CreateProc( const TCHAR* URL, const TCHAR* Parm
 			[ProcessHandle setStandardError: (id)PipeWrite];
 		}
 
-		[ProcessHandle launch];
-
-		if (PriorityModifier != 0)
+		@try
+		{
+			[ProcessHandle launch];
+			
+			if (PriorityModifier != 0)
+			{
+				PriorityModifier = MIN(PriorityModifier, -2);
+				PriorityModifier = MAX(PriorityModifier, 2);
+				// priority values: 20 = lowest, 10 = low, 0 = normal, -10 = high, -20 = highest
+				setpriority(PRIO_PROCESS, [ProcessHandle processIdentifier], -PriorityModifier * 10);
+			}
+		}
+		@catch (NSException* Exc)
 		{
-			PriorityModifier = MIN(PriorityModifier, -2);
-			PriorityModifier = MAX(PriorityModifier, 2);
-			// priority values: 20 = lowest, 10 = low, 0 = normal, -10 = high, -20 = highest
-			setpriority(PRIO_PROCESS, [ProcessHandle processIdentifier], -PriorityModifier * 10);
+			[ProcessHandle release];
+			ProcessHandle = nil;
+			bIsShellScript = false;
 		}
 
 		[Arguments release];
