Removed a spammy error message from the TSF IME management code

diff --git a/Engine/Source/Runtime/Core/Private/Windows/WindowsTextInputMethodSystem.cpp b/Engine/Source/Runtime/Core/Private/Windows/WindowsTextInputMethodSystem.cpp
--- a/Engine/Source/Runtime/Core/Private/Windows/WindowsTextInputMethodSystem.cpp
+++ b/Engine/Source/Runtime/Core/Private/Windows/WindowsTextInputMethodSystem.cpp
@@ -635,10 +635,7 @@ void FWindowsTextInputMethodSystem::DeactivateContext(const TSharedRef<ITextInpu
 			{
 				//  TSF Implementation
 				ITfDocumentMgr* Unused;
-				if(FAILED(TSFThreadManager->AssociateFocus(Hwnd, TSFDisabledDocumentManager, &Unused)))
-				{
-					UE_LOG(LogWindowsTextInputMethodSystem, Error, TEXT("Deactivating a context failed while setting focus to the disabled TSF document manager."));
-				}
+				TSFThreadManager->AssociateFocus(Hwnd, TSFDisabledDocumentManager, &Unused);
 
 				// IMM Implementation
 				::ImmAssociateContext(InternalContext.WindowHandle, nullptr);
