UE4: Test builds now have asserts enabled if shipping builds also have them enabled.

diff --git a/Engine/Source/Runtime/Core/Public/Misc/Build.h b/Engine/Source/Runtime/Core/Public/Misc/Build.h
--- a/Engine/Source/Runtime/Core/Public/Misc/Build.h
+++ b/Engine/Source/Runtime/Core/Public/Misc/Build.h
@@ -194,7 +194,7 @@
 	#define NO_LOGGING										0
 #elif UE_BUILD_TEST
 	#define DO_GUARD_SLOW									0
-	#define DO_CHECK										0
+	#define DO_CHECK										USE_CHECKS_IN_SHIPPING
 	#define STATS											0
 	#define ALLOW_DEBUG_FILES								1
 	#define NO_LOGGING										!USE_LOGGING_IN_SHIPPING
