Fixed failing compilation due to a delegate in FTextLocalizationManager specifying something other than itself for the owning type of the delegate, despite triggering the delegate itself.

diff --git a/Engine/Source/Runtime/Core/Public/Internationalization/TextLocalizationManager.h b/Engine/Source/Runtime/Core/Public/Internationalization/TextLocalizationManager.h
--- a/Engine/Source/Runtime/Core/Public/Internationalization/TextLocalizationManager.h
+++ b/Engine/Source/Runtime/Core/Public/Internationalization/TextLocalizationManager.h
@@ -97,7 +97,7 @@ public:
 	int CORE_API GetHeadCultureRevision() const { return HeadCultureRevisionIndex; }
 
 	/** Broadcasts whenever the new translations are available */
-	DECLARE_EVENT(FInternationalization, FTranslationsChangedEvent)
+	DECLARE_EVENT(FTextLocalizationManager, FTranslationsChangedEvent)
 	CORE_API FTranslationsChangedEvent&  OnTranslationsChanged() { return TranslationsChangedEvent; }
 
 private:
