Changed FMacApplication::LocateWindowUnderCursor() to use cursor location from Cocoa event instead of engine. This should solve issues with it not finding any window due to not updating mouse position in engine often enough in case of very low framerate (like when the editor is in a background and throttling is enabled).

diff --git a/Engine/Source/Runtime/Core/Private/Mac/MacApplication.cpp b/Engine/Source/Runtime/Core/Private/Mac/MacApplication.cpp
--- a/Engine/Source/Runtime/Core/Private/Mac/MacApplication.cpp
+++ b/Engine/Source/Runtime/Core/Private/Mac/MacApplication.cpp
@@ -937,7 +937,12 @@ FCocoaWindow* FMacApplication::FindEventWindow( NSEvent* Event )
 		}
 		else
 		{
-			const FVector2D CursorPos = bUsingHighPrecisionMouseInput ? HighPrecisionMousePos : ((FMacCursor*)Cursor.Get())->GetPosition();
+			NSPoint CursorPos = [Event locationInWindow];
+			if ([Event window])
+			{
+				CursorPos.x += [Event window].frame.origin.x;
+				CursorPos.y += [Event window].frame.origin.y;
+			}
 			TSharedPtr<FMacWindow> WindowUnderCursor = LocateWindowUnderCursor(CursorPos);
 			if (WindowUnderCursor.IsValid())
 			{
@@ -949,10 +954,8 @@ FCocoaWindow* FMacApplication::FindEventWindow( NSEvent* Event )
 	return EventWindow;
 }
 
-TSharedPtr<FMacWindow> FMacApplication::LocateWindowUnderCursor( const FVector2D& CursorPos )
+TSharedPtr<FMacWindow> FMacApplication::LocateWindowUnderCursor( const NSPoint Position )
 {
-	const NSPoint Position = NSMakePoint(CursorPos.X, FPlatformMisc::ConvertSlateYPositionToCocoa(CursorPos.Y));
-
 	NSScreen* MouseScreen = nil;
 	if ([NSScreen screensHaveSeparateSpaces])
 	{
diff --git a/Engine/Source/Runtime/Core/Public/Mac/MacApplication.h b/Engine/Source/Runtime/Core/Public/Mac/MacApplication.h
--- a/Engine/Source/Runtime/Core/Public/Mac/MacApplication.h
+++ b/Engine/Source/Runtime/Core/Public/Mac/MacApplication.h
@@ -131,7 +131,7 @@ private:
 
 	void ResendEvent( NSEvent* Event );
 	FCocoaWindow* FindEventWindow( NSEvent* CocoaEvent );
-	TSharedPtr<FMacWindow> LocateWindowUnderCursor( const FVector2D& CursorPos );
+	TSharedPtr<FMacWindow> LocateWindowUnderCursor( const NSPoint Position );
 
 	NSScreen* FindScreenByPoint( int32 X, int32 Y ) const;
 
