TFunction and TFunctionRef classes for better lambda-passing semantics.

TFunction is like a UE version of std::function.  It is nullable, re-assignable and 'owns' any callable object bound to it.
TFunctionRef is non-nullable, non-assignable, non-owning callable object proxy and is intended to be used for passing lambdas to functions to be called without transferring ownership and without making the function a template, which means it can be FAST.

Bunch of new metafunctions to make this work: TIsFunction, TDecay and TRemoveExtend, modelled on their std:: equivalents.
New TFunction and TFunctionRef visualisers to allow you to see/modify the contents of lambdas after being bound.
TFunction/IFunction code in ShadowRendering replaced with this new stuff.

diff --git a/Engine/Source/Runtime/Core/Public/Core.h b/Engine/Source/Runtime/Core/Public/Core.h
--- a/Engine/Source/Runtime/Core/Public/Core.h
+++ b/Engine/Source/Runtime/Core/Public/Core.h
@@ -22,6 +22,7 @@
 #include "UnrealMathUtility.h"					// FMath
 #include "UnrealTypeTraits.h"					// Type Traits
 #include "UnrealTemplate.h"						// Common template definitions.
+#include "Function.h"							// TFunction
 #include "MemoryOps.h"							// Functions for efficient handling of object arrays.
 
 #include "Misc/CoreDefines.h"
@@ -114,4 +115,4 @@
 #include "Attribute.h"
 #include "Optional.h"
 #include "EnumClassFlags.h"
-#include "Misc/DisableOldUETypes.h"
\ No newline at end of file
+#include "Misc/DisableOldUETypes.h"
diff --git a/Engine/Source/Runtime/Core/Public/Templates/Function.h b/Engine/Source/Runtime/Core/Public/Templates/Function.h
new file mode 100644
--- /dev/null
+++ b/Engine/Source/Runtime/Core/Public/Templates/Function.h
@@ -0,0 +1,681 @@
+// Copyright 1998-2014 Epic Games, Inc. All Rights Reserved.
+
+#pragma once
+
+#include <new>
+
+#include "UnrealMemory.h"
+
+// Disable visualization hack for shipping or test builds.
+#if !(UE_BUILD_SHIPPING || UE_BUILD_TEST)
+	#define ENABLE_TFUNCTIONREF_VISUALIZATION 1
+#else
+	#define ENABLE_TFUNCTIONREF_VISUALIZATION 0
+#endif
+
+// VS2012 doesn't support defaulted function template arguments, so we need a workaround for that.
+// The workaround technically changes the signature of the constructors, but we assume that the modified
+// signature will not be used.
+#if defined(_MSC_VER) && _MSC_VER < 1800
+	#define ENABLE_TFUNCTIONREF_WORKAROUND 1
+#else
+	#define ENABLE_TFUNCTIONREF_WORKAROUND 0
+#endif
+
+/**
+ * TFunction<FuncType>
+ *
+ * See the class definition for intended usage.
+ */
+template <typename FuncType>
+class TFunction;
+
+/**
+ * TFunctionRef<FuncType>
+ *
+ * See the class definition for intended usage.
+ */
+template <typename FuncType>
+class TFunctionRef;
+
+/**
+ * Traits class which checks if T is a TFunction<> type.
+ */
+template <typename T> struct TIsATFunction               { enum { Value = false }; };
+template <typename T> struct TIsATFunction<TFunction<T>> { enum { Value = true  }; };
+
+/**
+ * Traits class which checks if T is a TFunction<> type.
+ */
+template <typename T> struct TIsATFunctionRef                  { enum { Value = false }; };
+template <typename T> struct TIsATFunctionRef<TFunctionRef<T>> { enum { Value = true  }; };
+
+/**
+ * Private implementation details of TFunction and TFunctionRef.
+ */
+namespace UE4Function_Private
+{
+	/**
+	 * Common interface to a callable object owned by TFunction.
+	 */
+	struct IFunction_OwnedObject
+	{
+		/**
+		 * Returns a heap-allocated copy of the object, created by copy construction.
+		 */
+		virtual IFunction_OwnedObject* CloneByCopy() const = 0;
+
+		/**
+		 * Returns the address of the object.
+		 */
+		virtual void* GetAddress() = 0;
+
+		/**
+		 * Destructor.
+		 */
+		virtual ~IFunction_OwnedObject() = 0;
+	};
+
+	/**
+	 * Destructor.
+	 */
+	inline IFunction_OwnedObject::~IFunction_OwnedObject()
+	{
+	}
+
+	/**
+	 * Implementation of IFunction_OwnedObject for a given T.
+	 */
+	template <typename T>
+	struct TFunction_OwnedObject : public IFunction_OwnedObject
+	{
+		/**
+		 * Constructor which creates its T by copying.
+		 */
+		explicit TFunction_OwnedObject(const T& InObj)
+			: Obj(InObj)
+		{
+		}
+
+		/**
+		 * Constructor which creates its T by moving.
+		 */
+		explicit TFunction_OwnedObject(T&& InObj)
+			: Obj(MoveTemp(InObj))
+		{
+		}
+
+		/**
+		 * Returns a heap-allocated copy of the object, created by copy construction.
+		 */
+		IFunction_OwnedObject* CloneByCopy() const override
+		{
+			return new TFunction_OwnedObject(Obj);
+		}
+
+		/**
+		 * Returns the address of the object.
+		 */
+		virtual void* GetAddress() override
+		{
+			return &Obj;
+		}
+
+		T Obj;
+	};
+
+	/**
+	 * A class which is used to instantiate the code needed to call a bound function.
+	 */
+	template <typename Functor, typename FuncType>
+	struct TFunctionRefCaller;
+
+	/**
+	 * A class which is used to instantiate the code needed to assert when called - used for null bindings.
+	 */
+	template <typename FuncType>
+	struct TFunctionRefAsserter;
+
+	/**
+	 * A class which defines an operator() which will invoke the TFunctionRefCaller::Call function.
+	 */
+	template <typename FuncType>
+	struct TFunctionRefBase;
+
+	#if ENABLE_TFUNCTIONREF_VISUALIZATION
+		/**
+		 * Helper classes to help debugger visualization.
+		 */
+		struct IDebugHelper
+		{
+			virtual ~IDebugHelper() = 0;
+		};
+
+		inline IDebugHelper::~IDebugHelper()
+		{
+		}
+
+		template <typename T>
+		struct TDebugHelper : IDebugHelper
+		{
+			T* Ptr;
+		};
+	#endif
+
+	template <typename T>
+	inline T& FakeCall(T* Ptr)
+	{
+		return *Ptr;
+	}
+
+	inline void FakeCall(void* Ptr)
+	{
+	}
+
+
+	/**
+	 * Switch on the existence of variadics.  Once all our supported compilers support variadics, a lot of this code
+	 * can be collapsed into FFunctionRefCaller.  They're currently separated out to minimize the amount of workarounds
+	 * needed.
+	 */
+	#if PLATFORM_COMPILER_HAS_VARIADIC_TEMPLATES
+
+		template <typename Functor, typename Ret, typename... ParamTypes>
+		struct TFunctionRefCaller<Functor, Ret (ParamTypes...)>
+		{
+			static Ret Call(void* Obj, ParamTypes&... Params)
+			{
+				return (*(Functor*)Obj)(Forward<ParamTypes>(Params)...);
+			}
+		};
+
+		/**
+		 * Specialization for void return types.
+		 */
+		template <typename Functor, typename... ParamTypes>
+		struct TFunctionRefCaller<Functor, void (ParamTypes...)>
+		{
+			static void Call(void* Obj, ParamTypes&... Params)
+			{
+				(*(Functor*)Obj)(Forward<ParamTypes>(Params)...);
+			}
+		};
+
+		template <typename Ret, typename... ParamTypes>
+		struct TFunctionRefAsserter<Ret (ParamTypes...)>
+		{
+			static Ret Call(void* Obj, ParamTypes&...)
+			{
+				// Attempting to call a null TFunction!
+				check(false);
+
+				// This doesn't need to be valid, because it'll never be reached, but it does at least need to compile.
+				return FakeCall((Ret*)Obj);
+			}
+		};
+
+		template <typename Ret, typename... ParamTypes>
+		struct TFunctionRefBase<Ret (ParamTypes...)>
+		{
+			typedef Ret (*FuncPtr)(void*, ParamTypes&...);
+
+			Ret operator()(ParamTypes... Params) const
+			{
+				const auto* Derived = static_cast<const TFunctionRef<Ret (ParamTypes...)>*>(this);
+				return Derived->Func(Derived->Ptr, Params...);
+			}
+		};
+
+	#else
+
+		/**
+		 * Specialisations of the above classes which don't use variadic templates, up to a maximum arity of 4.
+		 */
+
+		template <typename Functor, typename Ret                                                    > struct TFunctionRefCaller<Functor, Ret (              )> { static Ret Call(void* Obj                                ) { return (*(Functor*)Obj)(                                                                  ); } };
+		template <typename Functor, typename Ret, typename T0                                       > struct TFunctionRefCaller<Functor, Ret (T0            )> { static Ret Call(void* Obj, T0& P0                        ) { return (*(Functor*)Obj)(Forward<T0>(P0)                                                   ); } };
+		template <typename Functor, typename Ret, typename T0, typename T1                          > struct TFunctionRefCaller<Functor, Ret (T0, T1        )> { static Ret Call(void* Obj, T0& P0, T1& P1                ) { return (*(Functor*)Obj)(Forward<T0>(P0), Forward<T1>(P1)                                  ); } };
+		template <typename Functor, typename Ret, typename T0, typename T1, typename T2             > struct TFunctionRefCaller<Functor, Ret (T0, T1, T2    )> { static Ret Call(void* Obj, T0& P0, T1& P1, T2& P2        ) { return (*(Functor*)Obj)(Forward<T0>(P0), Forward<T1>(P1), Forward<T2>(P2)                 ); } };
+		template <typename Functor, typename Ret, typename T0, typename T1, typename T2, typename T3> struct TFunctionRefCaller<Functor, Ret (T0, T1, T2, T3)> { static Ret Call(void* Obj, T0& P0, T1& P1, T2& P2, T3& P3) { return (*(Functor*)Obj)(Forward<T0>(P0), Forward<T1>(P1), Forward<T2>(P2), Forward<T3>(P3)); } };
+
+		template <typename Functor                                                    > struct TFunctionRefCaller<Functor, void (              )> { static void Call(void* Obj                                ) { (*(Functor*)Obj)(                                                                  ); } };
+		template <typename Functor, typename T0                                       > struct TFunctionRefCaller<Functor, void (T0            )> { static void Call(void* Obj, T0& P0                        ) { (*(Functor*)Obj)(Forward<T0>(P0)                                                   ); } };
+		template <typename Functor, typename T0, typename T1                          > struct TFunctionRefCaller<Functor, void (T0, T1        )> { static void Call(void* Obj, T0& P0, T1& P1                ) { (*(Functor*)Obj)(Forward<T0>(P0), Forward<T1>(P1)                                  ); } };
+		template <typename Functor, typename T0, typename T1, typename T2             > struct TFunctionRefCaller<Functor, void (T0, T1, T2    )> { static void Call(void* Obj, T0& P0, T1& P1, T2& P2        ) { (*(Functor*)Obj)(Forward<T0>(P0), Forward<T1>(P1), Forward<T2>(P2)                 ); } };
+		template <typename Functor, typename T0, typename T1, typename T2, typename T3> struct TFunctionRefCaller<Functor, void (T0, T1, T2, T3)> { static void Call(void* Obj, T0& P0, T1& P1, T2& P2, T3& P3) { (*(Functor*)Obj)(Forward<T0>(P0), Forward<T1>(P1), Forward<T2>(P2), Forward<T3>(P3)); } };
+
+		template <typename Ret                                                    > struct TFunctionRefAsserter<Ret (              )> { static Ret Call(void* Obj                    ) { /* Attempting to call a null TFunction! */ check(false); return FakeCall((Ret*)Obj); } };
+		template <typename Ret, typename T0                                       > struct TFunctionRefAsserter<Ret (T0            )> { static Ret Call(void* Obj, T0&               ) { /* Attempting to call a null TFunction! */ check(false); return FakeCall((Ret*)Obj); } };
+		template <typename Ret, typename T0, typename T1                          > struct TFunctionRefAsserter<Ret (T0, T1        )> { static Ret Call(void* Obj, T0&, T1&          ) { /* Attempting to call a null TFunction! */ check(false); return FakeCall((Ret*)Obj); } };
+		template <typename Ret, typename T0, typename T1, typename T2             > struct TFunctionRefAsserter<Ret (T0, T1, T2    )> { static Ret Call(void* Obj, T0&, T1&, T2&     ) { /* Attempting to call a null TFunction! */ check(false); return FakeCall((Ret*)Obj); } };
+		template <typename Ret, typename T0, typename T1, typename T2, typename T3> struct TFunctionRefAsserter<Ret (T0, T1, T2, T3)> { static Ret Call(void* Obj, T0&, T1&, T2&, T3&) { /* Attempting to call a null TFunction! */ check(false); return FakeCall((Ret*)Obj); } };
+
+		template <typename Ret                                                    > struct TFunctionRefBase<Ret (              )> { typedef Ret (*FuncPtr)(void*                    ); Ret operator()(                          ) const { const auto* Derived = static_cast<const TFunctionRef<Ret (              )>*>(this); return Derived->Func(Derived->Ptr                ); } };
+		template <typename Ret, typename T0                                       > struct TFunctionRefBase<Ret (T0            )> { typedef Ret (*FuncPtr)(void*, T0&               ); Ret operator()(T0 P0                     ) const { const auto* Derived = static_cast<const TFunctionRef<Ret (T0            )>*>(this); return Derived->Func(Derived->Ptr, P0            ); } };
+		template <typename Ret, typename T0, typename T1                          > struct TFunctionRefBase<Ret (T0, T1        )> { typedef Ret (*FuncPtr)(void*, T0&, T1&          ); Ret operator()(T0 P0, T1 P1              ) const { const auto* Derived = static_cast<const TFunctionRef<Ret (T0, T1        )>*>(this); return Derived->Func(Derived->Ptr, P0, P1        ); } };
+		template <typename Ret, typename T0, typename T1, typename T2             > struct TFunctionRefBase<Ret (T0, T1, T2    )> { typedef Ret (*FuncPtr)(void*, T0&, T1&, T2&     ); Ret operator()(T0 P0, T1 P1, T2 P2       ) const { const auto* Derived = static_cast<const TFunctionRef<Ret (T0, T1, T2    )>*>(this); return Derived->Func(Derived->Ptr, P0, P1, P2    ); } };
+		template <typename Ret, typename T0, typename T1, typename T2, typename T3> struct TFunctionRefBase<Ret (T0, T1, T2, T3)> { typedef Ret (*FuncPtr)(void*, T0&, T1&, T2&, T3&); Ret operator()(T0 P0, T1 P1, T2 P2, T3 P3) const { const auto* Derived = static_cast<const TFunctionRef<Ret (T0, T1, T2, T3)>*>(this); return Derived->Func(Derived->Ptr, P0, P1, P2, P3); } };
+
+	#endif
+}
+
+/**
+ * TFunctionRef<FuncType>
+ *
+ * A class which represents a reference to something callable.  The important part here is *reference* - if
+ * you bind it to a lambda and the lambda goes out of scope, you will be left with an invalid reference.
+ *
+ * FuncType represents a function type and so TFunctionRef should be defined as follows:
+ *
+ * // A function taking a string and float and returning int32.  Parameter names are optional.
+ * TFunctionRef<int32 (const FString& Name, float Scale)>
+ *
+ * If you also want to take ownership of the callable thing, e.g. you want to return a lambda from a
+ * function, you should use TFunction.  TFunctionRef does not concern itself with ownership because it's
+ * intended to be FAST.
+ *
+ * TFunctionRef is most useful when you want to parameterize a function with some caller-defined code
+ * without making it a template.
+ *
+ * Example:
+ *
+ * // Something.h
+ * void DoSomethingWithConvertingStringsToInts(TFunctionRef<int32 (const FString& Str)> Convert);
+ *
+ * // Something.cpp
+ * void DoSomethingWithConvertingStringsToInts(TFunctionRef<int32 (const FString& Str)> Convert)
+ * {
+ *     for (const FString& Str : SomeBunchOfStrings)
+ *     {
+ *         int32 Int = Func(Str);
+ *         DoSomething(Int);
+ *     }
+ * }
+ *
+ * // SomewhereElse.cpp
+ * #include "Something.h"
+ *
+ * void Func()
+ * {
+ *     // First do something using string length
+ *     DoSomethingWithConvertingStringsToInts([](const FString& Str) {
+ *         return Str.Len();
+ *     });
+ *
+ *     // Then do something using string conversion
+ *     DoSomethingWithConvertingStringsToInts([](const FString& Str) {
+ *         int32 Result;
+ *         TTypeFromString<int32>::FromString(Result, *Str);
+ *         return Result;
+ *     });
+ * }
+ */
+template <typename FuncType>
+class TFunctionRef : public UE4Function_Private::TFunctionRefBase<FuncType>
+{
+	template <typename T>
+	friend class TFunction;
+
+	friend struct UE4Function_Private::TFunctionRefBase<FuncType>;
+
+	typedef UE4Function_Private::TFunctionRefBase<FuncType> Super;
+
+public:
+	/**
+	 * Constructor which binds a TFunctionRef to a non-const lvalue function object.
+	 */
+#if ENABLE_TFUNCTIONREF_WORKAROUND
+	template <typename FunctorType>
+	TFunctionRef(FunctorType& Functor, typename TEnableIf<!TIsFunction<FunctorType>::Value && !TIsATFunctionRef<FunctorType>::Value>::Type* = nullptr)
+#else
+	template <typename FunctorType, typename = typename TEnableIf<!TIsFunction<FunctorType>::Value && !TIsATFunctionRef<FunctorType>::Value>::Type>
+	TFunctionRef(FunctorType& Functor)
+#endif
+	{
+		// This constructor is disabled for function types because we want it to call the function pointer overload.
+		// It is also disabled for TFunctionRef types because VC is incorrectly treating it as a copy constructor.
+
+		Set(&Functor);
+	}
+
+	/**
+	 * Constructor which binds a TFunctionRef to an rvalue or const lvalue function object.
+	 */
+#if ENABLE_TFUNCTIONREF_WORKAROUND
+	template <typename FunctorType>
+	TFunctionRef(const FunctorType& Functor, typename TEnableIf<!TIsFunction<FunctorType>::Value && !TIsATFunctionRef<FunctorType>::Value>::Type* = nullptr)
+#else
+	template <typename FunctorType, typename = typename TEnableIf<!TIsFunction<FunctorType>::Value && !TIsATFunctionRef<FunctorType>::Value>::Type>
+	TFunctionRef(const FunctorType& Functor)
+#endif
+	{
+		// This constructor is disabled for function types because we want it to call the function pointer overload.
+		// It is also disabled for TFunctionRef types because VC is incorrectly treating it as a copy constructor.
+
+		Set(&Functor);
+	}
+
+	/**
+	 * Constructor which binds a TFunctionRef to a function pointer.
+	 */
+#if ENABLE_TFUNCTIONREF_WORKAROUND
+	template <typename FunctionType>
+	TFunctionRef(FunctionType* Function, typename TEnableIf<TIsFunction<FunctionType>::Value>::Type* = nullptr)
+#else
+	template <typename FunctionType, typename = typename TEnableIf<TIsFunction<FunctionType>::Value>::Type>
+	TFunctionRef(FunctionType* Function)
+#endif
+	{
+		// This constructor is enabled only for function types because we don't want weird errors from it being called with arbitrary pointers.
+
+		Set(Function);
+	}
+
+	#if ENABLE_TFUNCTIONREF_VISUALIZATION
+
+		/**
+		 * Copy constructor.
+		 */
+		TFunctionRef(const TFunctionRef& Other)
+		{
+			// If visualization is enabled, then we need to do an explicit copy
+			// to ensure that our hacky DebugPtrStorage's vptr is copied.
+			CopyAndReseat(Other, Other.Ptr);
+		}
+
+	#endif
+
+	// We delete the assignment operators because we don't want it to be confused with being related to
+	// regular C++ reference assignment - i.e. calling the assignment operator of whatever the reference
+	// is bound to - because that's not what TFunctionRef does, or is it even capable of doing that.
+	#if PLATFORM_COMPILER_HAS_DEFAULTED_FUNCTIONS
+
+		#if !ENABLE_TFUNCTIONREF_VISUALIZATION
+			TFunctionRef(const TFunctionRef&) = default;
+		#endif
+		TFunctionRef& operator=(const TFunctionRef&) const = delete;
+		~TFunctionRef()                                    = default;
+
+	#else
+
+		// We mark copy assignment as private.  We don't define the others because it's mostly likely
+		// that this code path is only for VC and the compiler will still generate the others fine anyway,
+		// despite such behavior being deprecated.
+		private:
+			TFunctionRef& operator=(const TFunctionRef&) const;
+		public:
+
+	#endif
+
+private:
+	/**
+	 * Constructor which leaves the TFunctionRef uninitialized.  This is only for use with
+	 * TFunction.
+	 */
+	explicit TFunctionRef(ENoInit)
+	{
+	}
+
+	/**
+	 * Sets the state of the TFunctionRef given a pointer to a callable thing.
+	 */
+	template <typename FunctorType>
+	void Set(FunctorType* Functor)
+	{
+		// We force a void* cast here because if FunctorType is an actual function then
+		// this won't compile.  We convert it back again before we use it anyway.
+
+		Ptr  = (void*)Functor;
+		Func = &UE4Function_Private::TFunctionRefCaller<FunctorType, FuncType>::Call;
+
+		#if ENABLE_TFUNCTIONREF_VISUALIZATION
+			// We placement new over the top of the same object each time.  This is illegal,
+			// but it ensures that the vptr is set correctly for the bound type, and so is
+			// visualizable.  We never depend on the state of this object at runtime, so it's
+			// ok.
+			new ((void*)&DebugPtrStorage) UE4Function_Private::TDebugHelper<FunctorType>;
+			DebugPtrStorage.Ptr = Ptr;
+		#endif
+	}
+
+	/**
+	 * 'Nulls' the TFunctionRef.
+	 */
+	void Unset()
+	{
+		Ptr  = nullptr;
+		Func = &UE4Function_Private::TFunctionRefAsserter<FuncType>::Call;
+	}
+
+	/**
+	 * Copies another TFunctionRef and rebinds it to another object of the same type which was originally bound.
+	 * Only intended to be used by TFunction's copy constructor/assignment operator.
+	 */
+	void CopyAndReseat(const TFunctionRef& Other, void* Functor)
+	{
+		Ptr  = Functor;
+		Func = Other.Func;
+
+		#if ENABLE_TFUNCTIONREF_VISUALIZATION
+			// Use Memcpy to copy the other DebugPtrStorage, including vptr (because we don't know the bound type
+			// here), and then reseat the underlying pointer.  Possibly even more evil than the Set code.
+			FMemory::Memcpy(&DebugPtrStorage, &Other.DebugPtrStorage, sizeof(DebugPtrStorage));
+			DebugPtrStorage.Ptr = Ptr;
+		#endif
+	}
+
+	/**
+	 * Tests if the TFunctionRef is unset.
+	 */
+	bool IsUnset() const
+	{
+		return Ptr == nullptr;
+	}
+
+	// A pointer to the callable object
+	void* Ptr;
+
+	// A pointer to a function which invokes the call operator on the callable object
+	typename Super::FuncPtr Func;
+
+	#if ENABLE_TFUNCTIONREF_VISUALIZATION
+		// To help debug visualizers
+		UE4Function_Private::TDebugHelper<void> DebugPtrStorage;
+	#endif
+};
+
+/**
+ * TFunction<FuncType>
+ *
+ * A class which represents a copy of something callable.  FuncType represents a function type and so
+ * TFunction should be defined as follows:
+ *
+ * // A function taking a string and float and returning int32.  Parameter names are optional.
+ * TFunction<int32 (const FString& Name, float Scale)>
+ *
+ * Unlike TFunctionRef, this object is intended to be used like a UE4 version of std::function.  That is,
+ * it takes a copy of whatever is bound to it, meaning you can return it from functions and store them in
+ * objects without caring about the lifetime of the original object being bound.
+ *
+ * Example:
+ *
+ * // Something.h
+ * TFunction<FString (int32)> GetTransform();
+ *
+ * // Something.cpp
+ * TFunction<FString (int32)> GetTransform(const FString& Prefix)
+ * {
+ *     // Squares number and returns it as a string with the specified prefix
+ *     return [=](int32 Num) {
+ *         return Prefix + TEXT(": ") + TTypeToString<int32>::ToString(Num * Num);
+ *     };
+ * }
+ *
+ * // SomewhereElse.cpp
+ * #include "Something.h"
+ *
+ * void Func()
+ * {
+ *     TFunction<FString (int32)> Transform = GetTransform(TEXT("Hello"));
+ *
+ *     FString Result = Transform(5); // "Hello: 25"
+ * }
+ */
+template <typename FuncType>
+class TFunction : public TFunctionRef<FuncType>
+{
+	typedef TFunctionRef<FuncType> Super;
+
+public:
+	/**
+	 * Default constructor.
+	 */
+	TFunction(TYPE_OF_NULLPTR = nullptr)
+		: Super(NoInit)
+	{
+		Func = nullptr;
+		Super::Unset();
+	}
+
+	/**
+	 * Constructor which binds a TFunction to any function object.
+	 */
+#if ENABLE_TFUNCTIONREF_WORKAROUND
+	template <typename FunctorType>
+	TFunction(FunctorType&& InFunc, typename TEnableIf<!TIsATFunction<typename TDecay<FunctorType>::Type>::Value>::Type* = nullptr)
+#else
+	template <typename FunctorType, typename = typename TEnableIf<!TIsATFunction<typename TDecay<FunctorType>::Type>::Value>::Type>
+	TFunction(FunctorType&& InFunc)
+#endif
+		: Super(NoInit)
+	{
+		// This constructor is disabled for TFunction types because VC is incorrectly treating it as copy/move constructors.
+
+		typedef typename TDecay<FunctorType>::Type                             DecayedFunctorType;
+		typedef UE4Function_Private::TFunction_OwnedObject<DecayedFunctorType> OwnedType;
+
+		// This is probably a mistake if you expect TFunction to take a copy of what
+		// TFunctionRef is bound to, because that's not possible.
+		//
+		// If you really intended to bind a TFunction to a TFunctionRef, you can just
+		// wrap it in a lambda (and thus it's clear you're just binding to a call to another
+		// reference):
+		//
+		// TFunction<int32(float)> MyFunction = [=](float F) -> int32 { return MyFunctionRef(F); };
+		static_assert(!TIsATFunctionRef<DecayedFunctorType>::Value, "Cannot construct a TFunction from a TFunctionRef");
+
+		auto NewObj = new OwnedType(Forward<FunctorType>(InFunc));
+
+		Func = NewObj;
+		Super::Set(&NewObj->Obj);
+	}
+
+	/**
+	 * Copy constructor.
+	 */
+	TFunction(const TFunction& Other)
+		: Super(NoInit)
+	{
+		if (Other)
+		{
+			Func = Other.Func->CloneByCopy();
+			Super::CopyAndReseat(Other, Func->GetAddress());
+		}
+		else
+		{
+			Func = nullptr;
+			Super::Unset();
+		}
+	}
+
+	/**
+	 * Move constructor.
+	 */
+	TFunction(TFunction&& Other)
+		: Super(NoInit)
+	{
+		Func = Other.Func;
+		Super::CopyAndReseat(Other, Func->GetAddress());
+
+		Other.Func = nullptr;
+		Other.Unset();
+	}
+
+	/**
+	 * Copy/move assignment operator.
+	 */
+	TFunction& operator=(TFunction Other)
+	{
+		FMemory::Memswap(&Other, this, sizeof(TFunction));
+		return *this;
+	}
+
+	/**
+	 * Nullptr assignment operator.
+	 */
+	TFunction& operator=(TYPE_OF_NULLPTR)
+	{
+		delete Func;
+		Func = nullptr;
+		Super::Unset();
+
+		return *this;
+	}
+
+	/**
+	 * Destructor.
+	 */
+	~TFunction()
+	{
+		delete Func;
+	}
+
+	/**
+	 * Tests if the TFunction is callable.
+	 */
+	FORCEINLINE_EXPLICIT_OPERATOR_BOOL() const
+	{
+		return !Super::IsUnset();
+	}
+
+	/**
+	 * TFunctions are not usefully comparable, but we need to define operators for FORCEINLINE_EXPLICIT_OPERATOR_BOOL to work
+	 * in the case where explicit operator bool is not supported by the compiler.
+	 */
+	SAFE_BOOL_OPERATORS(TFunction)
+
+private:
+	/**
+	 * Nullptr equality operator.
+	 */
+	friend bool operator==(TYPE_OF_NULLPTR, const TFunction& Func)
+	{
+		return !Func;
+	}
+
+	/**
+	 * Nullptr equality operator.
+	 */
+	friend bool operator==(const TFunction& Func, TYPE_OF_NULLPTR)
+	{
+		return !Func;
+	}
+
+	/**
+	 * Nullptr inequality operator.
+	 */
+	friend bool operator!=(TYPE_OF_NULLPTR, const TFunction& Func)
+	{
+		return (bool)Func;
+	}
+
+	/**
+	 * Nullptr inequality operator.
+	 */
+	friend bool operator!=(const TFunction& Func, TYPE_OF_NULLPTR)
+	{
+		return (bool)Func;
+	}
+
+	UE4Function_Private::IFunction_OwnedObject* Func;
+};
diff --git a/Engine/Source/Runtime/Core/Public/Templates/UnrealTemplate.h b/Engine/Source/Runtime/Core/Public/Templates/UnrealTemplate.h
--- a/Engine/Source/Runtime/Core/Public/Templates/UnrealTemplate.h
+++ b/Engine/Source/Runtime/Core/Public/Templates/UnrealTemplate.h
@@ -464,6 +464,37 @@ template <typename T> struct TRValueToLValueReference<T&&> { typedef T& Type; };
 template <typename T>           struct TIsCPPArray       { enum { Value = false }; };
 template <typename T, uint32 N> struct TIsCPPArray<T[N]> { enum { Value = true  }; };
 
+/**
+ * Removes one dimension of extents from an array type.
+ */
+template <typename T>           struct TRemoveExtent       { typedef T Type; };
+template <typename T>           struct TRemoveExtent<T[]>  { typedef T Type; };
+template <typename T, uint32 N> struct TRemoveExtent<T[N]> { typedef T Type; };
+
+/**
+ * Returns the decayed type of T, meaning it removes all references, qualifiers and
+ * applies array-to-pointer and function-to-pointer conversions.
+ *
+ * http://en.cppreference.com/w/cpp/types/decay
+ */
+template <typename T>
+struct TDecay
+{
+private:
+	typedef typename TRemoveReference<T>::Type NoRefs;
+
+public:
+	typedef typename TChooseClass<
+		TIsCPPArray<NoRefs>::Value,
+		typename TRemoveExtent<NoRefs>::Type*,
+		typename TChooseClass<
+			TIsFunction<NoRefs>::Value,
+			NoRefs*,
+			typename TRemoveCV<NoRefs>::Type
+		>::Result
+	>::Result Type;
+};
+
 /**
  * Reverses the order of the bits of a value.
  * This is an TEnableIf'd template to ensure that no undesirable conversions occur.  Overloads for other types can be added in the same way.
diff --git a/Engine/Source/Runtime/Core/Public/Templates/UnrealTypeTraits.h b/Engine/Source/Runtime/Core/Public/Templates/UnrealTypeTraits.h
--- a/Engine/Source/Runtime/Core/Public/Templates/UnrealTypeTraits.h
+++ b/Engine/Source/Runtime/Core/Public/Templates/UnrealTypeTraits.h
@@ -210,6 +210,39 @@ struct TIsFundamentalType
 	enum { Value = TIsArithmeticType<T>::Value || TIsVoidType<T>::Value };
 };
 
+/**
+ * TIsFunction
+ *
+ * Tests is a type is a function.
+ */
+template <typename T>
+struct TIsFunction
+{
+	enum { Value = false };
+};
+
+#if PLATFORM_COMPILER_HAS_VARIADIC_TEMPLATES
+
+	template <typename RetType, typename... Params>
+	struct TIsFunction<RetType(Params...)>
+	{
+		enum { Value = true };
+	};
+
+#else
+
+	template <typename RetType                                                                                                        > struct TIsFunction<RetType(                              )> { enum { Value = true }; };
+	template <typename RetType, typename P0                                                                                           > struct TIsFunction<RetType(P0                            )> { enum { Value = true }; };
+	template <typename RetType, typename P0, typename P1                                                                              > struct TIsFunction<RetType(P0, P1                        )> { enum { Value = true }; };
+	template <typename RetType, typename P0, typename P1, typename P2                                                                 > struct TIsFunction<RetType(P0, P1, P2                    )> { enum { Value = true }; };
+	template <typename RetType, typename P0, typename P1, typename P2, typename P3                                                    > struct TIsFunction<RetType(P0, P1, P2, P3                )> { enum { Value = true }; };
+	template <typename RetType, typename P0, typename P1, typename P2, typename P3, typename P4                                       > struct TIsFunction<RetType(P0, P1, P2, P3, P4            )> { enum { Value = true }; };
+	template <typename RetType, typename P0, typename P1, typename P2, typename P3, typename P4, typename P5                          > struct TIsFunction<RetType(P0, P1, P2, P3, P4, P5        )> { enum { Value = true }; };
+	template <typename RetType, typename P0, typename P1, typename P2, typename P3, typename P4, typename P5, typename P6             > struct TIsFunction<RetType(P0, P1, P2, P3, P4, P5, P6    )> { enum { Value = true }; };
+	template <typename RetType, typename P0, typename P1, typename P2, typename P3, typename P4, typename P5, typename P6, typename P7> struct TIsFunction<RetType(P0, P1, P2, P3, P4, P5, P6, P7)> { enum { Value = true }; };
+
+#endif
+
 /**
  * TIsZeroConstructType
  */
@@ -237,6 +270,7 @@ struct TIsWeakPointerType
 	enum { Value = false };
 };
 
+
 /**
  * TNameOf
  */
diff --git a/Engine/Source/Runtime/Renderer/Private/ShadowRendering.cpp b/Engine/Source/Runtime/Renderer/Private/ShadowRendering.cpp
--- a/Engine/Source/Runtime/Renderer/Private/ShadowRendering.cpp
+++ b/Engine/Source/Runtime/Renderer/Private/ShadowRendering.cpp
@@ -1517,9 +1517,9 @@ static TAutoConsoleVariable<int32> CVarRHICmdShadowDeferredContexts(
 class FShadowParallelCommandListSet : public FParallelCommandListSet
 {
 	FProjectedShadowInfo& ProjectedShadowInfo;
-	IFunction<void, FRHICommandList&>& SetShadowRenderTargets;
+	TFunctionRef<void (FRHICommandList& RHICmdList)> SetShadowRenderTargets;
 public:
-	FShadowParallelCommandListSet(const FViewInfo& InView, FRHICommandList& InParentCmdList, bool* InOutDirty, bool bInParallelExecute, FProjectedShadowInfo& InProjectedShadowInfo, IFunction<void, FRHICommandList&>& InSetShadowRenderTargets)
+	FShadowParallelCommandListSet(const FViewInfo& InView, FRHICommandList& InParentCmdList, bool* InOutDirty, bool bInParallelExecute, FProjectedShadowInfo& InProjectedShadowInfo, TFunctionRef<void (FRHICommandList& RHICmdList)> InSetShadowRenderTargets)
 		: FParallelCommandListSet(InView, InParentCmdList, InOutDirty, bInParallelExecute)
 		, ProjectedShadowInfo(InProjectedShadowInfo)
 		, SetShadowRenderTargets(InSetShadowRenderTargets)
@@ -1533,7 +1533,7 @@ public:
 	}
 };
 
-void FProjectedShadowInfo::RenderDepthInner(FRHICommandList& RHICmdList, FSceneRenderer* SceneRenderer, const FViewInfo* FoundView, IFunction<void, FRHICommandList&>& SetShadowRenderTargets)
+void FProjectedShadowInfo::RenderDepthInner(FRHICommandList& RHICmdList, FSceneRenderer* SceneRenderer, const FViewInfo* FoundView, TFunctionRef<void (FRHICommandList& RHICmdList)> SetShadowRenderTargets)
 {
 	FShadowDepthDrawingPolicyContext PolicyContext(this);
 
@@ -1621,7 +1621,7 @@ void FProjectedShadowInfo::RenderDepthInner(FRHICommandList& RHICmdList, FSceneR
 	}
 }
 
-void FProjectedShadowInfo::RenderDepth(FRHICommandList& RHICmdList, FSceneRenderer* SceneRenderer, IFunction<void, FRHICommandList&>& SetShadowRenderTargets)
+void FProjectedShadowInfo::RenderDepth(FRHICommandList& RHICmdList, FSceneRenderer* SceneRenderer, TFunctionRef<void (FRHICommandList& RHICmdList)> SetShadowRenderTargets)
 {
 #if WANTS_DRAW_MESH_EVENTS
 	FString EventName;
@@ -2617,10 +2617,10 @@ bool FDeferredShadingSceneRenderer::RenderOnePassPointLightShadows(FRHICommandLi
 			if (!ProjectedShadowInfo->bRayTracedDistanceFieldShadow)
 			{
 				SCOPED_DRAW_EVENT(RHICmdList, ShadowDepthsFromOpaque);
-				auto SetShadowRenderTargets = MakeFunction<void, FRHICommandList&>([ProjectedShadowInfo](FRHICommandList& RHICmdList)
+				auto SetShadowRenderTargets = [ProjectedShadowInfo](FRHICommandList& RHICmdList)
 				{
 					GSceneRenderTargets.BeginRenderingCubeShadowDepth(RHICmdList, ProjectedShadowInfo->ResolutionX);
-				});
+				};
 
 				SetShadowRenderTargets(RHICmdList);  // run it now, maybe run it later for parallel command lists
 				ProjectedShadowInfo->bAllocated = true;
@@ -3027,10 +3027,10 @@ bool FDeferredShadingSceneRenderer::RenderReflectiveShadowMaps(FRHICommandListIm
 				FLightPropagationVolume* LightPropagationVolume = ViewState->GetLightPropagationVolume();
 
 				check(LightPropagationVolume);
-				auto SetShadowRenderTargets = MakeFunction<void, FRHICommandList&>([LightPropagationVolume](FRHICommandList& RHICmdList)
+				auto SetShadowRenderTargets = [LightPropagationVolume](FRHICommandList& RHICmdList)
 				{
 					GSceneRenderTargets.BeginRenderingReflectiveShadowMap(RHICmdList, LightPropagationVolume);
-				});
+				};
 
 				SetShadowRenderTargets(RHICmdList);  // run it now, maybe run it later for parallel command lists
 
@@ -3205,10 +3205,10 @@ bool FDeferredShadingSceneRenderer::RenderProjectedShadows(FRHICommandListImmedi
 			SCOPED_DRAW_EVENT(RHICmdList, ShadowDepthsFromOpaque);
 
 			bool bFirst = true;
-			auto SetShadowRenderTargets = MakeFunction<void, FRHICommandList&>([&bFirst](FRHICommandList& RHICmdList)
+			auto SetShadowRenderTargets = [&bFirst](FRHICommandList& RHICmdList)
 			{
 				GSceneRenderTargets.BeginRenderingShadowDepth(RHICmdList, bFirst);
-			});
+			};
 
 			SetShadowRenderTargets(RHICmdList);  // run it now, maybe run it later for parallel command lists
 			bFirst = false;
@@ -3342,10 +3342,10 @@ bool FDeferredShadingSceneRenderer::RenderCachedPreshadows(FRHICommandListImmedi
 		{
 			SCOPED_DRAW_EVENTF(RHICmdList, EventShadowDepths, TEXT("Preshadow Cache Depths"));
 
-			auto SetShadowRenderTargets = MakeFunction<void, FRHICommandList&>([](FRHICommandList& RHICmdList)
+			auto SetShadowRenderTargets = [](FRHICommandList& RHICmdList)
 			{
 				SetRenderTarget(RHICmdList, FTextureRHIRef(), GSceneRenderTargets.PreShadowCacheDepthZ->GetRenderTargetItem().TargetableTexture);
-			});
+			};
 
 			SetShadowRenderTargets(RHICmdList);  // run it now, maybe run it later for parallel command lists
 
@@ -3466,10 +3466,10 @@ bool FForwardShadingSceneRenderer::RenderShadowDepthMap(FRHICommandListImmediate
 		SCOPED_DRAW_EVENT(RHICmdList, ShadowDepthsFromOpaque);
 
 		bool bFirst = true;
-		auto SetShadowRenderTargets = MakeFunction<void, FRHICommandList&>([&bFirst](FRHICommandList& RHICmdList)
+		auto SetShadowRenderTargets = [&bFirst](FRHICommandList& RHICmdList)
 		{
 			GSceneRenderTargets.BeginRenderingShadowDepth(RHICmdList, bFirst);
-		});
+		};
 
 		SetShadowRenderTargets(RHICmdList);  // run it now, maybe run it later for parallel command lists
 		bFirst = false;
diff --git a/Engine/Source/Runtime/Renderer/Private/ShadowRendering.h b/Engine/Source/Runtime/Renderer/Private/ShadowRendering.h
--- a/Engine/Source/Runtime/Renderer/Private/ShadowRendering.h
+++ b/Engine/Source/Runtime/Renderer/Private/ShadowRendering.h
@@ -490,54 +490,6 @@ public:
 	bool bIsTwoSided;
 };
 
-template<typename TReturn, typename TArg>
-class IFunction
-{
-public:
-	virtual ~IFunction() {}
-	virtual TReturn operator()(TArg&& Arg) = 0;
-	virtual TReturn operator()(TArg&& Arg) const = 0;
-};
-
-template<typename TReturn, typename TArg, typename TLambda>
-class TFunction : public IFunction<TReturn, TArg>
-{
-	TLambda Func;
-public:
-	TFunction(TLambda&& InFunc)
-		: Func(MoveTemp(InFunc))
-	{
-	}
-	TFunction(TLambda& InFunc)
-		: Func(InFunc)
-	{
-	}
-	virtual ~TFunction() override final
-	{
-	}
-	virtual TReturn operator()(TArg&& Arg) override final
-	{
-		return Func(Forward<TArg>(Arg));
-	}
-	virtual TReturn operator()(TArg&& Arg) const override final
-	{
-		return Func(Forward<TArg>(Arg));
-	}
-};
-
-template<typename TReturn, typename TArg, typename TLambda>
-TFunction<TReturn, TArg, TLambda> MakeFunction(TLambda&& Func)
-{
-	return TFunction<TReturn, TArg, TLambda>(MoveTemp(Func));
-}
-
-template<typename TReturn, typename TArg, typename TLambda>
-TFunction<TReturn, TArg, TLambda>* NewFunction(TLambda&& Func)
-{
-	return new TFunction<TReturn, TArg, TLambda>(MoveTemp(Func));
-}
-
-
 /**
  * Information about a projected shadow.
  */
@@ -719,7 +671,7 @@ public:
 	/**
 	 * Renders the shadow subject depth.
 	 */
-	void RenderDepth(FRHICommandList& RHICmdList, class FSceneRenderer* SceneRenderer, IFunction<void, FRHICommandList&>& SetShadowRenderTargets);
+	void RenderDepth(FRHICommandList& RHICmdList, class FSceneRenderer* SceneRenderer, TFunctionRef<void (FRHICommandList& RHICmdList)> SetShadowRenderTargets);
 
 	/** Set state for depth rendering */
 	void SetStateForDepth(FRHICommandList& RHICmdList);
@@ -807,7 +759,7 @@ private:
 	/**
 	* Renders the shadow subject depth, to a particular hacked view
 	*/
-	void RenderDepthInner(FRHICommandList& RHICmdList, class FSceneRenderer* SceneRenderer, const FViewInfo* FoundView, IFunction<void, FRHICommandList&>& SetShadowRenderTargets);
+	void RenderDepthInner(FRHICommandList& RHICmdList, class FSceneRenderer* SceneRenderer, const FViewInfo* FoundView, TFunctionRef<void (FRHICommandList& RHICmdList)> SetShadowRenderTargets);
 
 	/**
 	* Renders the dynamic shadow subject depth, to a particular hacked view
