Stop comparisons with self in StableSort
Fixes a few implementation bugs, e.g. 0-size rotates occuring and extra unnecessary iterations

diff --git a/Engine/Source/Runtime/Core/Public/Templates/Sorting.h b/Engine/Source/Runtime/Core/Public/Templates/Sorting.h
--- a/Engine/Source/Runtime/Core/Public/Templates/Sorting.h
+++ b/Engine/Source/Runtime/Core/Public/Templates/Sorting.h
@@ -328,12 +328,16 @@ public:
 
 		while (AStart < BStart && BStart < Num)
 		{
-			int32 NewBOffset = BinarySearchFirst(First + BStart, Num - BStart, First[AStart], Predicate);
+			int32 NewAOffset = BinarySearchLast(First + AStart, BStart - AStart, First[BStart], Predicate) + 1;
+			AStart += NewAOffset;
 
-			TRotationPolicy::Rotate(First, AStart, BStart + NewBOffset, NewBOffset);
+			if (AStart >= BStart) // done
+				break;
 
+			int32 NewBOffset = BinarySearchFirst(First + BStart, Num - BStart, First[AStart], Predicate);
+			TRotationPolicy::Rotate(First, AStart, BStart + NewBOffset, NewBOffset);
 			BStart += NewBOffset;
-			AStart += BinarySearchLast(First + AStart, BStart - AStart, First[AStart + NewBOffset], Predicate) + 1;
+			AStart += NewBOffset + 1;
 		}
 	}
 
diff --git a/Engine/Source/Runtime/Landscape/Private/LandscapeEdit.cpp b/Engine/Source/Runtime/Landscape/Private/LandscapeEdit.cpp
--- a/Engine/Source/Runtime/Landscape/Private/LandscapeEdit.cpp
+++ b/Engine/Source/Runtime/Landscape/Private/LandscapeEdit.cpp
@@ -4769,8 +4769,6 @@ struct FMobileLayerAllocation
 
 	friend bool operator<(const FMobileLayerAllocation& lhs, const FMobileLayerAllocation& rhs)
 	{
-		if (&lhs == &rhs) return false; // I don't know why StableSort is comparing elements against themself
-
 		if (!lhs.Allocation.LayerInfo && !rhs.Allocation.LayerInfo) return false; // equally broken :P
 		if (!lhs.Allocation.LayerInfo && rhs.Allocation.LayerInfo) return false; // broken layers sort to the end
 		if (!rhs.Allocation.LayerInfo && lhs.Allocation.LayerInfo) return true;
