Fixed bug with servers stomping on each other's logs (UEPLAT-187).

- we were specifying O_TRUNC before testing if file is "locked", which corrupted the file.

diff --git a/Engine/Source/Runtime/Core/Private/Linux/LinuxPlatformFile.cpp b/Engine/Source/Runtime/Core/Private/Linux/LinuxPlatformFile.cpp
--- a/Engine/Source/Runtime/Core/Private/Linux/LinuxPlatformFile.cpp
+++ b/Engine/Source/Runtime/Core/Private/Linux/LinuxPlatformFile.cpp
@@ -253,10 +253,7 @@ IFileHandle* FLinuxPlatformFile::OpenWrite(const TCHAR* Filename, bool bAppend,
 	{
 		Flags |= O_APPEND;
 	}
-	else
-	{
-		Flags |= O_TRUNC;
-	}
+
 	if (bAllowRead)
 	{
 		Flags |= O_RDWR;
@@ -272,23 +269,36 @@ IFileHandle* FLinuxPlatformFile::OpenWrite(const TCHAR* Filename, bool bAppend,
 		return NULL;
 	}
 
+	// Caveat: cannot specify O_TRUNC in flags, as this will corrupt the file which may be "locked" by other process. We will ftruncate() it once we "lock" it
 	int32 Handle = open(TCHAR_TO_UTF8(*NormalizeFilename(Filename)), Flags, S_IRUSR | S_IWUSR);
 	if (Handle != -1)
 	{
 		// mimic Windows "exclusive write" behavior (we don't use FILE_SHARE_WRITE) by locking the file.
-		// note that the lock will be removed by itself when the last file descriptor is close()d
+		// note that the (non-mandatory) "lock" will be removed by itself when the last file descriptor is close()d
 		if (flock(Handle, LOCK_EX | LOCK_NB) == -1)
 		{
 			// if locked, consider operation a failure
-			if (EWOULDBLOCK == errno)
+			if (EAGAIN == errno || EWOULDBLOCK == errno)
 			{
 				close(Handle);
-				return NULL;
+				return nullptr;
 			}
-
 			// all the other locking errors are ignored.
 		}
 
+		// truncate the file now that we locked it
+		if (!bAppend)
+		{
+			if (ftruncate(Handle, 0) != 0)
+			{
+				int ErrNo = errno;
+				UE_LOG(LogLinuxPlatformFile, Warning, TEXT( "ftruncate() failed for '%s': errno=%d (%s)" ),
+															Filename, ErrNo, ANSI_TO_TCHAR(strerror(ErrNo)));
+				close(Handle);
+				return nullptr;
+			}
+		}
+
 		FFileHandleLinux* FileHandleLinux = new FFileHandleLinux(Handle);
 		if (bAppend)
 		{
@@ -296,8 +306,10 @@ IFileHandle* FLinuxPlatformFile::OpenWrite(const TCHAR* Filename, bool bAppend,
 		}
 		return FileHandleLinux;
 	}
-	UE_LOG(LogLinuxPlatformFile, Warning, TEXT( "open('%s', Flags=0x%08X) failed: errno=%d (%s)" ), *NormalizeFilename(Filename), Flags, errno, ANSI_TO_TCHAR(strerror(errno)));
-	return NULL;
+
+	int ErrNo = errno;
+	UE_LOG(LogLinuxPlatformFile, Warning, TEXT( "open('%s', Flags=0x%08X) failed: errno=%d (%s)" ), *NormalizeFilename(Filename), Flags, ErrNo, ANSI_TO_TCHAR(strerror(ErrNo)));
+	return nullptr;
 }
 
 bool FLinuxPlatformFile::DirectoryExists(const TCHAR* Directory)
