Fix Xbox One compile error.

diff --git a/Engine/Source/Runtime/Core/Public/Math/UnrealMathDirectX.h b/Engine/Source/Runtime/Core/Public/Math/UnrealMathDirectX.h
--- a/Engine/Source/Runtime/Core/Public/Math/UnrealMathDirectX.h
+++ b/Engine/Source/Runtime/Core/Public/Math/UnrealMathDirectX.h
@@ -167,7 +167,7 @@ FORCEINLINE VectorRegister MakeVectorRegister( float X, float Y, float Z, float
  * @param Vec	Vector to store
  * @param Ptr	Aligned memory pointer
  */
-#define VectorStoreAlignedStreamed( Vec, Ptr )	XM_STREAM_PS( Ptr , Vec )
+#define VectorStoreAlignedStreamed( Vec, Ptr )	XM_STREAM_PS( (float*)(Ptr), Vec )
 
 /**
  * Stores a vector to memory (aligned or unaligned).
