Make Core module public headers individually compilable. Part 27 of n.

diff --git a/Engine/Source/Runtime/Core/Private/Android/AndroidFile.cpp b/Engine/Source/Runtime/Core/Private/Android/AndroidFile.cpp
--- a/Engine/Source/Runtime/Core/Private/Android/AndroidFile.cpp
+++ b/Engine/Source/Runtime/Core/Private/Android/AndroidFile.cpp
@@ -5,7 +5,7 @@
 =============================================================================*/
 
 #include "CorePrivatePCH.h"
-
+#include "Misc/App.h"
 #include <dirent.h>
 #include <jni.h>
 #include <Android/asset_manager.h>
diff --git a/Engine/Source/Runtime/Core/Private/GenericPlatform/GenericPlatformContext.cpp b/Engine/Source/Runtime/Core/Private/GenericPlatform/GenericPlatformContext.cpp
--- a/Engine/Source/Runtime/Core/Private/GenericPlatform/GenericPlatformContext.cpp
+++ b/Engine/Source/Runtime/Core/Private/GenericPlatform/GenericPlatformContext.cpp
@@ -2,6 +2,7 @@
 
 #include "CorePrivatePCH.h"
 #include "GenericPlatform/GenericPlatformContext.h"
+#include "Misc/App.h"
 #include "Runtime/Launch/Resources/Version.h"
 
 const ANSICHAR* FGenericCrashContext::CrashContextRuntimeXMLNameA = "CrashContext.runtime-xml";
diff --git a/Engine/Source/Runtime/Core/Private/GenericPlatform/GenericPlatformOutputDevices.cpp b/Engine/Source/Runtime/Core/Private/GenericPlatform/GenericPlatformOutputDevices.cpp
--- a/Engine/Source/Runtime/Core/Private/GenericPlatform/GenericPlatformOutputDevices.cpp
+++ b/Engine/Source/Runtime/Core/Private/GenericPlatform/GenericPlatformOutputDevices.cpp
@@ -1,7 +1,9 @@
 // Copyright 1998-2014 Epic Games, Inc. All Rights Reserved.
 
 #include "CorePrivatePCH.h"
-#include "FeedbackContextAnsi.h"
+#include "HAL/FeedbackContextAnsi.h"
+#include "Misc/App.h"
+#include "Misc/OutputDeviceConsole.h"
 
 void							FGenericPlatformOutputDevices::SetupOutputDevices()
 {
diff --git a/Engine/Source/Runtime/Core/Private/HAL/OutputDevices.cpp b/Engine/Source/Runtime/Core/Private/HAL/OutputDevices.cpp
--- a/Engine/Source/Runtime/Core/Private/HAL/OutputDevices.cpp
+++ b/Engine/Source/Runtime/Core/Private/HAL/OutputDevices.cpp
@@ -5,9 +5,9 @@
 =============================================================================*/
 
 #include "CorePrivatePCH.h"
-
-#include <stdio.h>
+#include "Misc/App.h"
 #include "Templates/UniquePtr.h"
+#include <stdio.h>
 
 // #if _MSC_VER
 // #pragma warning (push)
diff --git a/Engine/Source/Runtime/Core/Private/Internationalization/TextLocalizationManager.cpp b/Engine/Source/Runtime/Core/Private/Internationalization/TextLocalizationManager.cpp
--- a/Engine/Source/Runtime/Core/Private/Internationalization/TextLocalizationManager.cpp
+++ b/Engine/Source/Runtime/Core/Private/Internationalization/TextLocalizationManager.cpp
@@ -1,6 +1,7 @@
 ﻿// Copyright 1998-2014 Epic Games, Inc. All Rights Reserved.
 
 #include "CorePrivatePCH.h"
+#include "Misc/App.h"
 #include "InternationalizationManifest.h"
 #include "TextLocalizationResourceGenerator.h"
 
diff --git a/Engine/Source/Runtime/Core/Private/Linux/LinuxPlatformCrashContext.cpp b/Engine/Source/Runtime/Core/Private/Linux/LinuxPlatformCrashContext.cpp
--- a/Engine/Source/Runtime/Core/Private/Linux/LinuxPlatformCrashContext.cpp
+++ b/Engine/Source/Runtime/Core/Private/Linux/LinuxPlatformCrashContext.cpp
@@ -2,6 +2,7 @@
 
 #include "CorePrivatePCH.h"
 #include "Linux/LinuxPlatformCrashContext.h"
+#include "Misc/App.h"
 #include "Runtime/Launch/Resources/Version.h"
 
 #include <sys/utsname.h>	// for uname()
diff --git a/Engine/Source/Runtime/Core/Private/Linux/LinuxPlatformFeedbackContextPrivate.h b/Engine/Source/Runtime/Core/Private/Linux/LinuxPlatformFeedbackContextPrivate.h
--- a/Engine/Source/Runtime/Core/Private/Linux/LinuxPlatformFeedbackContextPrivate.h
+++ b/Engine/Source/Runtime/Core/Private/Linux/LinuxPlatformFeedbackContextPrivate.h
@@ -1,6 +1,7 @@
 // Copyright 1998-2014 Epic Games, Inc. All Rights Reserved.
 
 #include <syslog.h>
+#include "Misc/App.h"
 
 #pragma once
 
diff --git a/Engine/Source/Runtime/Core/Private/Linux/LinuxPlatformOutputDevices.cpp b/Engine/Source/Runtime/Core/Private/Linux/LinuxPlatformOutputDevices.cpp
--- a/Engine/Source/Runtime/Core/Private/Linux/LinuxPlatformOutputDevices.cpp
+++ b/Engine/Source/Runtime/Core/Private/Linux/LinuxPlatformOutputDevices.cpp
@@ -11,6 +11,7 @@
 #include "LinuxApplication.h"
 #include "LinuxPlatformOutputDevicesPrivate.h"
 #include "LinuxPlatformFeedbackContextPrivate.h"
+#include "Misc/App.h"
 
 #define CONSOLE_RED		"\x1b[31m"
 #define CONSOLE_GREEN	"\x1b[32m"
diff --git a/Engine/Source/Runtime/Core/Private/Linux/LinuxPlatformSplash.cpp b/Engine/Source/Runtime/Core/Private/Linux/LinuxPlatformSplash.cpp
--- a/Engine/Source/Runtime/Core/Private/Linux/LinuxPlatformSplash.cpp
+++ b/Engine/Source/Runtime/Core/Private/Linux/LinuxPlatformSplash.cpp
@@ -6,6 +6,7 @@
 
 
 #include "CorePrivatePCH.h"
+#include "Misc/App.h"
 #include "EngineVersion.h"
 #include "EngineBuildSettings.h"
 
diff --git a/Engine/Source/Runtime/Core/Private/Linux/LinuxPlatformStackWalk.cpp b/Engine/Source/Runtime/Core/Private/Linux/LinuxPlatformStackWalk.cpp
--- a/Engine/Source/Runtime/Core/Private/Linux/LinuxPlatformStackWalk.cpp
+++ b/Engine/Source/Runtime/Core/Private/Linux/LinuxPlatformStackWalk.cpp
@@ -5,6 +5,7 @@
 =============================================================================*/
 
 #include "CorePrivatePCH.h"
+#include "Misc/App.h"
 #include <execinfo.h>
 #include <dlfcn.h>
 #include <cxxabi.h>
diff --git a/Engine/Source/Runtime/Core/Private/Mac/MacPlatformMisc.cpp b/Engine/Source/Runtime/Core/Private/Mac/MacPlatformMisc.cpp
--- a/Engine/Source/Runtime/Core/Private/Mac/MacPlatformMisc.cpp
+++ b/Engine/Source/Runtime/Core/Private/Mac/MacPlatformMisc.cpp
@@ -5,6 +5,7 @@
 =============================================================================*/
 
 #include "CorePrivatePCH.h"
+#include "Misc/App.h"
 #include "ExceptionHandling.h"
 #include "SecureHash.h"
 #include "VarargsHelper.h"
diff --git a/Engine/Source/Runtime/Core/Private/Mac/MacPlatformOutputDevices.cpp b/Engine/Source/Runtime/Core/Private/Mac/MacPlatformOutputDevices.cpp
--- a/Engine/Source/Runtime/Core/Private/Mac/MacPlatformOutputDevices.cpp
+++ b/Engine/Source/Runtime/Core/Private/Mac/MacPlatformOutputDevices.cpp
@@ -5,6 +5,7 @@
 =============================================================================*/
 
 #include "CorePrivatePCH.h"
+#include "Misc/App.h"
 
 #include "FeedbackContextAnsi.h"
 #include "MacApplication.h"
diff --git a/Engine/Source/Runtime/Core/Private/Mac/MacPlatformProcess.cpp b/Engine/Source/Runtime/Core/Private/Mac/MacPlatformProcess.cpp
--- a/Engine/Source/Runtime/Core/Private/Mac/MacPlatformProcess.cpp
+++ b/Engine/Source/Runtime/Core/Private/Mac/MacPlatformProcess.cpp
@@ -5,6 +5,7 @@
 =============================================================================*/
 
 #include "CorePrivatePCH.h"
+#include "Misc/App.h"
 #include "../../Launch/Resources/Version.h"
 #include "MacApplication.h"
 #include <mach-o/dyld.h>
diff --git a/Engine/Source/Runtime/Core/Private/Mac/MacPlatformSplash.cpp b/Engine/Source/Runtime/Core/Private/Mac/MacPlatformSplash.cpp
--- a/Engine/Source/Runtime/Core/Private/Mac/MacPlatformSplash.cpp
+++ b/Engine/Source/Runtime/Core/Private/Mac/MacPlatformSplash.cpp
@@ -1,6 +1,7 @@
 // Copyright 1998-2014 Epic Games, Inc. All Rights Reserved.
 
 #include "CorePrivatePCH.h"
+#include "Misc/App.h"
 #include "EngineVersion.h"
 #include "EngineBuildSettings.h"
 
diff --git a/Engine/Source/Runtime/Core/Private/Misc/App.cpp b/Engine/Source/Runtime/Core/Private/Misc/App.cpp
--- a/Engine/Source/Runtime/Core/Private/Misc/App.cpp
+++ b/Engine/Source/Runtime/Core/Private/Misc/App.cpp
@@ -2,6 +2,7 @@
 
 #include "CorePrivatePCH.h"
 #include "Runtime/Launch/Resources/Version.h"
+#include "Misc/App.h"
 
 
 /* FApp static initialization
diff --git a/Engine/Source/Runtime/Core/Private/Misc/ConfigCacheIni.cpp b/Engine/Source/Runtime/Core/Private/Misc/ConfigCacheIni.cpp
--- a/Engine/Source/Runtime/Core/Private/Misc/ConfigCacheIni.cpp
+++ b/Engine/Source/Runtime/Core/Private/Misc/ConfigCacheIni.cpp
@@ -1,6 +1,7 @@
 // Copyright 1998-2014 Epic Games, Inc. All Rights Reserved.
 
 #include "CorePrivatePCH.h"
+#include "Misc/App.h"
 #include "ConfigCacheIni.h"
 #include "RemoteConfigIni.h"
 #include "AES.h"
diff --git a/Engine/Source/Runtime/Core/Private/Misc/MonitoredProcess.cpp b/Engine/Source/Runtime/Core/Private/Misc/MonitoredProcess.cpp
--- a/Engine/Source/Runtime/Core/Private/Misc/MonitoredProcess.cpp
+++ b/Engine/Source/Runtime/Core/Private/Misc/MonitoredProcess.cpp
@@ -1,7 +1,7 @@
 // Copyright 1998-2014 Epic Games, Inc. All Rights Reserved.
 
 #include "CorePrivatePCH.h"
-
+#include "Misc/MonitoredProcess.h"
 
 /* FMonitoredProcess structors
  *****************************************************************************/
diff --git a/Engine/Source/Runtime/Core/Private/Misc/Paths.cpp b/Engine/Source/Runtime/Core/Private/Misc/Paths.cpp
--- a/Engine/Source/Runtime/Core/Private/Misc/Paths.cpp
+++ b/Engine/Source/Runtime/Core/Private/Misc/Paths.cpp
@@ -2,6 +2,7 @@
 
 // Core includes.
 #include "CorePrivatePCH.h"
+#include "Misc/App.h"
 #include "Runtime/Launch/Resources/Version.h" 
 #include "EngineVersion.h"
 
diff --git a/Engine/Source/Runtime/Core/Private/Misc/RemoteConfigIni.cpp b/Engine/Source/Runtime/Core/Private/Misc/RemoteConfigIni.cpp
--- a/Engine/Source/Runtime/Core/Private/Misc/RemoteConfigIni.cpp
+++ b/Engine/Source/Runtime/Core/Private/Misc/RemoteConfigIni.cpp
@@ -1,7 +1,8 @@
 // Copyright 1998-2014 Epic Games, Inc. All Rights Reserved.
 
 #include "CorePrivatePCH.h"
-#include "RemoteConfigIni.h"
+#include "Misc/App.h"
+#include "Misc/RemoteConfigIni.h"
 
 // Globals
 FRemoteConfig GRemoteConfig;
diff --git a/Engine/Source/Runtime/Core/Private/Serialization/Archive.cpp b/Engine/Source/Runtime/Core/Private/Serialization/Archive.cpp
--- a/Engine/Source/Runtime/Core/Private/Serialization/Archive.cpp
+++ b/Engine/Source/Runtime/Core/Private/Serialization/Archive.cpp
@@ -4,7 +4,8 @@
 	UnArchive.cpp: Core archive classes.
 =============================================================================*/
 #include "CorePrivatePCH.h"
-#include "Archive.h"
+#include "Serialization/Archive.h"
+#include "Serialization/CustomVersion.h"
 
 /*-----------------------------------------------------------------------------
 	FArchiveProxy implementation.
diff --git a/Engine/Source/Runtime/Core/Private/Serialization/CustomVersion.cpp b/Engine/Source/Runtime/Core/Private/Serialization/CustomVersion.cpp
--- a/Engine/Source/Runtime/Core/Private/Serialization/CustomVersion.cpp
+++ b/Engine/Source/Runtime/Core/Private/Serialization/CustomVersion.cpp
@@ -5,6 +5,7 @@
 =============================================================================*/
 
 #include "CorePrivatePCH.h"
+#include "Serialization/CustomVersion.h"
 
 namespace
 {
diff --git a/Engine/Source/Runtime/Core/Private/WinRT/WinRTFeedbackContextPrivate.h b/Engine/Source/Runtime/Core/Private/WinRT/WinRTFeedbackContextPrivate.h
--- a/Engine/Source/Runtime/Core/Private/WinRT/WinRTFeedbackContextPrivate.h
+++ b/Engine/Source/Runtime/Core/Private/WinRT/WinRTFeedbackContextPrivate.h
@@ -5,6 +5,7 @@
 =============================================================================*/
 
 #pragma once
+#include "Misc/App.h"
 
 /**
  * Feedback context implementation for WinRT.
diff --git a/Engine/Source/Runtime/Core/Private/WinRT/WinRTOutputDevices.cpp b/Engine/Source/Runtime/Core/Private/WinRT/WinRTOutputDevices.cpp
--- a/Engine/Source/Runtime/Core/Private/WinRT/WinRTOutputDevices.cpp
+++ b/Engine/Source/Runtime/Core/Private/WinRT/WinRTOutputDevices.cpp
@@ -5,6 +5,7 @@
 =============================================================================*/
 
 #include "CorePrivatePCH.h"
+#include "Misc/App.h"
 
 #include "FeedbackContextAnsi.h"
 #include "WinRTOutputDevicesPrivate.h"
diff --git a/Engine/Source/Runtime/Core/Private/Windows/WindowsPlatformFeedbackContextPrivate.h b/Engine/Source/Runtime/Core/Private/Windows/WindowsPlatformFeedbackContextPrivate.h
--- a/Engine/Source/Runtime/Core/Private/Windows/WindowsPlatformFeedbackContextPrivate.h
+++ b/Engine/Source/Runtime/Core/Private/Windows/WindowsPlatformFeedbackContextPrivate.h
@@ -5,6 +5,7 @@
 =============================================================================*/
 
 #pragma once
+#include "Misc/OutputDeviceConsole.h"
 
 /**
  * Feedback context implementation for windows.
diff --git a/Engine/Source/Runtime/Core/Private/Windows/WindowsPlatformOutputDevicesPrivate.h b/Engine/Source/Runtime/Core/Private/Windows/WindowsPlatformOutputDevicesPrivate.h
--- a/Engine/Source/Runtime/Core/Private/Windows/WindowsPlatformOutputDevicesPrivate.h
+++ b/Engine/Source/Runtime/Core/Private/Windows/WindowsPlatformOutputDevicesPrivate.h
@@ -1,7 +1,7 @@
 // Copyright 1998-2014 Epic Games, Inc. All Rights Reserved.
 
 #pragma once
-
+#include "Misc/OutputDeviceConsole.h"
 
 #ifndef WANTS_WINDOWS_EVENT_LOGGING
 	#if UE_SERVER
diff --git a/Engine/Source/Runtime/Core/Private/Windows/WindowsPlatformSplash.cpp b/Engine/Source/Runtime/Core/Private/Windows/WindowsPlatformSplash.cpp
--- a/Engine/Source/Runtime/Core/Private/Windows/WindowsPlatformSplash.cpp
+++ b/Engine/Source/Runtime/Core/Private/Windows/WindowsPlatformSplash.cpp
@@ -1,7 +1,7 @@
 // Copyright 1998-2014 Epic Games, Inc. All Rights Reserved.
 
 #include "CorePrivatePCH.h"
-
+#include "Misc/App.h"
 #include "EngineVersion.h"
 #include "AllowWindowsPlatformTypes.h"
 #include "EngineBuildSettings.h"
diff --git a/Engine/Source/Runtime/Core/Private/Windows/XInputInterface.cpp b/Engine/Source/Runtime/Core/Private/Windows/XInputInterface.cpp
--- a/Engine/Source/Runtime/Core/Private/Windows/XInputInterface.cpp
+++ b/Engine/Source/Runtime/Core/Private/Windows/XInputInterface.cpp
@@ -3,7 +3,7 @@
 #include "CorePrivatePCH.h"
 #include "WindowsApplication.h"
 #include "XInputInterface.h"
-#include "GenericApplication.h"
+#include "GenericPlatform/GenericApplication.h"
 
 #pragma pack (push,8)
 #include "AllowWindowsPlatformTypes.h"
diff --git a/Engine/Source/Runtime/Core/Private/Windows/XInputInterface.h b/Engine/Source/Runtime/Core/Private/Windows/XInputInterface.h
--- a/Engine/Source/Runtime/Core/Private/Windows/XInputInterface.h
+++ b/Engine/Source/Runtime/Core/Private/Windows/XInputInterface.h
@@ -1,7 +1,7 @@
 // Copyright 1998-2014 Epic Games, Inc. All Rights Reserved.
 
 #pragma once
-
+#include "GenericPlatform/GenericApplicationMessageHandler.h"
 
 /** Max number of controllers. */
 #define MAX_NUM_XINPUT_CONTROLLERS 4
diff --git a/Engine/Source/Runtime/Core/Public/GenericPlatform/GenericApplication.h b/Engine/Source/Runtime/Core/Public/GenericPlatform/GenericApplication.h
--- a/Engine/Source/Runtime/Core/Public/GenericPlatform/GenericApplication.h
+++ b/Engine/Source/Runtime/Core/Public/GenericPlatform/GenericApplication.h
@@ -1,8 +1,8 @@
 // Copyright 1998-2014 Epic Games, Inc. All Rights Reserved.
 
 #pragma once
+#include "GenericPlatform/GenericApplicationMessageHandler.h"
 
-class FGenericApplicationMessageHandler;
 class FGenericWindow;
 class ICursor;
 class ITextInputMethodSystem;
diff --git a/Engine/Source/Runtime/Core/Public/HAL/FeedbackContextAnsi.h b/Engine/Source/Runtime/Core/Public/HAL/FeedbackContextAnsi.h
--- a/Engine/Source/Runtime/Core/Public/HAL/FeedbackContextAnsi.h
+++ b/Engine/Source/Runtime/Core/Public/HAL/FeedbackContextAnsi.h
@@ -1,7 +1,8 @@
 // Copyright 1998-2014 Epic Games, Inc. All Rights Reserved.
 
 #pragma once
-
+#include "Misc/App.h"
+#include "Misc/FeedbackContext.h"
 
 /*-----------------------------------------------------------------------------
 	FFeedbackContextAnsi.
diff --git a/Engine/Source/Runtime/Core/Public/Misc/OutputDeviceConsole.h b/Engine/Source/Runtime/Core/Public/Misc/OutputDeviceConsole.h
--- a/Engine/Source/Runtime/Core/Public/Misc/OutputDeviceConsole.h
+++ b/Engine/Source/Runtime/Core/Public/Misc/OutputDeviceConsole.h
@@ -1,7 +1,7 @@
 // Copyright 1998-2014 Epic Games, Inc. All Rights Reserved.
 
 #pragma once
-
+#include "Misc/OutputDevice.h"
 
 /**
  * This class servers as the base class for console window output.
