Fix for static initialization order of debug visualizer delegates

diff --git a/Engine/Source/Runtime/Core/Private/Misc/CallbackDevice.cpp b/Engine/Source/Runtime/Core/Private/Misc/CallbackDevice.cpp
--- a/Engine/Source/Runtime/Core/Private/Misc/CallbackDevice.cpp
+++ b/Engine/Source/Runtime/Core/Private/Misc/CallbackDevice.cpp
@@ -20,6 +20,18 @@ FCoreDelegates::FHotFixDelegate& FCoreDelegates::GetHotfixDelegate(EHotfixDelega
 	return HotFixDelegates[HotFix];
 }
 
+FCoreDelegates::FObjectArrayForDebugVisualizersDelegate& FCoreDelegates::GetObjectArrayForDebugVisualizersDelegate()
+{
+	static FObjectArrayForDebugVisualizersDelegate StaticDelegate;
+	return StaticDelegate;
+}
+
+FCoreDelegates::FGetSerialNumberBlocksForDebugVisualizersDelegate& FCoreDelegates::GetSerialNumberBlocksForDebugVisualizersDelegate()
+{
+	static FGetSerialNumberBlocksForDebugVisualizersDelegate StaticDelegate;
+	return StaticDelegate;
+}
+
 FCoreDelegates::FOnMountPak FCoreDelegates::OnMountPak;
 FCoreDelegates::FOnUserLoginChangedEvent FCoreDelegates::OnUserLoginChangedEvent; 
 FCoreDelegates::FOnUserControllerConnectionChange FCoreDelegates::OnControllerConnectionChange;
@@ -36,8 +48,6 @@ FSimpleMulticastDelegate FCoreDelegates::OnShutdownAfterError;
 FSimpleMulticastDelegate FCoreDelegates::OnInit;
 FSimpleMulticastDelegate FCoreDelegates::OnExit;
 FSimpleMulticastDelegate FCoreDelegates::OnPreExit;
-FCoreDelegates::FGetSerialNumberBlocksForDebugVisualizersDelegate FCoreDelegates::GetSerialNumberBlocksDebugVisualizers;
-FCoreDelegates::FObjectArrayForDebugVisualizersDelegate FCoreDelegates::ObjectArrayForDebugVisualizers;
 FSimpleMulticastDelegate FCoreDelegates::ColorPickerChanged;
 FCoreDelegates::FOnModalMessageBox FCoreDelegates::ModalErrorMessage;
 FCoreDelegates::FOnInviteAccepted FCoreDelegates::OnInviteAccepted;
diff --git a/Engine/Source/Runtime/Core/Public/Misc/CallbackDevice.h b/Engine/Source/Runtime/Core/Public/Misc/CallbackDevice.h
--- a/Engine/Source/Runtime/Core/Public/Misc/CallbackDevice.h
+++ b/Engine/Source/Runtime/Core/Public/Misc/CallbackDevice.h
@@ -61,7 +61,7 @@ public:
 	DECLARE_DELEGATE_RetVal(TArray<class UObjectBase*>*, FObjectArrayForDebugVisualizersDelegate);
 
 	// Called in PER_MODULE_BOILERPLATE macro.
-	static FObjectArrayForDebugVisualizersDelegate ObjectArrayForDebugVisualizers;
+	static FObjectArrayForDebugVisualizersDelegate& GetObjectArrayForDebugVisualizersDelegate();
 
 	// Callback for handling an ensure
 	DECLARE_MULTICAST_DELEGATE(FOnHandleSystemEnsure);
@@ -127,7 +127,7 @@ public:
 	static FSimpleMulticastDelegate OnPreExit;
 
 	// Called in PER_MODULE_BOILERPLATE macro.
-	static FGetSerialNumberBlocksForDebugVisualizersDelegate GetSerialNumberBlocksDebugVisualizers;
+	static FGetSerialNumberBlocksForDebugVisualizersDelegate& GetSerialNumberBlocksForDebugVisualizersDelegate();
 
 	/** Color picker color has changed, please refresh as needed*/
 	static FSimpleMulticastDelegate ColorPickerChanged;
diff --git a/Engine/Source/Runtime/Core/Public/Modules/Boilerplate/ModuleBoilerplate.h b/Engine/Source/Runtime/Core/Public/Modules/Boilerplate/ModuleBoilerplate.h
--- a/Engine/Source/Runtime/Core/Public/Modules/Boilerplate/ModuleBoilerplate.h
+++ b/Engine/Source/Runtime/Core/Public/Modules/Boilerplate/ModuleBoilerplate.h
@@ -33,7 +33,7 @@
 #define PER_MODULE_BOILERPLATE \
 	TArray<FNameEntry const*>* GFNameTableForDebuggerVisualizers = FName::GetNameTableForDebuggerVisualizers_ST(); \
 	FNameEntry*** GFNameTableForDebuggerVisualizers_MT = FName::GetNameTableForDebuggerVisualizers_MT(); \
-	int32*** GSerialNumberBlocksForDebugVisualizers = FCoreDelegates::GetSerialNumberBlocksDebugVisualizers.IsBound() ? FCoreDelegates::GetSerialNumberBlocksDebugVisualizers.Execute() : NULL; \
-	TArray<UObjectBase*>* GObjectArrayForDebugVisualizers = FCoreDelegates::ObjectArrayForDebugVisualizers.IsBound() ? FCoreDelegates::ObjectArrayForDebugVisualizers.Execute() : NULL; \
+	int32*** GSerialNumberBlocksForDebugVisualizers = FCoreDelegates::GetSerialNumberBlocksForDebugVisualizersDelegate().IsBound() ? FCoreDelegates::GetSerialNumberBlocksForDebugVisualizersDelegate().Execute() : NULL; \
+	TArray<UObjectBase*>* GObjectArrayForDebugVisualizers = FCoreDelegates::GetObjectArrayForDebugVisualizersDelegate().IsBound() ? FCoreDelegates::GetObjectArrayForDebugVisualizersDelegate().Execute() : NULL; \
 	bool GFNameDebuggerVisualizersIsUE3=false; \
 	REPLACEMENT_OPERATOR_NEW_AND_DELETE
diff --git a/Engine/Source/Runtime/CoreUObject/Private/UObject/WeakObjectPtr.cpp b/Engine/Source/Runtime/CoreUObject/Private/UObject/WeakObjectPtr.cpp
--- a/Engine/Source/Runtime/CoreUObject/Private/UObject/WeakObjectPtr.cpp
+++ b/Engine/Source/Runtime/CoreUObject/Private/UObject/WeakObjectPtr.cpp
@@ -55,7 +55,7 @@ public:
 		}
 		GSerialNumberBlocksForDebugVisualizersRoot = (int32**)&Blocks;
 
-		FCoreDelegates::GetSerialNumberBlocksDebugVisualizers.BindStatic(&GetSerialNumberBlocksForDebugVisualizersPtr);
+		FCoreDelegates::GetSerialNumberBlocksForDebugVisualizersDelegate().BindStatic(&GetSerialNumberBlocksForDebugVisualizersPtr);
 	}
 	/** Destructor, does nothing, leaks the memory **/
 	~FSerialNumberManager()
diff --git a/Engine/Source/Runtime/CoreUObject/Public/UObject/UObjectArray.h b/Engine/Source/Runtime/CoreUObject/Public/UObject/UObjectArray.h
--- a/Engine/Source/Runtime/CoreUObject/Public/UObject/UObjectArray.h
+++ b/Engine/Source/Runtime/CoreUObject/Public/UObject/UObjectArray.h
@@ -67,7 +67,7 @@ public:
 		ObjLastNonGCIndex(INDEX_NONE),
 		OpenForDisregardForGC(true)
 	{
-		FCoreDelegates::ObjectArrayForDebugVisualizers.BindStatic(GetObjectArrayForDebugVisualizers);
+		FCoreDelegates::GetObjectArrayForDebugVisualizersDelegate().BindStatic(GetObjectArrayForDebugVisualizers);
 	}
 
 	/**
