Restored updating of Slate window position on mouse drag in MacApplication. Still working on a better solution. Also, removed some unused code.

diff --git a/Engine/Source/Runtime/Core/Private/Mac/MacApplication.cpp b/Engine/Source/Runtime/Core/Private/Mac/MacApplication.cpp
--- a/Engine/Source/Runtime/Core/Private/Mac/MacApplication.cpp
+++ b/Engine/Source/Runtime/Core/Private/Mac/MacApplication.cpp
@@ -106,7 +106,6 @@ FMacApplication::FMacApplication()
 	, bUsingTrackpad( false )
 	, HIDInput( HIDInputInterface::Create( MessageHandler ) )
 	, DraggedWindow( NULL )
-	, WindowDragCursorPos( FVector2D::ZeroVector )
 	, MouseCaptureWindow( NULL )
 	, bIsMouseCaptureEnabled( false )
 	, bIsMouseCursorLocked( false )
@@ -567,12 +566,21 @@ void FMacApplication::ProcessNSEvent(NSEvent* const Event, TSharedPtr< FMacWindo
 
 			FMacCursor* MacCursor = (FMacCursor*)Cursor.Get();
 
+			// Cocoa does not update NSWindow's frame until user stops dragging the window, so while window is being dragged, we calculate
+			// its position based on mouse move delta
 			if( DraggedWindow )
 			{
-				// Cocoa does not update NSWindow's frame until user stops dragging the window, so while the window is being dragged we don't let Slate know anything's moving
-				MacCursor->UpdateCurrentPosition(WindowDragCursorPos);
+				const int32 X = FMath::TruncToInt(CurrentEventWindow->PositionX + [Event deltaX]);
+				const int32 Y = FMath::TruncToInt(CurrentEventWindow->PositionY + [Event deltaY]);
+				if( CurrentEventWindow.IsValid() )
+				{
+					CurrentEventWindow->PositionX = X;
+					CurrentEventWindow->PositionY = Y;
+					MessageHandler->OnMovedWindow( CurrentEventWindow.ToSharedRef(), X, Y );
+				}
 			}
-			else if( CurrentEventWindow.IsValid() )
+
+			if( CurrentEventWindow.IsValid() )
 			{
 				// Under OS X we disassociate the cursor and mouse position during hi-precision mouse input.
 				// The game snaps the mouse cursor back to the starting point when this is disabled, which
@@ -870,18 +878,6 @@ void FMacApplication::ProcessNSEvent(NSEvent* const Event, TSharedPtr< FMacWindo
 	bIsProcessingNSEvent = bWasProcessingNSEvent;
 }
 
-void FMacApplication::ProcessEvent( NSEvent* Event )
-{
-	SCOPED_AUTORELEASE_POOL;
-
-	const NSEventType EventType = [Event type];
-
-	FCocoaWindow* NativeWindow = FindEventWindow( Event );
-	TSharedPtr< FMacWindow > CurrentEventWindow = FindWindowByNSWindow( Windows, &WindowsMutex, NativeWindow );
-
-	ProcessNSEvent(Event, CurrentEventWindow, ((FMacCursor*)Cursor.Get())->GetPosition());
-}
-
 void FMacApplication::ResendEvent(NSEvent* Event)
 {
 	MainThreadCall(^{
@@ -975,20 +971,6 @@ TSharedPtr<FMacWindow> FMacApplication::LocateWindowUnderCursor( const NSPoint P
 			continue;
 		}
 
-        if( [NativeWindow canBecomeKeyWindow] == false )
-        {
-            NSWindow* ParentWindow = [NativeWindow parentWindow];
-            while( ParentWindow )
-            {
-                if( [ParentWindow canBecomeKeyWindow] )
-                {
-                    NativeWindow = ParentWindow;
-                    break;
-                }
-                ParentWindow = [ParentWindow parentWindow];
-            }
-        }
-        
         NSRect VisibleFrame = [NativeWindow frame];
 #if WITH_EDITOR
         if(MouseScreen != nil)
@@ -1286,12 +1268,6 @@ void FMacApplication::OnWindowWillMove( FCocoaWindow* Window )
 	SCOPED_AUTORELEASE_POOL;
 
 	DraggedWindow = Window;
-
-	TSharedPtr< FMacWindow > EventWindow = FindWindowByNSWindow(Windows, &WindowsMutex, Window);
-	if (EventWindow.IsValid())
-	{
-		WindowDragCursorPos = ((FMacCursor*)Cursor.Get())->GetPosition();
-	}
 }
 
 void FMacApplication::OnWindowDidMove( FCocoaWindow* Window )
diff --git a/Engine/Source/Runtime/Core/Public/Mac/MacApplication.h b/Engine/Source/Runtime/Core/Public/Mac/MacApplication.h
--- a/Engine/Source/Runtime/Core/Public/Mac/MacApplication.h
+++ b/Engine/Source/Runtime/Core/Public/Mac/MacApplication.h
@@ -62,8 +62,6 @@ public:
 
 	void ProcessNSEvent(NSEvent* const Event, TSharedPtr< FMacWindow > CurrentEventWindow, FVector2D const MousePosition);
 
-	void ProcessEvent(NSEvent* Event);
-
 	void OnWindowDraggingFinished();
 
 	bool IsWindowMovable(FCocoaWindow* Win, bool* OutMovableByBackground);
@@ -161,7 +159,6 @@ private:
 	TSharedRef< class HIDInputInterface > HIDInput;
 
 	FCocoaWindow* DraggedWindow;
-	FVector2D WindowDragCursorPos;
 
 	FMouseCaptureWindow* MouseCaptureWindow;
 	bool bIsMouseCaptureEnabled;
