Simplified and cleaned up FMacApplication::ProcessNSEvent() thanks to the fact that SlateApplication no longer needs a SetEventWindow() call in there.

diff --git a/Engine/Source/Runtime/Core/Private/Mac/MacApplication.cpp b/Engine/Source/Runtime/Core/Private/Mac/MacApplication.cpp
--- a/Engine/Source/Runtime/Core/Private/Mac/MacApplication.cpp
+++ b/Engine/Source/Runtime/Core/Private/Mac/MacApplication.cpp
@@ -114,15 +114,14 @@ FMacApplication::FMacApplication()
 	, ModifierKeysFlags( 0 )
 	, CurrentModifierFlags( 0 )
 	, bIsWorkspaceSessionActive( true )
-	, EventMonitor( NULL )
 {
 	CGDisplayRegisterReconfigurationCallback(FMacApplication::OnDisplayReconfiguration, this);
 
 	[NSEvent addGlobalMonitorForEventsMatchingMask:NSMouseMovedMask handler:^(NSEvent* Event){
 		FMacEvent::SendToGameRunLoop(Event, EMacEventSendMethod::Async);
 	}];
-	
-	EventMonitor = (void*)[NSEvent addLocalMonitorForEventsMatchingMask:NSAnyEventMask handler:^(NSEvent* IncomingEvent)
+
+	EventMonitor = [NSEvent addLocalMonitorForEventsMatchingMask:NSAnyEventMask handler:^(NSEvent* IncomingEvent)
 	{
 		NSEvent* ReturnEvent = HandleNSEvent(IncomingEvent);
 		return ReturnEvent;
@@ -254,8 +253,8 @@ FMacApplication::~FMacApplication()
 {
 	if(EventMonitor)
 	{
-		[NSEvent removeMonitor:(id)EventMonitor];
-		EventMonitor = nullptr;
+		[NSEvent removeMonitor:EventMonitor];
+		EventMonitor = nil;
 	}
 	
 	if(AppActivationObserver)
@@ -341,8 +340,7 @@ void FMacApplication::ProcessEvent(FMacEvent const* const Event)
 	if(AppKitEvent)
 	{
 		// Process a standard NSEvent as we always did
-		TSharedPtr<FMacWindow> MacWindow = FindWindowByNSWindow(MacApplication->Windows, &MacApplication->WindowsMutex, Event->GetWindow());
-		MacApplication->ProcessNSEvent(AppKitEvent, MacWindow, Event->GetMousePosition());
+		MacApplication->ProcessNSEvent(AppKitEvent, Event->GetMousePosition());
 	}
 	else if(Notification)
 	{
@@ -480,7 +478,7 @@ bool FMacApplication::IsWindowMovable(FCocoaWindow* Win, bool* OutMovableByBackg
 	return true;
 }
 
-void FMacApplication::HandleModifierChange(TSharedPtr< FMacWindow > CurrentEventWindow, NSUInteger NewModifierFlags, NSUInteger FlagsShift, NSUInteger UE4Shift, EMacModifierKeys TranslatedCode)
+void FMacApplication::HandleModifierChange(NSUInteger NewModifierFlags, NSUInteger FlagsShift, NSUInteger UE4Shift, EMacModifierKeys TranslatedCode)
 {
 	bool CurrentPressed = (CurrentModifierFlags & FlagsShift) != 0;
 	bool NewPressed = (NewModifierFlags & FlagsShift) != 0;
@@ -489,159 +487,130 @@ void FMacApplication::HandleModifierChange(TSharedPtr< FMacWindow > CurrentEvent
 		if( NewPressed )
 		{
 			ModifierKeysFlags |= 1 << UE4Shift;
-			if( CurrentEventWindow.IsValid() )
-			{
-				MessageHandler->OnKeyDown( TranslatedCode, 0, false );
-			}
+			MessageHandler->OnKeyDown( TranslatedCode, 0, false );
 		}
 		else
 		{
 			ModifierKeysFlags &= ~(1 << UE4Shift);
-			if( CurrentEventWindow.IsValid() )
-			{
-				MessageHandler->OnKeyUp( TranslatedCode, 0, false );
-			}
+			MessageHandler->OnKeyUp( TranslatedCode, 0, false );
 		}
 	}
 }
 
-void FMacApplication::ProcessNSEvent(NSEvent* const Event, TSharedPtr< FMacWindow > CurrentEventWindow, FVector2D const MousePosition)
+void FMacApplication::ProcessNSEvent(NSEvent* const Event, FVector2D const MousePosition)
 {
 	SCOPED_AUTORELEASE_POOL;
 
-	const NSEventType EventType = [Event type];
-	
-	FCocoaWindow* NativeWindow = (FCocoaWindow*)[Event window];
-	if( NativeWindow && !CurrentEventWindow.IsValid() && LastEventWindow.IsValid() )
-	{
-		CurrentEventWindow = LastEventWindow;
-		NativeWindow = CurrentEventWindow->GetWindowHandle();
-	}
-
-	if( CurrentEventWindow.IsValid() )
-	{
-		LastEventWindow = CurrentEventWindow;
-		NativeWindow = CurrentEventWindow->GetWindowHandle();
-	}
-
-	if( !NativeWindow )
-	{
-		return;
-	}
-
-	const bool bWasProcessingNSEvent = bIsProcessingNSEvent; // ProcessNSEvent can be called recursively
+	const bool bWasProcessingNSEvent = bIsProcessingNSEvent; // This function can be called recursively
 	bIsProcessingNSEvent = true;
 
-	if( CurrentModifierFlags != [Event modifierFlags] )
+	if (CurrentModifierFlags != [Event modifierFlags])
 	{
 		NSUInteger ModifierFlags = [Event modifierFlags];
 		
-		HandleModifierChange(CurrentEventWindow, ModifierFlags, (1<<4), 7, MMK_RightCommand);
-		HandleModifierChange(CurrentEventWindow, ModifierFlags, (1<<3), 6, MMK_LeftCommand);
-		HandleModifierChange(CurrentEventWindow, ModifierFlags, (1<<1), 0, MMK_LeftShift);
-		HandleModifierChange(CurrentEventWindow, ModifierFlags, (1<<16), 8, MMK_CapsLock);
-		HandleModifierChange(CurrentEventWindow, ModifierFlags, (1<<5), 4, MMK_LeftAlt);
-		HandleModifierChange(CurrentEventWindow, ModifierFlags, (1<<0), 2, MMK_LeftControl);
-		HandleModifierChange(CurrentEventWindow, ModifierFlags, (1<<2), 1, MMK_RightShift);
-		HandleModifierChange(CurrentEventWindow, ModifierFlags, (1<<6), 5, MMK_RightAlt);
-		HandleModifierChange(CurrentEventWindow, ModifierFlags, (1<<13), 3, MMK_RightControl);
+		HandleModifierChange(ModifierFlags, (1<<4), 7, MMK_RightCommand);
+		HandleModifierChange(ModifierFlags, (1<<3), 6, MMK_LeftCommand);
+		HandleModifierChange(ModifierFlags, (1<<1), 0, MMK_LeftShift);
+		HandleModifierChange(ModifierFlags, (1<<16), 8, MMK_CapsLock);
+		HandleModifierChange(ModifierFlags, (1<<5), 4, MMK_LeftAlt);
+		HandleModifierChange(ModifierFlags, (1<<0), 2, MMK_LeftControl);
+		HandleModifierChange(ModifierFlags, (1<<2), 1, MMK_RightShift);
+		HandleModifierChange(ModifierFlags, (1<<6), 5, MMK_RightAlt);
+		HandleModifierChange(ModifierFlags, (1<<13), 3, MMK_RightControl);
 		
 		CurrentModifierFlags = ModifierFlags;
 	}
 
-	switch( EventType )
+	FCocoaWindow* NativeWindow = FindEventWindow(Event);
+	TSharedPtr<FMacWindow> CurrentEventWindow = FindWindowByNSWindow(MacApplication->Windows, &MacApplication->WindowsMutex, NativeWindow);
+
+	const NSEventType EventType = [Event type];
+	switch (EventType)
 	{
 		case NSMouseMoved:
 		case NSLeftMouseDragged:
 		case NSRightMouseDragged:
 		case NSOtherMouseDragged:
 		{
-			if( CurrentEventWindow.IsValid() && CurrentEventWindow->IsRegularWindow() )
+			if (CurrentEventWindow.IsValid() && CurrentEventWindow->IsRegularWindow())
 			{
 				bool IsMouseOverTitleBar = false;
-				bool IsMovable = IsWindowMovable(NativeWindow, &IsMouseOverTitleBar);
-				[NativeWindow setMovable: IsMovable];
-				[NativeWindow setMovableByWindowBackground: IsMouseOverTitleBar];
+				const bool IsMovable = IsWindowMovable(NativeWindow, &IsMouseOverTitleBar);
+				[NativeWindow setMovable:IsMovable];
+				[NativeWindow setMovableByWindowBackground:IsMouseOverTitleBar];
 			}
 
 			FMacCursor* MacCursor = (FMacCursor*)Cursor.Get();
 
 			// Cocoa does not update NSWindow's frame until user stops dragging the window, so while window is being dragged, we calculate
 			// its position based on mouse move delta
-			if( DraggedWindow )
+			if (CurrentEventWindow.IsValid() && DraggedWindow && DraggedWindow == NativeWindow)
 			{
 				const int32 X = FMath::TruncToInt(CurrentEventWindow->PositionX + [Event deltaX]);
 				const int32 Y = FMath::TruncToInt(CurrentEventWindow->PositionY + [Event deltaY]);
-				if( CurrentEventWindow.IsValid() )
-				{
-					CurrentEventWindow->PositionX = X;
-					CurrentEventWindow->PositionY = Y;
-					MessageHandler->OnMovedWindow( CurrentEventWindow.ToSharedRef(), X, Y );
-				}
+				CurrentEventWindow->PositionX = X;
+				CurrentEventWindow->PositionY = Y;
+				MessageHandler->OnMovedWindow(CurrentEventWindow.ToSharedRef(), X, Y);
 			}
 
-			if( CurrentEventWindow.IsValid() )
+			if (bUsingHighPrecisionMouseInput)
 			{
 				// Under OS X we disassociate the cursor and mouse position during hi-precision mouse input.
 				// The game snaps the mouse cursor back to the starting point when this is disabled, which
 				// accumulates mouse delta that we want to ignore.
-				const FVector2D AccumDelta = static_cast<FMacCursor*>( Cursor.Get() )->GetMouseWarpDelta(true);
+				const FVector2D AccumDelta = MacCursor->GetMouseWarpDelta(true);
+
+				// Find the screen the cursor is currently on.
+				NSEnumerator *ScreenEnumerator = [[NSScreen screens] objectEnumerator];
+				NSScreen *Screen;
+				while ((Screen = [ScreenEnumerator nextObject]) && !NSMouseInRect(NSMakePoint(HighPrecisionMousePos.X, HighPrecisionMousePos.Y), Screen.frame, NO))
+					;
+
+				// Clamp to no more than the reported delta - a single event of no mouse movement won't be noticed
+				// but going in the wrong direction will.
+				const FVector2D FullDelta([Event deltaX], [Event deltaY]);
+				const FVector2D WarpDelta(FMath::Abs(AccumDelta.X)<FMath::Abs(FullDelta.X) ? AccumDelta.X : FullDelta.X, FMath::Abs(AccumDelta.Y)<FMath::Abs(FullDelta.Y) ? AccumDelta.Y : FullDelta.Y);
 
-				if( bUsingHighPrecisionMouseInput )
+				FVector2D Delta = ((FullDelta - WarpDelta) / 2.f) * MacCursor->GetMouseScaling();
+
+				HighPrecisionMousePos = MacCursor->GetPosition() + Delta;
+				MacCursor->UpdateCursorClipping(HighPrecisionMousePos);
+
+				// Clamp to the current screen and avoid the menu bar and dock to prevent popups and other
+				// assorted potential for mouse abuse.
+				NSRect VisibleFrame = [Screen visibleFrame];
+				// Avoid the menu bar & dock disclosure borders at the top & bottom of fullscreen windows
+				if (CurrentEventWindow.IsValid() && CurrentEventWindow->GetWindowMode() != EWindowMode::Windowed)
 				{
-					// Find the screen the cursor is currently on.
-					NSEnumerator *screenEnumerator = [[NSScreen screens] objectEnumerator];
-					NSScreen *screen;
-					while ((screen = [screenEnumerator nextObject]) && !NSMouseInRect(NSMakePoint(HighPrecisionMousePos.X, HighPrecisionMousePos.Y), screen.frame, NO))
-						;
-					
-				
-					// Clamp to no more than the reported delta - a single event of no mouse movement won't be noticed
-					// but going in the wrong direction will.
-					const FVector2D FullDelta([Event deltaX], [Event deltaY]);
-					const FVector2D WarpDelta(FMath::Abs(AccumDelta.X)<FMath::Abs(FullDelta.X) ? AccumDelta.X : FullDelta.X, FMath::Abs(AccumDelta.Y)<FMath::Abs(FullDelta.Y) ? AccumDelta.Y : FullDelta.Y);
-					
-					FVector2D Delta = ((FullDelta - WarpDelta) / 2.f) * MacCursor->GetMouseScaling();
-					
-					HighPrecisionMousePos = MacCursor->GetPosition() + Delta;
-					MacCursor->UpdateCursorClipping( HighPrecisionMousePos );
-					
-					// Clamp to the current screen and avoid the menu bar and dock to prevent popups and other
-					// assorted potential for mouse abuse.
-					NSRect VisibleFrame = [screen visibleFrame];
-					// Avoid the menu bar & dock disclosure borders at the top & bottom of fullscreen windows
-					if(CurrentEventWindow->GetWindowMode() != EWindowMode::Windowed)
-					{
-						VisibleFrame.origin.y += 5;
-						VisibleFrame.size.height -= 10;
-					}
-					NSRect FullFrame = [screen frame];
-					VisibleFrame.origin.y = (FullFrame.origin.y+FullFrame.size.height) - (VisibleFrame.origin.y + VisibleFrame.size.height);
-					
-					HighPrecisionMousePos.X = FMath::Clamp(HighPrecisionMousePos.X / MacCursor->GetMouseScaling().X, (float)VisibleFrame.origin.x, (float)(VisibleFrame.origin.x + VisibleFrame.size.width)-1.f);
-					HighPrecisionMousePos.Y = FMath::Clamp(HighPrecisionMousePos.Y / MacCursor->GetMouseScaling().Y, (float)VisibleFrame.origin.y, (float)(VisibleFrame.origin.y + VisibleFrame.size.height)-1.f);
-					
-					MacCursor->WarpCursor( HighPrecisionMousePos.X, HighPrecisionMousePos.Y );
-					MessageHandler->OnRawMouseMove( Delta.X, Delta.Y );
+					VisibleFrame.origin.y += 5;
+					VisibleFrame.size.height -= 10;
 				}
-				else
+				NSRect FullFrame = [Screen frame];
+				VisibleFrame.origin.y = (FullFrame.origin.y+FullFrame.size.height) - (VisibleFrame.origin.y + VisibleFrame.size.height);
+
+				HighPrecisionMousePos.X = FMath::Clamp(HighPrecisionMousePos.X / MacCursor->GetMouseScaling().X, (float)VisibleFrame.origin.x, (float)(VisibleFrame.origin.x + VisibleFrame.size.width)-1.f);
+				HighPrecisionMousePos.Y = FMath::Clamp(HighPrecisionMousePos.Y / MacCursor->GetMouseScaling().Y, (float)VisibleFrame.origin.y, (float)(VisibleFrame.origin.y + VisibleFrame.size.height)-1.f);
+
+				MacCursor->WarpCursor(HighPrecisionMousePos.X, HighPrecisionMousePos.Y);
+				MessageHandler->OnRawMouseMove(Delta.X, Delta.Y);
+			}
+			else
+			{
+				FVector2D CurrentPosition = MousePosition * MacCursor->GetMouseScaling();
+				if (MacCursor->UpdateCursorClipping(CurrentPosition))
 				{
-					FVector2D CurrentPosition = MousePosition * MacCursor->GetMouseScaling();
-					if( MacCursor->UpdateCursorClipping( CurrentPosition ) )
-					{
-						MacCursor->SetPosition( CurrentPosition.X, CurrentPosition.Y );
-					}
-					else
-					{
-						MacCursor->UpdateCurrentPosition(MousePosition);
-					}
-					MessageHandler->OnMouseMove();
+					MacCursor->SetPosition(CurrentPosition.X, CurrentPosition.Y);
 				}
-
-				if( !DraggedWindow && !GetCapture() )
+				else
 				{
-					MessageHandler->OnCursorSet();
+					MacCursor->UpdateCurrentPosition(MousePosition);
 				}
+				MessageHandler->OnMouseMove();
+			}
+
+			if (!DraggedWindow && !GetCapture())
+			{
+				MessageHandler->OnCursorSet();
 			}
 			break;
 		}
@@ -669,18 +638,18 @@ void FMacApplication::ProcessNSEvent(NSEvent* const Event, TSharedPtr< FMacWindo
 				}
 			}
 
-			if( CurrentEventWindow.IsValid() )
+			if (CurrentEventWindow.IsValid())
 			{
 				if (Button == LastPressedMouseButton && ([Event clickCount] % 2) == 0)
 				{
-					MessageHandler->OnMouseDoubleClick( CurrentEventWindow, Button );
+					MessageHandler->OnMouseDoubleClick(CurrentEventWindow, Button);
 				}
 				else
 				{
-					MessageHandler->OnMouseDown( CurrentEventWindow, Button );
+					MessageHandler->OnMouseDown(CurrentEventWindow, Button);
 				}
 
-				if( !DraggedWindow && !GetCapture() )
+				if (!DraggedWindow && !GetCapture())
 				{
 					MessageHandler->OnCursorSet();
 				}
@@ -714,133 +683,111 @@ void FMacApplication::ProcessNSEvent(NSEvent* const Event, TSharedPtr< FMacWindo
 				}
 			}
 
-			if( CurrentEventWindow.IsValid() )
-			{
-				MessageHandler->OnMouseUp( Button );
+			MessageHandler->OnMouseUp(Button);
 
-				if( !DraggedWindow && !GetCapture() )
-				{
-					MessageHandler->OnCursorSet();
-				}
+			if (!DraggedWindow && !GetCapture())
+			{
+				MessageHandler->OnCursorSet();
 			}
+
 			FPlatformMisc::bChachedMacMenuStateNeedsUpdate = true;
 			break;
 		}
 
 		case NSScrollWheel:
 		{
-			if( CurrentEventWindow.IsValid() )
+			const float DeltaX = ([Event modifierFlags] & NSShiftKeyMask) ? [Event deltaY] : [Event deltaX];
+			const float DeltaY = ([Event modifierFlags] & NSShiftKeyMask) ? [Event deltaX] : [Event deltaY];
+
+			NSEventPhase Phase = [Event phase];
+
+			if ([Event momentumPhase] != NSEventPhaseNone || [Event phase] != NSEventPhaseNone)
 			{
-				const float DeltaX = ([Event modifierFlags] & NSShiftKeyMask) ? [Event deltaY] : [Event deltaX];
-				const float DeltaY = ([Event modifierFlags] & NSShiftKeyMask) ? [Event deltaX] : [Event deltaY];
-		
-				NSEventPhase Phase = [Event phase];
-				
-				if ([Event momentumPhase] != NSEventPhaseNone || [Event phase] != NSEventPhaseNone)
-				{
-					bool bInverted = [Event isDirectionInvertedFromDevice];
-					
-					FVector2D ScrollDelta( [Event scrollingDeltaX], [Event scrollingDeltaY] );
-					
-					
-					// This is actually a scroll gesture from trackpad
-					MessageHandler->OnTouchGesture( EGestureEvent::Scroll, bInverted ? -ScrollDelta : ScrollDelta, DeltaY );
-					RecordUsage( EGestureEvent::Scroll );
-				}
-				else
-				{
-					MessageHandler->OnMouseWheel( DeltaY );
-				}
+				const bool bInverted = [Event isDirectionInvertedFromDevice];
 
-				if( !DraggedWindow && !GetCapture() )
-				{
-					MessageHandler->OnCursorSet();
-				}
+				const FVector2D ScrollDelta([Event scrollingDeltaX], [Event scrollingDeltaY]);
+
+				// This is actually a scroll gesture from trackpad
+				MessageHandler->OnTouchGesture(EGestureEvent::Scroll, bInverted ? -ScrollDelta : ScrollDelta, DeltaY);
+				RecordUsage(EGestureEvent::Scroll);
+			}
+			else
+			{
+				MessageHandler->OnMouseWheel(DeltaY);
+			}
+
+			if (!DraggedWindow && !GetCapture())
+			{
+				MessageHandler->OnCursorSet();
 			}
 			break;
 		}
 
 		case NSEventTypeMagnify:
 		{
-			if( CurrentEventWindow.IsValid() )
-			{
-				MessageHandler->OnTouchGesture( EGestureEvent::Magnify, FVector2D( [Event magnification], [Event magnification] ), 0 );
-				RecordUsage( EGestureEvent::Magnify );
-			}
+			MessageHandler->OnTouchGesture(EGestureEvent::Magnify, FVector2D([Event magnification], [Event magnification]), 0);
+			RecordUsage(EGestureEvent::Magnify);
 			break;
 		}
 
 		case NSEventTypeSwipe:
 		{
-			if( CurrentEventWindow.IsValid() )
-			{
-				MessageHandler->OnTouchGesture( EGestureEvent::Swipe, FVector2D( [Event deltaX], [Event deltaY] ), 0 );
-				RecordUsage( EGestureEvent::Swipe );
-			}
+			MessageHandler->OnTouchGesture(EGestureEvent::Swipe, FVector2D([Event deltaX], [Event deltaY]), 0);
+			RecordUsage(EGestureEvent::Swipe);
 			break;
 		}
 
 		case NSEventTypeRotate:
 		{
-			if( CurrentEventWindow.IsValid() )
-			{
-				MessageHandler->OnTouchGesture( EGestureEvent::Rotate, FVector2D( [Event rotation], [Event rotation] ), 0 );
-				RecordUsage( EGestureEvent::Rotate );
-			}
+			MessageHandler->OnTouchGesture(EGestureEvent::Rotate, FVector2D([Event rotation], [Event rotation]), 0);
+			RecordUsage(EGestureEvent::Rotate);
 			break;
 		}
-	
+
 		case NSEventTypeBeginGesture:
 		{
-            if( CurrentEventWindow.IsValid() )
-			{
-				MessageHandler->OnBeginGesture();
-            }
+			MessageHandler->OnBeginGesture();
 			break;
 		}
-        
+
 		case NSEventTypeEndGesture:
 		{
-            if( CurrentEventWindow.IsValid() )
-			{
-				MessageHandler->OnEndGesture();
-            }
+			MessageHandler->OnEndGesture();
 #if WITH_EDITOR
 			LastGestureUsed = EGestureEvent::None;
 #endif
 			break;
 		}
-			
+
 		case NSKeyDown:
 		{
 			NSString *Characters = [Event characters];
 			bool bHandled = false;
-			if( !bSystemModalMode && [Characters length] && CurrentEventWindow.IsValid() )
+			if (!bSystemModalMode && [Characters length] && CurrentEventWindow.IsValid())
 			{
 				bHandled = CurrentEventWindow->OnIMKKeyDown(Event);
-				if(!bHandled)
+				if (!bHandled)
 				{
 					const bool IsRepeat = [Event isARepeat];
-					const TCHAR Character = ConvertChar( [Characters characterAtIndex:0] );
+					const TCHAR Character = ConvertChar([Characters characterAtIndex:0]);
 					const TCHAR CharCode = [[Event charactersIgnoringModifiers] characterAtIndex:0];
 					const uint32 KeyCode = [Event keyCode];
-					const bool IsPrintable = IsPrintableKey( Character );
-					
-					bHandled = MessageHandler->OnKeyDown( KeyCode, TranslateCharCode( CharCode, KeyCode ), IsRepeat );
-					
+					const bool IsPrintable = IsPrintableKey(Character);
+
+					bHandled = MessageHandler->OnKeyDown(KeyCode, TranslateCharCode( CharCode, KeyCode ), IsRepeat);
+
 					// First KeyDown, then KeyChar. This is important, as in-game console ignores first character otherwise
-					
 					bool bCmdKeyPressed = [Event modifierFlags] & 0x18;
-					if ( !bCmdKeyPressed && IsPrintable )
+					if (!bCmdKeyPressed && IsPrintable)
 					{
-						MessageHandler->OnKeyChar( Character, IsRepeat );
+						MessageHandler->OnKeyChar(Character, IsRepeat);
 					}
 				}
 			}
 			if (bHandled)
 			{
 				FCocoaMenu* MainMenu = [[NSApp mainMenu] isKindOfClass:[FCocoaMenu class]] ? (FCocoaMenu*)[NSApp mainMenu]: nil;
-				if ( MainMenu )
+				if (MainMenu)
 				{
 					MainThreadCall(^{ [MainMenu highlightKeyEquivalent:Event]; }, NSDefaultRunLoopMode, true);
 				}
@@ -856,15 +803,15 @@ void FMacApplication::ProcessNSEvent(NSEvent* const Event, TSharedPtr< FMacWindo
 		{
 			NSString *Characters = [Event characters];
 			bool bHandled = false;
-			if( !bSystemModalMode && [Characters length] && CurrentEventWindow.IsValid() )
+			if (!bSystemModalMode && [Characters length])
 			{
 				const bool IsRepeat = [Event isARepeat];
-				const TCHAR Character = ConvertChar( [Characters characterAtIndex:0] );
+				const TCHAR Character = ConvertChar([Characters characterAtIndex:0]);
 				const TCHAR CharCode = [[Event charactersIgnoringModifiers] characterAtIndex:0];
 				const uint32 KeyCode = [Event keyCode];
 				const bool IsPrintable = IsPrintableKey( Character );
 
-				bHandled = MessageHandler->OnKeyUp( KeyCode, TranslateCharCode( CharCode, KeyCode ), IsRepeat );
+				bHandled = MessageHandler->OnKeyUp(KeyCode, TranslateCharCode( CharCode, KeyCode ), IsRepeat);
 			}
 			if (!bHandled)
 			{
diff --git a/Engine/Source/Runtime/Core/Public/Mac/MacApplication.h b/Engine/Source/Runtime/Core/Public/Mac/MacApplication.h
--- a/Engine/Source/Runtime/Core/Public/Mac/MacApplication.h
+++ b/Engine/Source/Runtime/Core/Public/Mac/MacApplication.h
@@ -60,7 +60,7 @@ public:
 		return TextInputMethodSystem.Get();
 	}
 
-	void ProcessNSEvent(NSEvent* const Event, TSharedPtr< FMacWindow > CurrentEventWindow, FVector2D const MousePosition);
+	void ProcessNSEvent(NSEvent* const Event, FVector2D const MousePosition);
 
 	void OnWindowDraggingFinished();
 
@@ -135,7 +135,7 @@ private:
 
 	void UpdateMouseCaptureWindow( FCocoaWindow* TargetWindow );
 
-	void HandleModifierChange(TSharedPtr< FMacWindow > CurrentEventWindow, NSUInteger NewModifierFlags, NSUInteger FlagsShift, NSUInteger UE4Shift, EMacModifierKeys TranslatedCode);
+	void HandleModifierChange(NSUInteger NewModifierFlags, NSUInteger FlagsShift, NSUInteger UE4Shift, EMacModifierKeys TranslatedCode);
 
 #if WITH_EDITOR
 	void RecordUsage(EGestureEvent::Type Gesture);
@@ -167,8 +167,6 @@ private:
 	bool bSystemModalMode;
 	bool bIsProcessingNSEvent;
 
-	TSharedPtr< FMacWindow > LastEventWindow;
-
 	/** The current set of modifier keys that are pressed. This is used to detect differences between left and right modifier keys on key up events*/
 	uint32 ModifierKeysFlags;
 
@@ -201,7 +199,7 @@ private:
 	int32 GestureUsage[EGestureEvent::Count];
 #endif
 
-	void* EventMonitor;
+	id EventMonitor;
 
 	friend class FMacWindow;
 	friend class FMacEvent;
