Fixed handling of text dragged to Slate windows from other apps on Mac

diff --git a/Engine/Source/Runtime/Core/Private/Mac/MacApplication.cpp b/Engine/Source/Runtime/Core/Private/Mac/MacApplication.cpp
--- a/Engine/Source/Runtime/Core/Private/Mac/MacApplication.cpp
+++ b/Engine/Source/Runtime/Core/Private/Mac/MacApplication.cpp
@@ -1121,20 +1121,17 @@ void FDisplayMetrics::GetDisplayMetrics(FDisplayMetrics& OutDisplayMetrics)
 	OutDisplayMetrics.PrimaryDisplayWorkAreaRect.Bottom = OutDisplayMetrics.PrimaryDisplayWorkAreaRect.Top + VisibleFrame.size.height;
 }
 
-void FMacApplication::OnDragEnter( FCocoaWindow* Window, void *InPasteboard )
+void FMacApplication::OnDragEnter(FCocoaWindow* Window, NSPasteboard* Pasteboard)
 {
 	SCOPED_AUTORELEASE_POOL;
 
-	TSharedPtr< FMacWindow > EventWindow = FindWindowByNSWindow( Windows, &WindowsMutex, Window );
-	if( !EventWindow.IsValid() )
+	TSharedPtr< FMacWindow > EventWindow = FindWindowByNSWindow(Windows, &WindowsMutex, Window);
+	if (!EventWindow.IsValid())
 	{
 		return;
 	}
 
 	// Decipher the pasteboard data
-
-	NSPasteboard *Pasteboard = (NSPasteboard*)InPasteboard;
-
 	const bool bHaveText = [[Pasteboard types] containsObject:NSPasteboardTypeString];
 	const bool bHaveFiles = [[Pasteboard types] containsObject:NSFilenamesPboardType];
 
@@ -1150,17 +1147,12 @@ void FMacApplication::OnDragEnter( FCocoaWindow* Window, void *InPasteboard )
 			FileList.Add(ListElement);
 		}
 
-		MessageHandler->OnDragEnterFiles( EventWindow.ToSharedRef(), FileList );
+		MessageHandler->OnDragEnterFiles(EventWindow.ToSharedRef(), FileList);
 	}
 	else if (bHaveText)
 	{
-		NSString *Text = [Pasteboard propertyListForType:NSPasteboardTypeString];
-		TCHAR* TextData = (TCHAR*)FMemory::Malloc(([Text length] + 1) * sizeof(TCHAR));
-		FPlatformString::CFStringToTCHAR((CFStringRef)Text, TextData);
-
-		MessageHandler->OnDragEnterText( EventWindow.ToSharedRef(), FString(TextData) );
-
-		FMemory::Free(TextData);
+		NSString* Text = [Pasteboard stringForType:NSPasteboardTypeString];
+		MessageHandler->OnDragEnterText(EventWindow.ToSharedRef(), FString(Text));
 	}
 }
 
diff --git a/Engine/Source/Runtime/Core/Private/Mac/MacWindow.cpp b/Engine/Source/Runtime/Core/Private/Mac/MacWindow.cpp
--- a/Engine/Source/Runtime/Core/Private/Mac/MacWindow.cpp
+++ b/Engine/Source/Runtime/Core/Private/Mac/MacWindow.cpp
@@ -152,7 +152,7 @@ void FMacWindow::Initialize( FMacApplication* const Application, const TSharedRe
 
 				// Tell Cocoa that we are opting into drag and drop.
 				// Only makes sense for regular windows (windows that last a while.)
-				[WindowHandle registerForDraggedTypes: [NSArray arrayWithObject: NSFilenamesPboardType]];
+				[WindowHandle registerForDraggedTypes:@[NSFilenamesPboardType, NSPasteboardTypeString]];
 
 				if( Definition->HasOSWindowBorder )
 				{
diff --git a/Engine/Source/Runtime/Core/Public/Mac/CocoaWindow.h b/Engine/Source/Runtime/Core/Public/Mac/CocoaWindow.h
--- a/Engine/Source/Runtime/Core/Public/Mac/CocoaWindow.h
+++ b/Engine/Source/Runtime/Core/Public/Mac/CocoaWindow.h
@@ -7,7 +7,7 @@
 /**
  * Custom window class used for input handling
  */
-@interface FCocoaWindow : NSWindow <NSWindowDelegate>
+@interface FCocoaWindow : NSWindow <NSWindowDelegate, NSDraggingDestination>
 {
 	EWindowMode::Type WindowMode;
 	CGFloat DeferOpacity;
diff --git a/Engine/Source/Runtime/Core/Public/Mac/MacApplication.h b/Engine/Source/Runtime/Core/Public/Mac/MacApplication.h
--- a/Engine/Source/Runtime/Core/Public/Mac/MacApplication.h
+++ b/Engine/Source/Runtime/Core/Public/Mac/MacApplication.h
@@ -91,7 +91,7 @@ public:
 
 public:
 
-	void OnDragEnter( FCocoaWindow* Window, void *InPasteboard );
+	void OnDragEnter(FCocoaWindow* Window, NSPasteboard* InPasteboard);
 	void OnDragOver( FCocoaWindow* Window );
 	void OnDragOut( FCocoaWindow* Window );
 	void OnDragDrop( FCocoaWindow* Window );
