Core: Documentation cleanup pass on TEnumAsByte

diff --git a/Engine/Source/Runtime/Core/Public/Containers/EnumAsByte.h b/Engine/Source/Runtime/Core/Public/Containers/EnumAsByte.h
--- a/Engine/Source/Runtime/Core/Public/Containers/EnumAsByte.h
+++ b/Engine/Source/Runtime/Core/Public/Containers/EnumAsByte.h
@@ -1,11 +1,13 @@
 // Copyright 1998-2014 Epic Games, Inc. All Rights Reserved.
 
 #pragma once
+
 #include "Templates/UnrealTypeTraits.h"
 
+
 /**
- * Template to store enums as BYTEs with type safety
-**/
+ * Template to store enumeration values as bytes in a type-safe way.
+ */
 template<class TEnum>
 class TEnumAsByte
 {
@@ -15,7 +17,7 @@ public:
 	/**
 	 * Default Constructor (no initialization).
 	 */
-	FORCEINLINE TEnumAsByte( ) { }
+	FORCEINLINE TEnumAsByte() { }
 
 	/**
 	 * Copy constructor.
@@ -59,6 +61,7 @@ public:
 	 * Assignment operator.
 	 *
 	 * @param InValue value to set.
+	 * @return This instance.
 	 */
 	FORCEINLINE TEnumAsByte& operator=( TEnumAsByte InValue )
 	{
@@ -69,6 +72,7 @@ public:
 	 * Assignment operator.
 	 *
 	 * @param InValue value to set.
+	 * @return This instance.
 	 */
 	FORCEINLINE TEnumAsByte& operator=( TEnum InValue )
 	{
@@ -79,6 +83,7 @@ public:
 	/**
 	 * Compares two enumeration values for equality.
 	 *
+	 * @param InValue The value to compare with.
 	 * @return true if the two values are equal, false otherwise.
 	 */
 	bool operator==( TEnum InValue ) const
@@ -89,6 +94,7 @@ public:
 	/**
 	 * Compares two enumeration values for equality.
 	 *
+	 * @param InValue The value to compare with.
 	 * @return true if the two values are equal, false otherwise.
 	 */
 	bool operator==(TEnumAsByte InValue) const
@@ -96,10 +102,8 @@ public:
 		return Value == InValue.Value;
 	}
 
-	/**
-	 * Implicit conversion to TEnum.
-	 */
-	operator TEnum( ) const
+	/** Implicit conversion to TEnum. */
+	operator TEnum() const
 	{
 		return (TEnum)Value;
 	}
@@ -107,16 +111,18 @@ public:
 public:
 
 	/**
-	 * Gets the value.
+	 * Gets the enumeration value.
+	 *
+	 * @return The enumeration value.
 	 */
-	TEnum GetValue( ) const
+	TEnum GetValue() const
 	{
 		return (TEnum)Value;
 	}
 
 private:
 
-	/** Holds the actual value **/
+	/** Holds the value as a byte. **/
 	uint8 Value;
 };
 
