Linux: get rid of legacy OpenGL code.

- PR #526 by yaakuro, with changes. Only core 3.2 is now used.
- Editor is now able to start on Intel Linux drivers (but see UECOMP-116).

diff --git a/Engine/Shaders/StandaloneRenderer/OpenGL/SplashFragmentShader.glsl b/Engine/Shaders/StandaloneRenderer/OpenGL/SplashFragmentShader.glsl
new file mode 100644
--- /dev/null
+++ b/Engine/Shaders/StandaloneRenderer/OpenGL/SplashFragmentShader.glsl
@@ -0,0 +1,20 @@
+
+#if PLATFORM_USES_ES2
+precision highp float;
+#else
+// #version 120 at the beginning is added in FSlateOpenGLShader::CompileShader()
+#extension GL_EXT_gpu_shader4 : enable
+#endif
+
+varying vec2 textureCoordinate;
+
+uniform sampler2D SplashTexture;
+
+void main() 
+{
+	// OpenGL has 0,0 the "math" way
+	vec2 tc = vec2(textureCoordinate.s, 1.0-textureCoordinate.t);
+
+	gl_FragColor = texture2D(SplashTexture, tc);
+
+}
\ No newline at end of file
diff --git a/Engine/Shaders/StandaloneRenderer/OpenGL/SplashVertexShader.glsl b/Engine/Shaders/StandaloneRenderer/OpenGL/SplashVertexShader.glsl
new file mode 100644
--- /dev/null
+++ b/Engine/Shaders/StandaloneRenderer/OpenGL/SplashVertexShader.glsl
@@ -0,0 +1,12 @@
+attribute vec2 InPosition;
+
+varying vec2 textureCoordinate;
+
+void main() 
+{
+	// We do not need texture coordinates. We calculate using position.
+	textureCoordinate = InPosition * 0.5 + 0.5;
+
+	gl_Position = vec4(InPosition, 0.0, 1.0);
+
+}
\ No newline at end of file
diff --git a/Engine/Source/Developer/StandaloneRenderer/Private/Linux/OpenGL/SlateOpenGLContext.cpp b/Engine/Source/Developer/StandaloneRenderer/Private/Linux/OpenGL/SlateOpenGLContext.cpp
--- a/Engine/Source/Developer/StandaloneRenderer/Private/Linux/OpenGL/SlateOpenGLContext.cpp
+++ b/Engine/Source/Developer/StandaloneRenderer/Private/Linux/OpenGL/SlateOpenGLContext.cpp
@@ -24,6 +24,7 @@ FSlateOpenGLContext::FSlateOpenGLContext()
 	: WindowHandle(NULL)
 	, Context(NULL)
 	, bReleaseWindowOnDestroy(false)
+	, VertexArrayObject(0)
 {
 }
 
@@ -42,10 +43,16 @@ void FSlateOpenGLContext::Initialize( void* InWindow, const FSlateOpenGLContext*
 		bReleaseWindowOnDestroy = true;
 	}
 
+	if (SDL_GL_LoadLibrary(nullptr))
+	{
+		FString SdlError(ANSI_TO_TCHAR(SDL_GetError()));
+		UE_LOG(LogInit, Fatal, TEXT("FSlateOpenGLContext::Initialize - Unable to dynamically load libGL: %s\n."), *SdlError);
+	}
+
 	SDL_GL_SetAttribute( SDL_GL_CONTEXT_MAJOR_VERSION, 3 );
 	SDL_GL_SetAttribute( SDL_GL_CONTEXT_MINOR_VERSION, 2 );
-	SDL_GL_SetAttribute( SDL_GL_CONTEXT_FLAGS, SDL_GL_CONTEXT_FORWARD_COMPATIBLE_FLAG);
-	SDL_GL_SetAttribute( SDL_GL_CONTEXT_PROFILE_MASK, SDL_GL_CONTEXT_PROFILE_COMPATIBILITY );
+	SDL_GL_SetAttribute( SDL_GL_CONTEXT_PROFILE_MASK, SDL_GL_CONTEXT_PROFILE_CORE );
+	SDL_GL_SetAttribute( SDL_GL_CONTEXT_FLAGS, SDL_GL_CONTEXT_DEBUG_FLAG );
 
 	if( SharedContext )
 	{
@@ -67,7 +74,7 @@ void FSlateOpenGLContext::Initialize( void* InWindow, const FSlateOpenGLContext*
 		int OpenGLMinorVersion = -1;
 		SDL_GL_GetAttribute( SDL_GL_CONTEXT_MAJOR_VERSION, &OpenGLMajorVersion );
 		SDL_GL_GetAttribute( SDL_GL_CONTEXT_MINOR_VERSION, &OpenGLMinorVersion );
-		
+
 		UE_LOG(LogInit, Fatal, TEXT("FSlateOpenGLContext::Initialize - Could not create OpenGL %d.%d context, SDL error: '%s'"), 
 				OpenGLMajorVersion, OpenGLMinorVersion,
 				*SdlError
@@ -76,6 +83,12 @@ void FSlateOpenGLContext::Initialize( void* InWindow, const FSlateOpenGLContext*
 		return;
 	}
 	SDL_GL_MakeCurrent( WindowHandle, Context );
+
+	LoadOpenGLExtensions();
+
+	// one Vertex Array Object is reportedly needed for OpenGL 3.2+ core profiles
+	glGenVertexArrays(1, &VertexArrayObject);
+	glBindVertexArray(VertexArrayObject);
 }
 
 void FSlateOpenGLContext::Destroy()
@@ -83,6 +96,7 @@ void FSlateOpenGLContext::Destroy()
 	if	( WindowHandle != NULL )
 	{
 		SDL_GL_MakeCurrent( NULL, NULL );
+		glDeleteVertexArrays(1, &VertexArrayObject);
 		SDL_GL_DeleteContext( Context );
 
 		if	( bReleaseWindowOnDestroy )
@@ -99,9 +113,9 @@ void FSlateOpenGLContext::MakeCurrent()
 {
 	if	( WindowHandle )
 	{
-		CHECK_GL_ERRORS;
 		if(SDL_GL_MakeCurrent( WindowHandle, Context ) == 0)
 		{
+			glBindVertexArray(VertexArrayObject);
 			glGetError(); // SDL leaves glGetError in a dirty state even when successful?
 		}
 	}
diff --git a/Engine/Source/Developer/StandaloneRenderer/Private/OpenGL/SlateOpenGLExtensions.cpp b/Engine/Source/Developer/StandaloneRenderer/Private/OpenGL/SlateOpenGLExtensions.cpp
--- a/Engine/Source/Developer/StandaloneRenderer/Private/OpenGL/SlateOpenGLExtensions.cpp
+++ b/Engine/Source/Developer/StandaloneRenderer/Private/OpenGL/SlateOpenGLExtensions.cpp
@@ -58,6 +58,12 @@ PFNGLGETATTRIBLOCATIONPROC			glGetAttribLocation			= NULL;
 
 PFNWGLCREATECONTEXTATTRIBSARBPROC	wglCreateContextAttribsARB	= NULL;
 
+PFNGLDEBUGMESSAGECALLBACKARBPROC 	glDebugMessageCallback 		= NULL;
+PFNGLBINDVERTEXARRAYPROC 			glBindVertexArray 			= NULL;
+PFNGLDELETEVERTEXARRAYSPROC 		glDeleteVertexArrays 		= NULL;
+PFNGLGENVERTEXARRAYSPROC 			glGenVertexArrays 			= NULL;
+PFNGLISVERTEXARRAYPROC 				glIsVertexArray 			= NULL;
+
 #elif PLATFORM_LINUX
 
 #define DEFINE_GL_ENTRYPOINTS(Type,Func) Type Func = NULL;
diff --git a/Engine/Source/Developer/StandaloneRenderer/Private/OpenGL/SlateOpenGLExtensions.h b/Engine/Source/Developer/StandaloneRenderer/Private/OpenGL/SlateOpenGLExtensions.h
--- a/Engine/Source/Developer/StandaloneRenderer/Private/OpenGL/SlateOpenGLExtensions.h
+++ b/Engine/Source/Developer/StandaloneRenderer/Private/OpenGL/SlateOpenGLExtensions.h
@@ -57,6 +57,12 @@ extern PFNGLGETATTRIBLOCATIONPROC			glGetAttribLocation;
 
 extern PFNWGLCREATECONTEXTATTRIBSARBPROC	wglCreateContextAttribsARB;
 
+extern PFNGLDEBUGMESSAGECALLBACKARBPROC 	glDebugMessageCallback;
+extern PFNGLBINDVERTEXARRAYPROC 			glBindVertexArray;
+extern PFNGLDELETEVERTEXARRAYSPROC 			glDeleteVertexArrays;
+extern PFNGLGENVERTEXARRAYSPROC 			glGenVertexArrays;
+extern PFNGLISVERTEXARRAYPROC 				glIsVertexArray;
+
 #elif PLATFORM_LINUX
 
 #define ENUM_GL_ENTRYPOINTS(EnumMacro) \
@@ -104,6 +110,11 @@ extern PFNWGLCREATECONTEXTATTRIBSARBPROC	wglCreateContextAttribsARB;
 	EnumMacro(PFNGLUNIFORM4FVPROC,glUniform4fv) \
 	EnumMacro(PFNGLUNIFORMMATRIX4FVPROC,glUniformMatrix4fv) \
 	EnumMacro(PFNGLVERTEXATTRIBPOINTERPROC,glVertexAttribPointer) \
+	EnumMacro(PFNGLDEBUGMESSAGECALLBACKARBPROC,glDebugMessageCallback) \
+	EnumMacro(PFNGLBINDVERTEXARRAYPROC,glBindVertexArray) \
+	EnumMacro(PFNGLDELETEVERTEXARRAYSPROC,glDeleteVertexArrays) \
+	EnumMacro(PFNGLGENVERTEXARRAYSPROC,glGenVertexArrays) \
+	EnumMacro(PFNGLISVERTEXARRAYPROC,glIsVertexArray) \
 
 #define DECLARE_GL_ENTRYPOINTS(Type,Func) extern Type Func;
 ENUM_GL_ENTRYPOINTS(DECLARE_GL_ENTRYPOINTS);
@@ -112,4 +123,4 @@ ENUM_GL_ENTRYPOINTS(DECLARE_GL_ENTRYPOINTS);
 
 void LoadOpenGLExtensions();
 
-#endif //!PLATFORM_IOS
\ No newline at end of file
+#endif //!PLATFORM_IOS
diff --git a/Engine/Source/Developer/StandaloneRenderer/Private/OpenGL/SlateOpenGLRenderer.h b/Engine/Source/Developer/StandaloneRenderer/Private/OpenGL/SlateOpenGLRenderer.h
--- a/Engine/Source/Developer/StandaloneRenderer/Private/OpenGL/SlateOpenGLRenderer.h
+++ b/Engine/Source/Developer/StandaloneRenderer/Private/OpenGL/SlateOpenGLRenderer.h
@@ -45,6 +45,7 @@ struct FSlateOpenGLContext
 	SDL_Window* WindowHandle;
 	SDL_GLContext Context;
 	bool bReleaseWindowOnDestroy;
+	GLuint VertexArrayObject; 
 #else
 #error "Unknown platform"
 #endif
diff --git a/Engine/Source/Runtime/Core/Private/Linux/LinuxPlatformSplash.cpp b/Engine/Source/Runtime/Core/Private/Linux/LinuxPlatformSplash.cpp
--- a/Engine/Source/Runtime/Core/Private/Linux/LinuxPlatformSplash.cpp
+++ b/Engine/Source/Runtime/Core/Private/Linux/LinuxPlatformSplash.cpp
@@ -77,49 +77,186 @@ static int32 SplashBPP = 0;
 // subset of GL functions
 /** List all OpenGL entry points used by Unreal. */
 #define ENUM_GL_ENTRYPOINTS(EnumMacro) \
+	EnumMacro(PFNGLENABLEPROC, glEnable) \
+	EnumMacro(PFNGLDISABLEPROC, glDisable) \
+	EnumMacro(PFNGLCLEARPROC, glClear) \
+	EnumMacro(PFNGLCLEARCOLORPROC, glClearColor) \
+	EnumMacro(PFNGLDRAWARRAYSPROC, glDrawArrays) \
+	EnumMacro(PFNGLGENBUFFERSARBPROC, glGenBuffers) \
+	EnumMacro(PFNGLBINDBUFFERARBPROC, glBindBuffer) \
+	EnumMacro(PFNGLBUFFERDATAARBPROC, glBufferData) \
+	EnumMacro(PFNGLDELETEBUFFERSARBPROC, glDeleteBuffers) \
+	EnumMacro(PFNGLMAPBUFFERARBPROC, glMapBuffer) \
+	EnumMacro(PFNGLUNMAPBUFFERARBPROC, glDrawRangeElements) \
 	EnumMacro(PFNGLBINDTEXTUREPROC, glBindTexture) \
+	EnumMacro(PFNGLACTIVETEXTUREARBPROC, glActiveTexture) \
 	EnumMacro(PFNGLTEXIMAGE2DPROC, glTexImage2D) \
 	EnumMacro(PFNGLGENTEXTURESPROC, glGenTextures) \
+	EnumMacro(PFNGLCREATESHADERPROC, glCreateShader) \
+	EnumMacro(PFNGLSHADERSOURCEPROC, glShaderSource) \
+	EnumMacro(PFNGLCOMPILESHADERPROC, glCompileShader) \
+	EnumMacro(PFNGLGETSHADERINFOLOGPROC, glGetShaderInfoLog) \
+	EnumMacro(PFNGLATTACHSHADERPROC, glAttachShader) \
+	EnumMacro(PFNGLDETACHSHADERPROC, glDetachShader) \
+	EnumMacro(PFNGLLINKPROGRAMPROC, glLinkProgram) \
+	EnumMacro(PFNGLGETPROGRAMINFOLOGPROC, glGetProgramInfoLog) \
+	EnumMacro(PFNGLUSEPROGRAMPROC, glUseProgram) \
+	EnumMacro(PFNGLDELETESHADERPROC, glDeleteShader) \
+	EnumMacro(PFNGLCREATEPROGRAMPROC,glCreateProgram) \
+	EnumMacro(PFNGLDELETEPROGRAMPROC, glDeleteProgram) \
+	EnumMacro(PFNGLGETSHADERIVPROC, glGetShaderiv) \
+	EnumMacro(PFNGLGETPROGRAMIVPROC, glGetProgramiv) \
+	EnumMacro(PFNGLGETUNIFORMLOCATIONPROC, glGetUniformLocation) \
+	EnumMacro(PFNGLUNIFORM1FPROC, glUniform1f) \
+	EnumMacro(PFNGLUNIFORM2FPROC, glUniform2f) \
+	EnumMacro(PFNGLUNIFORM3FPROC, glUniform3f) \
+	EnumMacro(PFNGLUNIFORM4FPROC, glUniform4f) \
+	EnumMacro(PFNGLUNIFORM1IPROC, glUniform1i) \
+	EnumMacro(PFNGLUNIFORMMATRIX4FVPROC, glUniformMatrix4fv) \
+	EnumMacro(PFNGLVERTEXATTRIBPOINTERPROC, glVertexAttribPointer) \
+	EnumMacro(PFNGLBINDATTRIBLOCATIONPROC, glBindAttribLocation) \
+	EnumMacro(PFNGLENABLEVERTEXATTRIBARRAYPROC, glEnableVertexAttribArray) \
+	EnumMacro(PFNGLDISABLEVERTEXATTRIBARRAYPROC, glDisableVertexAttribArray) \
+	EnumMacro(PFNGLGETATTRIBLOCATIONPROC, glGetAttribLocation) \
+	EnumMacro(PFNGLBINDVERTEXARRAYPROC, glBindVertexArray) \
+	EnumMacro(PFNGLDELETEVERTEXARRAYSPROC, glDeleteVertexArrays) \
+	EnumMacro(PFNGLGENVERTEXARRAYSPROC, glGenVertexArrays) \
+	EnumMacro(PFNGLISVERTEXARRAYPROC, glIsVertexArray) \
 	EnumMacro(PFNGLTEXPARAMETERIPROC, glTexParameteri)
 
 #define DEFINE_GL_ENTRYPOINTS(Type,Func) Type Func = NULL;
 ENUM_GL_ENTRYPOINTS(DEFINE_GL_ENTRYPOINTS);
 
-// subset 
-#if !defined (GL_QUADS)
-	#define GL_QUADS                                0x0007
-#endif // GL_QUADS
 
-/* Matrix Mode */
-#if !defined (GL_MATRIX_MODE)
-	#define GL_MATRIX_MODE                          0x0BA0
-#endif // GL_MATRIX_MODE
+//////////////////////////////////
+
+
+/**
+ * Returns the current shader log for a GLSL shader                   
+ */
+static FString GetGLSLShaderLog( GLuint Shader )
+{
+	GLint Len;
+	FString ShaderLog;
 
-#if !defined (GL_MODELVIEW)
-	#define GL_MODELVIEW                            0x1700
-#endif // GL_MODELVIEW
+	glGetShaderiv(Shader, GL_INFO_LOG_LENGTH, &Len);
 
-#if !defined (GL_PROJECTION)
-	#define GL_PROJECTION                           0x1701
-#endif // GL_PROJECTION
+	if(Len > 0)
+	{
+		GLsizei ActualLen;
+		GLchar *Log = new GLchar[Len];
 
-#if !defined (GL_TEXTURE)
-	#define GL_TEXTURE                              0x1702
-#endif // GL_TEXTURE
+		glGetShaderInfoLog(Shader, Len, &ActualLen, Log);
+	
+		ShaderLog = ANSI_TO_TCHAR( Log );
+
+		delete[] Log;
+	}
 
-extern "C"
+	return ShaderLog;
+}
+
+static FString GetGLSLProgramLog( GLuint Program )
 {
-	GLAPI void glMatrixMode(GLenum mode);
-	GLAPI void glLoadIdentity(void);
-	GLAPI void glDisable(GLenum cap);
-	GLAPI void glEnable(GLenum cap);
-	GLAPI void glBegin(GLenum mode);
-	GLAPI void glEnd(void);
-	GLAPI void glVertex2i(GLint x, GLint y);
-	GLAPI void glTexCoord2i(GLint s, GLint t);
-};
+	GLint Len;
+	FString ProgramLog;
+	glGetProgramiv( Program, GL_INFO_LOG_LENGTH, &Len );
 
-//////////////////////////////////
+	if( Len > 0 )
+	{
+		GLchar* Log = new GLchar[Len];
+
+		GLsizei ActualLen;
+		glGetProgramInfoLog( Program, Len, &ActualLen, Log );
+
+		ProgramLog = ANSI_TO_TCHAR( Log );
+
+		delete[] Log;
+	}
+
+	return ProgramLog;
+}
+
+//---------------------------------------------------------
+static int CompileShaderFromString( const FString& Source, GLuint& ShaderID, GLenum ShaderType )
+{
+	// Create a new shader ID.
+	ShaderID = glCreateShader( ShaderType );
+	GLint CompileStatus = GL_FALSE;
+
+	check( ShaderID );
+	
+	// Allocate a buffer big enough to store the string in ascii format
+	ANSICHAR* Chars[2] = {nullptr};
+	// pass the #define along to the shader
+#if PLATFORM_USES_ES2
+	Chars[0] = (ANSICHAR*)"#define PLATFORM_USES_ES2 1\n\n#define PLATFORM_LINUX 0\n";
+#elif PLATFORM_LINUX
+	Chars[0] = (ANSICHAR*)"#version 150\n\n#define PLATFORM_USES_ES2 0\n\n#define PLATFORM_LINUX 1\n";
+#else
+	Chars[0] = (ANSICHAR*)"#version 120\n\n#define PLATFORM_USES_ES2 0\n\n#define PLATFORM_LINUX 0\n";
+#endif
+	Chars[1] = new ANSICHAR[Source.Len()+1];
+	FCStringAnsi::Strcpy(Chars[1], Source.Len() + 1, TCHAR_TO_ANSI(*Source));
+
+	// give opengl the source code for the shader
+	glShaderSource( ShaderID, 2, (const ANSICHAR**)Chars, NULL );
+	delete[] Chars[1];
+
+	// Compile the shader and check for success
+	glCompileShader( ShaderID );
+
+	glGetShaderiv( ShaderID, GL_COMPILE_STATUS, &CompileStatus );
+	if( CompileStatus == GL_FALSE )
+	{
+		// The shader did not compile.  Display why it failed.
+		FString Log = GetGLSLShaderLog( ShaderID );
+
+		checkf(false, TEXT("Failed to compile shader: %s\n"), *Log );
+
+		// Delete the shader since it failed.
+		glDeleteShader( ShaderID );
+		ShaderID = 0;
+		return -1;
+	}
+	
+	return 0;
+}
+
+//---------------------------------------------------------
+static int CompileShaderFromFile( const FString& Filename, GLuint& ShaderID, GLenum ShaderType )
+{
+	// Load the file to a string
+	FString Source;
+	bool bFileFound = FFileHelper::LoadFileToString( Source, *Filename );
+	check(bFileFound);
+
+	return CompileShaderFromString(Source, ShaderID, ShaderType);
+}
+
+static int LinkShaders(GLuint& ShaderProgram, GLuint& VertexShader, GLuint& FragmentShader)
+{
+	// Create a new program id and attach the shaders
+	ShaderProgram = glCreateProgram();
+	glAttachShader( ShaderProgram, VertexShader );
+	glAttachShader( ShaderProgram, FragmentShader );
+
+	// Set up attribute locations for per vertex data
+	glBindAttribLocation(ShaderProgram, 0, "InPosition");
+
+	// Link the shaders
+	glLinkProgram( ShaderProgram );
+
+	// Check to see if linking succeeded
+	GLint LinkStatus;
+	glGetProgramiv( ShaderProgram, GL_LINK_STATUS, &LinkStatus );
+	if( LinkStatus == GL_FALSE )
+	{
+		return -1;
+	}
+
+	return 0;
+}
 
 //---------------------------------------------------------
 static int32 OpenFonts ( )
@@ -395,9 +532,7 @@ static int StartSplashScreenThread(void *ptr)
 
 	SDL_GL_SetAttribute( SDL_GL_CONTEXT_MAJOR_VERSION, 3 );
 	SDL_GL_SetAttribute( SDL_GL_CONTEXT_MINOR_VERSION, 2 );
-	SDL_GL_SetAttribute( SDL_GL_CONTEXT_FLAGS, SDL_GL_CONTEXT_FORWARD_COMPATIBLE_FLAG);
-	SDL_GL_SetAttribute( SDL_GL_CONTEXT_PROFILE_MASK, SDL_GL_CONTEXT_PROFILE_COMPATIBILITY );
-
+	SDL_GL_SetAttribute( SDL_GL_CONTEXT_PROFILE_MASK, SDL_GL_CONTEXT_PROFILE_CORE);
 	SDL_GL_SetAttribute( SDL_GL_SHARE_WITH_CURRENT_CONTEXT, 0 );
 
 	SDL_GLContext Context = SDL_GL_CreateContext( GSplashWindow );
@@ -412,18 +547,59 @@ static int StartSplashScreenThread(void *ptr)
 	#define GET_GL_ENTRYPOINTS(Type,Func) Func = reinterpret_cast<Type>(SDL_GL_GetProcAddress(#Func));
 	ENUM_GL_ENTRYPOINTS(GET_GL_ENTRYPOINTS);
 
-	// set up viewport
-	glMatrixMode(GL_MODELVIEW);
-	glLoadIdentity();
-	glMatrixMode(GL_PROJECTION);
-	glLoadIdentity();
+	// Initialize Shaders, Programs, etc.
+	GLuint VertexShader 	= 0;
+	GLuint FragmentShader 	= 0;
+	GLuint ShaderProgram	= 0;
+	GLuint SplashTextureID 	= 0;
 
-	glDisable(GL_DEPTH_TEST);
-	glEnable(GL_TEXTURE_2D);
+	// Compile Vertex and Fragment Shaders from file.
+	if (CompileShaderFromFile( FString::Printf( TEXT("%sShaders/StandaloneRenderer/OpenGL/SplashVertexShader.glsl"), *FPaths::EngineDir()), VertexShader, GL_VERTEX_SHADER) < 0)
+	{
+		UE_LOG(LogHAL, Error, TEXT("Splash screen CompileShaderFromFile failed."));
+		return -1;		
+	}
+
+	if (CompileShaderFromFile(FString::Printf( TEXT("%sShaders/StandaloneRenderer/OpenGL/SplashFragmentShader.glsl"), *FPaths::EngineDir()), FragmentShader, GL_FRAGMENT_SHADER) < 0)
+	{
+		UE_LOG(LogHAL, Error, TEXT("Splash screen CompileShaderFromFile failed."));
+		return -1;		
+	}
+
+	// Create Shader Program and link it.
+	if(LinkShaders(ShaderProgram, VertexShader, FragmentShader) < 0) 
+	{
+		FString Log = GetGLSLProgramLog( ShaderProgram );
+		UE_LOG(LogHAL, Error, TEXT("Failed to link GLSL program: %s"), *Log);
+	}
+
+	// Returns the reference to the Splash Texture from the Shader.
+	SplashTextureID = glGetUniformLocation( ShaderProgram, "SplashTexture" );
+
+	// Create Vertex Array Object. We need that to be conform with OpenGL 3.2+ spec.
+	GLuint VertexArrayObject;
+	glGenVertexArrays(1, &VertexArrayObject);
+	glBindVertexArray(VertexArrayObject);
+	
+	// Create Vertex Buffer and upload data.
+	GLuint VertexBuffer[2];
+	glGenBuffers(2, &VertexBuffer[0]);
+	glBindBuffer(GL_ARRAY_BUFFER, VertexBuffer[0]);
+	float screen_vertex [] = {
+			-1.0, -1.0f,
+			-1.0f, 1.0f,
+			1.0f, 1.0f,
+
+			1.0f, 1.0f,
+			1.0f, -1.0f,
+			-1.0f, -1.0f
+		};
+	glBufferData(GL_ARRAY_BUFFER, sizeof(float)*12, (GLvoid*)screen_vertex, GL_DYNAMIC_DRAW );
 
 	// create texture slot
 	GLuint texture;
 	glGenTextures(1, &texture);
+	glActiveTexture(GL_TEXTURE0);
 	glBindTexture(GL_TEXTURE_2D, texture);
 
 	// set filtering
@@ -498,10 +674,20 @@ static int StartSplashScreenThread(void *ptr)
 
 	OpenFonts();
 
+	glClearColor(0.2f, 0.3f, 0.8f, 1.0f);
+	glClear(GL_COLOR_BUFFER_BIT | GL_DEPTH_BUFFER_BIT);
+	glDisable(GL_DEPTH_TEST);
+
 	// drawing loop
 	while (ThreadState >= 0)
 	{
-		SDL_Delay(100);
+		SDL_Event event;
+		// Poll events otherwise the window will get dark or
+		// on some Desktops it will complain that the thread
+		// does not work anymore.
+		while (SDL_PollEvent(&event)) {}
+
+		SDL_Delay(300);
 		
 		if (ThreadState > 0)
 		{
@@ -511,21 +697,33 @@ static int StartSplashScreenThread(void *ptr)
 		
 		if (SDL_GL_MakeCurrent( GSplashWindow, Context ) == 0)
 		{
-			// draw splash image			
-			glBindTexture(GL_TEXTURE_2D, texture);
-			
-			glBegin(GL_QUADS);
-			glTexCoord2i(0,1); glVertex2i(-1, -1);
-			glTexCoord2i(1,1); glVertex2i(1, -1);
-			glTexCoord2i(1,0); glVertex2i(1, 1);
-			glTexCoord2i(0,0); glVertex2i(-1, 1);
-			glEnd();
-			
+			// Activate Shader Program, Texture etc.
+			glUseProgram(ShaderProgram);
+			glActiveTexture(GL_TEXTURE0);
+			glUniform1i(SplashTextureID, 0);
+
+			// Set stream positions and draw.
+			glEnableVertexAttribArray(0); 
+			glVertexAttribPointer(0, 2, GL_FLOAT, GL_FALSE, 0, 0);
+			glDrawArrays(GL_TRIANGLES, 0, 6);
+
+			// Disable Shader Program, VAO, VBO etc.
+			glDisableVertexAttribArray(0); 
+			glUseProgram(0);
+			glBindBuffer(GL_ARRAY_BUFFER, 0);
+			glBindVertexArray(0);
+
 			SDL_GL_SwapWindow( GSplashWindow );
 		}
 	}
 
 	// clean up
+	glDeleteBuffers(1, &VertexBuffer[0]);
+	glDeleteVertexArrays(1, &VertexArrayObject);
+	glDeleteShader(VertexShader);
+	glDeleteShader(FragmentShader);
+	glDeleteProgram(ShaderProgram);
+ 
 	if (ScratchSpace)
 	{
 		FMemory::Free(ScratchSpace);
diff --git a/Engine/Source/Runtime/OpenGLDrv/Private/Linux/OpenGLLinux.cpp b/Engine/Source/Runtime/OpenGLDrv/Private/Linux/OpenGLLinux.cpp
--- a/Engine/Source/Runtime/OpenGLDrv/Private/Linux/OpenGLLinux.cpp
+++ b/Engine/Source/Runtime/OpenGLDrv/Private/Linux/OpenGLLinux.cpp
@@ -750,7 +750,7 @@ bool PlatformInitOpenGL()
 			UE_LOG(LogLinux, Fatal, TEXT("SDL error errno=%d (%s)"), SDL_GetError());
 		}
 
-		if	( SDL_GL_SetAttribute( SDL_GL_CONTEXT_FLAGS, SDL_GL_CONTEXT_FORWARD_COMPATIBLE_FLAG | DebugFlag ) )
+		if	( SDL_GL_SetAttribute( SDL_GL_CONTEXT_FLAGS, DebugFlag ) )
 		{
 			verifyf( false, TEXT("OpenGLRHI: %s\n."), SDL_GetError() );
 			UE_LOG(LogLinux, Fatal, TEXT("SDL error errno=%d (%s)"), SDL_GetError());
