Linux: support for external input plugins (PR #512).

- Contributed by [REDACTED_PERSON].

diff --git a/Engine/Source/Runtime/Core/Private/Linux/LinuxApplication.cpp b/Engine/Source/Runtime/Core/Private/Linux/LinuxApplication.cpp
--- a/Engine/Source/Runtime/Core/Private/Linux/LinuxApplication.cpp
+++ b/Engine/Source/Runtime/Core/Private/Linux/LinuxApplication.cpp
@@ -1,9 +1,13 @@
+// Copyright 1998-2014 Epic Games, Inc. All Rights Reserved.
 
 #include "CorePrivatePCH.h"
 #include "LinuxApplication.h"
 #include "LinuxWindow.h"
 #include "LinuxCursor.h"
 #include "GenericApplicationMessageHandler.h"
+#include "IForceFeedbackSystem.h"
+#include "IInputDeviceModule.h"
+#include "IInputDevice.h"
 #if STEAM_CONTROLLER_SUPPORT
 	#include "SteamControllerInterface.h"
 #endif // STEAM_CONTROLLER_SUPPORT
@@ -63,6 +67,7 @@ FLinuxApplication::FLinuxApplication() : GenericApplication( MakeShareable( new
 #if STEAM_CONTROLLER_SUPPORT
 	, SteamInput( SteamControllerInterface::Create(MessageHandler) )
 #endif // STEAM_CONTROLLER_SUPPORT
+	, bHasLoadedInputPlugins(false)
 {
 	bUsingHighPrecisionMouseInput = false;
 	bAllowedToDeferMessageProcessing = true;
@@ -735,10 +740,33 @@ void FLinuxApplication::ProcessDeferredEvents( const float TimeDelta )
 
 void FLinuxApplication::PollGameDeviceState( const float TimeDelta )
 {
+	// initialize any externally-implemented input devices (we delay load initialize the array so any plugins have had time to load)
+	if (!bHasLoadedInputPlugins)
+	{
+		TArray<IInputDeviceModule*> PluginImplementations = IModularFeatures::Get().GetModularFeatureImplementations<IInputDeviceModule>(IInputDeviceModule::GetModularFeatureName());
+		for (auto InputPluginIt = PluginImplementations.CreateIterator(); InputPluginIt; ++InputPluginIt)
+		{
+			TSharedPtr<IInputDevice> Device = (*InputPluginIt)->CreateInputDevice(MessageHandler);
+			if (Device.IsValid())
+			{
+				ExternalInputDevices.Add(Device);
+			}
+		}
+
+		bHasLoadedInputPlugins = true;
+	}
+	
 #if STEAM_CONTROLLER_SUPPORT
 	// Poll game device states and send new events
 	SteamInput->SendControllerEvents();
 #endif // STEAM_CONTROLLER_SUPPORT
+
+	// Poll externally-implemented devices
+	for (auto DeviceIt = ExternalInputDevices.CreateIterator(); DeviceIt; ++DeviceIt)
+	{
+		(*DeviceIt)->Tick(TimeDelta);
+		(*DeviceIt)->SendControllerEvents();
+	}
 }
 
 TCHAR FLinuxApplication::ConvertChar( SDL_Keysym Keysym )
diff --git a/Engine/Source/Runtime/Core/Public/Linux/LinuxApplication.h b/Engine/Source/Runtime/Core/Public/Linux/LinuxApplication.h
--- a/Engine/Source/Runtime/Core/Public/Linux/LinuxApplication.h
+++ b/Engine/Source/Runtime/Core/Public/Linux/LinuxApplication.h
@@ -1,4 +1,4 @@
-// Copyright 1998-2013 Epic Games, Inc. All Rights Reserved.
+// Copyright 1998-2014 Epic Games, Inc. All Rights Reserved.
 
 #pragma once
 
@@ -125,6 +125,13 @@ private:
 #if STEAM_CONTROLLER_SUPPORT
 	TSharedPtr< class SteamControllerInterface > SteamInput;
 #endif // STEAM_CONTROLLER_SUPPORT
+
+	/** List of input devices implemented in external modules. */
+	TArray< TSharedPtr<class IInputDevice> > ExternalInputDevices;
+
+	/** Whether input plugins have been loaded */
+	bool bHasLoadedInputPlugins;
+
 };
 
 
