Remove files from CorePrivatePCH.h which are included 1 or 0 times in Core module.

diff --git a/Engine/Source/Runtime/Core/Private/CorePrivatePCH.h b/Engine/Source/Runtime/Core/Private/CorePrivatePCH.h
--- a/Engine/Source/Runtime/Core/Private/CorePrivatePCH.h
+++ b/Engine/Source/Runtime/Core/Private/CorePrivatePCH.h
@@ -38,7 +38,6 @@ Includes.
 ----------------------------------------------------------------------------*/
 
 #include "FileManager.h"				// File manager interface
-#include "ScopedDebugInfo.h"			// Scoped debug info.
 #include "ExternalProfiler.h"			// External profiler integration.
 #include "MemoryBase.h"					// Base memory allocation
 #include "ByteSwap.h"					// ByteSwapping utility code.
@@ -55,7 +54,6 @@ Includes.
 #include "ScopedPointer.h"				// Scoped pointer definitions.
 #include "Sorting.h"					// Sorting definitions.
 #include "Array.h"						// Dynamic array definitions.
-#include "ArrayBuilder.h"				// Builder template for arrays.
 #include "BitArray.h"					// Bit array definition.
 #include "Bitstreams.h"					// Bit stream archiver.
 #include "SparseArray.h"				// Sparse array definitions.
@@ -63,34 +61,22 @@ Includes.
 #include "NameTypes.h"					// Global name subsystem.
 #include "CoreMisc.h"					// Low level utility code.
 #include "Paths.h"						// Path helper functions
-#include "StaticArray.h"                // Static array definition.
-#include "StaticBitArray.h"             // Static bit array definition.
 #include "Set.h"						// Set definitions.
 #include "Map.h"						// Dynamic map definitions.
-#include "MapBuilder.h"					// Builder template for maps.
-#include "List.h"						// Dynamic list definitions.
-#include "ResourceArray.h"				// Resource array definitions.
 #include "RefCounting.h"				// Reference counting definitions.
-#include "ScriptDelegates.h"
 #include "Delegate.h"					// C++ delegate system
 #include "ThreadingBase.h"				// Non-platform specific multi-threaded support.
 #include "Internationalization/Internationalization.h"
 #include "Internationalization/Culture.h"
 #include "Guid.h"						// FGuid class
-#include "NetworkGuid.h"				// FNetworkGuid class
 #include "UnrealMath.h"					// Vector math functions.
-#include "SHMath.h"						// SH math functions.
-#include "RandomStream.h"				// Random stream definitions.
 #include "OutputDevices.h"				// Output devices
 #include "CoreStats.h"
 #include "MemStack.h"					// Stack based memory management.
 #include "AsyncWork.h"					// Async threaded work
 #include "Archive.h"					// Utility archive classes
-#include "IOBase.h"						// base IO declarations, FIOManager, FIOSystem
 #include "Variant.h"
-#include "WildcardString.h"
 #include "CircularBuffer.h"
-#include "CircularQueue.h"
 #include "Queue.h"
 #include "Ticker.h"						// Efficient scheduled delegate manager
 #include "RocketSupport.h"				// Core support for launching in "Rocket" mode
@@ -100,18 +86,11 @@ Includes.
 #include "FeedbackContext.h"
 #include "AutomationTest.h"				// Automation testing functionality for game, editor, etc.
 #include "CallbackDevice.h"				// Base class for callback devices.
-#include "ObjectThumbnail.h"			// Object thumbnails
-#include "LocalTimestampDirectoryVisitor.h"
 #include "GenericWindowDefinition.h"
 #include "GenericWindow.h"
 #include "GenericApplicationMessageHandler.h"
 #include "GenericApplication.h"
 #include "ICursor.h"
 #include "CustomVersion.h"				// Custom versioning system.
-#include "App.h"
-#include "OutputDeviceConsole.h"
 #include "MonitoredProcess.h"
-#include "Attribute.h"
-#include "Optional.h"
-#include "EnumClassFlags.h"
 #include "Misc/DisableOldUETypes.h"
diff --git a/Engine/Source/Runtime/Core/Private/Math/SHMath.cpp b/Engine/Source/Runtime/Core/Private/Math/SHMath.cpp
--- a/Engine/Source/Runtime/Core/Private/Math/SHMath.cpp
+++ b/Engine/Source/Runtime/Core/Private/Math/SHMath.cpp
@@ -2,6 +2,7 @@
 
 
 #include "CorePrivatePCH.h"
+#include "Math/SHMath.h"
 
 //
 //	Spherical harmonic globals.
diff --git a/Engine/Source/Runtime/Core/Private/Misc/LocalTimestampDirectoryVisitor.cpp b/Engine/Source/Runtime/Core/Private/Misc/LocalTimestampDirectoryVisitor.cpp
--- a/Engine/Source/Runtime/Core/Private/Misc/LocalTimestampDirectoryVisitor.cpp
+++ b/Engine/Source/Runtime/Core/Private/Misc/LocalTimestampDirectoryVisitor.cpp
@@ -1,7 +1,7 @@
 // Copyright 1998-2014 Epic Games, Inc. All Rights Reserved.
 
 #include "CorePrivatePCH.h"
-
+#include "Misc/LocalTimestampDirectoryVisitor.h"
 
 /* FLocalTimestampVisitor structors
  *****************************************************************************/
diff --git a/Engine/Source/Runtime/Core/Private/Misc/ObjectThumbnail.cpp b/Engine/Source/Runtime/Core/Private/Misc/ObjectThumbnail.cpp
--- a/Engine/Source/Runtime/Core/Private/Misc/ObjectThumbnail.cpp
+++ b/Engine/Source/Runtime/Core/Private/Misc/ObjectThumbnail.cpp
@@ -1,7 +1,7 @@
 // Copyright 1998-2014 Epic Games, Inc. All Rights Reserved.
 
 #include "CorePrivatePCH.h"
-
+#include "Misc/ObjectThumbnail.h"
 
 /** Static: Thumbnail compressor */
 FThumbnailCompressionInterface* FObjectThumbnail::ThumbnailCompressor = NULL;
diff --git a/Engine/Source/Runtime/Core/Private/Misc/WildcardString.cpp b/Engine/Source/Runtime/Core/Private/Misc/WildcardString.cpp
--- a/Engine/Source/Runtime/Core/Private/Misc/WildcardString.cpp
+++ b/Engine/Source/Runtime/Core/Private/Misc/WildcardString.cpp
@@ -1,7 +1,7 @@
 // Copyright 1998-2014 Epic Games, Inc. All Rights Reserved.
 
 #include "CorePrivatePCH.h"
-
+#include "Misc/WildcardString.h"
 
 /* FWildcardString static functions
  *****************************************************************************/
diff --git a/Engine/Source/Runtime/Core/Private/ProfilingDebugging/ScopedDebugInfo.cpp b/Engine/Source/Runtime/Core/Private/ProfilingDebugging/ScopedDebugInfo.cpp
--- a/Engine/Source/Runtime/Core/Private/ProfilingDebugging/ScopedDebugInfo.cpp
+++ b/Engine/Source/Runtime/Core/Private/ProfilingDebugging/ScopedDebugInfo.cpp
@@ -5,6 +5,7 @@
 =============================================================================*/
 
 #include "CorePrivatePCH.h"
+#include "ProfilingDebugging/ScopedDebugInfo.h"
 
 /** The TLS index for the debug info stack. */
 static uint32 GThreadDebugInfoTLSID = FPlatformTLS::AllocTlsSlot();
diff --git a/Engine/Source/Runtime/Core/Private/Serialization/AsyncIOSystemBase.cpp b/Engine/Source/Runtime/Core/Private/Serialization/AsyncIOSystemBase.cpp
--- a/Engine/Source/Runtime/Core/Private/Serialization/AsyncIOSystemBase.cpp
+++ b/Engine/Source/Runtime/Core/Private/Serialization/AsyncIOSystemBase.cpp
@@ -5,7 +5,7 @@
 =============================================================================*/
 
 #include "CorePrivatePCH.h"
-#include "AsyncIOSystemBase.h"
+#include "Serialization/AsyncIOSystemBase.h"
 
 /*-----------------------------------------------------------------------------
 	FAsyncIOSystemBase implementation.
diff --git a/Engine/Source/Runtime/Core/Private/Serialization/AsyncIOSystemBase.h b/Engine/Source/Runtime/Core/Private/Serialization/AsyncIOSystemBase.h
--- a/Engine/Source/Runtime/Core/Private/Serialization/AsyncIOSystemBase.h
+++ b/Engine/Source/Runtime/Core/Private/Serialization/AsyncIOSystemBase.h
@@ -5,6 +5,7 @@
 =============================================================================*/
 
 #pragma once
+#include "HAL/IOBase.h"
 
 /**
  * Base implementation of an async IO system allowing most of the code to be shared across platforms.
diff --git a/Engine/Source/Runtime/Core/Private/Tests/Misc/CircularQueueTest.cpp b/Engine/Source/Runtime/Core/Private/Tests/Misc/CircularQueueTest.cpp
--- a/Engine/Source/Runtime/Core/Private/Tests/Misc/CircularQueueTest.cpp
+++ b/Engine/Source/Runtime/Core/Private/Tests/Misc/CircularQueueTest.cpp
@@ -5,7 +5,8 @@
 =============================================================================*/
 
 #include "CorePrivatePCH.h"
-#include "AutomationTest.h"
+#include "Containers/CircularQueue.h"
+#include "Misc/AutomationTest.h"
 
 
 IMPLEMENT_SIMPLE_AUTOMATION_TEST(FCircularQueueTest, "Core.Misc.CircularQueue", EAutomationTestFlags::ATF_SmokeTest)
diff --git a/Engine/Source/Runtime/Core/Public/Misc/Variant.h b/Engine/Source/Runtime/Core/Public/Misc/Variant.h
--- a/Engine/Source/Runtime/Core/Public/Misc/Variant.h
+++ b/Engine/Source/Runtime/Core/Public/Misc/Variant.h
@@ -1,7 +1,8 @@
 // Copyright 1998-2014 Epic Games, Inc. All Rights Reserved.
 
 #pragma once
-
+#include "Math/RandomStream.h"
+#include "Misc/NetworkGuid.h"
 
 /**
  * Enumerates the built-in types that can be stored in instances of FVariant.
