CIS fix.

diff --git a/Engine/Source/Runtime/Core/Public/Delegates/Delegate.h b/Engine/Source/Runtime/Core/Public/Delegates/Delegate.h
--- a/Engine/Source/Runtime/Core/Public/Delegates/Delegate.h
+++ b/Engine/Source/Runtime/Core/Public/Delegates/Delegate.h
@@ -11,6 +11,7 @@
 #include "IDelegateInstance.h"
 #include "DelegateBase.h"
 #include "MulticastDelegateBase.h"
+#include "UObject/ScriptDelegates.h"
 
 
 /**
