Changed ECompilationResult enum to match exit code 3 as assert or crash as used by the engine.

diff --git a/Engine/Source/Programs/UnrealBuildTool/System/ExternalExecution.cs b/Engine/Source/Programs/UnrealBuildTool/System/ExternalExecution.cs
--- a/Engine/Source/Programs/UnrealBuildTool/System/ExternalExecution.cs
+++ b/Engine/Source/Programs/UnrealBuildTool/System/ExternalExecution.cs
@@ -27,7 +27,9 @@ namespace UnrealBuildTool
 		/** Compilation failed because generated code changed which was not supported */
 		FailedDueToHeaderChange = 1,
 		/** Compilation failed due to compilation errors */
-		OtherCompilationError,
+		OtherCompilationError = 2,
+		/** The process has most likely crashed. This is what UE returns in case of an assert */
+		CrashOrAssert = 3,
 		/** Compilation is not supported in the current build */
 		Unsupported,
 		/** Unknown error */
diff --git a/Engine/Source/Runtime/Core/Public/Misc/CompilationResult.h b/Engine/Source/Runtime/Core/Public/Misc/CompilationResult.h
--- a/Engine/Source/Runtime/Core/Public/Misc/CompilationResult.h
+++ b/Engine/Source/Runtime/Core/Public/Misc/CompilationResult.h
@@ -23,7 +23,9 @@ namespace ECompilationResult
 		/** Compilation failed because generated code changed which was not supported */
 		FailedDueToHeaderChange = 1,
 		/** Compilation failed due to compilation errors */
-		OtherCompilationError,
+		OtherCompilationError = 2,
+		/** The process has most likely crashed. This is what UE returns in case of an assert */
+		CrashOrAssert = 3,
 		/** Compilation is not supported in the current build */
 		Unsupported,
 		/** Unknown error */
@@ -47,6 +49,8 @@ namespace ECompilationResult
 			return TEXT("FailedDueToHeaderChange");
 		case ECompilationResult::OtherCompilationError:
 			return TEXT("OtherCompilationError");
+		case ECompilationResult::CrashOrAssert:
+			return TEXT("CrashOrAssert");
 		case ECompilationResult::Unsupported:
 			return TEXT("Unsupported");
 		};
