Core: Removed unused file header comments

diff --git a/Engine/Source/Runtime/Core/Private/Tests/HAL/PlatformTest.cpp b/Engine/Source/Runtime/Core/Private/Tests/HAL/PlatformTest.cpp
--- a/Engine/Source/Runtime/Core/Private/Tests/HAL/PlatformTest.cpp
+++ b/Engine/Source/Runtime/Core/Private/Tests/HAL/PlatformTest.cpp
@@ -1,10 +1,8 @@
 // Copyright 1998-2014 Epic Games, Inc. All Rights Reserved.
 
-/*================================================================================
-	PlatformTest.cpp: Platform automation test code
-==================================================================================*/
 #include "CorePrivatePCH.h"
 
+
 float TheCompilerDoesntKnowThisIsAlwaysZero = 0.0f;
 
 
@@ -18,6 +16,7 @@ struct TestA
 	uint8 Space[64];
 };
 
+
 struct TestB
 {
 	virtual ~TestB() {}
@@ -28,13 +27,13 @@ struct TestB
 	uint8 Space[96];
 };
 
+
 struct TestC : public TestA, TestB
 {
 	int i;
 };
 
 
-
 IMPLEMENT_SIMPLE_AUTOMATION_TEST(FPlatformVerificationTest, "Core.HAL.Platform Verification", EAutomationTestFlags::ATF_SmokeTest)
 
 bool FPlatformVerificationTest::RunTest (const FString& Parameters)
diff --git a/Engine/Source/Runtime/Core/Private/Tests/Internationalization/ArchiveTest.cpp b/Engine/Source/Runtime/Core/Private/Tests/Internationalization/ArchiveTest.cpp
--- a/Engine/Source/Runtime/Core/Private/Tests/Internationalization/ArchiveTest.cpp
+++ b/Engine/Source/Runtime/Core/Private/Tests/Internationalization/ArchiveTest.cpp
@@ -1,14 +1,11 @@
 // Copyright 1998-2014 Epic Games, Inc. All Rights Reserved.
 
-/*===============================================================================================
-	ArchiveTest.cpp: Unit test for the FInternationalizationArchive class and serializers
-=================================================================================================*/
-
 #include "CorePrivatePCH.h"
 #include "AutomationTest.h"
 #include "Internationalization/InternationalizationArchive.h"
 #include "Internationalization/InternationalizationMetadata.h"
 
+
 IMPLEMENT_SIMPLE_AUTOMATION_TEST(FArchiveTest, "Core.Misc.Internationalization Archive", EAutomationTestFlags::ATF_SmokeTest)
 
 
@@ -212,9 +209,5 @@ bool FArchiveTest::RunTest( const FString& Parameters )
 		}
 	}
 
-
-
-	
-
 	return true;
 }
diff --git a/Engine/Source/Runtime/Core/Private/Tests/Internationalization/DateTimeFormattingRulesTest.cpp b/Engine/Source/Runtime/Core/Private/Tests/Internationalization/DateTimeFormattingRulesTest.cpp
--- a/Engine/Source/Runtime/Core/Private/Tests/Internationalization/DateTimeFormattingRulesTest.cpp
+++ b/Engine/Source/Runtime/Core/Private/Tests/Internationalization/DateTimeFormattingRulesTest.cpp
@@ -3,7 +3,8 @@
 #include "CorePrivatePCH.h"
 #include "AutomationTest.h"
 
-// Disable optimization for DateTimeFormattingRulesTest as it compiles very slowly in development builds
+
+// disable optimization for DateTimeFormattingRulesTest as it compiles very slowly in development builds
 PRAGMA_DISABLE_OPTIMIZATION
 
 IMPLEMENT_SIMPLE_AUTOMATION_TEST(FDateTimeFormattingRulesTest, "Core.Misc.DateTime Formatting Rules", EAutomationTestFlags::ATF_None)
@@ -19,6 +20,7 @@ namespace
 	}
 }
 
+
 bool FDateTimeFormattingRulesTest::RunTest (const FString& Parameters)
 {
 #if UE_ENABLE_ICU
@@ -98,4 +100,5 @@ bool FDateTimeFormattingRulesTest::RunTest (const FString& Parameters)
 	return true;
 }
 
+
 PRAGMA_ENABLE_OPTIMIZATION
\ No newline at end of file
diff --git a/Engine/Source/Runtime/Core/Private/Tests/Internationalization/ManifestTest.cpp b/Engine/Source/Runtime/Core/Private/Tests/Internationalization/ManifestTest.cpp
--- a/Engine/Source/Runtime/Core/Private/Tests/Internationalization/ManifestTest.cpp
+++ b/Engine/Source/Runtime/Core/Private/Tests/Internationalization/ManifestTest.cpp
@@ -1,17 +1,11 @@
 // Copyright 1998-2014 Epic Games, Inc. All Rights Reserved.
 
-/*===============================================================================================
-	ManifestTest.cpp: Unit test for the FInternationalizationManifest class
-=================================================================================================*/
-
 #include "CorePrivatePCH.h"
 #include "AutomationTest.h"
 #include "Internationalization/InternationalizationManifest.h"
 #include "Internationalization/InternationalizationMetadata.h"
 
 
-
-
 // Helper function to count the number of entries in a manifest
 int32 CountManifestEntries( const FInternationalizationManifest& Manifest )
 {
@@ -20,6 +14,7 @@ int32 CountManifestEntries( const FInternationalizationManifest& Manifest )
 	return EntryCounter;
 }
 
+
 IMPLEMENT_SIMPLE_AUTOMATION_TEST(FLocContextTest, "Core.Misc.Internationalization Context", EAutomationTestFlags::ATF_SmokeTest)
 
 bool FLocContextTest::RunTest( const FString& Parameters )
@@ -204,6 +199,7 @@ bool FLocContextTest::RunTest( const FString& Parameters )
 	return true;
 }
 
+
 IMPLEMENT_SIMPLE_AUTOMATION_TEST(FLocItemTest, "Core.Misc.Internationalization LocItem", EAutomationTestFlags::ATF_SmokeTest)
 
 bool FLocItemTest::RunTest( const FString& Parameters )
@@ -358,7 +354,6 @@ bool FLocItemTest::RunTest( const FString& Parameters )
 }
 
 
-
 IMPLEMENT_SIMPLE_AUTOMATION_TEST(FManifestTest, "Core.Misc.Internationalization Manifest", EAutomationTestFlags::ATF_SmokeTest)
 
 bool FManifestTest::RunTest( const FString& Parameters )
@@ -601,7 +596,6 @@ bool FManifestTest::RunTest( const FString& Parameters )
 
 		}
 
-
 		// Adding entries with Source that is NOT an exact match and different context
 		{
 			// Source mismatched by Source Text.
@@ -788,7 +782,6 @@ bool FManifestTest::RunTest( const FString& Parameters )
 
 		}
 
-
 		// Adding an entry that only differs in the optional flag 
 		{
 			FInternationalizationManifest TestManifest;
diff --git a/Engine/Source/Runtime/Core/Private/Tests/Internationalization/MetadataTest.cpp b/Engine/Source/Runtime/Core/Private/Tests/Internationalization/MetadataTest.cpp
--- a/Engine/Source/Runtime/Core/Private/Tests/Internationalization/MetadataTest.cpp
+++ b/Engine/Source/Runtime/Core/Private/Tests/Internationalization/MetadataTest.cpp
@@ -2,14 +2,10 @@
 
 #include "CorePrivatePCH.h"
 #include "AutomationTest.h"
-
 #include "Internationalization/InternationalizationMetadata.h"
 
-/*=============================================================================
-	Metadatatest.cpp: Unit test for the various loc metadata classes
-=============================================================================*/
-IMPLEMENT_SIMPLE_AUTOMATION_TEST(FMetadataTest, "Core.Misc.Internationalization Metadata", EAutomationTestFlags::ATF_SmokeTest)
 
+IMPLEMENT_SIMPLE_AUTOMATION_TEST(FMetadataTest, "Core.Misc.Internationalization Metadata", EAutomationTestFlags::ATF_SmokeTest)
 
 bool FMetadataTest::RunTest( const FString& Parameters )
 {
@@ -331,7 +327,6 @@ bool FMetadataTest::RunTest( const FString& Parameters )
 			TestFalse( TEXT("MetadataObjectAClone == MetadataObjectAClone2"),  MetadataObjectAClone.IsExactMatch( MetadataObjectAClone2) );
 		}
 
-
 		// Test less than operator
 		{
 			
@@ -364,8 +359,6 @@ bool FMetadataTest::RunTest( const FString& Parameters )
 			MetadataObjectAClone2.SetStringField( TEXT("*NoCompare"), TEXT("NoCompare") );
 			TestTrue( TEXT("MetadataObjectAClone < MetadataObjectAClone2"), MetadataObjectAClone < MetadataObjectAClone2 );
 		}
-
-		
 	}
 
 	return true;
diff --git a/Engine/Source/Runtime/Core/Private/Tests/Internationalization/NumberFormatingRulesTest.cpp b/Engine/Source/Runtime/Core/Private/Tests/Internationalization/NumberFormatingRulesTest.cpp
--- a/Engine/Source/Runtime/Core/Private/Tests/Internationalization/NumberFormatingRulesTest.cpp
+++ b/Engine/Source/Runtime/Core/Private/Tests/Internationalization/NumberFormatingRulesTest.cpp
@@ -1,12 +1,9 @@
 ﻿// Copyright 1998-2014 Epic Games, Inc. All Rights Reserved.
 
-/*=============================================================================
-	NumberFormattingRulesTest.cpp: Unit test for the FNumberFormattingRules class.
-=============================================================================*/
-
 #include "CorePrivatePCH.h"
 #include "AutomationTest.h"
 
+
 // Disable optimization for NumberFormatingRulesTest as it compiles very slowly in development builds
 PRAGMA_DISABLE_OPTIMIZATION
 
@@ -23,6 +20,7 @@ namespace
 	}
 }
 
+
 bool FNumberFormattingRulesTest::RunTest (const FString& Parameters)
 {
 /**
diff --git a/Engine/Source/Runtime/Core/Private/Tests/Math/RangeBoundTest.cpp b/Engine/Source/Runtime/Core/Private/Tests/Math/RangeBoundTest.cpp
--- a/Engine/Source/Runtime/Core/Private/Tests/Math/RangeBoundTest.cpp
+++ b/Engine/Source/Runtime/Core/Private/Tests/Math/RangeBoundTest.cpp
@@ -1,16 +1,11 @@
 // Copyright 1998-2014 Epic Games, Inc. All Rights Reserved.
 
-/*=============================================================================
-	RangeBoundTest.cpp: Unit test for the TRangeBound template.
-=============================================================================*/
-
 #include "CorePrivatePCH.h"
 #include "AutomationTest.h"
 
 
 IMPLEMENT_SIMPLE_AUTOMATION_TEST(FRangeBoundTest, "Core.Math.RangeBound", EAutomationTestFlags::ATF_SmokeTest)
 
-
 bool FRangeBoundTest::RunTest( const FString& Parameters )
 {
 	// ensure template instantiation works
@@ -120,4 +115,4 @@ bool FRangeBoundTest::RunTest( const FString& Parameters )
 	TestEqual(TEXT("'[3' must be greater than '[2' <2>"), FFloatRangeBound::MaxLower(b4_2, b3_2), b4_2);
 
 	return true;
-}
\ No newline at end of file
+}
diff --git a/Engine/Source/Runtime/Core/Private/Tests/Math/RangeSetTest.cpp b/Engine/Source/Runtime/Core/Private/Tests/Math/RangeSetTest.cpp
--- a/Engine/Source/Runtime/Core/Private/Tests/Math/RangeSetTest.cpp
+++ b/Engine/Source/Runtime/Core/Private/Tests/Math/RangeSetTest.cpp
@@ -1,17 +1,12 @@
 // Copyright 1998-2014 Epic Games, Inc. All Rights Reserved.
 
-/*=============================================================================
-	RangeSetTest.cpp: Unit test for the TRangeSet template.
-=============================================================================*/
-
 #include "CorePrivatePCH.h"
 #include "AutomationTest.h"
 
 
 IMPLEMENT_SIMPLE_AUTOMATION_TEST(FRangeSetTest, "Core.Math.RangeSet", EAutomationTestFlags::ATF_SmokeTest)
 
-
 bool FRangeSetTest::RunTest( const FString& Parameters )
 {
 	return true;
-}
\ No newline at end of file
+}
diff --git a/Engine/Source/Runtime/Core/Private/Tests/Math/RangeTest.cpp b/Engine/Source/Runtime/Core/Private/Tests/Math/RangeTest.cpp
--- a/Engine/Source/Runtime/Core/Private/Tests/Math/RangeTest.cpp
+++ b/Engine/Source/Runtime/Core/Private/Tests/Math/RangeTest.cpp
@@ -1,9 +1,5 @@
 // Copyright 1998-2014 Epic Games, Inc. All Rights Reserved.
 
-/*=============================================================================
-	RangeTest.cpp: Unit test for the TRange template.
-=============================================================================*/
-
 #include "CorePrivatePCH.h"
 #include "AutomationTest.h"
 
@@ -11,8 +7,8 @@
 // Disable optimization for RangeTest as it compiles very slowly in development builds
 PRAGMA_DISABLE_OPTIMIZATION
 
-IMPLEMENT_SIMPLE_AUTOMATION_TEST(FRangeTest, "Core.Math.Range", EAutomationTestFlags::ATF_SmokeTest)
 
+IMPLEMENT_SIMPLE_AUTOMATION_TEST(FRangeTest, "Core.Math.Range", EAutomationTestFlags::ATF_SmokeTest)
 
 bool FRangeTest::RunTest( const FString& Parameters )
 {
@@ -113,16 +109,12 @@ bool FRangeTest::RunTest( const FString& Parameters )
 
 	// @todo gmp: conjoined ranges
 
-
 	// @todo gmp: element containment
 
-
 	// @todo gmp: range containment
 
-
 	// @todo gmp: contiguous ranges
 
-
 	// bounds access
 	FFloatRange r11_1 = FFloatRange(FFloatRangeBound::Inclusive(1.0f), FFloatRangeBound::Exclusive(4.0f));
 	FFloatRange r11_2 = FFloatRange(FFloatRangeBound::Open(), FFloatRangeBound::Open());
@@ -174,7 +166,6 @@ bool FRangeTest::RunTest( const FString& Parameters )
 
 	// @todo gmp: range splitting
 
-
 	// difference
 	FFloatRange r15_1 = FFloatRange(FFloatRangeBound::Exclusive(7.0f), FFloatRangeBound::Exclusive(14.0f));		// X
 	FFloatRange r15_2 = FFloatRange(FFloatRangeBound::Inclusive(7.0f), FFloatRangeBound::Inclusive(14.0f));
@@ -403,4 +394,4 @@ bool FRangeTest::RunTest( const FString& Parameters )
 }
 
 
-PRAGMA_ENABLE_OPTIMIZATION
\ No newline at end of file
+PRAGMA_ENABLE_OPTIMIZATION
diff --git a/Engine/Source/Runtime/Core/Private/Tests/Math/UnrealMathTest.cpp b/Engine/Source/Runtime/Core/Private/Tests/Math/UnrealMathTest.cpp
--- a/Engine/Source/Runtime/Core/Private/Tests/Math/UnrealMathTest.cpp
+++ b/Engine/Source/Runtime/Core/Private/Tests/Math/UnrealMathTest.cpp
@@ -1,10 +1,8 @@
 // Copyright 1998-2014 Epic Games, Inc. All Rights Reserved.
 
-/*=============================================================================
-	UnMathTest.cpp: Test harness for vector INTrinsics abstraction.
-=============================================================================*/
 #include "CorePrivatePCH.h" 
 
+
 DEFINE_LOG_CATEGORY_STATIC(LogUnrealMathTest, Log, All);
 
 // Create some temporary storage variables
@@ -12,6 +10,7 @@ MS_ALIGN( 16 ) static float GScratch[16] GCC_ALIGN( 16 );
 static float GSum;
 static bool GPassing;
 
+
 /**
  * Tests if two vectors (xyzw) are bitwise equal
  *
@@ -31,6 +30,7 @@ bool TestVectorsEqualBitwise( VectorRegister Vec0, VectorRegister Vec1)
 	return Passed;
 }
 
+
 /**
  * Tests if two vectors (xyzw) are equal within an optional tolerance
  *
@@ -54,6 +54,7 @@ bool TestVectorsEqual( VectorRegister Vec0, VectorRegister Vec1, float Tolerance
 	return GSum <= Tolerance;
 }
 
+
 /**
  * Tests if two vectors (xyz) are equal within an optional tolerance
  *
@@ -76,6 +77,7 @@ bool TestVectorsEqual3( VectorRegister Vec0, VectorRegister Vec1, float Toleranc
 	return GSum <= Tolerance;
 }
 
+
 /**
  * Tests if two matrices (4x4 xyzw) are equal within an optional tolerance
  *
@@ -104,6 +106,7 @@ bool TestMatricesEqual( FMatrix &Mat0, FMatrix &Mat1, float Tolerance = 0.0f)
 	return true;
 }
 
+
 /**
  * Multiplies two 4x4 matrices.
  *
@@ -139,6 +142,7 @@ void TestVectorMatrixMultiply( void* Result, const void* Matrix1, const void* Ma
 	memcpy( Result, &Temp, 16*sizeof(float) );
 }
 
+
 /**
  * Calculate the inverse of an FMatrix.
  *
@@ -229,6 +233,7 @@ void TestVectorMatrixMultiply( void* Result, const void* Matrix1, const void* Ma
 	memcpy( DstMatrix, &Result, 16*sizeof(float) );
 }
 
+
 /**
  * Calculate Homogeneous transform.
  *
@@ -251,6 +256,7 @@ VectorRegister TestVectorTransformVector(const VectorRegister&  VecP,  const voi
 	return Result.v;
 }
 
+
 /**
 * Multiplies two quaternions: The order matters.
 *
@@ -278,6 +284,7 @@ void TestVectorQuaternionMultiply( void *Result, const void* Quat1, const void*
 	R[3] = TW;
 }
 
+
 /**
  * Helper debugf function to print out success or failure information for a test
  *
@@ -294,6 +301,7 @@ void LogTest( const TCHAR *TestName, bool bHasPassed )
 	}
 }
 
+
 /** 
  * Set the contents of the scratch memory
  * 
@@ -308,6 +316,7 @@ void SetScratch( float X, float Y, float Z, float W, float U = 0.0f )
 	GScratch[4] = U;
 }
 
+
 IMPLEMENT_SIMPLE_AUTOMATION_TEST(FVectorRegisterAbstractionTest, "Core.Math.Vector Register Abstraction Test", EAutomationTestFlags::ATF_SmokeTest)
 
 /**
@@ -626,7 +635,6 @@ bool FVectorRegisterAbstractionTest::RunTest(const FString& Parameters)
 	V3 = MakeVectorRegister( 1.0f, -3.0f, 24.0f, -36.0f );
 	LogTest( TEXT("VectorBitwiseXor-Float"), TestVectorsEqual( V2, V3 ) );
 
-
 	V0 = MakeVectorRegister( -1.0f, -3.0f, -24.0f, 36.0f );
 	V1 = MakeVectorRegister( 5.0f, 35.0f, 23.0f, 48.0f );
 	V2 = VectorMergeVecXYZ_VecW( V0, V1 );
@@ -639,7 +647,6 @@ bool FVectorRegisterAbstractionTest::RunTest(const FString& Parameters)
 	V3 = MakeVectorRegister( 5.0f, 35.0f, 23.0f, 36.0f );
 	LogTest( TEXT("VectorMergeXYZ_VecW-2"), TestVectorsEqual( V2, V3 ) );
 
-
 	V0 = MakeVectorRegister( 1.0f, 1.0e6f, 1.3e-8f, 35.0f );
 	V1 = VectorReciprocal( V0 );
 	V3 = VectorMultiply(V1, V0);
@@ -650,7 +657,6 @@ bool FVectorRegisterAbstractionTest::RunTest(const FString& Parameters)
 	V3 = VectorMultiply(V1, V0);
 	LogTest( TEXT("VectorReciprocalAccurate"), TestVectorsEqual( VectorOne(), V3, 1e-7f ) );
 
-
 	V0 = MakeVectorRegister( 1.0f, 1.0e6f, 1.3e-8f, 35.0f );
 	V1 = VectorReciprocalSqrt( V0 );
 	V3 = VectorMultiply(VectorMultiply(V1, V1), V0);
@@ -679,7 +685,7 @@ bool FVectorRegisterAbstractionTest::RunTest(const FString& Parameters)
 	const float ZNear = -100.0f;
 	const float ZFar = 100.0f;
 
-	M1= 	FMatrix(FPlane(2.0f/(Right-Left),			0,							0,					0 ),
+	M1 = FMatrix(FPlane(2.0f/(Right-Left),	0,							0,					0 ),
 		FPlane(0,							2.0f/(Top-Bottom),			0,					0 ),
 		FPlane(0,							0,							1/(ZNear-ZFar),		0 ),
 		FPlane((Left+Right)/(Left-Right),	(Top+Bottom)/(Bottom-Top),	ZNear/(ZNear-ZFar), 1 ) );
@@ -738,4 +744,3 @@ bool FVectorRegisterAbstractionTest::RunTest(const FString& Parameters)
 
 	return true;
 }
-
diff --git a/Engine/Source/Runtime/Core/Private/Tests/Misc/CircularBufferTest.cpp b/Engine/Source/Runtime/Core/Private/Tests/Misc/CircularBufferTest.cpp
--- a/Engine/Source/Runtime/Core/Private/Tests/Misc/CircularBufferTest.cpp
+++ b/Engine/Source/Runtime/Core/Private/Tests/Misc/CircularBufferTest.cpp
@@ -1,16 +1,11 @@
 // Copyright 1998-2014 Epic Games, Inc. All Rights Reserved.
 
-/*=============================================================================
-	CircularBufferTest.cpp: Unit test for the TCircularBuffer template class.
-=============================================================================*/
-
 #include "CorePrivatePCH.h"
 #include "AutomationTest.h"
 
 
 IMPLEMENT_SIMPLE_AUTOMATION_TEST(FCircularBufferTest, "Core.Misc.CircularBuffer", EAutomationTestFlags::ATF_SmokeTest)
 
-
 bool FCircularBufferTest::RunTest( const FString& Parameters )
 {
 	// buffer capacity
diff --git a/Engine/Source/Runtime/Core/Private/Tests/Misc/CircularQueueTest.cpp b/Engine/Source/Runtime/Core/Private/Tests/Misc/CircularQueueTest.cpp
--- a/Engine/Source/Runtime/Core/Private/Tests/Misc/CircularQueueTest.cpp
+++ b/Engine/Source/Runtime/Core/Private/Tests/Misc/CircularQueueTest.cpp
@@ -1,9 +1,5 @@
 // Copyright 1998-2014 Epic Games, Inc. All Rights Reserved.
 
-/*=============================================================================
-	CircularQueueTest.cpp: Unit test for the TCircularQueue template class.
-=============================================================================*/
-
 #include "CorePrivatePCH.h"
 #include "Containers/CircularQueue.h"
 #include "Misc/AutomationTest.h"
@@ -11,7 +7,6 @@
 
 IMPLEMENT_SIMPLE_AUTOMATION_TEST(FCircularQueueTest, "Core.Misc.CircularQueue", EAutomationTestFlags::ATF_SmokeTest)
 
-
 bool FCircularQueueTest::RunTest( const FString& Parameters )
 {
 	const uint32 QueueSize = 8;
diff --git a/Engine/Source/Runtime/Core/Private/Tests/Misc/GuidTest.cpp b/Engine/Source/Runtime/Core/Private/Tests/Misc/GuidTest.cpp
--- a/Engine/Source/Runtime/Core/Private/Tests/Misc/GuidTest.cpp
+++ b/Engine/Source/Runtime/Core/Private/Tests/Misc/GuidTest.cpp
@@ -1,16 +1,11 @@
 // Copyright 1998-2014 Epic Games, Inc. All Rights Reserved.
 
-/*=============================================================================
-	GuidTest.cpp: Unit test for the FGuid class.
-=============================================================================*/
-
 #include "CorePrivatePCH.h"
 #include "AutomationTest.h"
 
 
 IMPLEMENT_SIMPLE_AUTOMATION_TEST(FGuidTest, "Core.Misc.Guid", EAutomationTestFlags::ATF_SmokeTest)
 
-
 bool FGuidTest::RunTest( const FString& Parameters )
 {
 	FGuid g = FGuid(305419896, 2271560481, 305419896, 2271560481);
@@ -81,4 +76,4 @@ bool FGuidTest::RunTest( const FString& Parameters )
 	TestFalse(TEXT("Invalidated GUIDs must be invalid"), g4_1.IsValid());
 
 	return true;
-}
\ No newline at end of file
+}
diff --git a/Engine/Source/Runtime/Core/Private/Tests/Misc/PathsTest.cpp b/Engine/Source/Runtime/Core/Private/Tests/Misc/PathsTest.cpp
--- a/Engine/Source/Runtime/Core/Private/Tests/Misc/PathsTest.cpp
+++ b/Engine/Source/Runtime/Core/Private/Tests/Misc/PathsTest.cpp
@@ -1,19 +1,14 @@
 // Copyright 1998-2014 Epic Games, Inc. All Rights Reserved.
 
-/*=============================================================================
-	PathsTest.cpp: Unit test for the FPaths class.
-=============================================================================*/
-
 #include "CorePrivatePCH.h"
 #include "AutomationTest.h"
 
 
-IMPLEMENT_SIMPLE_AUTOMATION_TEST(FPathTests, "Core.Misc.Paths", EAutomationTestFlags::ATF_SmokeTest)
-
-
 FString FPaths::GameProjectFilePath;
 
 
+IMPLEMENT_SIMPLE_AUTOMATION_TEST(FPathTests, "Core.Misc.Paths", EAutomationTestFlags::ATF_SmokeTest)
+
 bool FPathTests::RunTest( const FString& Parameters )
 {
 	struct FCollapseRelativeDirectoriesTest
diff --git a/Engine/Source/Runtime/Core/Private/Tests/Misc/QueueTest.cpp b/Engine/Source/Runtime/Core/Private/Tests/Misc/QueueTest.cpp
--- a/Engine/Source/Runtime/Core/Private/Tests/Misc/QueueTest.cpp
+++ b/Engine/Source/Runtime/Core/Private/Tests/Misc/QueueTest.cpp
@@ -1,16 +1,11 @@
 // Copyright 1998-2014 Epic Games, Inc. All Rights Reserved.
 
-/*=============================================================================
-	QueueTest.cpp: Unit test for the TQueue template class.
-=============================================================================*/
-
 #include "CorePrivatePCH.h"
 #include "AutomationTest.h"
 
 
 IMPLEMENT_SIMPLE_AUTOMATION_TEST(FQueueTest, "Core.Misc.Queue", EAutomationTestFlags::ATF_SmokeTest)
 
-
 bool FQueueTest::RunTest( const FString& Parameters )
 {
 	// empty queues
@@ -44,4 +39,4 @@ bool FQueueTest::RunTest( const FString& Parameters )
 	TestTrue(TEXT("After removing all items, the queue must be empty"), q2_1.IsEmpty());
 
 	return true;
-}
\ No newline at end of file
+}
diff --git a/Engine/Source/Runtime/Core/Private/Tests/Misc/TimespanTest.cpp b/Engine/Source/Runtime/Core/Private/Tests/Misc/TimespanTest.cpp
--- a/Engine/Source/Runtime/Core/Private/Tests/Misc/TimespanTest.cpp
+++ b/Engine/Source/Runtime/Core/Private/Tests/Misc/TimespanTest.cpp
@@ -1,16 +1,11 @@
 // Copyright 1998-2014 Epic Games, Inc. All Rights Reserved.
 
-/*=============================================================================
-	TimespanTest.cpp: Unit test for the FTimespan class.
-=============================================================================*/
-
 #include "CorePrivatePCH.h"
 #include "AutomationTest.h"
 
 
 IMPLEMENT_SIMPLE_AUTOMATION_TEST(FTimespanTest, "Core.Misc.Timespan", EAutomationTestFlags::ATF_SmokeTest)
 
-
 bool FTimespanTest::RunTest( const FString& Parameters )
 {
 	// constructors must create equal objects
diff --git a/Engine/Source/Runtime/Core/Private/UObject/UnrealNames.cpp b/Engine/Source/Runtime/Core/Private/UObject/UnrealNames.cpp
--- a/Engine/Source/Runtime/Core/Private/UObject/UnrealNames.cpp
+++ b/Engine/Source/Runtime/Core/Private/UObject/UnrealNames.cpp
@@ -1,9 +1,5 @@
 // Copyright 1998-2014 Epic Games, Inc. All Rights Reserved.
 
-/*=============================================================================
-	UnName.cpp: Unreal global name code.
-=============================================================================*/
-
 #include "CorePrivatePCH.h"
 
 
