#ue4 - Add version of LerpStable() that works with floats, not just doubles as before.

diff --git a/Engine/Source/Runtime/Core/Public/Math/UnrealMathUtility.h b/Engine/Source/Runtime/Core/Public/Math/UnrealMathUtility.h
--- a/Engine/Source/Runtime/Core/Public/Math/UnrealMathUtility.h
+++ b/Engine/Source/Runtime/Core/Public/Math/UnrealMathUtility.h
@@ -481,6 +481,13 @@ struct FMath : public FPlatformMath
 		return (T)((A * (1.0 - Alpha)) + (B * Alpha));
 	}
 
+	/** Performs a linear interpolation between two values, Alpha ranges from 0-1. Handles full numeric range of T */
+	template< class T >
+	static FORCEINLINE_DEBUGGABLE T LerpStable(const T& A, const T& B, float Alpha)
+	{
+		return (T)((A * (1.0f - Alpha)) + (B * Alpha));
+	}
+
 	/** Performs a 2D linear interpolation between four values values, FracX, FracY ranges from 0-1 */
 	template< class T, class U > 
 	static FORCEINLINE_DEBUGGABLE T BiLerp(const T& P00,const T& P10,const T& P01,const T& P11, const U& FracX, const U& FracY)
