Core: Added better documentation for DEPRECATED macro

diff --git a/Engine/Source/Runtime/Core/Public/GenericPlatform/GenericPlatformCompilerPreSetup.h b/Engine/Source/Runtime/Core/Public/GenericPlatform/GenericPlatformCompilerPreSetup.h
--- a/Engine/Source/Runtime/Core/Public/GenericPlatform/GenericPlatformCompilerPreSetup.h
+++ b/Engine/Source/Runtime/Core/Public/GenericPlatform/GenericPlatformCompilerPreSetup.h
@@ -1,22 +1,31 @@
 // Copyright 1998-2014 Epic Games, Inc. All Rights Reserved.
 
-/*=============================================================================================
-	GenericPlatformCompilerPreSetup.h: pragmas, version checks and other things for generic compiler
-==============================================================================================*/
-
 #pragma once
 
+
 #ifndef DEPRECATED
-	// Sample usage (note slightly unintuitive syntax for classes and structs) :
-	// 
-	// DEPRECATED(4.6, "Message")
-	// void Function();
-	// 
-	// struct DEPRECATED(4.6, "Message") MyStruct;
-	// class DEPRECATED(4.6, "Message") MyClass;
+	/**
+	 * Macro for marking up deprecated code, functions and types.
+	 *
+	 * Features that are marked as deprecated are scheduled to be removed from the code base
+	 * in a future release. If you are using a deprecated feature in your code, you should
+	 * replace it before upgrading to the next release. See the Upgrade Notes in the release
+	 * notes for the release in which the feature was marked deprecated.
+	 *
+	 * Sample usage (note the slightly different syntax for classes and structures):
+	 *
+	 *		DEPRECATED(4.6, "Message")
+	 *		void Function();
+	 *
+	 *		struct DEPRECATED(4.6, "Message") MyStruct;
+	 *		class DEPRECATED(4.6, "Message") MyClass;
+	 *
+	 * @param VERSION The release number in which the feature was marked deprecated.
+	 * @param MESSAGE A message text containing additional upgrade notes.
+	 */
 	#define DEPRECATED(VERSION, MESSAGE)
 #endif // DEPRECATED
 
 #ifndef EMIT_DEPRECATED_WARNING_MESSAGE
 	#define EMIT_DEPRECATED_WARNING_MESSAGE(Msg)
-#endif // EMIT_DEPRECATED_WARNING_MESSAGE
\ No newline at end of file
+#endif // EMIT_DEPRECATED_WARNING_MESSAGE
diff --git a/Engine/Source/Runtime/Core/Public/Mac/MacPlatformCompilerPreSetup.h b/Engine/Source/Runtime/Core/Public/Mac/MacPlatformCompilerPreSetup.h
--- a/Engine/Source/Runtime/Core/Public/Mac/MacPlatformCompilerPreSetup.h
+++ b/Engine/Source/Runtime/Core/Public/Mac/MacPlatformCompilerPreSetup.h
@@ -1,20 +1,29 @@
 // Copyright 1998-2014 Epic Games, Inc. All Rights Reserved.
 
-/*==========================================================================================================
-	MacPlatformCompilerPreSetup.h: setup for the compiler that needs to be defined before everything else
-===========================================================================================================*/
-
 #pragma once
 
+
 #ifndef DISABLE_DEPRECATION
 	#pragma clang diagnostic warning "-Wdeprecated-declarations"
 
-	// Sample usage (note slightly unintuitive syntax for classes and structs) :
-	// 
-	// DEPRECATED(4.6, "Message")
-	// void Function();
-	// 
-	// struct DEPRECATED(4.6, "Message") MyStruct;
-	// class DEPRECATED(4.6, "Message") MyClass;
+	/**
+	 * Macro for marking up deprecated code, functions and types.
+	 *
+	 * Features that are marked as deprecated are scheduled to be removed from the code base
+	 * in a future release. If you are using a deprecated feature in your code, you should
+	 * replace it before upgrading to the next release. See the Upgrade Notes in the release
+	 * notes for the release in which the feature was marked deprecated.
+	 *
+	 * Sample usage (note the slightly different syntax for classes and structures):
+	 *
+	 *		DEPRECATED(4.6, "Message")
+	 *		void Function();
+	 *
+	 *		struct DEPRECATED(4.6, "Message") MyStruct;
+	 *		class DEPRECATED(4.6, "Message") MyClass;
+	 *
+	 * @param VERSION The release number in which the feature was marked deprecated.
+	 * @param MESSAGE A message containing upgrade notes.
+	 */
 	#define DEPRECATED(VERSION, MESSAGE) __attribute__((deprecated(MESSAGE " Please update your code to the new API before upgrading to the next release, otherwise your project will no longer compile.")))
 #endif // DISABLE_DEPRECATION
diff --git a/Engine/Source/Runtime/Core/Public/Windows/WindowsPlatformCompilerPreSetup.h b/Engine/Source/Runtime/Core/Public/Windows/WindowsPlatformCompilerPreSetup.h
--- a/Engine/Source/Runtime/Core/Public/Windows/WindowsPlatformCompilerPreSetup.h
+++ b/Engine/Source/Runtime/Core/Public/Windows/WindowsPlatformCompilerPreSetup.h
@@ -1,19 +1,28 @@
 // Copyright 1998-2014 Epic Games, Inc. All Rights Reserved.
 
-/*==========================================================================================================
-	WindowsPlatformCompilerPreSetup.h: setup for the compiler that needs to be defined before everything else
-===========================================================================================================*/
-
 #pragma once
 
+
 #ifndef DISABLE_DEPRECATION
-	// Sample usage (note slightly unintuitive syntax for classes and structs) :
-	// 
-	// DEPRECATED(4.6, "Message")
-	// void Function();
-	// 
-	// struct DEPRECATED(4.6, "Message") MyStruct;
-	// class DEPRECATED(4.6, "Message") MyClass;
+	/**
+	 * Macro for marking up deprecated code, functions and types.
+	 *
+	 * Features that are marked as deprecated are scheduled to be removed from the code base
+	 * in a future release. If you are using a deprecated feature in your code, you should
+	 * replace it before upgrading to the next release. See the Upgrade Notes in the release
+	 * notes for the release in which the feature was marked deprecated.
+	 *
+	 * Sample usage (note the slightly different syntax for classes and structures):
+	 *
+	 *		DEPRECATED(4.6, "Message")
+	 *		void Function();
+	 *
+	 *		struct DEPRECATED(4.6, "Message") MyStruct;
+	 *		class DEPRECATED(4.6, "Message") MyClass;
+	 *
+	 * @param VERSION The release number in which the feature was marked deprecated.
+	 * @param MESSAGE A message containing upgrade notes.
+	 */
 	#define DEPRECATED(VERSION, MESSAGE) __declspec(deprecated(MESSAGE " Please update your code to the new API before upgrading to the next release, otherwise your project will no longer compile."))
 
 	#ifndef EMIT_DEPRECATED_WARNING_MESSAGE_STR
