Stats - Added better mechanism for flushing the stats messages, should reduce the amount of the stat packets, will also reduce the memory usage, improve performance, will
 happen only once per frame

diff --git a/Engine/Source/Developer/Profiler/Private/ProfilerSession.cpp b/Engine/Source/Developer/Profiler/Private/ProfilerSession.cpp
--- a/Engine/Source/Developer/Profiler/Private/ProfilerSession.cpp
+++ b/Engine/Source/Developer/Profiler/Private/ProfilerSession.cpp
@@ -584,7 +584,7 @@ void FRawProfilerSession::PrepareLoading()
 			for( int32 PacketIndex = 0; PacketIndex < StatPacketArray.Packets.Num(); PacketIndex++ )
 			{
 				const FStatPacket& StatPacket = *StatPacketArray.Packets[PacketIndex];
-				const FStatPacket::TStatMessagesArray& Data = StatPacket.StatMessages;
+				const FStatMessagesArray& Data = StatPacket.StatMessages;
 
 				for( int32 Index = 0; Index < Data.Num(); Index++ )
 				{
@@ -919,7 +919,7 @@ void FRawProfilerSession::ProcessStatPacketArray( const FStatPacketArray& StatPa
 		Stack.Add( ThreadNode );
 		FProfilerStackNode* Current = Stack.Last();
 
-		const FStatPacket::TStatMessagesArray& Data = StatPacket.StatMessages;
+		const FStatMessagesArray& Data = StatPacket.StatMessages;
 		for( int32 Index = 0; Index < Data.Num(); Index++ )
 		{
 			const FStatMessage& Item = Data[Index];
diff --git a/Engine/Source/Runtime/Core/Private/Stats/Stats2.cpp b/Engine/Source/Runtime/Core/Private/Stats/Stats2.cpp
--- a/Engine/Source/Runtime/Core/Private/Stats/Stats2.cpp
+++ b/Engine/Source/Runtime/Core/Private/Stats/Stats2.cpp
@@ -24,20 +24,22 @@ static struct FForceInitAtBootFStats2 : public TForceInitAtBoot<FStats2Globals>
 {} FForceInitAtBootFStats2;
 
 DECLARE_FLOAT_COUNTER_STAT( TEXT( "Seconds Per Cycle" ), STAT_SecondsPerCycle, STATGROUP_Engine );
+DEFINE_STAT(STAT_StatMessagesMemory);
 
 /*-----------------------------------------------------------------------------
 	FStats2
 -----------------------------------------------------------------------------*/
 
+int32 FStats::GameThreadStatsFrame = 1;
+
 void FStats::AdvanceFrame( bool bDiscardCallstack, const FOnAdvanceRenderingThreadStats& AdvanceRenderingThreadStatsDelegate /*= FOnAdvanceRenderingThreadStats()*/ )
 {
 #if STATS
 	check( IsInGameThread() );
 	static int32 MasterDisableChangeTagStartFrame = -1;
-	static int64 StatsFrame = 1;
-	StatsFrame++;
+	FPlatformAtomics::InterlockedIncrement(&GameThreadStatsFrame);
 
-	int64 Frame = StatsFrame;
+	int64 Frame = GameThreadStatsFrame;
 	if( bDiscardCallstack )
 	{
 		FThreadStats::FrameDataIsIncomplete(); // we won't collect call stack stats this frame
@@ -48,7 +50,7 @@ void FStats::AdvanceFrame( bool bDiscardCallstack, const FOnAdvanceRenderingThre
 	}
 	if( !FThreadStats::IsCollectingData() || MasterDisableChangeTagStartFrame != FThreadStats::MasterDisableChangeTag() )
 	{
-		Frame = -StatsFrame; // mark this as a bad frame
+		Frame = -GameThreadStatsFrame; // mark this as a bad frame
 	}
 	static FStatNameAndInfo Adv( NAME_AdvanceFrame, "", "", TEXT( "" ), EStatDataType::ST_int64, true, false );
 	FThreadStats::AddMessage( Adv.GetEncodedName(), EStatOperation::AdvanceFrameEventGameThread, Frame ); // we need to flush here if we aren't collecting stats to make sure the meta data is up to date
@@ -59,7 +61,7 @@ void FStats::AdvanceFrame( bool bDiscardCallstack, const FOnAdvanceRenderingThre
 
 	if( AdvanceRenderingThreadStatsDelegate.IsBound() )
 	{
-		AdvanceRenderingThreadStatsDelegate.Execute( bDiscardCallstack, StatsFrame, MasterDisableChangeTagStartFrame );
+		AdvanceRenderingThreadStatsDelegate.Execute( bDiscardCallstack, GameThreadStatsFrame, MasterDisableChangeTagStartFrame );
 	}
 	else
 	{
@@ -88,8 +90,6 @@ set a DEBUG_STATS define that allows this all to be debugged. Otherwise, inline
 
 It should be possible to load stats data without STATS
 
-Put version and FPlatformTime::GetSecondsPerCycle() in the header
-
 //@todo this probably goes away after we redo the programmer interface
 
 We are only saving condensed frames. For the "core view", we want a single frame capture of FStatsThreadState::History. That is the whole enchillada including start and end time for every scope.
@@ -618,7 +618,6 @@ FName FStatNameAndInfo::GetGroupCategoryFrom(FName InLongName)
 	return NAME_None;
 }
 
-uint32 StatsThreadId = 0;
 static TAutoConsoleVariable<int32> CVarDumpStatPackets(	TEXT("DumpStatPackets"),0,	TEXT("If true, dump stat packets."));
 
 /** The rendering thread runnable object. */
@@ -648,7 +647,6 @@ public:
 
 	virtual uint32 Run()
 	{
-		StatsThreadId = FPlatformTLS::GetCurrentThreadId();
 		FTaskGraphInterface::Get().AttachToThread(ENamedThreads::StatsThread);
 		FTaskGraphInterface::Get().ProcessThreadUntilRequestReturn(ENamedThreads::StatsThread);
 		return 0;
@@ -710,6 +708,7 @@ public:
 
 		bReadyToProcess = Packet->ThreadType != EThreadType::Other;
 		IncomingData.Packets.Add(Packet);
+		INC_MEMORY_STAT_BY(STAT_StatMessagesMemory,Packet->StatMessages.GetAllocatedSize());
 		Tick();
 	}
 
@@ -747,10 +746,11 @@ public:
 	}
 };
 
-FThreadStats::FThreadStats()
-	: ScopeCount(0)
-	, bSawExplicitFlush(false)
-	, bWaitForExplicitFlush(0)
+FThreadStats::FThreadStats():
+	CurrentGameFrame(FStatsThreadState::GetLocalState().CurrentGameFrame),
+	ScopeCount(0), 
+	bWaitForExplicitFlush(0),
+	bSawExplicitFlush(false)
 {
 	Packet.ThreadId = FPlatformTLS::GetCurrentThreadId();
 	if (Packet.ThreadId == GGameThreadId)
@@ -769,16 +769,6 @@ FThreadStats::FThreadStats()
 	FPlatformTLS::SetTlsValue(TlsSlot, this);
 }
 
-// copy for sending to the stats thread
-FThreadStats::FThreadStats(FThreadStats const& Other)
-	: ScopeCount(0)
-	, bSawExplicitFlush(false)
-	, bWaitForExplicitFlush(0)
-{
-	Packet.ThreadId = Other.Packet.ThreadId;
-	Packet.ThreadType = Other.Packet.ThreadType;
-}
-
 void FThreadStats::CheckEnable()
 {
 	bool bOldMasterEnable(bMasterEnable);
@@ -794,44 +784,59 @@ void FThreadStats::CheckEnable()
 
 void FThreadStats::Flush(bool bHasBrokenCallstacks)
 {
-	const int32 PresizeMaxNumEntries = 10;
-	const int32 PresizeMaxSize = 256 * 1024;
-
 	if (bMasterDisableForever)
 	{
 		Packet.StatMessages.Empty();
 		return;
 	}
+
+	if( Packet.ThreadType == EThreadType::Other )
+	{
+		FPlatformMisc::MemoryBarrier();
+		const bool bFrameHasChanged = FStats::GameThreadStatsFrame > CurrentGameFrame;
+		if( bFrameHasChanged )
+		{
+			CurrentGameFrame = FStats::GameThreadStatsFrame;
+		}
+		else
+		{
+			return;
+		}
+	}
+
 	if (!ScopeCount && Packet.StatMessages.Num())
 	{
-		if( Packet.StatMessagesPresize.Num() >= PresizeMaxNumEntries )
+		if( Packet.StatMessagesPresize.Num() >= PRESIZE_MAX_NUM_ENTRIES )
 		{
 			Packet.StatMessagesPresize.RemoveAt(0);
 		}
-		if (Packet.StatMessages.Num() < PresizeMaxSize)
+		if (Packet.StatMessages.Num() < PRESIZE_MAX_SIZE)
 		{
 			Packet.StatMessagesPresize.Add(Packet.StatMessages.Num());
 		}
 		else
 		{
-			UE_LOG( LogStats, Verbose, TEXT( "StatMessage Packet has more than %i messages.  Ignoring for the presize history." ), (int32)PresizeMaxSize );
+			UE_LOG( LogStats, Verbose, TEXT( "StatMessage Packet has more than %i messages.  Ignoring for the presize history." ), (int32)PRESIZE_MAX_SIZE );
 		}
 		FStatPacket* ToSend = new FStatPacket(Packet);
 		Exchange(ToSend->StatMessages, Packet.StatMessages);
 		ToSend->bBrokenCallstacks = bHasBrokenCallstacks;
+
 		check(!Packet.StatMessages.Num());
-		int32 MaxPresize = Packet.StatMessagesPresize[0];
-		for (int32 Index = 0; Index < Packet.StatMessagesPresize.Num(); ++Index)
+		if( Packet.StatMessagesPresize.Num() > 0 )
 		{
-			if (MaxPresize < Packet.StatMessagesPresize[Index])
+			int32 MaxPresize = Packet.StatMessagesPresize[0];
+			for (int32 Index = 0; Index < Packet.StatMessagesPresize.Num(); ++Index)
 			{
-				MaxPresize = Packet.StatMessagesPresize[Index];
+				if (MaxPresize < Packet.StatMessagesPresize[Index])
+				{
+					MaxPresize = Packet.StatMessagesPresize[Index];
+				}
 			}
+			Packet.StatMessages.Empty(MaxPresize);
 		}
-		Packet.StatMessages.Empty(MaxPresize);
-		{
-			TGraphTask<FStatMessagesTask>::CreateTask().ConstructAndDispatchWhenReady(ToSend);
-		}
+
+		TGraphTask<FStatMessagesTask>::CreateTask().ConstructAndDispatchWhenReady(ToSend);
 		if (Packet.ThreadType != EThreadType::Other && bSawExplicitFlush)
 		{
 			bWaitForExplicitFlush = 1;
diff --git a/Engine/Source/Runtime/Core/Private/Stats/StatsData.cpp b/Engine/Source/Runtime/Core/Private/Stats/StatsData.cpp
--- a/Engine/Source/Runtime/Core/Private/Stats/StatsData.cpp
+++ b/Engine/Source/Runtime/Core/Private/Stats/StatsData.cpp
@@ -455,6 +455,16 @@ void FComplexRawStatStackNode::CopyExclusivesFromSelf()
 	FStatsThreadState
 -----------------------------------------------------------------------------*/
 
+void FStatPacketArray::Empty()
+{
+	for (int32 Index = 0; Index < Packets.Num(); Index++)
+	{
+		DEC_MEMORY_STAT_BY(STAT_StatMessagesMemory,Packets[Index]->StatMessages.GetAllocatedSize());
+		delete Packets[Index];
+	}
+	Packets.Empty();
+}
+
 FStatsThreadState::FStatsThreadState(int32 InHistoryFrames)
 	: HistoryFrames(InHistoryFrames)
 	, MaxFrameSeen(0)
@@ -462,6 +472,7 @@ FStatsThreadState::FStatsThreadState(int32 InHistoryFrames)
 	, LastFullFrameMetaAndNonFrame(-1)
 	, LastFullFrameProcessed(-1)
 	, bWasLoaded(false)
+	, bFindMemoryExtensiveStats(false)
 	, CurrentGameFrame(1)
 	, CurrentRenderFrame(1)
 {
@@ -523,9 +534,9 @@ int64 FStatsThreadState::GetLatestValidFrame() const
 	return Result;
 }
 
-void FStatsThreadState::ScanForAdvance(FStatPacket::TStatMessagesArray const& Data)
+void FStatsThreadState::ScanForAdvance(const FStatMessagesArray& Data)
 {
-	for (int32 Index = Data.Num() - 1; Index >= 0 ; Index--)
+	for (int32 Index = 0; Index < Data.Num(); Index++)
 	{
 		FStatMessage const& Item = Data[Index];
 		EStatOperation::Type Op = Item.NameAndInfo.GetField<EStatOperation>();
@@ -548,7 +559,8 @@ void FStatsThreadState::ScanForAdvance(FStatPacket::TStatMessagesArray const& Da
 					BadFrames.Add(Frame - 1);
 				}
 			}
-			CurrentGameFrame = NewGameFrame;
+			// CurrentGameFrame can be accessed by other threads.
+			FPlatformAtomics::InterlockedExchange( &CurrentGameFrame, NewGameFrame );
 		}
 		else if (Op == EStatOperation::AdvanceFrameEventRenderThread)
 		{
@@ -587,7 +599,7 @@ void FStatsThreadState::ScanForAdvance(FStatPacketArray& NewData)
 	{
 		int64 FrameNum = NewData.Packets[Index]->ThreadType == EThreadType::Renderer ? CurrentRenderFrame : CurrentGameFrame;
 		NewData.Packets[Index]->Frame = FrameNum;
-		FStatPacket::TStatMessagesArray const& Data = NewData.Packets[Index]->StatMessages;
+		const FStatMessagesArray& Data = NewData.Packets[Index]->StatMessages;
 		ScanForAdvance(Data);
 		Count += Data.Num();
 	}
@@ -595,52 +607,6 @@ void FStatsThreadState::ScanForAdvance(FStatPacketArray& NewData)
 	INC_DWORD_STAT_BY(STAT_StatFrameMessages, Count);
 }
 
-void FStatsThreadState::ProcessMetaDataForLoad(TArray<FStatMessage>& Data)
-{
-	check(bWasLoaded);
-	for (int32 Index = 0; Index < Data.Num() ; Index++)
-	{
-		FStatMessage& Item = Data[Index];
-		EStatOperation::Type Op = Item.NameAndInfo.GetField<EStatOperation>();
-		if (Op == EStatOperation::SetLongName)
-		{
-			FindOrAddMetaData(Item);
-		}
-		else if (Op == EStatOperation::AdvanceFrameEventGameThread)
-		{
-			check(Item.NameAndInfo.GetField<EStatDataType>() == EStatDataType::ST_int64);
-			if (Item.GetValue_int64() > 0)
-			{
-				CurrentGameFrame = Item.GetValue_int64();
-				if (CurrentGameFrame > MaxFrameSeen)
-				{
-					MaxFrameSeen = CurrentGameFrame;
-				}
-				if (MinFrameSeen < 0)
-				{
-					MinFrameSeen = CurrentGameFrame;
-				}
-			}
-		}
-		else if (Op == EStatOperation::AdvanceFrameEventRenderThread)
-		{
-			check(Item.NameAndInfo.GetField<EStatDataType>() == EStatDataType::ST_int64);
-			if (Item.GetValue_int64() > 0)
-			{
-				CurrentRenderFrame = Item.GetValue_int64();
-				if (CurrentGameFrame > MaxFrameSeen)
-				{
-					MaxFrameSeen = CurrentGameFrame;
-				}
-				if (MinFrameSeen < 0)
-				{
-					MinFrameSeen = CurrentGameFrame;
-				}
-			}
-		}
-	}
-}
-
 void FStatsThreadState::ProcessMetaDataOnly(TArray<FStatMessage>& Data)
 {
 	for (int32 Index = 0; Index < Data.Num() ; Index++)
@@ -652,7 +618,14 @@ void FStatsThreadState::ProcessMetaDataOnly(TArray<FStatMessage>& Data)
 	}
 }
 
-void FStatsThreadState::ProcessNonFrameStats(FStatPacket::TStatMessagesArray& Data, TSet<FName>* NonFrameStatsFound)
+
+void FStatsThreadState::ToggleFindMemoryExtensiveStats()
+{
+	bFindMemoryExtensiveStats = !bFindMemoryExtensiveStats;
+	UE_LOG(LogStats, Log, TEXT("bFindMemoryExtensiveStats is %s now"), bFindMemoryExtensiveStats?TEXT("enabled"):TEXT("disabled"));
+}
+
+void FStatsThreadState::ProcessNonFrameStats(FStatMessagesArray& Data, TSet<FName>* NonFrameStatsFound)
 {
 	check(!bWasLoaded);
 	for (int32 Index = 0; Index < Data.Num() ; Index++)
@@ -744,6 +717,11 @@ void FStatsThreadState::AddToHistoryAndEmpty(FStatPacketArray& NewData)
 
 			FStatPacket const* PacketToCopyForNonFrame = NULL;
 
+			if( bFindMemoryExtensiveStats )
+			{
+				FindAndDumpMemoryExtensiveStats(Frame);
+			}
+
 			TSet<FName> NonFrameStatsFound;
 			for (int32 PacketIndex = 0; PacketIndex < Frame.Packets.Num(); PacketIndex++)
 			{
@@ -760,7 +738,7 @@ void FStatsThreadState::AddToHistoryAndEmpty(FStatPacketArray& NewData)
 
 				check(PacketToCopyForNonFrame);
 
-				FStatPacket::TStatMessagesArray* NonFrameMessages = NULL;
+				FStatMessagesArray* NonFrameMessages = NULL;
 
 				for (TMap<FName, FStatMessage>::TConstIterator It(NotClearedEveryFrame); It; ++It)
 				{
@@ -768,12 +746,15 @@ void FStatsThreadState::AddToHistoryAndEmpty(FStatPacketArray& NewData)
 					{
 						if (!NonFrameMessages)
 						{
-							int32 NonFrameIndex = Frame.Packets.Add(new FStatPacket(*PacketToCopyForNonFrame));
+							const int32 NonFrameIndex = Frame.Packets.Add(new FStatPacket(*PacketToCopyForNonFrame));
 							NonFrameMessages = &Frame.Packets[NonFrameIndex]->StatMessages;
 						}
 						new (*NonFrameMessages) FStatMessage(It.Value());
 					}
 				}
+
+				INC_MEMORY_STAT_BY(STAT_StatMessagesMemory,NonFrameMessages->GetAllocatedSize());
+
 				GoodFrames.Add(FrameNum);
 			}
 			LastFullFrameMetaAndNonFrame = FrameNum;
@@ -1018,7 +999,7 @@ void FStatsThreadState::GetRawStackStats(int64 TargetFrame, FRawStatStackNode& R
 			Stack.Add(ThreadRoot);
 			FRawStatStackNode* Current = Stack.Last();
 
-			FStatPacket::TStatMessagesArray const& Data = Packet.StatMessages;
+			const FStatMessagesArray& Data = Packet.StatMessages;
 			for (int32 Index = 0; Index < Data.Num(); Index++)
 			{
 				FStatMessage const& Item = Data[Index];
@@ -1154,7 +1135,7 @@ int64 FStatsThreadState::GetFastThreadFrameTimeInternal( int64 TargetFrame, int3
 
 		if( Packet.ThreadId == ThreadID || Packet.ThreadType == Thread )
 		{
-			FStatPacket::TStatMessagesArray const& Data = Packet.StatMessages;
+			const FStatMessagesArray& Data = Packet.StatMessages;
 			for (int32 Index = 0; Index < Data.Num(); Index++)
 			{
 				FStatMessage const& Item = Data[Index];
@@ -1321,6 +1302,43 @@ void FStatsThreadState::AddMissingStats(TArray<FStatMessage>& Dest, TSet<FName>
 	}
 }
 
+void FStatsThreadState::FindAndDumpMemoryExtensiveStats( FStatPacketArray &Frame )
+{
+	int32 TotalMessages = 0;
+	TMap<FName,int32> NameToCount;
+
+	// Generate some data statistics.
+	for( const FStatPacket* StatPacket : Frame.Packets )
+	{
+		for( int32 MessageIndex = 0; MessageIndex < StatPacket->StatMessages.Num(); MessageIndex++)
+		{
+			const FName ShortName = StatPacket->StatMessages[MessageIndex].NameAndInfo.GetShortName();
+			NameToCount.FindOrAdd(ShortName) += 1;
+			TotalMessages++;
+		}
+	}
+
+	// Dump stats to the log.
+	NameToCount.ValueSort( TGreater<uint32>() );
+
+	const float MaxPctDisplayed = 0.9f;
+	int32 CurrentIndex = 0;
+	int32 DisplayedSoFar = 0;
+	UE_LOG( LogStats, Warning, TEXT("%2s, %32s, %5s"), TEXT("No"), TEXT("Name"), TEXT("Count") );
+	for( const auto& It : NameToCount )
+	{
+		UE_LOG( LogStats, Warning, TEXT("%2i, %32s, %5i"), CurrentIndex, *It.Key.ToString(), It.Value );
+		CurrentIndex++;
+		DisplayedSoFar += It.Value;
+
+		const float CurrentPct = (float)DisplayedSoFar/(float)TotalMessages;
+		if( CurrentPct > MaxPctDisplayed )
+		{
+			break;
+		}
+	}
+}
+
 FString FStatsUtils::DebugPrint(FStatMessage const& Item)
 {
 	FString Result(TEXT("Invalid"));
@@ -1731,4 +1749,4 @@ CORE_API void ClearPendingStatGroups()
 	FStatGroupGameThreadNotifier::Get().ClearData();
 }
 
-#endif
\ No newline at end of file
+#endif
diff --git a/Engine/Source/Runtime/Core/Private/Stats/StatsFile.cpp b/Engine/Source/Runtime/Core/Private/Stats/StatsFile.cpp
--- a/Engine/Source/Runtime/Core/Private/Stats/StatsFile.cpp
+++ b/Engine/Source/Runtime/Core/Private/Stats/StatsFile.cpp
@@ -84,6 +84,7 @@ FStatsThreadState::FStatsThreadState( FString const& Filename )
 	, LastFullFrameMetaAndNonFrame( -1 )
 	, LastFullFrameProcessed( -1 )
 	, bWasLoaded( true )
+	, bFindMemoryExtensiveStats( false )
 	, CurrentGameFrame( -1 )
 	, CurrentRenderFrame( -1 )
 {
@@ -222,6 +223,52 @@ void FStatsThreadState::AddMessages( TArray<FStatMessage>& InMessages )
 	bWasLoaded = false;
 }
 
+void FStatsThreadState::ProcessMetaDataForLoad(TArray<FStatMessage>& Data)
+{
+	check(bWasLoaded);
+	for (int32 Index = 0; Index < Data.Num() ; Index++)
+	{
+		FStatMessage& Item = Data[Index];
+		EStatOperation::Type Op = Item.NameAndInfo.GetField<EStatOperation>();
+		if (Op == EStatOperation::SetLongName)
+		{
+			FindOrAddMetaData(Item);
+		}
+		else if (Op == EStatOperation::AdvanceFrameEventGameThread)
+		{
+			check(Item.NameAndInfo.GetField<EStatDataType>() == EStatDataType::ST_int64);
+			if (Item.GetValue_int64() > 0)
+			{
+				CurrentGameFrame = Item.GetValue_int64();
+				if (CurrentGameFrame > MaxFrameSeen)
+				{
+					MaxFrameSeen = CurrentGameFrame;
+				}
+				if (MinFrameSeen < 0)
+				{
+					MinFrameSeen = CurrentGameFrame;
+				}
+			}
+		}
+		else if (Op == EStatOperation::AdvanceFrameEventRenderThread)
+		{
+			check(Item.NameAndInfo.GetField<EStatDataType>() == EStatDataType::ST_int64);
+			if (Item.GetValue_int64() > 0)
+			{
+				CurrentRenderFrame = Item.GetValue_int64();
+				if (CurrentGameFrame > MaxFrameSeen)
+				{
+					MaxFrameSeen = CurrentGameFrame;
+				}
+				if (MinFrameSeen < 0)
+				{
+					MinFrameSeen = CurrentGameFrame;
+				}
+			}
+		}
+	}
+}
+
 /*-----------------------------------------------------------------------------
 	FStatsWriteFile
 -----------------------------------------------------------------------------*/
diff --git a/Engine/Source/Runtime/Core/Public/Stats/Stats2.h b/Engine/Source/Runtime/Core/Public/Stats/Stats2.h
--- a/Engine/Source/Runtime/Core/Public/Stats/Stats2.h
+++ b/Engine/Source/Runtime/Core/Public/Stats/Stats2.h
@@ -55,6 +55,9 @@ struct CORE_API FStats
 
 	/** Advances stats for the current frame. */
 	static void AdvanceFrame( bool bDiscardCallstack, const FOnAdvanceRenderingThreadStats& AdvanceRenderingThreadStatsDelegate = FOnAdvanceRenderingThreadStats() );
+
+	/** Current game thread stats frame. */
+	static int32 GameThreadStatsFrame;
 };
 
 #if STATS
@@ -949,19 +952,18 @@ typedef TStatMessage<EComplexStatField> FComplexStatMessage;
 
 template<> struct TIsPODType<FComplexStatMessage> { enum { Value = true }; };
 
+
+enum class EStatMessagesArrayConstants
+{
+	MESSAGES_CHUNK_SIZE = 64*1024,
+};
+typedef TChunkedArray<FStatMessage,(uint32)EStatMessagesArrayConstants::MESSAGES_CHUNK_SIZE> FStatMessagesArray;
+
 /**
 * A stats packet. Sent between threads. Includes and array of messages and some information about the thread. 
 */
 struct FStatPacket
 {
-	enum 
-	{
-		MESSAGES_CHUNK_SIZE = 1024*1024,
-	};
-
-	//typedef TChunkedArray<FStatMessage,MESSAGES_CHUNK_SIZE> TStatMessagesArray;
-	typedef TArray<FStatMessage> TStatMessagesArray;
-
 	/** Assigned later, this is the frame number this packet is for **/
 	int64 Frame;
 	/** ThreadId this packet came from **/
@@ -971,8 +973,8 @@ struct FStatPacket
 	/** true if this packet has broken callstacks **/
 	bool bBrokenCallstacks;
 	/** messages in this packet **/
-	TStatMessagesArray StatMessages;
-	/** Size we presize the message buffer to, currently the max of what we have seen so far. **/
+	FStatMessagesArray StatMessages;
+	/** Size we presize the message buffer to, currently the max of what we have seen for the last PRESIZE_MAX_NUM_ENTRIES. **/
 	TArray<int32> StatMessagesPresize;
 
 	/** constructor **/
@@ -1000,6 +1002,12 @@ struct FStatPacket
 */
 class FThreadStats
 {
+	enum
+	{
+		PRESIZE_MAX_NUM_ENTRIES = 10,
+		PRESIZE_MAX_SIZE = EStatMessagesArrayConstants::MESSAGES_CHUNK_SIZE,
+	};
+
 	/** Used to control when we are collecting stats. User of the stats system increment and decrement this counter as they need data. **/
 	CORE_API static FThreadSafeCounter MasterEnableCounter;
 	/** Every time bMasterEnable changes, we update this. This is used to determine frames that have complete data. **/
@@ -1015,12 +1023,18 @@ class FThreadStats
 
 	friend class FStatsThread;
 
+	/** Current game frame for this thread stat. */
+	int32 CurrentGameFrame;
+
 	/** Tracks current stack depth for cycle counters. **/
 	uint32 ScopeCount;
-	/** Tracks current stack depth for cycle counters. **/
-	bool bSawExplicitFlush;
+
 	/** Tracks current stack depth for cycle counters. **/
 	uint32 bWaitForExplicitFlush;
+
+	/** Tracks current stack depth for cycle counters. **/
+	bool bSawExplicitFlush;
+
 	/** The data we are eventually going to send to the stats thread. **/
 	FStatPacket Packet;
 
@@ -1058,17 +1072,6 @@ public:
 		}
 	}
 
-	/** Used to send messages that came in before the thread was ready. Adds em and flushes **/
-	static void AddMessages(TArray<FStatMessage> const& Messages)
-	{
-		if (WillEverCollectData())
-		{
-			FThreadStats* ThreadStats = GetThreadStats();
-			ThreadStats->Packet.StatMessages += Messages;
-			ThreadStats->Flush();
-		}
-	}
-
 	/** Clock operation. **/
 	static FORCEINLINE_STATS void AddMessage(FName InStatName, EStatOperation::Type InStatOperation)
 	{
@@ -1103,7 +1106,7 @@ public:
 		{
 			FThreadStats* ThreadStats = GetThreadStats();
 			new (ThreadStats->Packet.StatMessages) FStatMessage(InStatName, InStatOperation, Value, bIsCycle);
-			if (!ThreadStats->ScopeCount) // we can't guarantee other threads will ever be flushed, so we need to flush ever counter item!
+			if(!ThreadStats->ScopeCount)
 			{
 				ThreadStats->Flush();
 			}
@@ -1677,6 +1680,7 @@ DECLARE_STATS_GROUP(TEXT("Navigation"),STATGROUP_Navigation, STATCAT_Advanced);
 DECLARE_STATS_GROUP(TEXT("Net"),STATGROUP_Net, STATCAT_Advanced);
 DECLARE_STATS_GROUP(TEXT("Niagara"),STATGROUP_Niagara, STATCAT_Advanced);
 DECLARE_STATS_GROUP(TEXT("Object"),STATGROUP_Object, STATCAT_Advanced);
+DECLARE_STATS_GROUP_VERBOSE(TEXT("ObjectVerbose"),STATGROUP_ObjectVerbose, STATCAT_Advanced);
 DECLARE_STATS_GROUP(TEXT("OpenGL RHI"),STATGROUP_OpenGLRHI, STATCAT_Advanced);
 DECLARE_STATS_GROUP(TEXT("Pak File"),STATGROUP_PakFile, STATCAT_Advanced);
 DECLARE_STATS_GROUP(TEXT("Particle Mem"),STATGROUP_ParticleMem, STATCAT_Advanced);
@@ -1714,5 +1718,6 @@ DECLARE_FLOAT_COUNTER_STAT_EXTERN(TEXT("StatUnit FPS"), STAT_FPS, STATGROUP_Engi
 /** Stats for the stat system */
 
 DECLARE_CYCLE_STAT_EXTERN(TEXT("DrawStats"),STAT_DrawStats,STATGROUP_StatSystem, CORE_API);
+DECLARE_MEMORY_STAT_EXTERN(TEXT("Stat messages memory"),STAT_StatMessagesMemory,STATGROUP_StatSystem, CORE_API);
 
 #endif
diff --git a/Engine/Source/Runtime/Core/Public/Stats/StatsData.h b/Engine/Source/Runtime/Core/Public/Stats/StatsData.h
--- a/Engine/Source/Runtime/Core/Public/Stats/StatsData.h
+++ b/Engine/Source/Runtime/Core/Public/Stats/StatsData.h
@@ -180,7 +180,7 @@ struct FStatCallCountComparer<FStatMessage>
 * An indirect array of stat packets.
 */
 //@todo can we just use TIndirectArray here?
-struct FStatPacketArray
+struct CORE_API FStatPacketArray
 {
 	TArray<FStatPacket*> Packets;
 
@@ -189,14 +189,7 @@ struct FStatPacketArray
 		Empty();
 	}
 
-	void Empty()
-	{
-		for (int32 Index = 0; Index < Packets.Num(); Index++)
-		{
-			delete Packets[Index];
-		}
-		Packets.Empty();
-	}
+	void Empty();
 };
 
 /**
@@ -417,7 +410,7 @@ class CORE_API FStatsThreadState
 	friend class FStatsThread;
 
 	/** Internal method to scan the messages to update the current frame **/
-	void ScanForAdvance(FStatPacket::TStatMessagesArray const& Data);
+	void ScanForAdvance(const FStatMessagesArray& Data);
 
 	/** Internal method to scan the messages to update the current frame and accumulate any non-frame stats. **/
 	void ScanForAdvance(FStatPacketArray& NewData);
@@ -435,13 +428,19 @@ public:
 		bWasLoaded =  true;
 	}
 
+	/** Toggles tracking the most memory expensive stats. */
+	void ToggleFindMemoryExtensiveStats();
+
 private:
 	/** Internal method to scan the messages to accumulate any non-frame stats. **/
-	void ProcessNonFrameStats(FStatPacket::TStatMessagesArray& Data, TSet<FName>* NonFrameStatsFound);
+	void ProcessNonFrameStats(FStatMessagesArray& Data, TSet<FName>* NonFrameStatsFound);
 
 	/** Internal method to place the data into the history, discard and broadcast any new frames to anyone who cares. **/
 	void AddToHistoryAndEmpty(FStatPacketArray& NewData);
 
+	/** Generates a list of most memory expensive stats and dump to the log. */
+	void FindAndDumpMemoryExtensiveStats( FStatPacketArray &Frame );
+
 	/** Called in response to SetLongName messages to update ShortNameToLongName and NotClearedEveryFrame **/
 	void FindOrAddMetaData(FStatMessage const& Item);
 
@@ -472,6 +471,9 @@ private:
 	/** Valid frame computation is different if we just loaded these stats **/
 	bool bWasLoaded;
 
+	/** If true, stats each frame will dump to the log a list of most memory expensive stats. */
+	bool bFindMemoryExtensiveStats;
+
 	/** Cached condensed frames. This saves a lot of time since multiple listeners can use the same condensed data **/
 	mutable TMap<int64, TArray<FStatMessage>* > CondensedStackHistory;
 
