Stats - CIS Fix

diff --git a/Engine/Source/Runtime/Core/Public/Stats/Stats2.h b/Engine/Source/Runtime/Core/Public/Stats/Stats2.h
--- a/Engine/Source/Runtime/Core/Public/Stats/Stats2.h
+++ b/Engine/Source/Runtime/Core/Public/Stats/Stats2.h
@@ -953,9 +953,12 @@ typedef TStatMessage<EComplexStatField> FComplexStatMessage;
 template<> struct TIsPODType<FComplexStatMessage> { enum { Value = true }; };
 
 
-enum class EStatMessagesArrayConstants
+struct EStatMessagesArrayConstants
 {
-	MESSAGES_CHUNK_SIZE = 64*1024,
+	enum
+	{
+		MESSAGES_CHUNK_SIZE = 64*1024,
+	};
 };
 typedef TChunkedArray<FStatMessage,(uint32)EStatMessagesArrayConstants::MESSAGES_CHUNK_SIZE> FStatMessagesArray;
 
