Efficient appending and 'copy' construction for TArrays of a different type.

* TArray can now be constructed from a TArray of a different type, which will do a per-element copy of the elements, including memcpy if it can.
* TArray can now be appended with a TArray of a different type too.
* New TIsBitwiseContructible<To, From> trait which can be specialized to say that your type is 'memcpy constructible' from another.
* CopyConstructItems and RelocateItems deprecated in favor of ConstructItems<DestType> and RelocateConstructItems<DestType>.
* Separated out some metafunctions into their own headers.
* Some PLATFORM_COMPILER_HAS_RVALUE_REFERENCES usage removed.

diff --git a/Engine/Source/Runtime/Core/Public/Containers/Array.h b/Engine/Source/Runtime/Core/Public/Containers/Array.h
--- a/Engine/Source/Runtime/Core/Public/Containers/Array.h
+++ b/Engine/Source/Runtime/Core/Public/Containers/Array.h
@@ -345,6 +345,29 @@ public:
 	}
 };
 
+namespace UE4Array_Private
+{
+	template <typename FromArrayType, typename ToArrayType>
+	struct TCanMoveTArrayPointersBetweenArrayTypes
+	{
+		typedef typename FromArrayType::Allocator   FromAllocatorType;
+		typedef typename ToArrayType  ::Allocator   ToAllocatorType;
+		typedef typename FromArrayType::ElementType FromElementType;
+		typedef typename ToArrayType  ::ElementType ToElementType;
+
+		enum
+		{
+			Value =
+				TAreTypesEqual<FromAllocatorType, ToAllocatorType>::Value && // Allocators must be equal
+				TContainerTraits<FromArrayType>::MoveWillEmptyContainer &&   // A move must be allowed to leave the source array empty
+				(
+					TAreTypesEqual         <ToElementType, FromElementType>::Value || // The element type of the container must be the same, or...
+					TIsBitwiseConstructible<ToElementType, FromElementType>::Value    // ... the element type of the source container must be bitwise constructible from the element type in the destination container
+				)
+		};
+	};
+}
+
 /**
  * Templated dynamic array
  *
@@ -356,11 +379,15 @@ public:
  * Caution: as noted below some methods are not safe for element types that require constructors.
  *
  **/
-template<typename InElementType, typename Allocator>
+template<typename InElementType, typename InAllocator>
 class TArray
 {
+	template <typename OtherInElementType, typename OtherAllocator>
+	friend class TArray;
+
 public:
 	typedef InElementType ElementType;
+	typedef InAllocator   Allocator;
 
 	/**
 	 * Constructor, initializes element number counters.
@@ -375,8 +402,8 @@ public:
 	 *
 	 * @param Other The source array to copy.
 	 */
-	template<typename OtherAllocator>
-	explicit TArray(const TArray<ElementType,OtherAllocator>& Other)
+	template <typename OtherElementType, typename OtherAllocator>
+	explicit TArray(const TArray<OtherElementType, OtherAllocator>& Other)
 	{
 		CopyToEmpty(Other);
 	}
@@ -425,7 +452,7 @@ public:
 	 *
 	 * @param Other The source array to assign from.
 	 */
-	TArray& operator=(const TArray<ElementType, Allocator>& Other)
+	TArray& operator=(const TArray& Other)
 	{
 		if (this != &Other)
 		{
@@ -435,8 +462,6 @@ public:
 		return *this;
 	}
 
-#if PLATFORM_COMPILER_HAS_RVALUE_REFERENCES
-
 private:
 	/**
 	 * Moves or copies array. Depends on the array type traits.
@@ -446,8 +471,8 @@ private:
 	 * @param ToArray Array to move into.
 	 * @param FromArray Array to move from.
 	 */
-	template <typename ArrayType>
-	static FORCEINLINE typename TEnableIf<TContainerTraits<ArrayType>::MoveWillEmptyContainer>::Type MoveOrCopy(ArrayType& ToArray, ArrayType& FromArray)
+	template <typename FromArrayType, typename ToArrayType>
+	static FORCEINLINE typename TEnableIf<UE4Array_Private::TCanMoveTArrayPointersBetweenArrayTypes<FromArrayType, ToArrayType>::Value>::Type MoveOrCopy(ToArrayType& ToArray, FromArrayType& FromArray)
 	{
 		ToArray.AllocatorInstance.MoveToEmpty(FromArray.AllocatorInstance);
 
@@ -467,8 +492,8 @@ private:
 	 * @param ExtraSlack Tells how much extra memory should be preallocated
 	 *                   at the end of the array in the number of elements.
 	 */
-	template <typename ArrayType>
-	static FORCEINLINE typename TEnableIf<TContainerTraits<ArrayType>::MoveWillEmptyContainer>::Type MoveOrCopy(ArrayType& ToArray, ArrayType& FromArray, int32 ExtraSlack)
+	template <typename FromArrayType, typename ToArrayType>
+	static FORCEINLINE typename TEnableIf<UE4Array_Private::TCanMoveTArrayPointersBetweenArrayTypes<FromArrayType, ToArrayType>::Value>::Type MoveOrCopy(ToArrayType& ToArray, FromArrayType& FromArray, int32 ExtraSlack)
 	{
 		MoveOrCopy(ToArray, FromArray);
 
@@ -485,23 +510,30 @@ private:
 	 * @param ExtraSlack Tells how much extra memory should be preallocated
 	 *                   at the end of the array in the number of elements.
 	 */
-	template <typename ArrayType>
-	static FORCEINLINE typename TEnableIf<!TContainerTraits<ArrayType>::MoveWillEmptyContainer>::Type MoveOrCopy(ArrayType& ToArray, ArrayType& FromArray, int32 ExtraSlack = 0)
+	template <typename FromArrayType, typename ToArrayType>
+	static FORCEINLINE typename TEnableIf<!UE4Array_Private::TCanMoveTArrayPointersBetweenArrayTypes<FromArrayType, ToArrayType>::Value>::Type MoveOrCopy(ToArrayType& ToArray, FromArrayType& FromArray, int32 ExtraSlack = 0)
 	{
 		ToArray.CopyToEmpty(FromArray, ExtraSlack);
 	}
 
 public:
-	// We don't implement move semantics for general OtherAllocators, as there's no way
-	// to tell if they're compatible with the current one.  Probably going to be a pretty
-	// rare requirement anyway.
+	/**
+	 * Move constructor.
+	 *
+	 * @param Other Array to move from.
+	 */
+	FORCEINLINE TArray(TArray&& Other)
+	{
+		MoveOrCopy(*this, Other);
+	}
 
 	/**
 	 * Move constructor.
 	 *
 	 * @param Other Array to move from.
 	 */
-	TArray(TArray<ElementType, Allocator>&& Other)
+	template <typename OtherElementType, typename OtherAllocator>
+	FORCEINLINE explicit TArray(TArray<OtherElementType, OtherAllocator>&& Other)
 	{
 		MoveOrCopy(*this, Other);
 	}
@@ -513,8 +545,13 @@ public:
 	 * @param ExtraSlack Tells how much extra memory should be preallocated
 	 *                   at the end of the array in the number of elements.
 	 */
-	TArray(TArray<ElementType, Allocator>&& Other, int32 ExtraSlack)
+	template <typename OtherElementType>
+	TArray(TArray<OtherElementType, Allocator>&& Other, int32 ExtraSlack)
 	{
+		// We don't implement move semantics for general OtherAllocators, as there's no way
+		// to tell if they're compatible with the current one.  Probably going to be a pretty
+		// rare requirement anyway.
+
 		MoveOrCopy(*this, Other, ExtraSlack);
 	}
 
@@ -523,7 +560,7 @@ public:
 	 *
 	 * @param Other Array to assign and move from.
 	 */
-	TArray& operator=(TArray<ElementType, Allocator>&& Other)
+	TArray& operator=(TArray&& Other)
 	{
 		if (this != &Other)
 		{
@@ -533,8 +570,6 @@ public:
 		return *this;
 	}
 
-#endif
-
 	/**
 	 * Destructor.
 	 */
@@ -1267,7 +1302,7 @@ public:
 			AllocatorInstance.ResizeAllocation(OldNum, ArrayMax, sizeof(ElementType));
 		}
 		ElementType* Data = GetData() + Index;
-		RelocateItems(Data + Count, Data, OldNum - Index);
+		RelocateConstructItems<ElementType>(Data + Count, Data, OldNum - Index);
 	}
 
 	/**
@@ -1322,7 +1357,7 @@ public:
 		check(Ptr != NULL);
 
 		InsertUninitialized(Index, Count);
-		CopyConstructItems(GetData() + Index, Ptr, Count);
+		ConstructItems<ElementType>(GetData() + Index, Ptr, Count);
 
 		return Index;
 	}
@@ -1556,8 +1591,8 @@ public:
 	 *
 	 * @param Source The array to append.
 	 */
-	template <typename OtherAllocator>
-	FORCEINLINE void Append(const TArray<InElementType, OtherAllocator>& Source)
+	template <typename OtherElementType, typename OtherAllocator>
+	FORCEINLINE void Append(const TArray<OtherElementType, OtherAllocator>& Source)
 	{
 		check((void*)this != (void*)&Source);
 
@@ -1572,7 +1607,7 @@ public:
 		// Allocate memory for the new elements.
 		Reserve(ArrayNum + SourceCount);
 
-		CopyConstructItems(GetData() + ArrayNum, Source.GetData(), SourceCount);
+		ConstructItems<ElementType>(GetData() + ArrayNum, Source.GetData(), SourceCount);
 
 		ArrayNum += SourceCount;
 	}
@@ -1582,7 +1617,8 @@ public:
 	 *
 	 * @param Source The array to append.
 	 */
-	FORCEINLINE void Append(TArray&& Source)
+	template <typename OtherElementType, typename OtherAllocator>
+	FORCEINLINE void Append(TArray<OtherElementType, OtherAllocator>&& Source)
 	{
 		check((void*)this != (void*)&Source);
 
@@ -1597,7 +1633,7 @@ public:
 		// Allocate memory for the new elements.
 		Reserve(ArrayNum + SourceCount);
 
-		RelocateItems(GetData() + ArrayNum, Source.GetData(), SourceCount);
+		RelocateConstructItems<ElementType>(GetData() + ArrayNum, Source.GetData(), SourceCount);
 		Source.ArrayNum = 0;
 
 		ArrayNum += SourceCount;
@@ -1614,7 +1650,7 @@ public:
 		check(Ptr != NULL);
 
 		int32 Pos = AddUninitialized(Count);
-		CopyConstructItems(GetData() + Pos, Ptr, Count);
+		ConstructItems<ElementType>(GetData() + Pos, Ptr, Count);
 	}
 
 	/**
@@ -1625,7 +1661,11 @@ public:
 	 *
 	 * @param Other The array to append.
 	 */
-	FORCEINLINE TArray& operator+=(TArray&& Other) { Append(MoveTemp(Other)); return *this; }
+	FORCEINLINE TArray& operator+=(TArray&& Other)
+	{
+		Append(MoveTemp(Other));
+		return *this;
+	}
 
 	/**
 	 * Appends the specified array to this array.
@@ -1633,7 +1673,11 @@ public:
 	 *
 	 * @param Other The array to append.
 	 */
-	FORCEINLINE TArray& operator+=(const TArray& Other) { Append(Other); return *this; }
+	FORCEINLINE TArray& operator+=(const TArray& Other)
+	{
+		Append(Other);
+		return *this;
+	}
 
 	/**
 	 * Adds a new item to the end of the array, possibly reallocating the whole array to fit.
@@ -2152,15 +2196,15 @@ private:
 	 *                   the end of the buffer. Counted in elements. Zero by
 	 *                   default.
 	 */
-	template<typename OtherAllocator>
-	void CopyToEmpty(const TArray<ElementType, OtherAllocator>& Source, int32 ExtraSlack = 0)
+	template <typename OtherElementType, typename OtherAllocator>
+	void CopyToEmpty(const TArray<OtherElementType, OtherAllocator>& Source, int32 ExtraSlack = 0)
 	{
 		check(ExtraSlack >= 0);
 
 		int32 SourceCount = Source.Num();
 		AllocatorInstance.ResizeAllocation(0, SourceCount + ExtraSlack, sizeof(ElementType));
 
-		CopyConstructItems(GetData(), Source.GetData(), SourceCount);
+		ConstructItems<ElementType>(GetData(), Source.GetData(), SourceCount);
 
 		ArrayNum = SourceCount;
 		ArrayMax = SourceCount + ExtraSlack;
@@ -2525,9 +2569,7 @@ struct TIsZeroConstructType<TArray<InElementType, Allocator>>
 template <typename InElementType, typename Allocator>
 struct TContainerTraits<TArray<InElementType, Allocator> > : public TContainerTraitsBase<TArray<InElementType, Allocator> >
 {
-	enum { MoveWillEmptyContainer =
-		PLATFORM_COMPILER_HAS_RVALUE_REFERENCES &&
-		TAllocatorTraits<Allocator>::SupportsMove };
+	enum { MoveWillEmptyContainer = TAllocatorTraits<Allocator>::SupportsMove };
 };
 
 /**
@@ -2613,15 +2655,10 @@ public:
 
 #if PLATFORM_COMPILER_HAS_DEFAULTED_FUNCTIONS
 
-	TMRUArray(const TMRUArray&) = default;
-	TMRUArray& operator=(const TMRUArray&) = default;
-
-#if PLATFORM_COMPILER_HAS_RVALUE_REFERENCES
-
 	TMRUArray(TMRUArray&&) = default;
+	TMRUArray(const TMRUArray&) = default;
 	TMRUArray& operator=(TMRUArray&&) = default;
-
-#endif
+	TMRUArray& operator=(const TMRUArray&) = default;
 
 #else
 
@@ -2648,7 +2685,6 @@ public:
 		return *this;
 	}
 
-#if PLATFORM_COMPILER_HAS_RVALUE_REFERENCES
 	/**
 	 * Move constructor.
 	 *
@@ -2672,8 +2708,6 @@ public:
 		return *this;
 	}
 
-#endif
-
 #endif
 	/**
 	 * Adds item to the array. Makes sure that we don't add more than the
@@ -2749,9 +2783,7 @@ public:
 template<typename T, typename Allocator>
 struct TContainerTraits<TMRUArray<T, Allocator> > : public TContainerTraitsBase<TMRUArray<T, Allocator> >
 {
-	enum { MoveWillEmptyContainer =
-		PLATFORM_COMPILER_HAS_RVALUE_REFERENCES &&
-		TContainerTraitsBase<typename TMRUArray<T, Allocator>::Super>::MoveWillEmptyContainer };
+	enum { MoveWillEmptyContainer = TContainerTraitsBase<typename TMRUArray<T, Allocator>::Super>::MoveWillEmptyContainer };
 };
 
 /*-----------------------------------------------------------------------------
@@ -3266,9 +3298,7 @@ private:
 template<typename T, typename Allocator>
 struct TContainerTraits<TIndirectArray<T, Allocator> > : public TContainerTraitsBase<TIndirectArray<T, Allocator> >
 {
-	enum { MoveWillEmptyContainer =
-		PLATFORM_COMPILER_HAS_RVALUE_REFERENCES &&
-		TContainerTraitsBase<typename TIndirectArray<T, Allocator>::InternalArrayType>::MoveWillEmptyContainer };
+	enum { MoveWillEmptyContainer = TContainerTraitsBase<typename TIndirectArray<T, Allocator>::InternalArrayType>::MoveWillEmptyContainer };
 };
 
 template <typename T,typename Allocator> void* operator new( size_t Size, TIndirectArray<T,Allocator>& Array )
@@ -3340,16 +3370,11 @@ public:
 
 #if PLATFORM_COMPILER_HAS_DEFAULTED_FUNCTIONS
 
+	TTransArray(TTransArray&&) = default;
 	TTransArray(const TTransArray&) = default;
+	TTransArray& operator=(TTransArray&&) = default;
 	TTransArray& operator=(const TTransArray&) = default;
 
-	#if PLATFORM_COMPILER_HAS_RVALUE_REFERENCES
-
-		TTransArray(TTransArray&&) = default;
-		TTransArray& operator=(TTransArray&&) = default;
-
-	#endif
-
 #else
 
 	FORCEINLINE TTransArray(const TTransArray& Other)
@@ -3365,22 +3390,18 @@ public:
 		return *this;
 	}
 
-	#if PLATFORM_COMPILER_HAS_RVALUE_REFERENCES
-
-		FORCEINLINE TTransArray(TTransArray&& Other)
-			: Super((TTransArray&&)Other)
-			, Owner(Other.Owner)
-		{
-		}
-
-		FORCEINLINE TTransArray& operator=(TTransArray&& Other)
-		{
-			(Super&)*this = (Super&&)Other;
-			Owner         = Other.Owner;
-			return *this;
-		}
+	FORCEINLINE TTransArray(TTransArray&& Other)
+		: Super((TTransArray&&)Other)
+		, Owner(Other.Owner)
+	{
+	}
 
-	#endif
+	FORCEINLINE TTransArray& operator=(TTransArray&& Other)
+	{
+		(Super&)*this = (Super&&)Other;
+		Owner         = Other.Owner;
+		return *this;
+	}
 
 #endif
 
@@ -3510,9 +3531,7 @@ protected:
 template<typename T>
 struct TContainerTraits<TTransArray<T> > : public TContainerTraitsBase<TTransArray<T> >
 {
-	enum { MoveWillEmptyContainer =
-		PLATFORM_COMPILER_HAS_RVALUE_REFERENCES &&
-		TContainerTraitsBase<typename TTransArray<T>::Super>::MoveWillEmptyContainer };
+	enum { MoveWillEmptyContainer = TContainerTraitsBase<typename TTransArray<T>::Super>::MoveWillEmptyContainer };
 };
 
 //
diff --git a/Engine/Source/Runtime/Core/Public/Containers/ContainerAllocationPolicies.h b/Engine/Source/Runtime/Core/Public/Containers/ContainerAllocationPolicies.h
--- a/Engine/Source/Runtime/Core/Public/Containers/ContainerAllocationPolicies.h
+++ b/Engine/Source/Runtime/Core/Public/Containers/ContainerAllocationPolicies.h
@@ -426,7 +426,7 @@ public:
 			if (!Other.SecondaryData.GetAllocation())
 			{
 				// Relocate objects from other inline storage only if it was stored inline in Other
-				RelocateItems((void*)InlineData, Other.GetInlineElements(), NumInlineElements);
+				RelocateConstructItems<ElementType>((void*)InlineData, Other.GetInlineElements(), NumInlineElements);
 			}
 
 			// Move secondary storage in any case.
@@ -448,7 +448,7 @@ public:
 				// If the old allocation wasn't in the inline data area, relocate it into the inline data area.
 				if(SecondaryData.GetAllocation())
 				{
-					RelocateItems((void*)InlineData, (ElementType*)SecondaryData.GetAllocation(), PreviousNumElements);
+					RelocateConstructItems<ElementType>((void*)InlineData, (ElementType*)SecondaryData.GetAllocation(), PreviousNumElements);
 
 					// Free the old indirect allocation.
 					SecondaryData.ResizeAllocation(0,0,NumBytesPerElement);
@@ -462,7 +462,7 @@ public:
 					SecondaryData.ResizeAllocation(0,NumElements,NumBytesPerElement);
 
 					// Move the data out of the inline data area into the new allocation.
-					RelocateItems((void*)SecondaryData.GetAllocation(), GetInlineElements(), PreviousNumElements);
+					RelocateConstructItems<ElementType>((void*)SecondaryData.GetAllocation(), GetInlineElements(), PreviousNumElements);
 				}
 				else
 				{
@@ -543,7 +543,7 @@ public:
 			check(this != &Other);
 
 			// Relocate objects from other inline storage
-			RelocateItems((void*)InlineData, Other.GetInlineElements(), NumInlineElements);
+			RelocateConstructItems<ElementType>((void*)InlineData, Other.GetInlineElements(), NumInlineElements);
 		}
 
 		// FContainerAllocatorInterface
diff --git a/Engine/Source/Runtime/Core/Public/Containers/Set.h b/Engine/Source/Runtime/Core/Public/Containers/Set.h
--- a/Engine/Source/Runtime/Core/Public/Containers/Set.h
+++ b/Engine/Source/Runtime/Core/Public/Containers/Set.h
@@ -69,7 +69,7 @@ FORCEINLINE void MoveByRelocate(T& A, T& B)
 	A.~T();
 
 	// Relocate B into the 'hole' left by the destruction of A, leaving a hole in B instead.
-	RelocateItems(&A, &B, 1);
+	RelocateConstructItems<T>(&A, &B, 1);
 }
 
 /** Either NULL or an identifier for an element of a set. */
diff --git a/Engine/Source/Runtime/Core/Public/Containers/StaticArray.h b/Engine/Source/Runtime/Core/Public/Containers/StaticArray.h
--- a/Engine/Source/Runtime/Core/Public/Containers/StaticArray.h
+++ b/Engine/Source/Runtime/Core/Public/Containers/StaticArray.h
@@ -24,7 +24,7 @@ public:
 	/** Copy constructor. */
 	TStaticArray(const TStaticArray& Other)
 	{
-		CopyConstructItems((void*)Elements, (const TElement*)Other.Elements, NumElements);
+		ConstructItems<TElement>((void*)Elements, (const TElement*)Other.Elements, NumElements);
 	}
 
 #if PLATFORM_COMPILER_HAS_RVALUE_REFERENCES
diff --git a/Engine/Source/Runtime/Core/Public/Templates/AreTypesEqual.h b/Engine/Source/Runtime/Core/Public/Templates/AreTypesEqual.h
new file mode 100644
--- /dev/null
+++ b/Engine/Source/Runtime/Core/Public/Templates/AreTypesEqual.h
@@ -0,0 +1,21 @@
+// Copyright 1998-2014 Epic Games, Inc. All Rights Reserved.
+
+#pragma once
+
+/** Tests whether two typenames refer to the same type. */
+template<typename A,typename B>
+struct TAreTypesEqual;
+
+template<typename,typename>
+struct TAreTypesEqual
+{
+	enum { Value = false };
+};
+
+template<typename A>
+struct TAreTypesEqual<A,A>
+{
+	enum { Value = true };
+};
+
+#define ARE_TYPES_EQUAL(A,B) TAreTypesEqual<A,B>::Value
diff --git a/Engine/Source/Runtime/Core/Public/Templates/MemoryOps.h b/Engine/Source/Runtime/Core/Public/Templates/MemoryOps.h
--- a/Engine/Source/Runtime/Core/Public/Templates/MemoryOps.h
+++ b/Engine/Source/Runtime/Core/Public/Templates/MemoryOps.h
@@ -10,6 +10,23 @@
 #include "UnrealTypeTraits.h"
 
 
+namespace UE4MemoryOps_Private
+{
+	template <typename DestinationElementType, typename SourceElementType>
+	struct TCanBitwiseRelocate
+	{
+		enum
+		{
+			Value =
+				(
+					TIsBitwiseConstructible<DestinationElementType, SourceElementType>::Value &&
+					!TTypeTraits<SourceElementType>::NeedsDestructor
+				) ||
+				TAreTypesEqual<DestinationElementType, SourceElementType>::Value
+		};
+	};
+}
+
 /**
  * Default constructs a range of items in memory.
  *
@@ -64,29 +81,37 @@ FORCEINLINE typename TEnableIf<!TTypeTraits<ElementType>::NeedsDestructor>::Type
 
 
 /**
- * Copy constructs a range of items into memory.
+ * Constructs a range of items into memory from a set of arguments.  The arguments come from an another array.
  *
  * @param	Dest		The memory location to start copying into.
- * @param	Source		A pointer to the first item to copy from.
+ * @param	Source		A pointer to the first argument to pass to the constructor.
  * @param	Count		The number of elements to copy.
  */
-template <typename ElementType>
-FORCEINLINE typename TEnableIf<TTypeTraits<ElementType>::NeedsCopyConstructor>::Type CopyConstructItems(void* Dest, const ElementType* Source, int32 Count)
+template <typename DestinationElementType, typename SourceElementType>
+FORCEINLINE typename TEnableIf<!TIsBitwiseConstructible<DestinationElementType, SourceElementType>::Value>::Type ConstructItems(void* Dest, const SourceElementType* Source, int32 Count)
 {
 	while (Count)
 	{
-		new (Dest) ElementType(*Source);
-		++(ElementType*&)Dest;
+		new (Dest) DestinationElementType(*Source);
+		++(DestinationElementType*&)Dest;
 		++Source;
 		--Count;
 	}
 }
 
 
+template <typename DestinationElementType, typename SourceElementType>
+FORCEINLINE typename TEnableIf<TIsBitwiseConstructible<DestinationElementType, SourceElementType>::Value>::Type ConstructItems(void* Dest, const SourceElementType* Source, int32 Count)
+{
+	FMemory::Memcpy(Dest, Source, sizeof(SourceElementType) * Count);
+}
+
+
 template <typename ElementType>
-FORCEINLINE typename TEnableIf<!TTypeTraits<ElementType>::NeedsCopyConstructor>::Type CopyConstructItems(void* Dest, const ElementType* Source, int32 Count)
+DEPRECATED(4.6, "CopyConstructItems is deprecated, use ConstructItems<ElementType> instead.")
+FORCEINLINE void CopyConstructItems(void* Dest, const ElementType* Source, int32 Count)
 {
-	FMemory::Memcpy(Dest, Source, sizeof(ElementType) * Count);
+	ConstructItems<ElementType>(Dest, Source, Count);
 }
 
 
@@ -118,15 +143,30 @@ FORCEINLINE typename TEnableIf<!TTypeTraits<ElementType>::NeedsCopyAssignment>::
 
 
 /**
- * Relocates a range of items to a new memory location. This is a so-called 'destructive move' for which
+ * Relocates a range of items to a new memory location as a new type. This is a so-called 'destructive move' for which
  * there is no single operation in C++ but which can be implemented very efficiently in general.
  *
  * @param	Dest		The memory location to relocate to.
  * @param	Source		A pointer to the first item to relocate.
  * @param	Count		The number of elements to relocate.
  */
-template <typename ElementType>
-FORCEINLINE void RelocateItems(void* Dest, const ElementType* Source, int32 Count)
+template <typename DestinationElementType, typename SourceElementType>
+FORCEINLINE typename TEnableIf<!UE4MemoryOps_Private::TCanBitwiseRelocate<DestinationElementType, SourceElementType>::Value>::Type RelocateConstructItems(void* Dest, const SourceElementType* Source, int32 Count)
+{
+	while (Count)
+	{
+		// We need a typedef here because VC won't compile the destructor call below if SourceElementType itself has a member called SourceElementType
+		typedef SourceElementType RelocateConstructItemsElementTypeTypedef;
+
+		new (Dest) DestinationElementType(*Source);
+		++(DestinationElementType*&)Dest;
+		(Source++)->RelocateConstructItemsElementTypeTypedef::~RelocateConstructItemsElementTypeTypedef();
+		--Count;
+	}
+}
+
+template <typename DestinationElementType, typename SourceElementType>
+FORCEINLINE typename TEnableIf<UE4MemoryOps_Private::TCanBitwiseRelocate<DestinationElementType, SourceElementType>::Value>::Type RelocateConstructItems(void* Dest, const SourceElementType* Source, int32 Count)
 {
 	/* All existing UE containers seem to assume trivial relocatability (i.e. memcpy'able) of their members,
 	 * so we're going to assume that this is safe here.  However, it's not generally possible to assume this
@@ -136,77 +176,66 @@ FORCEINLINE void RelocateItems(void* Dest, const ElementType* Source, int32 Coun
 	 * However, it is not yet possible to automatically infer this at compile time, so we can't enable
 	 * different (i.e. safer) implementations anyway. */
 
-	FMemory::Memmove(Dest, Source, sizeof(ElementType) * Count);
+	FMemory::Memmove(Dest, Source, sizeof(SourceElementType) * Count);
 }
 
+template <typename ElementType>
+DEPRECATED(4.6, "RelocateItems is deprecated, use RelocateConstructItems<ElementType> instead.")
+FORCEINLINE void RelocateItems(void* Dest, const ElementType* Source, int32 Count)
+{
+	RelocateConstructItems<ElementType>(Dest, Source, Count);
+}
 
-#if PLATFORM_COMPILER_HAS_RVALUE_REFERENCES
-
-	/**
-	 * Move constructs a range of items into memory.
-	 *
-	 * @param	Dest		The memory location to start moving into.
-	 * @param	Source		A pointer to the first item to move from.
-	 * @param	Count		The number of elements to move.
-	 */
-	template <typename ElementType>
-	FORCEINLINE typename TEnableIf<TTypeTraits<ElementType>::NeedsMoveConstructor>::Type MoveConstructItems(void* Dest, ElementType* Source, int32 Count)
-	{
-		while (Count)
-		{
-			new (Dest) ElementType((ElementType&&)*Source);
-			++(ElementType*&)Dest;
-			++Source;
-			--Count;
-		}
-	}
-
-	template <typename ElementType>
-	FORCEINLINE typename TEnableIf<!TTypeTraits<ElementType>::NeedsMoveConstructor>::Type MoveConstructItems(void* Dest, const ElementType* Source, int32 Count)
-	{
-		FMemory::Memmove(Dest, Source, sizeof(ElementType) * Count);
-	}
 
-	/**
-	 * Move assigns a range of items.
-	 *
-	 * @param	Dest		The memory location to start move assigning to.
-	 * @param	Source		A pointer to the first item to move assign.
-	 * @param	Count		The number of elements to move assign.
-	 */
-	template <typename ElementType>
-	FORCEINLINE typename TEnableIf<TTypeTraits<ElementType>::NeedsMoveAssignment>::Type MoveAssignItems(ElementType* Dest, ElementType* Source, int32 Count)
-	{
-		while (Count)
-		{
-			*Dest = (ElementType&&)*Source;
-			++Dest;
-			++Source;
-			--Count;
-		}
-	}
-
-	template <typename ElementType>
-	FORCEINLINE typename TEnableIf<!TTypeTraits<ElementType>::NeedsMoveAssignment>::Type MoveAssignItems(ElementType* Dest, const ElementType* Source, int32 Count)
+/**
+ * Move constructs a range of items into memory.
+ *
+ * @param	Dest		The memory location to start moving into.
+ * @param	Source		A pointer to the first item to move from.
+ * @param	Count		The number of elements to move.
+ */
+template <typename ElementType>
+FORCEINLINE typename TEnableIf<TTypeTraits<ElementType>::NeedsMoveConstructor>::Type MoveConstructItems(void* Dest, ElementType* Source, int32 Count)
+{
+	while (Count)
 	{
-		FMemory::Memmove(Dest, Source, sizeof(ElementType) * Count);
+		new (Dest) ElementType((ElementType&&)*Source);
+		++(ElementType*&)Dest;
+		++Source;
+		--Count;
 	}
+}
 
-#else
-
-	template <typename ElementType>
-	FORCEINLINE void MoveConstructItems(void* Dest, const ElementType* Source, int32 Count)
-	{
-		CopyConstructItems(Dest, Source, Count);
-	}
+template <typename ElementType>
+FORCEINLINE typename TEnableIf<!TTypeTraits<ElementType>::NeedsMoveConstructor>::Type MoveConstructItems(void* Dest, const ElementType* Source, int32 Count)
+{
+	FMemory::Memmove(Dest, Source, sizeof(ElementType) * Count);
+}
 
-	template <typename ElementType>
-	FORCEINLINE void MoveAssignItems(ElementType* Dest, const ElementType* Source, int32 Count)
+/**
+ * Move assigns a range of items.
+ *
+ * @param	Dest		The memory location to start move assigning to.
+ * @param	Source		A pointer to the first item to move assign.
+ * @param	Count		The number of elements to move assign.
+ */
+template <typename ElementType>
+FORCEINLINE typename TEnableIf<TTypeTraits<ElementType>::NeedsMoveAssignment>::Type MoveAssignItems(ElementType* Dest, ElementType* Source, int32 Count)
+{
+	while (Count)
 	{
-		CopyAssignItems(Dest, Source, Count);
+		*Dest = (ElementType&&)*Source;
+		++Dest;
+		++Source;
+		--Count;
 	}
+}
 
-#endif
+template <typename ElementType>
+FORCEINLINE typename TEnableIf<!TTypeTraits<ElementType>::NeedsMoveAssignment>::Type MoveAssignItems(ElementType* Dest, const ElementType* Source, int32 Count)
+{
+	FMemory::Memmove(Dest, Source, sizeof(ElementType) * Count);
+}
 
 template <typename ElementType>
 FORCEINLINE typename TEnableIf<TTypeTraits<ElementType>::IsBytewiseComparable, bool>::Type CompareItems(const ElementType* A, const ElementType* B, int32 Count)
diff --git a/Engine/Source/Runtime/Core/Public/Templates/PointerIsConvertibleFromTo.h b/Engine/Source/Runtime/Core/Public/Templates/PointerIsConvertibleFromTo.h
new file mode 100644
--- /dev/null
+++ b/Engine/Source/Runtime/Core/Public/Templates/PointerIsConvertibleFromTo.h
@@ -0,0 +1,17 @@
+// Copyright 1998-2014 Epic Games, Inc. All Rights Reserved.
+
+#pragma once
+
+/**
+ * Tests if a From* is convertible to a To*
+ **/
+template <typename From, typename To>
+struct TPointerIsConvertibleFromTo
+{
+private:
+	static uint8  Test(...);
+	static uint16 Test(To*);
+
+public:
+	enum { Value  = sizeof(Test((From*)nullptr)) - 1 };
+};
diff --git a/Engine/Source/Runtime/Core/Public/Templates/RemoveCV.h b/Engine/Source/Runtime/Core/Public/Templates/RemoveCV.h
new file mode 100644
--- /dev/null
+++ b/Engine/Source/Runtime/Core/Public/Templates/RemoveCV.h
@@ -0,0 +1,13 @@
+// Copyright 1998-2014 Epic Games, Inc. All Rights Reserved.
+
+#pragma once
+
+/**
+ * TRemoveCV<type> will remove any const/volatile qualifiers from a type.
+ * (based on std::remove_cv<>
+ * note: won't remove the const from "const int*", as the pointer is not const
+ */
+template <typename T> struct TRemoveCV                   { typedef T Type; };
+template <typename T> struct TRemoveCV<const T>          { typedef T Type; };
+template <typename T> struct TRemoveCV<volatile T>       { typedef T Type; };
+template <typename T> struct TRemoveCV<const volatile T> { typedef T Type; };
diff --git a/Engine/Source/Runtime/Core/Public/Templates/TypeWrapper.h b/Engine/Source/Runtime/Core/Public/Templates/TypeWrapper.h
new file mode 100644
--- /dev/null
+++ b/Engine/Source/Runtime/Core/Public/Templates/TypeWrapper.h
@@ -0,0 +1,68 @@
+// Copyright 1998-2014 Epic Games, Inc. All Rights Reserved.
+
+#pragma once
+
+/**
+ * Wraps a type.
+ *
+ * The intended use of this template is to allow types to be passed around where the unwrapped type would give different behavior.
+ * An example of this is when you want a template specialization to refer to the primary template, but doing that would cause
+ * infinite recursion through the specialization, e.g.:
+ *
+ * // Before
+ * template <typename T>
+ * struct TThing
+ * {
+ *     void f(T t)
+ *     {
+ *          DoSomething(t);
+ *     }
+ * };
+ *
+ * template <>
+ * struct TThing<int>
+ * {
+ *     void f(int t)
+ *     {
+ *         DoSomethingElseFirst(t);
+ *         TThing<int>::f(t); // Infinite recursion
+ *     }
+ * };
+ *
+ *
+ * // After
+ * template <typename T>
+ * struct TThing
+ * {
+ *     typedef typename TUnwrapType<T>::Type RealT;
+ *
+ *     void f(RealT t)
+ *     {
+ *          DoSomething(t);
+ *     }
+ * };
+ *
+ * template <>
+ * struct TThing<int>
+ * {
+ *     void f(int t)
+ *     {
+ *         DoSomethingElseFirst(t);
+ *         TThing<TTypeWrapper<int>>::f(t); // works
+ *     }
+ * };
+ */
+template <typename T>
+struct TTypeWrapper;
+
+template <typename T>
+struct TUnwrapType
+{
+	typedef T Type;
+};
+
+template <typename T>
+struct TUnwrapType<TTypeWrapper<T>>
+{
+	typedef T Type;
+};
diff --git a/Engine/Source/Runtime/Core/Public/Templates/UniquePtr.h b/Engine/Source/Runtime/Core/Public/Templates/UniquePtr.h
--- a/Engine/Source/Runtime/Core/Public/Templates/UniquePtr.h
+++ b/Engine/Source/Runtime/Core/Public/Templates/UniquePtr.h
@@ -221,7 +221,19 @@ FORCEINLINE bool operator!=(const TUniquePtr<T>& Lhs, const TUniquePtr<T>& Rhs)
 	return Lhs.Get() != Rhs.Get();
 }
 
-template <typename T> struct TIsZeroConstructType<TUniquePtr<T>> { enum { Value = true }; };
+// Trait which allows TUniquePtr to be default constructed by memsetting to zero.
+template <typename T>
+struct TIsZeroConstructType<TUniquePtr<T>>
+{
+	enum { Value = true };
+};
+
+// Trait which allows TUniquePtr to be memcpy'able from pointers.
+template <typename T>
+struct TIsBitwiseConstructible<TUniquePtr<T>, T*>
+{
+	enum { Value = true };
+};
 
 #if PLATFORM_COMPILER_HAS_VARIADIC_TEMPLATES
 
@@ -279,4 +291,4 @@ template <typename T> struct TIsZeroConstructType<TUniquePtr<T>> { enum { Value
 		return TUniquePtr<T>(new T(Forward<TArg0>(Arg0), Forward<TArg1>(Arg1), Forward<TArg2>(Arg2), Forward<TArg3>(Arg3)));
 	}
 
-#endif
\ No newline at end of file
+#endif
diff --git a/Engine/Source/Runtime/Core/Public/Templates/UnrealTemplate.h b/Engine/Source/Runtime/Core/Public/Templates/UnrealTemplate.h
--- a/Engine/Source/Runtime/Core/Public/Templates/UnrealTemplate.h
+++ b/Engine/Source/Runtime/Core/Public/Templates/UnrealTemplate.h
@@ -3,6 +3,8 @@
 #pragma once
 
 #include "EnableIf.h"
+#include "PointerIsConvertibleFromTo.h"
+#include "TypeWrapper.h"
 #include "HAL/Platform.h"
 #include "Templates/UnrealTypeTraits.h"
 
@@ -11,20 +13,16 @@
 	Standard templates.
 -----------------------------------------------------------------------------*/
 
-#if PLATFORM_SUPPORTS_CanConvertPointerFromTo
-
 /** 
 * CanConvertPointerFromTo<From, To>::Result is an enum value equal to 1 if From* is automatically convertable to a To* (without regard to const!)
 **/
 template<class From, class To>
 class CanConvertPointerFromTo
 {
-	static uint8 Test(...);
-	static uint16 Test(To const*);
 public:
 	enum Type
 	{
-		Result = sizeof(Test((From*)0)) - 1
+		Result = TPointerIsConvertibleFromTo<From, const To>::Value
 	};
 };
 
@@ -58,8 +56,6 @@ static_assert(!(CanConvertPointerFromTo<void, CanConvertPointerFromTo_TestBase>:
 static_assert(!(CanConvertPointerFromTo<CanConvertPointerFromTo_TestBase, bool>::Result), "Platform CanConvertPointerFromTo test failed.");
 static_assert(!(CanConvertPointerFromTo<void, bool>::Result), "Platform CanConvertPointerFromTo test failed.");
 
-#endif
-
 
 
 /**
@@ -346,24 +342,6 @@ struct TKeyValuePair
 	ValueType	Value;
 };
 
-/** Tests whether two typenames refer to the same type. */
-template<typename A,typename B>
-struct TAreTypesEqual;
-
-template<typename,typename>
-struct TAreTypesEqual
-{
-	enum { Value = false };
-};
-
-template<typename A>
-struct TAreTypesEqual<A,A>
-{
-	enum { Value = true };
-};
-
-#define ARE_TYPES_EQUAL(A,B) TAreTypesEqual<A,B>::Value
-
 //
 // Macros that can be used to specify multiple template parameters in a macro parameter.
 // This is necessary to prevent the macro parsing from interpreting the template parameter
@@ -442,16 +420,6 @@ FORCEINLINE T StaticCast(ArgType&& Arg)
 	return static_cast<T>(Arg);
 }
 
-/**
- * TRemoveCV<type> will remove any const/volatile qualifiers from a type.
- * (based on std::remove_cv<>
- * note: won't remove the const from "const int*", as the pointer is not const
- */
-template <typename T> struct TRemoveCV                   { typedef T Type; };
-template <typename T> struct TRemoveCV<const T>          { typedef T Type; };
-template <typename T> struct TRemoveCV<volatile T>       { typedef T Type; };
-template <typename T> struct TRemoveCV<const volatile T> { typedef T Type; };
-
 /**
  * TRValueToLValueReference converts any rvalue reference type into the equivalent lvalue reference, otherwise returns the same type.
  */
@@ -610,68 +578,3 @@ template <typename LHS, typename RHS>
 struct TOr : TOrValue<LHS::Value, RHS>
 {
 };
-
-/**
- * Wraps a type.
- *
- * The intended use of this template is to allow types to be passed around where the unwrapped type would give different behavior.
- * An example of this is when you want a template specialization to refer to the primary template, but doing that would cause
- * infinite recursion through the specialization, e.g.:
- *
- * // Before
- * template <typename T>
- * struct TThing
- * {
- *     void f(T t)
- *     {
- *          DoSomething(t);
- *     }
- * };
- *
- * template <>
- * struct TThing<int>
- * {
- *     void f(int t)
- *     {
- *         DoSomethingElseFirst(t);
- *         TThing<int>::f(t); // Infinite recursion
- *     }
- * };
- *
- *
- * // After
- * template <typename T>
- * struct TThing
- * {
- *     typedef typename TUnwrapType<T>::Type RealT;
- *
- *     void f(RealT t)
- *     {
- *          DoSomething(t);
- *     }
- * };
- *
- * template <>
- * struct TThing<int>
- * {
- *     void f(int t)
- *     {
- *         DoSomethingElseFirst(t);
- *         TThing<TTypeWrapper<int>>::f(t); // works
- *     }
- * };
- */
-template <typename T>
-struct TTypeWrapper;
-
-template <typename T>
-struct TUnwrapType
-{
-	typedef T Type;
-};
-
-template <typename T>
-struct TUnwrapType<TTypeWrapper<T>>
-{
-	typedef T Type;
-};
diff --git a/Engine/Source/Runtime/Core/Public/Templates/UnrealTypeTraits.h b/Engine/Source/Runtime/Core/Public/Templates/UnrealTypeTraits.h
--- a/Engine/Source/Runtime/Core/Public/Templates/UnrealTypeTraits.h
+++ b/Engine/Source/Runtime/Core/Public/Templates/UnrealTypeTraits.h
@@ -7,6 +7,8 @@
 =============================================================================*/
 
 #pragma once
+#include "AreTypesEqual.h"
+#include "RemoveCV.h"
 
 #include "AssertionMacros.h"
 #include "Misc/OutputDevice.h"
@@ -522,4 +524,66 @@ struct TMoveSupportTraits : TMoveSupportTraitsBase<T, typename TCallTraits<T>::P
 {
 };
 
+/**
+ * Tests if a type T is bitwise-constructible from a given argument type U.  That is, whether or not
+ * the U can be memcpy'd in order to produce an instance of T, rather than having to go
+ * via a constructor.
+ *
+ * Examples:
+ * TIsBitwiseConstructible<PODType,    PODType   >::Value == true  // PODs can be trivially copied
+ * TIsBitwiseConstructible<const int*, int*      >::Value == true  // a non-const Derived pointer is trivially copyable as a const Base pointer
+ * TIsBitwiseConstructible<int*,       const int*>::Value == false // not legal the other way because it would be a const-correctness violation
+ * TIsBitwiseConstructible<int32,      uint32    >::Value == true  // signed integers can be memcpy'd as unsigned integers
+ * TIsBitwiseConstructible<uint32,     int32     >::Value == true  // and vice versa
+ */
+
+template <typename T, typename Arg>
+struct TIsBitwiseConstructible
+{
+	static_assert(
+		!TIsReferenceType<T  >::Value &&
+		!TIsReferenceType<Arg>::Value,
+		"TIsBitwiseConstructible is not designed to accept reference types");
+
+	static_assert(
+		TAreTypesEqual<T,   typename TRemoveCV<T  >::Type>::Value &&
+		TAreTypesEqual<Arg, typename TRemoveCV<Arg>::Type>::Value,
+		"TIsBitwiseConstructible is not designed to accept qualified types");
+
+	// Assume no bitwise construction in general
+	enum { Value = false };
+};
+
+template <typename T>
+struct TIsBitwiseConstructible<T, T>
+{
+	// Ts can always be bitwise constructed from itself if it is trivially copyable.
+	enum { Value = !TTypeTraits<T>::NeedsCopyConstructor };
+};
+
+template <typename T, typename U>
+struct TIsBitwiseConstructible<const T, U> : TIsBitwiseConstructible<T, U>
+{
+	// Constructing a const T is the same as constructing a T
+};
+
+// Const pointers can be bitwise constructed from non-const pointers.
+// This is not true for pointer conversions in general, e.g. where an offset may need to be applied in the case
+// of multiple inheritance, but there is no way of detecting that at compile-time.
+template <typename T>
+struct TIsBitwiseConstructible<const T*, T*>
+{
+	// Constructing a const T is the same as constructing a T
+	enum { Value = true };
+};
 
+// Unsigned types can be bitwise converted to their signed equivalents, and vice versa.
+// (assuming two's-complement, which we are)
+template <> struct TIsBitwiseConstructible<uint8,   int8>  { enum { Value = true }; };
+template <> struct TIsBitwiseConstructible< int8,  uint8>  { enum { Value = true }; };
+template <> struct TIsBitwiseConstructible<uint16,  int16> { enum { Value = true }; };
+template <> struct TIsBitwiseConstructible< int16, uint16> { enum { Value = true }; };
+template <> struct TIsBitwiseConstructible<uint32,  int32> { enum { Value = true }; };
+template <> struct TIsBitwiseConstructible< int32, uint32> { enum { Value = true }; };
+template <> struct TIsBitwiseConstructible<uint64,  int64> { enum { Value = true }; };
+template <> struct TIsBitwiseConstructible< int64, uint64> { enum { Value = true }; };
