Applied changes from CL# 2331703 to all Clang platforms.

diff --git a/Engine/Source/Runtime/Core/Public/Android/AndroidPlatform.h b/Engine/Source/Runtime/Core/Public/Android/AndroidPlatform.h
--- a/Engine/Source/Runtime/Core/Public/Android/AndroidPlatform.h
+++ b/Engine/Source/Runtime/Core/Public/Android/AndroidPlatform.h
@@ -63,8 +63,13 @@ typedef FAndroidTypes FPlatformTypes;
 #define GCC_ALIGN(n)		__attribute__((aligned(n)))
 #define REQUIRES_ALIGNED_ACCESS 1
 
-// Operator new/delete handling.
+// operator new/delete operators
+// As of 10.9 we need to use _NOEXCEPT & cxx_noexcept compatible definitions
+#if __has_feature(cxx_noexcept)
+#define OPERATOR_NEW_THROW_SPEC
+#else
 #define OPERATOR_NEW_THROW_SPEC throw (std::bad_alloc)
-#define OPERATOR_DELETE_THROW_SPEC throw()
-#define OPERATOR_NEW_NOTHROW_SPEC throw ()
-#define OPERATOR_DELETE_NOTHROW_SPEC throw()
+#endif
+#define OPERATOR_DELETE_THROW_SPEC noexcept
+#define OPERATOR_NEW_NOTHROW_SPEC  noexcept
+#define OPERATOR_DELETE_NOTHROW_SPEC  noexcept
diff --git a/Engine/Source/Runtime/Core/Public/HTML5/HTML5Platform.h b/Engine/Source/Runtime/Core/Public/HTML5/HTML5Platform.h
--- a/Engine/Source/Runtime/Core/Public/HTML5/HTML5Platform.h
+++ b/Engine/Source/Runtime/Core/Public/HTML5/HTML5Platform.h
@@ -91,10 +91,16 @@ typedef FHTML5Types FPlatformTypes;
 
 // Operator new/delete handling.
 #if !PLATFORM_HTML5_WIN32
+// operator new/delete operators
+// As of 10.9 we need to use _NOEXCEPT & cxx_noexcept compatible definitions
+#if __has_feature(cxx_noexcept)
+#define OPERATOR_NEW_THROW_SPEC
+#else
 #define OPERATOR_NEW_THROW_SPEC throw (std::bad_alloc)
-#define OPERATOR_DELETE_THROW_SPEC throw()
-#define OPERATOR_NEW_NOTHROW_SPEC throw ()
-#define OPERATOR_DELETE_NOTHROW_SPEC throw()
+#endif
+#define OPERATOR_DELETE_THROW_SPEC noexcept
+#define OPERATOR_NEW_NOTHROW_SPEC  noexcept
+#define OPERATOR_DELETE_NOTHROW_SPEC  noexcept
 #endif
 
 
diff --git a/Engine/Source/Runtime/Core/Public/IOS/IOSPlatform.h b/Engine/Source/Runtime/Core/Public/IOS/IOSPlatform.h
--- a/Engine/Source/Runtime/Core/Public/IOS/IOSPlatform.h
+++ b/Engine/Source/Runtime/Core/Public/IOS/IOSPlatform.h
@@ -67,11 +67,16 @@ typedef FIOSPlatformTypes FPlatformTypes;
 #define GCC_ALIGN(n) __attribute__((aligned(n)))
 #define REQUIRES_ALIGNED_ACCESS 1
 
-// new/delete operators
+// operator new/delete operators
+// As of 10.9 we need to use _NOEXCEPT & cxx_noexcept compatible definitions
+#if __has_feature(cxx_noexcept)
+#define OPERATOR_NEW_THROW_SPEC
+#else
 #define OPERATOR_NEW_THROW_SPEC throw (std::bad_alloc)
-#define OPERATOR_DELETE_THROW_SPEC throw()
-#define OPERATOR_NEW_NOTHROW_SPEC throw ()
-#define OPERATOR_DELETE_NOTHROW_SPEC throw()
+#endif
+#define OPERATOR_DELETE_THROW_SPEC noexcept
+#define OPERATOR_NEW_NOTHROW_SPEC  noexcept
+#define OPERATOR_DELETE_NOTHROW_SPEC  noexcept
 
 // DLL export and import definitions
 #define DLLEXPORT
diff --git a/Engine/Source/Runtime/Core/Public/Linux/LinuxPlatform.h b/Engine/Source/Runtime/Core/Public/Linux/LinuxPlatform.h
--- a/Engine/Source/Runtime/Core/Public/Linux/LinuxPlatform.h
+++ b/Engine/Source/Runtime/Core/Public/Linux/LinuxPlatform.h
@@ -72,8 +72,13 @@ typedef FLinuxPlatformTypes FPlatformTypes;
 #define GCC_PACK(n)			__attribute__((packed,aligned(n)))
 #define GCC_ALIGN(n)		__attribute__((aligned(n)))
 
-// Operator new/delete handling.
+// operator new/delete operators
+// As of 10.9 we need to use _NOEXCEPT & cxx_noexcept compatible definitions
+#if __has_feature(cxx_noexcept)
+#define OPERATOR_NEW_THROW_SPEC
+#else
 #define OPERATOR_NEW_THROW_SPEC throw (std::bad_alloc)
-#define OPERATOR_DELETE_THROW_SPEC throw()
-#define OPERATOR_NEW_NOTHROW_SPEC throw ()
-#define OPERATOR_DELETE_NOTHROW_SPEC throw()
+#endif
+#define OPERATOR_DELETE_THROW_SPEC noexcept
+#define OPERATOR_NEW_NOTHROW_SPEC  noexcept
+#define OPERATOR_DELETE_NOTHROW_SPEC  noexcept
