Added config setting for audio volume when app is unfocused

Removed GVolumeMultiplier and added an equivalent to FApp to contain the code to load the unfocused volume multiplier from config.

UE-4449 - Allow users to choose whether audio can be heard when focus is lost

diff --git a/Engine/Config/BaseEngine.ini b/Engine/Config/BaseEngine.ini
--- a/Engine/Config/BaseEngine.ini
+++ b/Engine/Config/BaseEngine.ini
@@ -1481,6 +1481,7 @@ UseStrictEngineVersioning=True
 [Audio]
 MaxChannels=32
 CommonAudioPoolSize=0
+UnfocusedVolumeMultiplier=0.0
 
 [/Script/Engine.AudioSettings]
 DefaultSoundClassName=/Engine/EngineSounds/Master.Master
diff --git a/Engine/Source/Editor/UnrealEd/Private/Editor.cpp b/Engine/Source/Editor/UnrealEd/Private/Editor.cpp
--- a/Engine/Source/Editor/UnrealEd/Private/Editor.cpp
+++ b/Engine/Source/Editor/UnrealEd/Private/Editor.cpp
@@ -991,12 +991,12 @@ void UEditorEngine::Tick( float DeltaSeconds, bool bIdleMode )
 		if (!PlayWorld)
 		{
 			// Adjust the global volume multiplier if the window has focus and there is no pie world or no viewport overriding audio.
-			GVolumeMultiplier = GetDefault<ULevelEditorMiscSettings>()->EditorVolumeLevel;
+			FApp::SetVolumeMultiplier( GetDefault<ULevelEditorMiscSettings>()->EditorVolumeLevel );
 		}
 		else
 		{
 			// If there is currently a pie world a viewport is overriding audio settings do not adjust the volume.
-			GVolumeMultiplier = 1.0f;
+			FApp::SetVolumeMultiplier( 1.0f );
 		}
 	}
 
diff --git a/Engine/Source/Runtime/Android/AndroidAudio/Private/AndroidAudioSource.cpp b/Engine/Source/Runtime/Android/AndroidAudio/Private/AndroidAudioSource.cpp
--- a/Engine/Source/Runtime/Android/AndroidAudio/Private/AndroidAudioSource.cpp
+++ b/Engine/Source/Runtime/Android/AndroidAudio/Private/AndroidAudioSource.cpp
@@ -354,7 +354,7 @@ void FSLESSoundSource::Update( void )
 		// Emulate the bleed to rear speakers followed by stereo fold down
 		Volume *= 1.25f;
 	}
-	Volume *= GVolumeMultiplier;
+	Volume *= FApp::GetVolumeMultiplier();
 	Volume = FMath::Clamp(Volume, 0.0f, MAX_VOLUME);
 	
 	const float Pitch = FMath::Clamp<float>(WaveInstance->Pitch, MIN_PITCH, MAX_PITCH);
diff --git a/Engine/Source/Runtime/Core/Private/Mac/MacApplication.cpp b/Engine/Source/Runtime/Core/Private/Mac/MacApplication.cpp
--- a/Engine/Source/Runtime/Core/Private/Mac/MacApplication.cpp
+++ b/Engine/Source/Runtime/Core/Private/Mac/MacApplication.cpp
@@ -182,7 +182,7 @@ FMacApplication::FMacApplication()
 								}
 
 								// app is active, allow sound
-								GVolumeMultiplier = 1.0f;
+								FApp::SetVolumeMultiplier( 1.0f );
 							}];
 
 	AppDeactivationObserver = [[NSNotificationCenter defaultCenter] addObserverForName:NSApplicationDidResignActiveNotification object:[NSApplication sharedApplication] queue:[NSOperationQueue mainQueue] usingBlock:^(NSNotification* Notification){
@@ -229,8 +229,8 @@ FMacApplication::FMacApplication()
 										FPlatformProcess::Sleep(0.005f);
 									}
 
-									// app is inactive, silence it
-									GVolumeMultiplier = 0.0f;
+									// app is inactive, apply multiplier
+									FApp::SetVolumeMultiplier(FApp::GetUnfocusedVolumeMultiplier());
 								}];
 	
 	WorkspaceActivationObserver = [[[NSWorkspace sharedWorkspace] notificationCenter] addObserverForName:NSWorkspaceSessionDidBecomeActiveNotification object:[NSWorkspace sharedWorkspace] queue:[NSOperationQueue mainQueue] usingBlock:^(NSNotification* Notification){
diff --git a/Engine/Source/Runtime/Core/Private/Misc/App.cpp b/Engine/Source/Runtime/Core/Private/Misc/App.cpp
--- a/Engine/Source/Runtime/Core/Private/Misc/App.cpp
+++ b/Engine/Source/Runtime/Core/Private/Misc/App.cpp
@@ -18,6 +18,8 @@ double FApp::FixedDeltaTime = 1 / 30.0;
 double FApp::CurrentTime = 0.0;
 double FApp::LastTime = 0.0;
 double FApp::DeltaTime = 1 / 30.0;
+float FApp::VolumeMultiplier = 1.0f;
+float FApp::UnfocusedVolumeMultiplier = 0.0f;
 
 
 /* FApp static interface
@@ -130,3 +132,15 @@ bool FApp::ShouldUseThreadingForPerformance( )
 
 
 #endif // HAVE_RUNTIME_THREADING_SWITCHES
+
+float FApp::GetUnfocusedVolumeMultiplier()
+{
+	static bool bInitialisedFromConfig = false;
+
+	if (!bInitialisedFromConfig)
+	{
+		bInitialisedFromConfig = true;
+		GConfig->GetFloat(TEXT("Audio"), TEXT("UnfocusedVolumeMultiplier"), UnfocusedVolumeMultiplier, GEngineIni);
+	}
+	return UnfocusedVolumeMultiplier;
+}
diff --git a/Engine/Source/Runtime/Core/Private/Misc/Core.cpp b/Engine/Source/Runtime/Core/Private/Misc/Core.cpp
--- a/Engine/Source/Runtime/Core/Private/Misc/Core.cpp
+++ b/Engine/Source/Runtime/Core/Private/Misc/Core.cpp
@@ -90,9 +90,6 @@ bool GIsReconstructingBlueprintInstances = false;
 /** Force blueprints to not compile on load */
 bool GForceDisableBlueprintCompileOnLoad = false;
 
-
-float					GVolumeMultiplier				= 1.0f;						/* Use to silence the app when it loses focus */
-
 #if WITH_EDITORONLY_DATA
 bool					GIsEditor						= false;					/* Whether engine was launched for editing */
 bool					GIsImportingT3D					= false;					/* Whether editor is importing T3D */
diff --git a/Engine/Source/Runtime/Core/Private/WinRT/WinRTMisc.cpp b/Engine/Source/Runtime/Core/Private/WinRT/WinRTMisc.cpp
--- a/Engine/Source/Runtime/Core/Private/WinRT/WinRTMisc.cpp
+++ b/Engine/Source/Runtime/Core/Private/WinRT/WinRTMisc.cpp
@@ -123,8 +123,8 @@ void FWinRTMisc::PumpMessages(bool bFromMainLoop)
 	// have focus)
 	[REDACTED_EMAIL]: Will this always be true?
 	bool HasFocus = true;
-	// if its our window, allow sound, otherwise silence it
-	GVolumeMultiplier = HasFocus ? 1.0f : 0.0f;
+	// if its our window, allow sound, otherwise apply multiplier
+	FApp::SetVolumeMultiplier(HasFocus ? 1.0f : FApp::GetUnfocusedVolumeMultiplier());
 }
 
 uint32 FWinRTMisc::GetCharKeyMap(uint16* KeyCodes, FString* KeyNames, uint32 MaxMappings)
diff --git a/Engine/Source/Runtime/Core/Private/Windows/WindowsPlatformMisc.cpp b/Engine/Source/Runtime/Core/Private/Windows/WindowsPlatformMisc.cpp
--- a/Engine/Source/Runtime/Core/Private/Windows/WindowsPlatformMisc.cpp
+++ b/Engine/Source/Runtime/Core/Private/Windows/WindowsPlatformMisc.cpp
@@ -801,8 +801,8 @@ void FWindowsPlatformMisc::PumpMessages(bool bFromMainLoop)
 		HadFocus = HasFocus;
 	}
 
-	// if its our window, allow sound, otherwise silence it
-	GVolumeMultiplier = HasFocus ? 1.0f : 0.0f;
+	// if its our window, allow sound, otherwise apply multiplier
+	FApp::SetVolumeMultiplier( HasFocus ? 1.0f : FApp::GetUnfocusedVolumeMultiplier() );
 }
 
 uint32 FWindowsPlatformMisc::GetCharKeyMap(uint16* KeyCodes, FString* KeyNames, uint32 MaxMappings)
diff --git a/Engine/Source/Runtime/Core/Public/CoreGlobals.h b/Engine/Source/Runtime/Core/Public/CoreGlobals.h
--- a/Engine/Source/Runtime/Core/Public/CoreGlobals.h
+++ b/Engine/Source/Runtime/Core/Public/CoreGlobals.h
@@ -56,8 +56,6 @@ extern CORE_API bool GForceDisableBlueprintCompileOnLoad;
 /** Helper function to flush resource streaming. */
 extern CORE_API void(*GFlushStreamingFunc)(void);
 
-extern CORE_API float GVolumeMultiplier;
-
 
 #if WITH_EDITORONLY_DATA
 
diff --git a/Engine/Source/Runtime/Core/Public/Misc/App.h b/Engine/Source/Runtime/Core/Public/Misc/App.h
--- a/Engine/Source/Runtime/Core/Public/Misc/App.h
+++ b/Engine/Source/Runtime/Core/Public/Misc/App.h
@@ -388,6 +388,33 @@ public:
 		DeltaTime = Seconds;
 	}
 
+	/**
+	 * Get Volume Multiplier
+	 * 
+	 * @return Current volume multiplier
+	 */
+	FORCEINLINE static float GetVolumeMultiplier()
+	{
+		return VolumeMultiplier;
+	}
+
+	/**
+	 * Set Volume Multiplier
+	 * 
+	 * @param  InVolumeMultiplier	new volume multiplier
+	 */
+	FORCEINLINE static void SetVolumeMultiplier(float InVolumeMultiplier)
+	{
+		VolumeMultiplier = InVolumeMultiplier;
+	}
+
+	/**
+	 * Helper function to get UnfocusedVolumeMultiplier from config and store so it's not retrieved every frame
+	 * 
+	 * @return Volume multiplier to use when app loses focus
+	 */
+	static float GetUnfocusedVolumeMultiplier();
+
 private:
 
 	/** Holds the instance identifier. */
@@ -419,6 +446,12 @@ private:
 
 	/** Holds current delta time in seconds. */
 	static double DeltaTime;
+
+	/** Use to affect the app volume when it loses focus */
+	static float VolumeMultiplier;
+
+	/** Read from config to define the volume when app loses focus */
+	static float UnfocusedVolumeMultiplier;
 };
 
 
diff --git a/Engine/Source/Runtime/HTML5/HTML5Audio/Private/HTML5AudioSource.cpp b/Engine/Source/Runtime/HTML5/HTML5Audio/Private/HTML5AudioSource.cpp
--- a/Engine/Source/Runtime/HTML5/HTML5Audio/Private/HTML5AudioSource.cpp
+++ b/Engine/Source/Runtime/HTML5/HTML5Audio/Private/HTML5AudioSource.cpp
@@ -81,7 +81,7 @@ void FALSoundSource::Update( void )
 		// Emulate the bleed to rear speakers followed by stereo fold down
 		Volume *= 1.25f;
 	}
-	Volume *= GVolumeMultiplier;
+	Volume *= FApp::GetVolumeMultiplier();
 
  			Volume	= 	FMath::Clamp(Volume, 0.0f, MAX_VOLUME ); 
 	float	Pitch	=	FMath::Clamp(WaveInstance->Pitch, MIN_PITCH, MAX_PITCH ); 
diff --git a/Engine/Source/Runtime/IOS/IOSAudio/Private/IOSAudioSource.cpp b/Engine/Source/Runtime/IOS/IOSAudio/Private/IOSAudioSource.cpp
--- a/Engine/Source/Runtime/IOS/IOSAudio/Private/IOSAudioSource.cpp
+++ b/Engine/Source/Runtime/IOS/IOSAudio/Private/IOSAudioSource.cpp
@@ -251,7 +251,7 @@ void FIOSAudioSoundSource::Update(void)
 	}
 
 	// Apply global multiplier to disable sound when not the foreground app
-	Volume = FMath::Clamp(Volume, 0.0f, MAX_VOLUME) * GVolumeMultiplier;
+	Volume = FMath::Clamp(Volume, 0.0f, MAX_VOLUME) * FApp::GetVolumeMultiplier();
 
 	// Convert to dB
 	const AudioUnitParameterValue Gain = FMath::Clamp<float>(20.0f * log10(Volume), -100, 20.0f);
diff --git a/Engine/Source/Runtime/Linux/ALAudio/Private/ALAudioSource.cpp b/Engine/Source/Runtime/Linux/ALAudio/Private/ALAudioSource.cpp
--- a/Engine/Source/Runtime/Linux/ALAudio/Private/ALAudioSource.cpp
+++ b/Engine/Source/Runtime/Linux/ALAudio/Private/ALAudioSource.cpp
@@ -84,7 +84,7 @@ void FALSoundSource::Update( void )
 		// Emulate the bleed to rear speakers followed by stereo fold down
 		Volume *= 1.25f;
 	}
-	Volume *= GVolumeMultiplier;
+	Volume *= FApp::GetVolumeMultiplier();
 
  			Volume	= 	FMath::Clamp(Volume, 0.0f, MAX_VOLUME ); 
 	float	Pitch	=	FMath::Clamp(WaveInstance->Pitch, MIN_PITCH, MAX_PITCH ); 
diff --git a/Engine/Source/Runtime/Mac/CoreAudio/Private/CoreAudioSource.cpp b/Engine/Source/Runtime/Mac/CoreAudio/Private/CoreAudioSource.cpp
--- a/Engine/Source/Runtime/Mac/CoreAudio/Private/CoreAudioSource.cpp
+++ b/Engine/Source/Runtime/Mac/CoreAudio/Private/CoreAudioSource.cpp
@@ -265,7 +265,7 @@ void FCoreAudioSoundSource::Update( void )
 		}
 		
 		// apply global multiplier (ie to disable sound when not the foreground app)
-		Volume *= GVolumeMultiplier;
+		Volume *= FApp::GetVolumeMultiplier();
 		Volume = FMath::Clamp<float>( Volume, 0.0f, MAX_VOLUME );
 		
 		// Convert to dB
@@ -391,7 +391,7 @@ void FCoreAudioSoundSource::Update( void )
 	else
 	{
 		// apply global multiplier (ie to disable sound when not the foreground app)
-		Volume *= GVolumeMultiplier;
+		Volume *= FApp::GetVolumeMultiplier();
 		Volume = FMath::Clamp<float>( Volume, 0.0f, MAX_VOLUME );
 		
 		if( AudioDevice->GetMixDebugState() == DEBUGSTATE_IsolateReverb )
diff --git a/Engine/Source/Runtime/Windows/XAudio2/Private/XAudio2Effects.cpp b/Engine/Source/Runtime/Windows/XAudio2/Private/XAudio2Effects.cpp
--- a/Engine/Source/Runtime/Windows/XAudio2/Private/XAudio2Effects.cpp
+++ b/Engine/Source/Runtime/Windows/XAudio2/Private/XAudio2Effects.cpp
@@ -197,7 +197,7 @@ public:
 
 		// Check the global volume multiplier because this effect 
 		// will continue to play if the editor loses focus.
-		if( IsEnabled && GVolumeMultiplier != 0.0f )
+		if (IsEnabled && FApp::GetVolumeMultiplier() != 0.0f)
 		{
 			FAudioRadioEffect* RadioParameters = ( FAudioRadioEffect* )BeginProcess();
 			check( RadioParameters );
diff --git a/Engine/Source/Runtime/Windows/XAudio2/Private/XAudio2Source.cpp b/Engine/Source/Runtime/Windows/XAudio2/Private/XAudio2Source.cpp
--- a/Engine/Source/Runtime/Windows/XAudio2/Private/XAudio2Source.cpp
+++ b/Engine/Source/Runtime/Windows/XAudio2/Private/XAudio2Source.cpp
@@ -398,7 +398,7 @@ bool FXAudio2SoundSource::Init( FWaveInstance* InWaveInstance )
  */
 void FXAudio2SoundSource::GetChannelVolumes( float ChannelVolumes[CHANNELOUT_COUNT], float AttenuatedVolume )
 {
-	if( GVolumeMultiplier == 0.0f )
+	if (FApp::GetVolumeMultiplier() == 0.0f)
 	{
 		for( int32 i = 0; i < CHANNELOUT_COUNT; i++ )
 		{
@@ -584,7 +584,7 @@ void FXAudio2SoundSource::GetChannelVolumes( float ChannelVolumes[CHANNELOUT_COU
 			UE_LOG(LogXAudio2, Warning, TEXT("FXAudio2SoundSource contains unreasonble value %f in channel %d: %s"), ChannelVolumes[i], i, *Describe_Internal(true, false));
 		}
 
-		ChannelVolumes[i] = FMath::Clamp<float>(ChannelVolumes[i] * GVolumeMultiplier, 0.0f, MAX_VOLUME);
+		ChannelVolumes[i] = FMath::Clamp<float>(ChannelVolumes[i] * FApp::GetVolumeMultiplier(), 0.0f, MAX_VOLUME);
 	}
 }
 
