Stats - Missing 'toggledebug' command added

diff --git a/Engine/Source/Runtime/Core/Private/Stats/StatsCommand.cpp b/Engine/Source/Runtime/Core/Private/Stats/StatsCommand.cpp
--- a/Engine/Source/Runtime/Core/Private/Stats/StatsCommand.cpp
+++ b/Engine/Source/Runtime/Core/Private/Stats/StatsCommand.cpp
@@ -1075,6 +1075,8 @@ static void PrintStatsHelpToOutputDevice( FOutputDevice& Ar )
 
 	Ar.Log( TEXT("stat startfile - starts dumping a capture"));
 	Ar.Log( TEXT("stat stopfile - stops dumping a capture"));
+
+	Ar.Log( TEXT("stat toggledebug - toggles tracking the most memory expensive stats"));
 }
 
 static void CommandTestFile()
@@ -1226,6 +1228,10 @@ static void StatCmd(FString InCmd)
 	{
 		IStatGroupEnableManager::Get().StatGroupEnableManagerCommand(Cmd);
 	}
+	else if ( FParse::Command( &Cmd, TEXT( "toggledebug" ) ) )
+	{
+		FStatsThreadState::GetLocalState().ToggleFindMemoryExtensiveStats();
+	}
 	// @see FStatHierParams
 	else if ( FParse::Command( &Cmd, TEXT( "hier" ) ) )
 	{
@@ -1351,6 +1357,9 @@ bool DirectStatsCommand(const TCHAR* Cmd, bool bBlockForCompletion /*= false*/,
 		else if ( FParse::Command( &TempCmd, TEXT( "net" ) ) )
 		{
 		}
+		else if ( FParse::Command( &TempCmd, TEXT( "toggledebug" ) ) )
+		{
+		}
 		else
 #endif
 		{
