Stats - Removed code that is no longer needed, was replaced with FStats::GameThreadStatsFrame

diff --git a/Engine/Source/Runtime/Core/Private/Stats/StatsData.cpp b/Engine/Source/Runtime/Core/Private/Stats/StatsData.cpp
--- a/Engine/Source/Runtime/Core/Private/Stats/StatsData.cpp
+++ b/Engine/Source/Runtime/Core/Private/Stats/StatsData.cpp
@@ -559,8 +559,7 @@ void FStatsThreadState::ScanForAdvance(const FStatMessagesArray& Data)
 					BadFrames.Add(Frame - 1);
 				}
 			}
-			// CurrentGameFrame can be accessed by other threads.
-			FPlatformAtomics::InterlockedExchange( &CurrentGameFrame, NewGameFrame );
+			CurrentGameFrame = NewGameFrame;
 		}
 		else if (Op == EStatOperation::AdvanceFrameEventRenderThread)
 		{
