changed cvars Set() behavior: Each Set should specify a priority, order no longer matters if the priority is higher (solved a lot of corner cases), might need some iterations.

diff --git a/Engine/Source/Editor/LevelEditor/Private/LevelEditorActions.cpp b/Engine/Source/Editor/LevelEditor/Private/LevelEditorActions.cpp
--- a/Engine/Source/Editor/LevelEditor/Private/LevelEditorActions.cpp
+++ b/Engine/Source/Editor/LevelEditor/Private/LevelEditorActions.cpp
@@ -449,7 +449,7 @@ void FLevelEditorActionCallbacks::SetMaterialQualityLevel( EMaterialQualityLevel
 
 	//Ensure the material quality cvar is also set.
 	static IConsoleVariable* MaterialQualityLevelVar = IConsoleManager::Get().FindConsoleVariable(TEXT("r.MaterialQualityLevel"));
-	MaterialQualityLevelVar->Set(NewQualityLevel, ECVF_SetByEditorSetting);
+	MaterialQualityLevelVar->Set(NewQualityLevel, ECVF_SetByConsole);
 
 	GUnrealEd->RedrawAllViewports();
 }
@@ -1911,7 +1911,7 @@ void FLevelEditorActionCallbacks::OnToggleFreezeParticleSimulation()
 	IConsoleVariable* CVar = ConsoleManager.FindConsoleVariable(TEXT("FX.FreezeParticleSimulation"));
 	if (CVar)
 	{
-		CVar->Set(CVar->GetInt() == 0 ? 1 : 0, ECVF_SetByEditorSetting);
+		CVar->Set(CVar->GetInt() == 0 ? 1 : 0, ECVF_SetByConsole);
 	}
 }
 
diff --git a/Engine/Source/Editor/UnrealEd/Private/Editor/EditorUserSettings.cpp b/Engine/Source/Editor/UnrealEd/Private/Editor/EditorUserSettings.cpp
--- a/Engine/Source/Editor/UnrealEd/Private/Editor/EditorUserSettings.cpp
+++ b/Engine/Source/Editor/UnrealEd/Private/Editor/EditorUserSettings.cpp
@@ -21,7 +21,7 @@ void UEditorUserSettings::PostInitProperties()
 
 	//Ensure the material quality cvar is set to the settings loaded.
 	static IConsoleVariable* MaterialQualityLevelVar = IConsoleManager::Get().FindConsoleVariable(TEXT("r.MaterialQualityLevel"));
-	MaterialQualityLevelVar->Set(MaterialQualityLevel, ECVF_SetByEditorSetting);
+	MaterialQualityLevelVar->Set(MaterialQualityLevel, ECVF_SetByConsole);
 }
 
 void UEditorUserSettings::PostEditChangeProperty( FPropertyChangedEvent& PropertyChangedEvent )
diff --git a/Engine/Source/Runtime/Core/Private/HAL/ConsoleManager.cpp b/Engine/Source/Runtime/Core/Private/HAL/ConsoleManager.cpp
--- a/Engine/Source/Runtime/Core/Private/HAL/ConsoleManager.cpp
+++ b/Engine/Source/Runtime/Core/Private/HAL/ConsoleManager.cpp
@@ -14,6 +14,31 @@ DEFINE_LOG_CATEGORY_STATIC(LogConsoleManager, Log, All);
 
 static inline bool IsWhiteSpace(TCHAR Value) { return Value == TCHAR(' '); }
 
+// Get human readable string
+// @return never 0
+static TCHAR* GetSetByTCHAR(EConsoleVariableFlags InSetBy)
+{
+	EConsoleVariableFlags SetBy = (EConsoleVariableFlags)((uint32)InSetBy & ECVF_SetByMask);
+
+	switch(SetBy)
+	{
+		// CASE() macro was disliked by Mac compiler (Clang)
+		//   error: ISO C++11 does not allow conversion from string literal to 'TCHAR *' (aka 'wchar_t *') [-Werror,-Wc++11-compat-deprecated-writable-strings]
+		// so we do it with more code duplication. Even beter would be to use enum reflection instead.
+		case ECVF_SetByConstructor:			return TEXT("Constructor");
+		case ECVF_SetByScalability:			return TEXT("Scalability");
+		case ECVF_SetByGameSetting:			return TEXT("GameSetting");
+		case ECVF_SetByProjectSetting:		return TEXT("EditorSetting");
+		case ECVF_SetByDeviceProfile:		return TEXT("DeviceProfile");
+		case ECVF_SetBySystemSettingsIni:	return TEXT("SystemSettingsIni");
+		case ECVF_SetByConsoleVariablesIni: return TEXT("ConsoleVariablesIni");
+		case ECVF_SetByCommandline:			return TEXT("Commandline");
+		case ECVF_SetByCode:				return TEXT("Code");
+		case ECVF_SetByConsole:				return TEXT("Console");
+	}
+	return TEXT("<UNKNOWN>");
+}
+
 class FConsoleVariableBase : public IConsoleVariable
 {
 public:
@@ -59,8 +84,36 @@ public:
 
 	// ------
 
+	bool CanChange(EConsoleVariableFlags SetBy) const
+	{
+		uint32 OldPri =	(uint32)Flags & ECVF_SetByMask;
+		uint32 NewPri =	(uint32)SetBy & ECVF_SetByMask;
+
+		bool bRet = NewPri >= OldPri;
+
+		if(!bRet)
+		{
+			FConsoleManager& ConsoleManager = (FConsoleManager&)IConsoleManager::Get();
+			FString CVarName = ConsoleManager.FindConsoleObjectName(this);
+			UE_LOG(LogConsoleManager, Warning,
+				TEXT("Console variable '%s' wasn't set (Priority %s < %s)"),
+				CVarName.IsEmpty() ? TEXT("unknown?") : *CVarName,
+				GetSetByTCHAR((EConsoleVariableFlags)OldPri),
+				GetSetByTCHAR((EConsoleVariableFlags)NewPri)
+				);
+		}
+
+		return bRet;
+	}
+
 	void OnChanged(EConsoleVariableFlags SetBy)
 	{
+		// you have to specify a SetBy e.g. ECVF_SetByCommandline
+		check((uint32)SetBy & ECVF_SetByMask);
+
+		// double check, if this fires we miss a if(CanChange(SetBy))
+		check(CanChange(SetBy));
+
 		// only change on main thread
 
 		Flags = (EConsoleVariableFlags)((uint32)Flags | ECVF_Changed);
@@ -208,8 +261,11 @@ public:
 	} 
 	virtual void Set(const TCHAR* InValue, EConsoleVariableFlags SetBy)
 	{
-		TTypeFromString<T>::FromString(Data.ShadowedValue[0], InValue);
-		OnChanged(SetBy);
+		if(CanChange(SetBy))
+		{
+			TTypeFromString<T>::FromString(Data.ShadowedValue[0], InValue);
+			OnChanged(SetBy);
+		}
 	}
 
 	virtual int32 GetInt() const;
@@ -281,8 +337,11 @@ template<> TConsoleVariableData<float>* FConsoleVariable<float>::AsVariableFloat
 
 template<> void FConsoleVariable<FString>::Set(const TCHAR* InValue, EConsoleVariableFlags SetBy)
 {
-	Data.ShadowedValue[0] = InValue;
-	OnChanged(SetBy);
+	if(CanChange(SetBy))
+	{
+		Data.ShadowedValue[0] = InValue;
+		OnChanged(SetBy);
+	}
 }
 template<> int32 FConsoleVariable<FString>::GetInt() const
 {
@@ -322,8 +381,11 @@ public:
 	} 
 	virtual void Set(const TCHAR* InValue, EConsoleVariableFlags SetBy)
 	{
-		TTypeFromString<T>::FromString(MainValue, InValue);
-		OnChanged(SetBy);
+		if(CanChange(SetBy))
+		{
+			TTypeFromString<T>::FromString(MainValue, InValue);
+			OnChanged(SetBy);
+		}
 	}
 	virtual int32 GetInt() const
 	{
@@ -354,9 +416,12 @@ private: // ----------------------------------------------------
 
 	void OnChanged(EConsoleVariableFlags SetBy)
 	{
-		// propagate from main thread to render thread or to reference
-		OnCVarChange(RefValue, MainValue, Flags);
-		FConsoleVariableBase::OnChanged(SetBy);
+		if(CanChange(SetBy))
+		{
+			// propagate from main thread to render thread or to reference
+			OnCVarChange(RefValue, MainValue, Flags);
+			FConsoleVariableBase::OnChanged(SetBy);
+		}
 	}
 };
 
@@ -387,14 +452,17 @@ public:
 	} 
 	virtual void Set(const TCHAR* InValue, EConsoleVariableFlags SetBy)
 	{
-		int32 Value = FCString::Atoi(InValue);
+		if(CanChange(SetBy))
+		{
+			int32 Value = FCString::Atoi(InValue);
 
-		check(IsInGameThread());
+			check(IsInGameThread());
 
-		FMath::SetBoolInBitField(Force0MaskPtr, BitNumber, Value == 0);
-		FMath::SetBoolInBitField(Force1MaskPtr, BitNumber, Value == 1);
+			FMath::SetBoolInBitField(Force0MaskPtr, BitNumber, Value == 0);
+			FMath::SetBoolInBitField(Force1MaskPtr, BitNumber, Value == 1);
 
-		OnChanged(SetBy);
+			OnChanged(SetBy);
+		}
 	}
 	virtual int32 GetInt() const
 	{
@@ -433,30 +501,6 @@ IConsoleVariable* FConsoleManager::RegisterConsoleVariableBitRef(const TCHAR* CV
 	return AddConsoleObject(CVarName, new FConsoleVariableBitRef(FlagName, BitNumber, Force0MaskPtr, Force1MaskPtr, Help, (EConsoleVariableFlags)Flags))->AsVariable();
 }
 
-
-static TCHAR* GetLastSetBy(EConsoleVariableFlags InSetBy)
-{
-	EConsoleVariableFlags SetBy = (EConsoleVariableFlags)((uint32)InSetBy & ECVF_SetByMask);
-
-	switch(SetBy)
-	{
-#define CASE(A) case ECVF_SetBy##A: return TEXT(#A);
-		CASE(Constructor)
-		CASE(Code)
-		CASE(Scalability)
-		CASE(ConsoleVariablesIni)
-		CASE(SystemSettingsIni)
-		CASE(DeviceProfile)
-		CASE(Console)
-		CASE(Commandline)
-		CASE(EditorSetting)
-		CASE(GameSetting)
-		CASE(Network)
-#undef CASE 
-	}
-	return TEXT("<UNKNOWN>");
-}
-
 void FConsoleManager::CallAllConsoleVariableSinks()
 {
 	if(bCallAllConsoleVariableSinks)
@@ -918,7 +962,7 @@ bool FConsoleManager::ProcessUserConsoleInput(const TCHAR* InInput, FOutputDevic
 		if(*It == 0)
 		{
 			// get current state
-			Ar.Logf(TEXT("%s = \"%s\"      LastSetBy: %s"), *Param1, *CVar->GetString(), GetLastSetBy(CVar->GetFlags()));
+			Ar.Logf(TEXT("%s = \"%s\"      LastSetBy: %s"), *Param1, *CVar->GetString(), GetSetByTCHAR(CVar->GetFlags()));
 		}
 		else
 		{
@@ -1253,6 +1297,9 @@ void FConsoleManager::Test()
 		IConsoleVariable* VarD = IConsoleManager::Get().RegisterConsoleVariableRef(TEXT("TestNameD"), RefD, TEXT("TestHelpD"), Flags);
 		IConsoleVariable* VarE = IConsoleManager::Get().RegisterConsoleVariableRef(TEXT("TestNameE"), RefE, TEXT("TestHelpE"), Flags);
 
+		// at the moment ECVF_SetByConstructor has to be 0 or we set ECVF_Default to ECVF_SetByConstructor
+		check((VarA->GetFlags() & ECVF_SetByMask) == ECVF_SetByConstructor);
+
 		GConsoleVariableCallbackTestCounter = 0;
 		VarB->SetOnChangedCallback(FConsoleVariableDelegate::CreateStatic(&TestConsoleVariableCallback));
 		check(GConsoleVariableCallbackTestCounter == 0);
@@ -1293,10 +1340,10 @@ void FConsoleManager::Test()
 
 		// call Set(string)
 
-		VarA->Set(TEXT("3.1"), ECVF_SetByConsole);
-		VarB->Set(TEXT("3.1"), ECVF_SetByConsole);
-		VarD->Set(TEXT("3.1"), ECVF_SetByConsole);
-		VarE->Set(TEXT("3.1"), ECVF_SetByConsole);
+		VarA->Set(TEXT("3.1"), ECVF_SetByConsoleVariablesIni);
+		VarB->Set(TEXT("3.1"), ECVF_SetByConsoleVariablesIni);
+		VarD->Set(TEXT("3.1"), ECVF_SetByConsoleVariablesIni);
+		VarE->Set(TEXT("3.1"), ECVF_SetByConsoleVariablesIni);
 
 		check(GConsoleVariableCallbackTestCounter == 1);
 
@@ -1316,7 +1363,7 @@ void FConsoleManager::Test()
 		check(RefD == 3);
 		check(RefE == 3.1f);
 
-		VarB->Set(TEXT("3.1"), ECVF_SetByCode);
+		VarB->Set(TEXT("3.1"), ECVF_SetByConsoleVariablesIni);
 		check(GConsoleVariableCallbackTestCounter == 2);
 
 		// unregister
@@ -1350,12 +1397,43 @@ void FConsoleManager::Test()
 			check(FMath::IsNearlyEqual(VarC->GetFloat(), 1.23f, KINDA_SMALL_NUMBER));
 			check(VarC->GetString() == FString(TEXT("1.23")));
 			check(!VarC->TestFlags(ECVF_Changed));
-			VarC->Set(TEXT("3.1"), ECVF_SetByCode);
+			VarC->Set(TEXT("3.1"), ECVF_SetByConsole);
 			check(VarC->TestFlags(ECVF_Changed));
 			check(VarC->GetString() == FString(TEXT("3.1")));
 			UnregisterConsoleObject(TEXT("TestNameC"), false);
 			check(!IConsoleManager::Get().FindConsoleVariable(TEXT("TestNameC")));
 		}
+
+		// verify priority
+		{
+			IConsoleVariable* VarX = IConsoleManager::Get().RegisterConsoleVariable(TEXT("TestNameX"), 1, TEXT("TestHelpX"), Flags);
+			check(((uint32)VarX->GetFlags() & ECVF_SetByMask) == ECVF_SetByConstructor);
+
+			VarX->Set(TEXT("3.1"), ECVF_SetByConsoleVariablesIni);
+			check(((uint32)VarX->GetFlags() & ECVF_SetByMask) == ECVF_SetByConsoleVariablesIni);
+
+			// lower should fail
+			VarX->Set(TEXT("111"), ECVF_SetByScalability);
+			check(VarX->GetString() == FString(TEXT("3")));
+			check(((uint32)VarX->GetFlags() & ECVF_SetByMask) == ECVF_SetByConsoleVariablesIni);
+
+			// higher should work
+			VarX->Set(TEXT("222"), ECVF_SetByCommandline);
+			check(VarX->GetString() == FString(TEXT("222")));
+			check(((uint32)VarX->GetFlags() & ECVF_SetByMask) == ECVF_SetByCommandline);
+
+			// lower should fail
+			VarX->Set(TEXT("333"), ECVF_SetByConsoleVariablesIni);
+			check(VarX->GetString() == FString(TEXT("222")));
+			check(((uint32)VarX->GetFlags() & ECVF_SetByMask) == ECVF_SetByCommandline);
+
+			// higher should work
+			VarX->Set(TEXT("444"), ECVF_SetByConsole);
+			check(VarX->GetString() == FString(TEXT("444")));
+			check(((uint32)VarX->GetFlags() & ECVF_SetByMask) == ECVF_SetByConsole);
+
+			IConsoleManager::Get().UnregisterConsoleObject(VarX, false);
+		}
 	}
 
 	// this should trigger the callback
diff --git a/Engine/Source/Runtime/Core/Public/HAL/IConsoleManager.h b/Engine/Source/Runtime/Core/Public/HAL/IConsoleManager.h
--- a/Engine/Source/Runtime/Core/Public/HAL/IConsoleManager.h
+++ b/Engine/Source/Runtime/Core/Public/HAL/IConsoleManager.h
@@ -33,7 +33,7 @@
 enum EConsoleVariableFlags
 {
 	/**
-	 * Default, no flags are set
+	 * Default, no flags are set, the value is set by the constructor 
 	 */
 	ECVF_Default = 0x0,
 	/**
@@ -82,19 +82,35 @@ enum EConsoleVariableFlags
 	/* those cvars control other cvars with the flag ECVF_Scalability, names should start with "sg." */
 	ECVF_ScalabilityGroup = 0x80,
 
+	// ------------------------------------------------
+
 	/* to get some history of where the last value was set by ( useful for track down why a cvar is in a specific state */
 	ECVF_SetByMask =				0xff000000,
+	
+	// the ECVF_SetBy are sorted in override order (weak to strong), the value is not serialized, it only affects it's override behavior when calling Set()
+
+	// lowest priority (default after console variable creation)
 	ECVF_SetByConstructor =			0x00000000,
-	ECVF_SetByCode =				0x01000000,	// least useful, likely a hack, maybe better to start a new SetBy...
-	ECVF_SetByScalability =			0x02000000,
-	ECVF_SetByConsoleVariablesIni = 0x03000000,
-	ECVF_SetBySystemSettingsIni =	0x04000000,
-	ECVF_SetByDeviceProfile =		0x05000000,
-	ECVF_SetByConsole =				0x06000000,
-	ECVF_SetByCommandline =			0x07000000,	// a minus comman e.g. -VSync
-	ECVF_SetByEditorSetting =		0x08000000,
-	ECVF_SetByGameSetting =			0x09000000,
-	ECVF_SetByNetwork =				0x0a000000,
+	// from Scalability.ini (lower priority than game settings so it's easier to override partially)
+	ECVF_SetByScalability =			0x01000000,
+	// (in game UI or from file)
+	ECVF_SetByGameSetting =			0x02000000,
+	// project settings (editor UI or from file, higher priority than game setting to allow to enforce some setting fro this project)
+	ECVF_SetByProjectSetting =		0x03000000,
+	// per device setting (e.g. specific iOS device, higher priority than per project to do device specific settings)
+	ECVF_SetByDeviceProfile =		0x04000000,
+	// per project setting (ini file e.g. Engine.ini or Game.ini)
+	ECVF_SetBySystemSettingsIni =	0x05000000,
+	// consolevariables.ini (for multiple projects)
+	ECVF_SetByConsoleVariablesIni = 0x06000000,
+	// a minus command e.g. -VSync (very high priority to enforce the setting for the application)
+	ECVF_SetByCommandline =			0x07000000,
+	// least useful, likely a hack, maybe better to find the correct SetBy...
+	ECVF_SetByCode =				0x08000000,
+	// editor UI or consoel in game or editor
+	ECVF_SetByConsole =				0x09000000,
+
+	// ------------------------------------------------
 };
 
 class IConsoleVariable;
@@ -427,7 +443,7 @@ struct CORE_API IConsoleManager
 	 * @param	Help		Help text for this command
 	 * @param	Flags		Optional flags bitmask
 	 */
-	virtual IConsoleCommand* RegisterConsoleCommand(const TCHAR* Name, const TCHAR* Help, uint32 Flags = ECVF_Default) = 0;
+	virtual IConsoleCommand* RegisterConsoleCommand(const TCHAR* Name, const TCHAR* Help, uint32 Flags = (uint32)ECVF_Default) = 0;
 
 	/**
 	 * Unregisters a console object, if that object was registered. O(n), n is the console object count
diff --git a/Engine/Source/Runtime/Engine/Classes/Engine/NetworkSettings.h b/Engine/Source/Runtime/Engine/Classes/Engine/NetworkSettings.h
--- a/Engine/Source/Runtime/Engine/Classes/Engine/NetworkSettings.h
+++ b/Engine/Source/Runtime/Engine/Classes/Engine/NetworkSettings.h
@@ -6,7 +6,7 @@
 
 
 /**
- * Implements project settings for the Rendering sub-system.
+ * Implements project settings for the Network sub-system.
  */
 UCLASS(config=Engine, defaultconfig)
 class ENGINE_API UNetworkSettings
diff --git a/Engine/Source/Runtime/Engine/Private/NetworkSettings.cpp b/Engine/Source/Runtime/Engine/Private/NetworkSettings.cpp
--- a/Engine/Source/Runtime/Engine/Private/NetworkSettings.cpp
+++ b/Engine/Source/Runtime/Engine/Private/NetworkSettings.cpp
@@ -60,19 +60,19 @@ void UNetworkSettings::PostEditChangeProperty(FPropertyChangedEvent& PropertyCha
 				UByteProperty* ByteProperty = Cast<UByteProperty>(PropertyChangedEvent.Property);
 				if (ByteProperty != NULL && ByteProperty->Enum != NULL)
 				{
-					CVar->Set(ByteProperty->GetPropertyValue_InContainer(this), ECVF_SetByNetwork);
+					CVar->Set(ByteProperty->GetPropertyValue_InContainer(this), ECVF_SetByProjectSetting);
 				}
 				else if (UBoolProperty* BoolProperty = Cast<UBoolProperty>(PropertyChangedEvent.Property))
 				{
-					CVar->Set((int32)BoolProperty->GetPropertyValue_InContainer(this), ECVF_SetByNetwork);
+					CVar->Set((int32)BoolProperty->GetPropertyValue_InContainer(this), ECVF_SetByProjectSetting);
 				}
 				else if (UIntProperty* IntProperty = Cast<UIntProperty>(PropertyChangedEvent.Property))
 				{
-					CVar->Set(IntProperty->GetPropertyValue_InContainer(this), ECVF_SetByNetwork);
+					CVar->Set(IntProperty->GetPropertyValue_InContainer(this), ECVF_SetByProjectSetting);
 				}
 				else if (UFloatProperty* FloatProperty = Cast<UFloatProperty>(PropertyChangedEvent.Property))
 				{
-					CVar->Set(FloatProperty->GetPropertyValue_InContainer(this), ECVF_SetByNetwork);
+					CVar->Set(FloatProperty->GetPropertyValue_InContainer(this), ECVF_SetByProjectSetting);
 				}
 			}
 		}
diff --git a/Engine/Source/Runtime/Engine/Private/RenderingSettings.cpp b/Engine/Source/Runtime/Engine/Private/RenderingSettings.cpp
--- a/Engine/Source/Runtime/Engine/Private/RenderingSettings.cpp
+++ b/Engine/Source/Runtime/Engine/Private/RenderingSettings.cpp
@@ -79,19 +79,19 @@ void URendererSettings::PostEditChangeProperty(FPropertyChangedEvent& PropertyCh
 				UByteProperty* ByteProperty = Cast<UByteProperty>(PropertyChangedEvent.Property);
 				if (ByteProperty != NULL && ByteProperty->Enum != NULL)
 				{
-					CVar->Set(ByteProperty->GetPropertyValue_InContainer(this), ECVF_SetByEditorSetting);
+					CVar->Set(ByteProperty->GetPropertyValue_InContainer(this), ECVF_SetByProjectSetting);
 				}
 				else if (UBoolProperty* BoolProperty = Cast<UBoolProperty>(PropertyChangedEvent.Property))
 				{
-					CVar->Set((int32)BoolProperty->GetPropertyValue_InContainer(this), ECVF_SetByEditorSetting);
+					CVar->Set((int32)BoolProperty->GetPropertyValue_InContainer(this), ECVF_SetByProjectSetting);
 				}
 				else if (UIntProperty* IntProperty = Cast<UIntProperty>(PropertyChangedEvent.Property))
 				{
-					CVar->Set(IntProperty->GetPropertyValue_InContainer(this), ECVF_SetByEditorSetting);
+					CVar->Set(IntProperty->GetPropertyValue_InContainer(this), ECVF_SetByProjectSetting);
 				}
 				else if (UFloatProperty* FloatProperty = Cast<UFloatProperty>(PropertyChangedEvent.Property))
 				{
-					CVar->Set(FloatProperty->GetPropertyValue_InContainer(this), ECVF_SetByEditorSetting);
+					CVar->Set(FloatProperty->GetPropertyValue_InContainer(this), ECVF_SetByProjectSetting);
 				}
 			}
 			else
diff --git a/Engine/Source/Runtime/Engine/Private/SystemSettings.cpp b/Engine/Source/Runtime/Engine/Private/SystemSettings.cpp
--- a/Engine/Source/Runtime/Engine/Private/SystemSettings.cpp
+++ b/Engine/Source/Runtime/Engine/Private/SystemSettings.cpp
@@ -476,64 +476,57 @@ void FSystemSettings::ApplyNewSettings( const FSystemSettingsData& NewSettings,
 
 void FSystemSettings::ApplyOverrides()
 {
-	bool bUseMaxQualityMode = CVarUseMaxQualityMode.GetValueOnGameThread() != 0;
+	EConsoleVariableFlags SetBy = ECVF_SetByMask;
 
-	if (FParse::Param(FCommandLine::Get(),TEXT("MAXQUALITYMODE")))
+	if (FPlatformProperties::SupportsWindowedMode())
 	{
-		bUseMaxQualityMode = true;
-	}
-
-	if (FParse::Param(FCommandLine::Get(),TEXT("MSAA")))
-	{
-		check(0);
-		// todo: set console variable to activate MSAA
-	}
+		if (CVarUseMaxQualityMode.GetValueOnGameThread() != 0)
+		{
+			SetBy = (EConsoleVariableFlags)(CVarUseMaxQualityMode.AsVariable()->GetFlags() & ECVF_SetByMask);
+		}
 
-	if (!FPlatformProperties::SupportsWindowedMode())
-	{
-		bUseMaxQualityMode = false;
-	}
-	else
-	{
-		// Dump(TEXT("Startup System Settings:"));
+		if (FParse::Param(FCommandLine::Get(),TEXT("MAXQUALITYMODE")))
+		{
+			SetBy = ECVF_SetByCommandline;
+		}
 	}
 
-	if (bUseMaxQualityMode)
+	if (SetBy != ECVF_SetByMask)
 	{
 		// Modify various system settings to get the best quality regardless of performance impact
 		{
 			static auto CVar = IConsoleManager::Get().FindConsoleVariable(TEXT("r.Shadow.MinResolution"));
-			CVar->Set(16, ECVF_SetByCode);
+			CVar->Set(16, SetBy);
 		}
 
 		// Disable shadow fading out over distance
 		{
 			static auto CVar = IConsoleManager::Get().FindConsoleVariable(TEXT("r.Shadow.FadeResolution"));
-			CVar->Set(1, ECVF_SetByCode);
+			CVar->Set(1, SetBy);
 		}
 
 		// Increase minimum preshadow resolution
 		{
 			static auto CVar = IConsoleManager::Get().FindConsoleVariable(TEXT("r.Shadow.MinPreShadowResolution"));
-			CVar->Set(16, ECVF_SetByCode);
+			CVar->Set(16, SetBy);
 		}
 
 		// Disable preshadow fading out over distance
 		{
 			static auto CVar = IConsoleManager::Get().FindConsoleVariable(TEXT("r.Shadow.PreShadowFadeResolution"));
-			CVar->Set(1, ECVF_SetByCode);
+			CVar->Set(1, SetBy);
 		}
 
 		// Increase shadow texel density
 		{
 			static auto CVar = IConsoleManager::Get().FindConsoleVariable(TEXT("r.Shadow.TexelsPerPixel"));
-			CVar->Set(4.0f, ECVF_SetByCode);
+			CVar->Set(4.0f, SetBy);
 		}
 
 		// Don't downsample preshadows
 		{
 			static auto CVar = IConsoleManager::Get().FindConsoleVariable(TEXT("r.Shadow.PreShadowResolutionFactor"));
-			CVar->Set(1.0f, ECVF_SetByCode);
+			CVar->Set(1.0f, SetBy);
 		}
 
 		for (int32 GroupIndex = 0; GroupIndex < TEXTUREGROUP_MAX; GroupIndex++)
diff --git a/Engine/Source/Runtime/Launch/Private/LaunchEngineLoop.cpp b/Engine/Source/Runtime/Launch/Private/LaunchEngineLoop.cpp
--- a/Engine/Source/Runtime/Launch/Private/LaunchEngineLoop.cpp
+++ b/Engine/Source/Runtime/Launch/Private/LaunchEngineLoop.cpp
@@ -1065,7 +1065,7 @@ int32 FEngineLoop::PreInit( const TCHAR* CmdLine )
 	GSystemSettings.Initialize( bHasEditorToken );
 
 	// Apply renderer settings from console variables stored in the INI.
-	ApplyCVarSettingsFromIni(TEXT("/Script/Engine.RendererSettings"),*GEngineIni, ECVF_SetBySystemSettingsIni);
+	ApplyCVarSettingsFromIni(TEXT("/Script/Engine.RendererSettings"),*GEngineIni, ECVF_SetByProjectSetting);
 
 #if !UE_SERVER
 	if (!IsRunningDedicatedServer())
