fixed clang

diff --git a/Engine/Source/Runtime/Core/Private/HAL/ConsoleManager.cpp b/Engine/Source/Runtime/Core/Private/HAL/ConsoleManager.cpp
--- a/Engine/Source/Runtime/Core/Private/HAL/ConsoleManager.cpp
+++ b/Engine/Source/Runtime/Core/Private/HAL/ConsoleManager.cpp
@@ -16,7 +16,7 @@ static inline bool IsWhiteSpace(TCHAR Value) { return Value == TCHAR(' '); }
 
 // Get human readable string
 // @return never 0
-static TCHAR* GetSetByTCHAR(EConsoleVariableFlags InSetBy)
+static const TCHAR* GetSetByTCHAR(EConsoleVariableFlags InSetBy)
 {
 	EConsoleVariableFlags SetBy = (EConsoleVariableFlags)((uint32)InSetBy & ECVF_SetByMask);
 
