RadiansToDegrees(3)/DegreesToRadians(90) now returns a float. If you want to convert a triple of rotation values to degrees from radians or vice versa you have to explicitly use the FVector version.

diff --git a/Engine/Source/Runtime/Core/Public/Math/UnrealMathUtility.h b/Engine/Source/Runtime/Core/Public/Math/UnrealMathUtility.h
--- a/Engine/Source/Runtime/Core/Public/Math/UnrealMathUtility.h
+++ b/Engine/Source/Runtime/Core/Public/Math/UnrealMathUtility.h
@@ -301,7 +301,13 @@ struct FMath : public FPlatformMath
 	 * @return					Value in degrees.
 	 */
 	template<class T>
-	static FORCEINLINE T RadiansToDegrees(T const& RadVal)
+	static FORCEINLINE float RadiansToDegrees(T const& RadVal)
+	{
+		return RadVal * (180.f / PI);
+	}
+
+	/** This overload allows us to maintain double precision */
+	static FORCEINLINE double RadiansToDegrees(double RadVal)
 	{
 		return RadVal * (180.f / PI);
 	}
@@ -312,7 +318,13 @@ struct FMath : public FPlatformMath
 	 * @return					Value in radians.
 	 */
 	template<class T>
-	static FORCEINLINE T DegreesToRadians(T const& DegVal)
+	static FORCEINLINE float DegreesToRadians(T const& DegVal)
+	{
+		return DegVal * (PI / 180.f);
+	}
+
+	/** This overload allows us to maintain double precision */
+	static FORCEINLINE double DegreesToRadians(double DegVal)
 	{
 		return DegVal * (PI / 180.f);
 	}
diff --git a/Engine/Source/Runtime/Core/Public/Math/Vector.h b/Engine/Source/Runtime/Core/Public/Math/Vector.h
--- a/Engine/Source/Runtime/Core/Public/Math/Vector.h
+++ b/Engine/Source/Runtime/Core/Public/Math/Vector.h
@@ -851,6 +851,22 @@ public:
 	 */
 	static CORE_API float EvaluateBezier(const FVector* ControlPoints, int32 NumPoints, TArray<FVector>& OutPoints);
 
+	/**
+	 * Converts a vector containing radian values to a vector containing degree values.
+	 *
+	 * @param RadVector	Vector containing radian values
+	 * @return Vector  containing degree values
+	 */
+	static FVector RadiansToDegrees(const FVector& RadVector);
+
+	/**
+	 * Converts a vector containing degree values to a vector containing radian values.
+	 *
+	 * @param DegVector	Vector containing degree values
+	 * @return Vector containing radian values
+	 */
+	static FVector DegreesToRadians(const FVector& DegVector);
+
 	/**
 	 * Given a current set of cluster centers, a set of points, iterate N times to move clusters to be central. 
 	 *
@@ -1102,6 +1118,16 @@ inline float FVector::Triple( const FVector& X, const FVector& Y, const FVector&
 	+	(X.Z * (Y.X * Z.Y - Y.Y * Z.X)) );
 }
 
+inline FVector FVector::RadiansToDegrees(const FVector& RadVector)
+{
+	return RadVector * (180.f / PI);
+}
+
+inline FVector FVector::DegreesToRadians(const FVector& DegVector)
+{
+	return DegVector * (PI / 180.f);
+}
+
 FORCEINLINE FVector::FVector()
 {}
 
diff --git a/Engine/Source/Runtime/Engine/Private/Character.cpp b/Engine/Source/Runtime/Engine/Private/Character.cpp
--- a/Engine/Source/Runtime/Engine/Private/Character.cpp
+++ b/Engine/Source/Runtime/Engine/Private/Character.cpp
@@ -509,7 +509,7 @@ namespace MovementBaseUtility
 					if (MovementBaseUtility::GetMovementBaseTransform(MovementBase, BoneName, BaseLocation, BaseRotation))
 					{
 						const FVector RadialDistanceToBase = WorldLocation - BaseLocation;
-						const FVector TangentialVel = FMath::DegreesToRadians(BaseAngVel) ^ RadialDistanceToBase;
+						const FVector TangentialVel = FVector::DegreesToRadians(BaseAngVel) ^ RadialDistanceToBase;
 						return TangentialVel;
 					}
 				}
diff --git a/Engine/Source/Runtime/Engine/Private/PhysicsEngine/BodyInstance.cpp b/Engine/Source/Runtime/Engine/Private/PhysicsEngine/BodyInstance.cpp
--- a/Engine/Source/Runtime/Engine/Private/PhysicsEngine/BodyInstance.cpp
+++ b/Engine/Source/Runtime/Engine/Private/PhysicsEngine/BodyInstance.cpp
@@ -2421,7 +2421,7 @@ FVector FBodyInstance::GetUnrealWorldAngularVelocity() const
 	if (PRigidBody != NULL)
 	{
 		SCOPED_SCENE_READ_LOCK(PRigidBody->getScene());
-		AngVel = FMath::RadiansToDegrees(P2UVector(PRigidBody->getAngularVelocity()));
+		AngVel = FVector::RadiansToDegrees(P2UVector(PRigidBody->getAngularVelocity()));
 	}
 #endif // WITH_PHYSX
 
@@ -2978,7 +2978,7 @@ void FBodyInstance::SetAngularVelocity(const FVector& NewAngVel, bool bAddToCurr
 #if WITH_PHYSX
 	if (PxRigidBody* PRigidBody = GetPxRigidBody())
 	{
-		PxVec3 PNewAngVel = U2PVector( FMath::DegreesToRadians(NewAngVel) );
+		PxVec3 PNewAngVel = U2PVector( FVector::DegreesToRadians(NewAngVel) );
 
 		if (bAddToCurrent)
 		{
