fixed typo

diff --git a/Engine/Source/Runtime/Core/Public/HAL/IConsoleManager.h b/Engine/Source/Runtime/Core/Public/HAL/IConsoleManager.h
--- a/Engine/Source/Runtime/Core/Public/HAL/IConsoleManager.h
+++ b/Engine/Source/Runtime/Core/Public/HAL/IConsoleManager.h
@@ -107,7 +107,7 @@ enum EConsoleVariableFlags
 	ECVF_SetByCommandline =			0x07000000,
 	// least useful, likely a hack, maybe better to find the correct SetBy...
 	ECVF_SetByCode =				0x08000000,
-	// editor UI or consoel in game or editor
+	// editor UI or console in game or editor
 	ECVF_SetByConsole =				0x09000000,
 
 	// ------------------------------------------------
