Changed logic for getting cultures and setting the current culture to check if the culture would just fallback to root data, in which case it is considered unsupported. Automation tests involving culture changes now report warnings if a needed culture is unsupported.

diff --git a/Engine/Source/Runtime/Core/Private/Internationalization/ICUInternationalization.cpp b/Engine/Source/Runtime/Core/Private/Internationalization/ICUInternationalization.cpp
--- a/Engine/Source/Runtime/Core/Private/Internationalization/ICUInternationalization.cpp
+++ b/Engine/Source/Runtime/Core/Private/Internationalization/ICUInternationalization.cpp
@@ -122,8 +122,12 @@ bool FICUInternationalization::Initialize()
 
 	FICUBreakIteratorManager::Create();
 
-	I18N->InvariantCulture = FindOrMakeCulture(TEXT("en-US-POSIX"));
-	I18N->DefaultCulture = FindOrMakeCulture(FPlatformMisc::GetDefaultLocale());
+	I18N->InvariantCulture = FindOrMakeCulture(TEXT("en-US-POSIX"), false);
+	if (!I18N->InvariantCulture.IsValid())
+	{
+		I18N->InvariantCulture = FindOrMakeCulture(TEXT(""), true);
+	}
+	I18N->DefaultCulture = FindOrMakeCulture(FPlatformMisc::GetDefaultLocale(), true);
 	SetCurrentCulture( I18N->GetDefaultCulture()->GetName() );
 	return U_SUCCESS(ICUStatus) ? true : false;
 }
@@ -229,7 +233,7 @@ namespace
 }
 #endif
 
-void FICUInternationalization::SetCurrentCulture(const FString& Name)
+bool FICUInternationalization::SetCurrentCulture(const FString& Name)
 {
 	FCulturePtr NewCurrentCulture = FindOrMakeCulture(Name);
 
@@ -245,6 +249,8 @@ void FICUInternationalization::SetCurrentCulture(const FString& Name)
 			FInternationalization::Get().BroadcastCultureChanged();
 		}
 	}
+
+	return I18N->CurrentCulture == NewCurrentCulture;
 }
 
 void FICUInternationalization::GetCultureNames(TArray<FString>& CultureNames) const
@@ -264,15 +270,32 @@ FCulturePtr FICUInternationalization::GetCulture(const FString& Name)
 	return FindOrMakeCulture(Name);
 }
 
-FCulturePtr FICUInternationalization::FindOrMakeCulture(const FString& Name)
+FCulturePtr FICUInternationalization::FindOrMakeCulture(const FString& Name, const bool AllowDefaultFallback)
 {
-	FCultureRef* FoundCulture = CachedCultures.Find(Name);
+	const FString CanonicalName = FCulture::GetCanonicalName(Name);
+
+	// Find the cached culture.
+	FCultureRef* FoundCulture = CachedCultures.Find(CanonicalName);
+
+	// If no cached culture is found, try to make one.
 	if (!FoundCulture)
 	{
-		FCulturePtr NewCulture = FCulture::Create(Name);
-		if (NewCulture.IsValid())
+		UErrorCode ICUStatus = U_ZERO_ERROR;
+
+		// Confirm if data for the desired culture exists.
+		if (UResourceBundle* ICUResourceBundle = ures_open(nullptr, StringCast<char>(*CanonicalName).Get(), &ICUStatus))
 		{
-			FoundCulture = &(CachedCultures.Add(Name, NewCulture.ToSharedRef()));
+			// Make the culture only if it actually has some form of data and doesn't fallback to default "root" data, unless overriden ot allow it.
+			if (ICUStatus != U_USING_DEFAULT_WARNING || AllowDefaultFallback)
+			{
+				FCulturePtr NewCulture = FCulture::Create(CanonicalName);
+				if (NewCulture.IsValid())
+				{
+					FoundCulture = &(CachedCultures.Add(CanonicalName, NewCulture.ToSharedRef()));
+				}
+			}
+
+			ures_close(ICUResourceBundle);
 		}
 	}
 
diff --git a/Engine/Source/Runtime/Core/Private/Internationalization/ICUInternationalization.h b/Engine/Source/Runtime/Core/Private/Internationalization/ICUInternationalization.h
--- a/Engine/Source/Runtime/Core/Private/Internationalization/ICUInternationalization.h
+++ b/Engine/Source/Runtime/Core/Private/Internationalization/ICUInternationalization.h
@@ -13,7 +13,7 @@ public:
 	bool Initialize();
 	void Terminate();
 
-	void SetCurrentCulture(const FString& Name);
+	bool SetCurrentCulture(const FString& Name);
 	void GetCultureNames(TArray<FString>& CultureNames) const;
 	FCulturePtr GetCulture(const FString& Name);
 
@@ -23,7 +23,7 @@ private:
 	void UnloadDLLs();
 #endif
 
-	FCulturePtr FindOrMakeCulture(const FString& Name);
+	FCulturePtr FindOrMakeCulture(const FString& Name, const bool AllowDefaultFallback = false);
 
 private:
 	FInternationalization* const I18N;
diff --git a/Engine/Source/Runtime/Core/Private/Internationalization/Internationalization.cpp b/Engine/Source/Runtime/Core/Private/Internationalization/Internationalization.cpp
--- a/Engine/Source/Runtime/Core/Private/Internationalization/Internationalization.cpp
+++ b/Engine/Source/Runtime/Core/Private/Internationalization/Internationalization.cpp
@@ -33,9 +33,9 @@ void FInternationalization::TearDown()
 	}
 }
 
-void FInternationalization::SetCurrentCulture(const FString& Name)
+bool FInternationalization::SetCurrentCulture(const FString& Name)
 {
-	Implementation->SetCurrentCulture(Name);
+	return Implementation->SetCurrentCulture(Name);
 }
 
 FCulturePtr FInternationalization::GetCulture(const FString& Name)
diff --git a/Engine/Source/Runtime/Core/Private/Internationalization/LegacyInternationalization.cpp b/Engine/Source/Runtime/Core/Private/Internationalization/LegacyInternationalization.cpp
--- a/Engine/Source/Runtime/Core/Private/Internationalization/LegacyInternationalization.cpp
+++ b/Engine/Source/Runtime/Core/Private/Internationalization/LegacyInternationalization.cpp
@@ -25,8 +25,9 @@ void FLegacyInternationalization::Terminate()
 {
 }
 
-void FLegacyInternationalization::SetCurrentCulture(const FString& Name)
+bool FLegacyInternationalization::SetCurrentCulture(const FString& Name)
 {
+	return Name.IsEmpty();
 }
 
 void FLegacyInternationalization::GetCultureNames(TArray<FString>& CultureNames) const
diff --git a/Engine/Source/Runtime/Core/Private/Internationalization/LegacyInternationalization.h b/Engine/Source/Runtime/Core/Private/Internationalization/LegacyInternationalization.h
--- a/Engine/Source/Runtime/Core/Private/Internationalization/LegacyInternationalization.h
+++ b/Engine/Source/Runtime/Core/Private/Internationalization/LegacyInternationalization.h
@@ -11,7 +11,7 @@ public:
 	bool Initialize();
 	void Terminate();
 
-	void SetCurrentCulture(const FString& Name);
+	bool SetCurrentCulture(const FString& Name);
 	void GetCultureNames(TArray<FString>& CultureNames) const;
 	FCulturePtr GetCulture(const FString& Name);
 
diff --git a/Engine/Source/Runtime/Core/Private/Tests/Internationalization/DateTimeFormattingRulesTest.cpp b/Engine/Source/Runtime/Core/Private/Tests/Internationalization/DateTimeFormattingRulesTest.cpp
--- a/Engine/Source/Runtime/Core/Private/Tests/Internationalization/DateTimeFormattingRulesTest.cpp
+++ b/Engine/Source/Runtime/Core/Private/Tests/Internationalization/DateTimeFormattingRulesTest.cpp
@@ -7,9 +7,9 @@
 // disable optimization for DateTimeFormattingRulesTest as it compiles very slowly in development builds
 PRAGMA_DISABLE_OPTIMIZATION
 
-IMPLEMENT_SIMPLE_AUTOMATION_TEST(FDateTimeFormattingRulesTest, "Core.Misc.DateTime Formatting Rules", EAutomationTestFlags::ATF_None)
+	IMPLEMENT_SIMPLE_AUTOMATION_TEST(FDateTimeFormattingRulesTest, "Core.Misc.DateTime Formatting Rules", EAutomationTestFlags::ATF_None)
 
-namespace
+	namespace
 {
 	void Test(FDateTimeFormattingRulesTest* const ExplicitThis, const TCHAR* const Desc, const FText& A, const FText& B)
 	{
@@ -34,63 +34,73 @@ bool FDateTimeFormattingRulesTest::RunTest (const FString& Parameters)
 	const FDateTime UnixDecimalSequence = FDateTime::FromUnixTimestamp(1234567890);
 	const FDateTime TestDateTime( 1990, 6, 13, 12, 34, 56, 789 );
 
-	I18N.SetCurrentCulture("en-US");
-
-	// Unix Time Values via Date Time
-	Test( this, TEXT("Testing Unix Epoch"), FText::AsDateTime(UnixEpoch, EDateTimeStyle::Short, EDateTimeStyle::Short, "GMT"), FText::FromString( TEXT("1/1/70, 12:00 AM") ) );
-	Test( this, TEXT("Testing Unix Epoch"), FText::AsDateTime(UnixEpoch, EDateTimeStyle::Medium, EDateTimeStyle::Medium, "GMT"), FText::FromString( TEXT("Jan 1, 1970, 12:00:00 AM") ) );
-	Test( this, TEXT("Testing Unix Epoch"), FText::AsDateTime(UnixEpoch, EDateTimeStyle::Long, EDateTimeStyle::Long, "GMT"), FText::FromString( TEXT("January 1, 1970 at 12:00:00 AM GMT") ) );
-	Test( this, TEXT("Testing Unix Epoch"), FText::AsDateTime(UnixEpoch, EDateTimeStyle::Full, EDateTimeStyle::Full, "GMT"), FText::FromString( TEXT("Thursday, January 1, 1970 at 12:00:00 AM GMT") ) );
-
-	Test( this, TEXT("Testing Unix Billennium"), FText::AsDateTime(UnixBillennium, EDateTimeStyle::Short, EDateTimeStyle::Short, "GMT"), FText::FromString( TEXT("9/9/01, 1:46 AM") ) );
-	Test( this, TEXT("Testing Unix Billennium"), FText::AsDateTime(UnixBillennium, EDateTimeStyle::Medium, EDateTimeStyle::Medium, "GMT"), FText::FromString( TEXT("Sep 9, 2001, 1:46:40 AM") ) );
-	Test( this, TEXT("Testing Unix Billennium"), FText::AsDateTime(UnixBillennium, EDateTimeStyle::Long, EDateTimeStyle::Long, "GMT"), FText::FromString( TEXT("September 9, 2001 at 1:46:40 AM GMT") ) );
-	Test( this, TEXT("Testing Unix Billennium"), FText::AsDateTime(UnixBillennium, EDateTimeStyle::Full, EDateTimeStyle::Full, "GMT"), FText::FromString( TEXT("Sunday, September 9, 2001 at 1:46:40 AM GMT") ) );
-
-	Test( this, TEXT("Testing Unix Ones"), FText::AsDateTime(UnixOnes, EDateTimeStyle::Short, EDateTimeStyle::Short, "GMT"), FText::FromString( TEXT("3/18/05, 1:58 AM") ) );
-	Test( this, TEXT("Testing Unix Ones"), FText::AsDateTime(UnixOnes, EDateTimeStyle::Medium, EDateTimeStyle::Medium, "GMT"), FText::FromString( TEXT("Mar 18, 2005, 1:58:31 AM") ) );
-	Test( this, TEXT("Testing Unix Ones"), FText::AsDateTime(UnixOnes, EDateTimeStyle::Long, EDateTimeStyle::Long, "GMT"), FText::FromString( TEXT("March 18, 2005 at 1:58:31 AM GMT") ) );
-	Test( this, TEXT("Testing Unix Ones"), FText::AsDateTime(UnixOnes, EDateTimeStyle::Full, EDateTimeStyle::Full, "GMT"), FText::FromString( TEXT("Friday, March 18, 2005 at 1:58:31 AM GMT") ) );
-
-	Test( this, TEXT("Testing Unix Decimal Sequence"), FText::AsDateTime(UnixDecimalSequence, EDateTimeStyle::Short, EDateTimeStyle::Short, "GMT"), FText::FromString( TEXT("2/13/09, 11:31 PM") ) );
-	Test( this, TEXT("Testing Unix Decimal Sequence"), FText::AsDateTime(UnixDecimalSequence, EDateTimeStyle::Medium, EDateTimeStyle::Medium, "GMT"), FText::FromString( TEXT("Feb 13, 2009, 11:31:30 PM") ) );
-	Test( this, TEXT("Testing Unix Decimal Sequence"), FText::AsDateTime(UnixDecimalSequence, EDateTimeStyle::Long, EDateTimeStyle::Long, "GMT"), FText::FromString( TEXT("February 13, 2009 at 11:31:30 PM GMT") ) );
-	Test( this, TEXT("Testing Unix Decimal Sequence"), FText::AsDateTime(UnixDecimalSequence, EDateTimeStyle::Full, EDateTimeStyle::Full, "GMT"), FText::FromString( TEXT("Friday, February 13, 2009 at 11:31:30 PM GMT") ) );
-
-	// Date Time
-	Test( this, TEXT("Testing Date-Time"), FText::AsDateTime(TestDateTime, EDateTimeStyle::Short,EDateTimeStyle::Short, "GMT"), FText::FromString( TEXT("6/13/90, 12:34 PM") ) );
-	Test( this, TEXT("Testing Date-Time"), FText::AsDateTime(TestDateTime, EDateTimeStyle::Medium, EDateTimeStyle::Medium, "GMT"), FText::FromString( TEXT("Jun 13, 1990, 12:34:56 PM") ) );
-	Test( this, TEXT("Testing Date-Time"), FText::AsDateTime(TestDateTime, EDateTimeStyle::Long, EDateTimeStyle::Long, "GMT"), FText::FromString( TEXT("June 13, 1990 at 12:34:56 PM GMT") ) );
-	Test( this, TEXT("Testing Date-Time"), FText::AsDateTime(TestDateTime, EDateTimeStyle::Full, EDateTimeStyle::Full, "GMT"), FText::FromString( TEXT("Wednesday, June 13, 1990 at 12:34:56 PM GMT") ) );
-
-	I18N.SetCurrentCulture("ja-JP");
-
-	// Unix Time Values via Date Time
-	Test( this, TEXT("Testing Unix Epoch"), FText::AsDateTime(UnixEpoch, EDateTimeStyle::Short, EDateTimeStyle::Short, "GMT"), FText::FromString( TEXT("1970/01/01 0:00") ) );
-	Test( this, TEXT("Testing Unix Epoch"), FText::AsDateTime(UnixEpoch, EDateTimeStyle::Medium, EDateTimeStyle::Medium, "GMT"), FText::FromString( TEXT("1970/01/01 0:00:00") ) );
-	Test( this, TEXT("Testing Unix Epoch"), FText::AsDateTime(UnixEpoch, EDateTimeStyle::Long, EDateTimeStyle::Long, "GMT"), FText::FromString( TEXT("1970\x5E74") TEXT("1\x6708") TEXT("1\x65E5") TEXT(" 0:00:00 GMT") ) );
-	Test( this, TEXT("Testing Unix Epoch"), FText::AsDateTime(UnixEpoch, EDateTimeStyle::Full, EDateTimeStyle::Full, "GMT"), FText::FromString( TEXT("1970\x5E74") TEXT("1\x6708") TEXT("1\x65E5\x6728\x66DC\x65E5") TEXT(" ") TEXT("0\x6642") TEXT("00\x5206") TEXT("00\x79D2") TEXT(" GMT") ) );
-
-	Test( this, TEXT("Testing Unix Billennium"), FText::AsDateTime(UnixBillennium, EDateTimeStyle::Short, EDateTimeStyle::Short, "GMT"), FText::FromString( TEXT("2001/09/09 1:46") ) );
-	Test( this, TEXT("Testing Unix Billennium"), FText::AsDateTime(UnixBillennium, EDateTimeStyle::Medium, EDateTimeStyle::Medium, "GMT"), FText::FromString( TEXT("2001/09/09 1:46:40") ) );
-	Test( this, TEXT("Testing Unix Billennium"), FText::AsDateTime(UnixBillennium, EDateTimeStyle::Long, EDateTimeStyle::Long, "GMT"), FText::FromString( TEXT("2001\x5E74") TEXT("9\x6708") TEXT("9\x65E5") TEXT(" 1:46:40 GMT") ) );
-	Test( this, TEXT("Testing Unix Billennium"), FText::AsDateTime(UnixBillennium, EDateTimeStyle::Full, EDateTimeStyle::Full, "GMT"), FText::FromString( TEXT("2001\x5E74") TEXT("9\x6708") TEXT("9\x65E5\x65E5\x66DC\x65E5") TEXT(" ") TEXT("1\x6642") TEXT("46\x5206") TEXT("40\x79D2") TEXT(" GMT") ) );
-
-	Test( this, TEXT("Testing Unix Ones"), FText::AsDateTime(UnixOnes, EDateTimeStyle::Short, EDateTimeStyle::Short, "GMT"), FText::FromString( TEXT("2005/03/18 1:58") ) );
-	Test( this, TEXT("Testing Unix Ones"), FText::AsDateTime(UnixOnes, EDateTimeStyle::Medium, EDateTimeStyle::Medium, "GMT"), FText::FromString( TEXT("2005/03/18 1:58:31") ) );
-	Test( this, TEXT("Testing Unix Ones"), FText::AsDateTime(UnixOnes, EDateTimeStyle::Long, EDateTimeStyle::Long, "GMT"), FText::FromString( TEXT("2005\x5E74") TEXT("3\x6708") TEXT("18\x65E5") TEXT(" 1:58:31 GMT") ) );
-	Test( this, TEXT("Testing Unix Ones"), FText::AsDateTime(UnixOnes, EDateTimeStyle::Full, EDateTimeStyle::Full, "GMT"), FText::FromString( TEXT("2005\x5E74") TEXT("3\x6708") TEXT("18\x65E5\x91D1\x66DC\x65E5") TEXT(" ") TEXT("1\x6642") TEXT("58\x5206") TEXT("31\x79D2") TEXT(" GMT") ) );
-
-	Test( this, TEXT("Testing Unix Decimal Sequence"), FText::AsDateTime(UnixDecimalSequence, EDateTimeStyle::Short, EDateTimeStyle::Short, "GMT"), FText::FromString( TEXT("2009/02/13 23:31") ) );
-	Test( this, TEXT("Testing Unix Decimal Sequence"), FText::AsDateTime(UnixDecimalSequence, EDateTimeStyle::Medium, EDateTimeStyle::Medium, "GMT"), FText::FromString( TEXT("2009/02/13 23:31:30") ) );
-	Test( this, TEXT("Testing Unix Decimal Sequence"), FText::AsDateTime(UnixDecimalSequence, EDateTimeStyle::Long, EDateTimeStyle::Long, "GMT"), FText::FromString( TEXT("2009\x5E74") TEXT("2\x6708") TEXT("13\x65E5 23:31:30 GMT") ) );
-	Test( this, TEXT("Testing Unix Decimal Sequence"), FText::AsDateTime(UnixDecimalSequence, EDateTimeStyle::Full, EDateTimeStyle::Full, "GMT"), FText::FromString( TEXT("2009\x5E74") TEXT("2\x6708") TEXT("13\x65E5\x91D1\x66DC\x65E5") TEXT(" ") TEXT("23\x6642") TEXT("31\x5206") TEXT("30\x79D2") TEXT(" GMT") ) );
-
-	// Date Time
-	Test( this, TEXT("Testing Date-Time"), FText::AsDateTime(TestDateTime, EDateTimeStyle::Short,EDateTimeStyle::Short, "GMT"), FText::FromString( TEXT("1990/06/13 12:34") ) );
-	Test( this, TEXT("Testing Date-Time"), FText::AsDateTime(TestDateTime, EDateTimeStyle::Medium, EDateTimeStyle::Medium, "GMT"), FText::FromString( TEXT("1990/06/13 12:34:56") ) );
-	Test( this, TEXT("Testing Date-Time"), FText::AsDateTime(TestDateTime, EDateTimeStyle::Long, EDateTimeStyle::Long, "GMT"), FText::FromString( TEXT("1990\x5E74") TEXT("6\x6708") TEXT("13\x65E5") TEXT(" 12:34:56 GMT") ) );
-	Test( this, TEXT("Testing Date-Time"), FText::AsDateTime(TestDateTime, EDateTimeStyle::Full, EDateTimeStyle::Full, "GMT"), FText::FromString( TEXT("1990\x5E74") TEXT("6\x6708") TEXT("13\x65E5\x6C34\x66DC\x65E5") TEXT(" ") TEXT("12\x6642") TEXT("34\x5206") TEXT("56\x79D2 GMT") ) );
+	if (I18N.SetCurrentCulture("en-US"))
+	{
+		// Unix Time Values via Date Time
+		Test( this, TEXT("Testing Unix Epoch"), FText::AsDateTime(UnixEpoch, EDateTimeStyle::Short, EDateTimeStyle::Short, "GMT"), FText::FromString( TEXT("1/1/70, 12:00 AM") ) );
+		Test( this, TEXT("Testing Unix Epoch"), FText::AsDateTime(UnixEpoch, EDateTimeStyle::Medium, EDateTimeStyle::Medium, "GMT"), FText::FromString( TEXT("Jan 1, 1970, 12:00:00 AM") ) );
+		Test( this, TEXT("Testing Unix Epoch"), FText::AsDateTime(UnixEpoch, EDateTimeStyle::Long, EDateTimeStyle::Long, "GMT"), FText::FromString( TEXT("January 1, 1970 at 12:00:00 AM GMT") ) );
+		Test( this, TEXT("Testing Unix Epoch"), FText::AsDateTime(UnixEpoch, EDateTimeStyle::Full, EDateTimeStyle::Full, "GMT"), FText::FromString( TEXT("Thursday, January 1, 1970 at 12:00:00 AM GMT") ) );
+
+		Test( this, TEXT("Testing Unix Billennium"), FText::AsDateTime(UnixBillennium, EDateTimeStyle::Short, EDateTimeStyle::Short, "GMT"), FText::FromString( TEXT("9/9/01, 1:46 AM") ) );
+		Test( this, TEXT("Testing Unix Billennium"), FText::AsDateTime(UnixBillennium, EDateTimeStyle::Medium, EDateTimeStyle::Medium, "GMT"), FText::FromString( TEXT("Sep 9, 2001, 1:46:40 AM") ) );
+		Test( this, TEXT("Testing Unix Billennium"), FText::AsDateTime(UnixBillennium, EDateTimeStyle::Long, EDateTimeStyle::Long, "GMT"), FText::FromString( TEXT("September 9, 2001 at 1:46:40 AM GMT") ) );
+		Test( this, TEXT("Testing Unix Billennium"), FText::AsDateTime(UnixBillennium, EDateTimeStyle::Full, EDateTimeStyle::Full, "GMT"), FText::FromString( TEXT("Sunday, September 9, 2001 at 1:46:40 AM GMT") ) );
+
+		Test( this, TEXT("Testing Unix Ones"), FText::AsDateTime(UnixOnes, EDateTimeStyle::Short, EDateTimeStyle::Short, "GMT"), FText::FromString( TEXT("3/18/05, 1:58 AM") ) );
+		Test( this, TEXT("Testing Unix Ones"), FText::AsDateTime(UnixOnes, EDateTimeStyle::Medium, EDateTimeStyle::Medium, "GMT"), FText::FromString( TEXT("Mar 18, 2005, 1:58:31 AM") ) );
+		Test( this, TEXT("Testing Unix Ones"), FText::AsDateTime(UnixOnes, EDateTimeStyle::Long, EDateTimeStyle::Long, "GMT"), FText::FromString( TEXT("March 18, 2005 at 1:58:31 AM GMT") ) );
+		Test( this, TEXT("Testing Unix Ones"), FText::AsDateTime(UnixOnes, EDateTimeStyle::Full, EDateTimeStyle::Full, "GMT"), FText::FromString( TEXT("Friday, March 18, 2005 at 1:58:31 AM GMT") ) );
+
+		Test( this, TEXT("Testing Unix Decimal Sequence"), FText::AsDateTime(UnixDecimalSequence, EDateTimeStyle::Short, EDateTimeStyle::Short, "GMT"), FText::FromString( TEXT("2/13/09, 11:31 PM") ) );
+		Test( this, TEXT("Testing Unix Decimal Sequence"), FText::AsDateTime(UnixDecimalSequence, EDateTimeStyle::Medium, EDateTimeStyle::Medium, "GMT"), FText::FromString( TEXT("Feb 13, 2009, 11:31:30 PM") ) );
+		Test( this, TEXT("Testing Unix Decimal Sequence"), FText::AsDateTime(UnixDecimalSequence, EDateTimeStyle::Long, EDateTimeStyle::Long, "GMT"), FText::FromString( TEXT("February 13, 2009 at 11:31:30 PM GMT") ) );
+		Test( this, TEXT("Testing Unix Decimal Sequence"), FText::AsDateTime(UnixDecimalSequence, EDateTimeStyle::Full, EDateTimeStyle::Full, "GMT"), FText::FromString( TEXT("Friday, February 13, 2009 at 11:31:30 PM GMT") ) );
+
+		// Date Time
+		Test( this, TEXT("Testing Date-Time"), FText::AsDateTime(TestDateTime, EDateTimeStyle::Short,EDateTimeStyle::Short, "GMT"), FText::FromString( TEXT("6/13/90, 12:34 PM") ) );
+		Test( this, TEXT("Testing Date-Time"), FText::AsDateTime(TestDateTime, EDateTimeStyle::Medium, EDateTimeStyle::Medium, "GMT"), FText::FromString( TEXT("Jun 13, 1990, 12:34:56 PM") ) );
+		Test( this, TEXT("Testing Date-Time"), FText::AsDateTime(TestDateTime, EDateTimeStyle::Long, EDateTimeStyle::Long, "GMT"), FText::FromString( TEXT("June 13, 1990 at 12:34:56 PM GMT") ) );
+		Test( this, TEXT("Testing Date-Time"), FText::AsDateTime(TestDateTime, EDateTimeStyle::Full, EDateTimeStyle::Full, "GMT"), FText::FromString( TEXT("Wednesday, June 13, 1990 at 12:34:56 PM GMT") ) );
+	}
+	else
+	{
+		AddWarning(FString::Printf(TEXT("Internationalization data for %s missing - test is partially disabled."), TEXT("en-US")));
+	}
+
+	if (I18N.SetCurrentCulture("ja-JP"))
+	{
+		// Unix Time Values via Date Time
+		Test( this, TEXT("Testing Unix Epoch"), FText::AsDateTime(UnixEpoch, EDateTimeStyle::Short, EDateTimeStyle::Short, "GMT"), FText::FromString( TEXT("1970/01/01 0:00") ) );
+		Test( this, TEXT("Testing Unix Epoch"), FText::AsDateTime(UnixEpoch, EDateTimeStyle::Medium, EDateTimeStyle::Medium, "GMT"), FText::FromString( TEXT("1970/01/01 0:00:00") ) );
+		Test( this, TEXT("Testing Unix Epoch"), FText::AsDateTime(UnixEpoch, EDateTimeStyle::Long, EDateTimeStyle::Long, "GMT"), FText::FromString( TEXT("1970\x5E74") TEXT("1\x6708") TEXT("1\x65E5") TEXT(" 0:00:00 GMT") ) );
+		Test( this, TEXT("Testing Unix Epoch"), FText::AsDateTime(UnixEpoch, EDateTimeStyle::Full, EDateTimeStyle::Full, "GMT"), FText::FromString( TEXT("1970\x5E74") TEXT("1\x6708") TEXT("1\x65E5\x6728\x66DC\x65E5") TEXT(" ") TEXT("0\x6642") TEXT("00\x5206") TEXT("00\x79D2") TEXT(" GMT") ) );
+
+		Test( this, TEXT("Testing Unix Billennium"), FText::AsDateTime(UnixBillennium, EDateTimeStyle::Short, EDateTimeStyle::Short, "GMT"), FText::FromString( TEXT("2001/09/09 1:46") ) );
+		Test( this, TEXT("Testing Unix Billennium"), FText::AsDateTime(UnixBillennium, EDateTimeStyle::Medium, EDateTimeStyle::Medium, "GMT"), FText::FromString( TEXT("2001/09/09 1:46:40") ) );
+		Test( this, TEXT("Testing Unix Billennium"), FText::AsDateTime(UnixBillennium, EDateTimeStyle::Long, EDateTimeStyle::Long, "GMT"), FText::FromString( TEXT("2001\x5E74") TEXT("9\x6708") TEXT("9\x65E5") TEXT(" 1:46:40 GMT") ) );
+		Test( this, TEXT("Testing Unix Billennium"), FText::AsDateTime(UnixBillennium, EDateTimeStyle::Full, EDateTimeStyle::Full, "GMT"), FText::FromString( TEXT("2001\x5E74") TEXT("9\x6708") TEXT("9\x65E5\x65E5\x66DC\x65E5") TEXT(" ") TEXT("1\x6642") TEXT("46\x5206") TEXT("40\x79D2") TEXT(" GMT") ) );
+
+		Test( this, TEXT("Testing Unix Ones"), FText::AsDateTime(UnixOnes, EDateTimeStyle::Short, EDateTimeStyle::Short, "GMT"), FText::FromString( TEXT("2005/03/18 1:58") ) );
+		Test( this, TEXT("Testing Unix Ones"), FText::AsDateTime(UnixOnes, EDateTimeStyle::Medium, EDateTimeStyle::Medium, "GMT"), FText::FromString( TEXT("2005/03/18 1:58:31") ) );
+		Test( this, TEXT("Testing Unix Ones"), FText::AsDateTime(UnixOnes, EDateTimeStyle::Long, EDateTimeStyle::Long, "GMT"), FText::FromString( TEXT("2005\x5E74") TEXT("3\x6708") TEXT("18\x65E5") TEXT(" 1:58:31 GMT") ) );
+		Test( this, TEXT("Testing Unix Ones"), FText::AsDateTime(UnixOnes, EDateTimeStyle::Full, EDateTimeStyle::Full, "GMT"), FText::FromString( TEXT("2005\x5E74") TEXT("3\x6708") TEXT("18\x65E5\x91D1\x66DC\x65E5") TEXT(" ") TEXT("1\x6642") TEXT("58\x5206") TEXT("31\x79D2") TEXT(" GMT") ) );
+
+		Test( this, TEXT("Testing Unix Decimal Sequence"), FText::AsDateTime(UnixDecimalSequence, EDateTimeStyle::Short, EDateTimeStyle::Short, "GMT"), FText::FromString( TEXT("2009/02/13 23:31") ) );
+		Test( this, TEXT("Testing Unix Decimal Sequence"), FText::AsDateTime(UnixDecimalSequence, EDateTimeStyle::Medium, EDateTimeStyle::Medium, "GMT"), FText::FromString( TEXT("2009/02/13 23:31:30") ) );
+		Test( this, TEXT("Testing Unix Decimal Sequence"), FText::AsDateTime(UnixDecimalSequence, EDateTimeStyle::Long, EDateTimeStyle::Long, "GMT"), FText::FromString( TEXT("2009\x5E74") TEXT("2\x6708") TEXT("13\x65E5 23:31:30 GMT") ) );
+		Test( this, TEXT("Testing Unix Decimal Sequence"), FText::AsDateTime(UnixDecimalSequence, EDateTimeStyle::Full, EDateTimeStyle::Full, "GMT"), FText::FromString( TEXT("2009\x5E74") TEXT("2\x6708") TEXT("13\x65E5\x91D1\x66DC\x65E5") TEXT(" ") TEXT("23\x6642") TEXT("31\x5206") TEXT("30\x79D2") TEXT(" GMT") ) );
+
+		// Date Time
+		Test( this, TEXT("Testing Date-Time"), FText::AsDateTime(TestDateTime, EDateTimeStyle::Short,EDateTimeStyle::Short, "GMT"), FText::FromString( TEXT("1990/06/13 12:34") ) );
+		Test( this, TEXT("Testing Date-Time"), FText::AsDateTime(TestDateTime, EDateTimeStyle::Medium, EDateTimeStyle::Medium, "GMT"), FText::FromString( TEXT("1990/06/13 12:34:56") ) );
+		Test( this, TEXT("Testing Date-Time"), FText::AsDateTime(TestDateTime, EDateTimeStyle::Long, EDateTimeStyle::Long, "GMT"), FText::FromString( TEXT("1990\x5E74") TEXT("6\x6708") TEXT("13\x65E5") TEXT(" 12:34:56 GMT") ) );
+		Test( this, TEXT("Testing Date-Time"), FText::AsDateTime(TestDateTime, EDateTimeStyle::Full, EDateTimeStyle::Full, "GMT"), FText::FromString( TEXT("1990\x5E74") TEXT("6\x6708") TEXT("13\x65E5\x6C34\x66DC\x65E5") TEXT(" ") TEXT("12\x6642") TEXT("34\x5206") TEXT("56\x79D2 GMT") ) );
+	}
+	else
+	{
+		AddWarning(FString::Printf(TEXT("Internationalization data for %s missing - test is partially disabled."), TEXT("ja-JP")));
+	}
 
 	I18N.SetCurrentCulture(OriginalCulture);
 #else
diff --git a/Engine/Source/Runtime/Core/Private/Tests/Internationalization/NumberFormatingRulesTest.cpp b/Engine/Source/Runtime/Core/Private/Tests/Internationalization/NumberFormatingRulesTest.cpp
--- a/Engine/Source/Runtime/Core/Private/Tests/Internationalization/NumberFormatingRulesTest.cpp
+++ b/Engine/Source/Runtime/Core/Private/Tests/Internationalization/NumberFormatingRulesTest.cpp
@@ -7,9 +7,9 @@
 // Disable optimization for NumberFormatingRulesTest as it compiles very slowly in development builds
 PRAGMA_DISABLE_OPTIMIZATION
 
-IMPLEMENT_SIMPLE_AUTOMATION_TEST(FNumberFormattingRulesTest, "Core.Misc.Number Formatting Rules", EAutomationTestFlags::ATF_None)
+	IMPLEMENT_SIMPLE_AUTOMATION_TEST(FNumberFormattingRulesTest, "Core.Misc.Number Formatting Rules", EAutomationTestFlags::ATF_None)
 
-namespace
+	namespace
 {
 	void Test(FNumberFormattingRulesTest* const ExplicitThis, const TCHAR* const Desc, const FText& A, const FText& B)
 	{
@@ -23,34 +23,34 @@ namespace
 
 bool FNumberFormattingRulesTest::RunTest (const FString& Parameters)
 {
-/**
- *	List of functions and types to test.
- *
- *	Function	Type
- *	---------	------
- *	AsNumber 	double
- *	AsNumber 	float
- *	AsNumber 	uint8 
- *	AsNumber 	uint16
- *	AsNumber 	uint32
- *	AsNumber 	uint64
- *	AsNumber 	int8
- *	AsNumber 	int16 
- *	AsNumber 	int32 
- *	AsNumber 	int64 
- *	AsCurrency 	double 
- *	AsCurrency 	float 
- *	AsCurrency 	uint8 
- *	AsCurrency 	uint16
- *	AsCurrency 	uint32
- *	AsCurrency 	uint64
- *	AsCurrency 	int8 
- *	AsCurrency 	int16 
- *	AsCurrency 	int32 
- *	AsCurrency 	int64 
- *	AsPercent 	double 
- *	AsPercent	float 
- */
+	/**
+	*	List of functions and types to test.
+	*
+	*	Function	Type
+	*	---------	------
+	*	AsNumber 	double
+	*	AsNumber 	float
+	*	AsNumber 	uint8 
+	*	AsNumber 	uint16
+	*	AsNumber 	uint32
+	*	AsNumber 	uint64
+	*	AsNumber 	int8
+	*	AsNumber 	int16 
+	*	AsNumber 	int32 
+	*	AsNumber 	int64 
+	*	AsCurrency 	double 
+	*	AsCurrency 	float 
+	*	AsCurrency 	uint8 
+	*	AsCurrency 	uint16
+	*	AsCurrency 	uint32
+	*	AsCurrency 	uint64
+	*	AsCurrency 	int8 
+	*	AsCurrency 	int16 
+	*	AsCurrency 	int32 
+	*	AsCurrency 	int64 
+	*	AsPercent 	double 
+	*	AsPercent	float 
+	*/
 
 #if UE_ENABLE_ICU
 	double DoubleValue = 12345678.901;
@@ -73,221 +73,231 @@ bool FNumberFormattingRulesTest::RunTest (const FString& Parameters)
 	FInternationalization& I18N = FInternationalization::Get();
 	const FString OriginalCulture = I18N.GetCurrentCulture()->GetName();
 
-	I18N.SetCurrentCulture("en-US");
-
-	Test(this, TEXT("Convert a Double to a number formatted correct for en-US"),			FText::AsNumber(DoubleValue),			FText::FromString(TEXT("12,345,678.901")));
-	Test(this, TEXT("Convert a Float to a number formatted correct for en-US"),			FText::AsNumber(FloatValue),			FText::FromString(TEXT("1,234.567")));
-	Test(this, TEXT("Convert a Negative Double to a number formatted correct for en-US"), FText::AsNumber(DoubleNegativeValue),	FText::FromString(TEXT("-12,345,678.901")));
-	Test(this, TEXT("Convert a Negative Float to a number formatted correct for en-US"),	FText::AsNumber(FloatNegativeValue),	FText::FromString(TEXT("-1,234.567")));
-	Test(this, TEXT("Convert a uint8 to a number formatted correct for en-US"),			FText::AsNumber(Uint8Value),			FText::FromString(TEXT("255")));
-	Test(this, TEXT("Convert a uint16 to a number formatted correct for en-US"),			FText::AsNumber(Uint16Value),			FText::FromString(TEXT("65,535")));
-	Test(this, TEXT("Convert a uint32 to a number formatted correct for en-US"),			FText::AsNumber(Uint32Value),			FText::FromString(TEXT("4,294,967,295")));
-	Test(this, TEXT("Convert a uint64 to a number formatted correct for en-US"),			FText::AsNumber(Uint64Value),			FText::FromString(TEXT("1,844,674,407,370,955,161")));
-	Test(this, TEXT("Convert a int8 to a number formatted correct for en-US"),			FText::AsNumber(Int8Value),				FText::FromString(TEXT("123")));
-	Test(this, TEXT("Convert a int16 to a number formatted correct for en-US"),			FText::AsNumber(Int16Value),			FText::FromString(TEXT("12,345")));
-	Test(this, TEXT("Convert a int32 to a number formatted correct for en-US"),			FText::AsNumber(Int32Value),			FText::FromString(TEXT("12,345")));
-	Test(this, TEXT("Convert a int64 to a number formatted correct for en-US"),			FText::AsNumber(Int64Value),			FText::FromString(TEXT("12,345")));
-	Test(this, TEXT("Convert a Negative int8 to a number formatted correct for en-US"),	FText::AsNumber(Int8NegativeValue),		FText::FromString(TEXT("-123")));
-	Test(this, TEXT("Convert a Negative int16 to a number formatted correct for en-US"),	FText::AsNumber(Int16NegativeValue),	FText::FromString(TEXT("-12,345")));
-	Test(this, TEXT("Convert a Negative int32 to a number formatted correct for en-US"),	FText::AsNumber(Int32NegativeValue),	FText::FromString(TEXT("-12,345")));
-	Test(this, TEXT("Convert a Negative int64 to a number formatted correct for en-US"),	FText::AsNumber(Int64NegativeValue),	FText::FromString(TEXT("-12,345")));
-
-	// Test Number Formatting Options
+	if (I18N.SetCurrentCulture("en-US"))
 	{
-		FNumberFormattingOptions NumberFormattingOptions;
-		NumberFormattingOptions.MinimumIntegralDigits = 20;
-		NumberFormattingOptions.MaximumIntegralDigits = 20;
-		NumberFormattingOptions.MinimumFractionalDigits = 3;
-		NumberFormattingOptions.MaximumFractionalDigits = 3;
-		NumberFormattingOptions.UseGrouping = false;
-
-		Test(this, TEXT("Convert a Double to a number formatted correct for en-US"),			FText::AsNumber(DoubleValue, &(NumberFormattingOptions)),			FText::FromString(TEXT( "00000000000012345678.901")));
-		Test(this, TEXT("Convert a Float to a number formatted correct for en-US"),			FText::AsNumber(FloatValue, &(NumberFormattingOptions)),			FText::FromString(TEXT( "00000000000000001234.567")));
-		Test(this, TEXT("Convert a Negative Double to a number formatted correct for en-US"),	FText::AsNumber(DoubleNegativeValue, &(NumberFormattingOptions)),	FText::FromString(TEXT("-00000000000012345678.901")));
-		Test(this, TEXT("Convert a Negative Float to a number formatted correct for en-US"),	FText::AsNumber(FloatNegativeValue, &(NumberFormattingOptions)),	FText::FromString(TEXT("-00000000000000001234.567")));
-		Test(this, TEXT("Convert a uint8 to a number formatted correct for en-US"),			FText::AsNumber(Uint8Value, &(NumberFormattingOptions)),			FText::FromString(TEXT( "00000000000000000255.000")));
-		Test(this, TEXT("Convert a uint16 to a number formatted correct for en-US"),			FText::AsNumber(Uint16Value, &(NumberFormattingOptions)),			FText::FromString(TEXT( "00000000000000065535.000")));
-		Test(this, TEXT("Convert a uint32 to a number formatted correct for en-US"),			FText::AsNumber(Uint32Value, &(NumberFormattingOptions)),			FText::FromString(TEXT( "00000000004294967295.000")));
-		Test(this, TEXT("Convert a uint64 to a number formatted correct for en-US"),			FText::AsNumber(Uint64Value, &(NumberFormattingOptions)),			FText::FromString(TEXT( "01844674407370955161.000")));
-		Test(this, TEXT("Convert a int8 to a number formatted correct for en-US"),			FText::AsNumber(Int8Value, &(NumberFormattingOptions)),				FText::FromString(TEXT( "00000000000000000123.000")));
-		Test(this, TEXT("Convert a int16 to a number formatted correct for en-US"),			FText::AsNumber(Int16Value, &(NumberFormattingOptions)),			FText::FromString(TEXT( "00000000000000012345.000")));
-		Test(this, TEXT("Convert a int32 to a number formatted correct for en-US"),			FText::AsNumber(Int32Value, &(NumberFormattingOptions)),			FText::FromString(TEXT( "00000000000000012345.000")));
-		Test(this, TEXT("Convert a int64 to a number formatted correct for en-US"),			FText::AsNumber(Int64Value, &(NumberFormattingOptions)),			FText::FromString(TEXT( "00000000000000012345.000")));
-		Test(this, TEXT("Convert a Negative int8 to a number formatted correct for en-US"),	FText::AsNumber(Int8NegativeValue, &(NumberFormattingOptions)),		FText::FromString(TEXT("-00000000000000000123.000")));
-		Test(this, TEXT("Convert a Negative int16 to a number formatted correct for en-US"),	FText::AsNumber(Int16NegativeValue, &(NumberFormattingOptions)),	FText::FromString(TEXT("-00000000000000012345.000")));
-		Test(this, TEXT("Convert a Negative int32 to a number formatted correct for en-US"),	FText::AsNumber(Int32NegativeValue, &(NumberFormattingOptions)),	FText::FromString(TEXT("-00000000000000012345.000")));
-		Test(this, TEXT("Convert a Negative int64 to a number formatted correct for en-US"),	FText::AsNumber(Int64NegativeValue, &(NumberFormattingOptions)),	FText::FromString(TEXT("-00000000000000012345.000")));
-	}
+		Test(this, TEXT("Convert a Double to a number formatted correct for en-US"),			FText::AsNumber(DoubleValue),			FText::FromString(TEXT("12,345,678.901")));
+		Test(this, TEXT("Convert a Float to a number formatted correct for en-US"),			FText::AsNumber(FloatValue),			FText::FromString(TEXT("1,234.567")));
+		Test(this, TEXT("Convert a Negative Double to a number formatted correct for en-US"), FText::AsNumber(DoubleNegativeValue),	FText::FromString(TEXT("-12,345,678.901")));
+		Test(this, TEXT("Convert a Negative Float to a number formatted correct for en-US"),	FText::AsNumber(FloatNegativeValue),	FText::FromString(TEXT("-1,234.567")));
+		Test(this, TEXT("Convert a uint8 to a number formatted correct for en-US"),			FText::AsNumber(Uint8Value),			FText::FromString(TEXT("255")));
+		Test(this, TEXT("Convert a uint16 to a number formatted correct for en-US"),			FText::AsNumber(Uint16Value),			FText::FromString(TEXT("65,535")));
+		Test(this, TEXT("Convert a uint32 to a number formatted correct for en-US"),			FText::AsNumber(Uint32Value),			FText::FromString(TEXT("4,294,967,295")));
+		Test(this, TEXT("Convert a uint64 to a number formatted correct for en-US"),			FText::AsNumber(Uint64Value),			FText::FromString(TEXT("1,844,674,407,370,955,161")));
+		Test(this, TEXT("Convert a int8 to a number formatted correct for en-US"),			FText::AsNumber(Int8Value),				FText::FromString(TEXT("123")));
+		Test(this, TEXT("Convert a int16 to a number formatted correct for en-US"),			FText::AsNumber(Int16Value),			FText::FromString(TEXT("12,345")));
+		Test(this, TEXT("Convert a int32 to a number formatted correct for en-US"),			FText::AsNumber(Int32Value),			FText::FromString(TEXT("12,345")));
+		Test(this, TEXT("Convert a int64 to a number formatted correct for en-US"),			FText::AsNumber(Int64Value),			FText::FromString(TEXT("12,345")));
+		Test(this, TEXT("Convert a Negative int8 to a number formatted correct for en-US"),	FText::AsNumber(Int8NegativeValue),		FText::FromString(TEXT("-123")));
+		Test(this, TEXT("Convert a Negative int16 to a number formatted correct for en-US"),	FText::AsNumber(Int16NegativeValue),	FText::FromString(TEXT("-12,345")));
+		Test(this, TEXT("Convert a Negative int32 to a number formatted correct for en-US"),	FText::AsNumber(Int32NegativeValue),	FText::FromString(TEXT("-12,345")));
+		Test(this, TEXT("Convert a Negative int64 to a number formatted correct for en-US"),	FText::AsNumber(Int64NegativeValue),	FText::FromString(TEXT("-12,345")));
+
+		// Test Number Formatting Options
+		{
+			FNumberFormattingOptions NumberFormattingOptions;
+			NumberFormattingOptions.MinimumIntegralDigits = 20;
+			NumberFormattingOptions.MaximumIntegralDigits = 20;
+			NumberFormattingOptions.MinimumFractionalDigits = 3;
+			NumberFormattingOptions.MaximumFractionalDigits = 3;
+			NumberFormattingOptions.UseGrouping = false;
+
+			Test(this, TEXT("Convert a Double to a number formatted correct for en-US"),			FText::AsNumber(DoubleValue, &(NumberFormattingOptions)),			FText::FromString(TEXT( "00000000000012345678.901")));
+			Test(this, TEXT("Convert a Float to a number formatted correct for en-US"),			FText::AsNumber(FloatValue, &(NumberFormattingOptions)),			FText::FromString(TEXT( "00000000000000001234.567")));
+			Test(this, TEXT("Convert a Negative Double to a number formatted correct for en-US"),	FText::AsNumber(DoubleNegativeValue, &(NumberFormattingOptions)),	FText::FromString(TEXT("-00000000000012345678.901")));
+			Test(this, TEXT("Convert a Negative Float to a number formatted correct for en-US"),	FText::AsNumber(FloatNegativeValue, &(NumberFormattingOptions)),	FText::FromString(TEXT("-00000000000000001234.567")));
+			Test(this, TEXT("Convert a uint8 to a number formatted correct for en-US"),			FText::AsNumber(Uint8Value, &(NumberFormattingOptions)),			FText::FromString(TEXT( "00000000000000000255.000")));
+			Test(this, TEXT("Convert a uint16 to a number formatted correct for en-US"),			FText::AsNumber(Uint16Value, &(NumberFormattingOptions)),			FText::FromString(TEXT( "00000000000000065535.000")));
+			Test(this, TEXT("Convert a uint32 to a number formatted correct for en-US"),			FText::AsNumber(Uint32Value, &(NumberFormattingOptions)),			FText::FromString(TEXT( "00000000004294967295.000")));
+			Test(this, TEXT("Convert a uint64 to a number formatted correct for en-US"),			FText::AsNumber(Uint64Value, &(NumberFormattingOptions)),			FText::FromString(TEXT( "01844674407370955161.000")));
+			Test(this, TEXT("Convert a int8 to a number formatted correct for en-US"),			FText::AsNumber(Int8Value, &(NumberFormattingOptions)),				FText::FromString(TEXT( "00000000000000000123.000")));
+			Test(this, TEXT("Convert a int16 to a number formatted correct for en-US"),			FText::AsNumber(Int16Value, &(NumberFormattingOptions)),			FText::FromString(TEXT( "00000000000000012345.000")));
+			Test(this, TEXT("Convert a int32 to a number formatted correct for en-US"),			FText::AsNumber(Int32Value, &(NumberFormattingOptions)),			FText::FromString(TEXT( "00000000000000012345.000")));
+			Test(this, TEXT("Convert a int64 to a number formatted correct for en-US"),			FText::AsNumber(Int64Value, &(NumberFormattingOptions)),			FText::FromString(TEXT( "00000000000000012345.000")));
+			Test(this, TEXT("Convert a Negative int8 to a number formatted correct for en-US"),	FText::AsNumber(Int8NegativeValue, &(NumberFormattingOptions)),		FText::FromString(TEXT("-00000000000000000123.000")));
+			Test(this, TEXT("Convert a Negative int16 to a number formatted correct for en-US"),	FText::AsNumber(Int16NegativeValue, &(NumberFormattingOptions)),	FText::FromString(TEXT("-00000000000000012345.000")));
+			Test(this, TEXT("Convert a Negative int32 to a number formatted correct for en-US"),	FText::AsNumber(Int32NegativeValue, &(NumberFormattingOptions)),	FText::FromString(TEXT("-00000000000000012345.000")));
+			Test(this, TEXT("Convert a Negative int64 to a number formatted correct for en-US"),	FText::AsNumber(Int64NegativeValue, &(NumberFormattingOptions)),	FText::FromString(TEXT("-00000000000000012345.000")));
+		}
 
+		{
+			FNumberFormattingOptions NumberFormattingOptions;
+			NumberFormattingOptions.MinimumIntegralDigits = 1;
+			NumberFormattingOptions.MaximumIntegralDigits = 1;
+			NumberFormattingOptions.MinimumFractionalDigits = 0;
+			NumberFormattingOptions.MaximumFractionalDigits = 0;
+			NumberFormattingOptions.UseGrouping = false;
+
+			NumberFormattingOptions.RoundingMode = ERoundingMode::HalfToEven;
+
+			Test(this, TEXT("Round a Double to a number formatted correct for en-US using HalfToEven"),	FText::AsNumber(-1.50, &(NumberFormattingOptions)),	FText::FromString(TEXT("-2")));
+			Test(this, TEXT("Round a Double to a number formatted correct for en-US using HalfToEven"),	FText::AsNumber(-1.00, &(NumberFormattingOptions)),	FText::FromString(TEXT("-1")));
+			Test(this, TEXT("Round a Double to a number formatted correct for en-US using HalfToEven"),	FText::AsNumber(-0.75, &(NumberFormattingOptions)),	FText::FromString(TEXT("-1")));
+			Test(this, TEXT("Round a Double to a number formatted correct for en-US using HalfToEven"),	FText::AsNumber(-0.50, &(NumberFormattingOptions)),	FText::FromString(TEXT("-0")));
+			Test(this, TEXT("Round a Double to a number formatted correct for en-US using HalfToEven"),	FText::AsNumber(-0.25, &(NumberFormattingOptions)),	FText::FromString(TEXT("-0")));
+			Test(this, TEXT("Round a Double to a number formatted correct for en-US using HalfToEven"),	FText::AsNumber(-0.00, &(NumberFormattingOptions)),	FText::FromString(TEXT("-0")));
+			Test(this, TEXT("Round a Double to a number formatted correct for en-US using HalfToEven"),	FText::AsNumber( 0.00, &(NumberFormattingOptions)),	FText::FromString(TEXT( "0")));
+			Test(this, TEXT("Round a Double to a number formatted correct for en-US using HalfToEven"),	FText::AsNumber( 0.25, &(NumberFormattingOptions)),	FText::FromString(TEXT( "0")));
+			Test(this, TEXT("Round a Double to a number formatted correct for en-US using HalfToEven"),	FText::AsNumber( 0.50, &(NumberFormattingOptions)),	FText::FromString(TEXT( "0")));
+			Test(this, TEXT("Round a Double to a number formatted correct for en-US using HalfToEven"),	FText::AsNumber( 0.75, &(NumberFormattingOptions)),	FText::FromString(TEXT( "1")));
+			Test(this, TEXT("Round a Double to a number formatted correct for en-US using HalfToEven"),	FText::AsNumber( 1.00, &(NumberFormattingOptions)),	FText::FromString(TEXT( "1")));
+			Test(this, TEXT("Round a Double to a number formatted correct for en-US using HalfToEven"),	FText::AsNumber( 1.50, &(NumberFormattingOptions)),	FText::FromString(TEXT( "2")));
+
+			NumberFormattingOptions.RoundingMode = ERoundingMode::HalfFromZero;
+
+			Test(this, TEXT("Round a Double to a number formatted correct for en-US using HalfFromZero"),	FText::AsNumber(-1.00, &(NumberFormattingOptions)),	FText::FromString(TEXT("-1")));
+			Test(this, TEXT("Round a Double to a number formatted correct for en-US using HalfFromZero"),	FText::AsNumber(-0.75, &(NumberFormattingOptions)),	FText::FromString(TEXT("-1")));
+			Test(this, TEXT("Round a Double to a number formatted correct for en-US using HalfFromZero"),	FText::AsNumber(-0.50, &(NumberFormattingOptions)),	FText::FromString(TEXT("-1")));
+			Test(this, TEXT("Round a Double to a number formatted correct for en-US using HalfFromZero"),	FText::AsNumber(-0.25, &(NumberFormattingOptions)),	FText::FromString(TEXT("-0")));
+			Test(this, TEXT("Round a Double to a number formatted correct for en-US using HalfFromZero"),	FText::AsNumber(-0.00, &(NumberFormattingOptions)),	FText::FromString(TEXT("-0")));
+			Test(this, TEXT("Round a Double to a number formatted correct for en-US using HalfFromZero"),	FText::AsNumber( 0.00, &(NumberFormattingOptions)),	FText::FromString(TEXT( "0")));
+			Test(this, TEXT("Round a Double to a number formatted correct for en-US using HalfFromZero"),	FText::AsNumber( 0.25, &(NumberFormattingOptions)),	FText::FromString(TEXT( "0")));
+			Test(this, TEXT("Round a Double to a number formatted correct for en-US using HalfFromZero"),	FText::AsNumber( 0.50, &(NumberFormattingOptions)),	FText::FromString(TEXT( "1")));
+			Test(this, TEXT("Round a Double to a number formatted correct for en-US using HalfFromZero"),	FText::AsNumber( 0.75, &(NumberFormattingOptions)),	FText::FromString(TEXT( "1")));
+			Test(this, TEXT("Round a Double to a number formatted correct for en-US using HalfFromZero"),	FText::AsNumber( 1.00, &(NumberFormattingOptions)),	FText::FromString(TEXT( "1")));
+
+			NumberFormattingOptions.RoundingMode = ERoundingMode::HalfToZero;
+
+			Test(this, TEXT("Round a Double to a number formatted correct for en-US using HalfToZero"),	FText::AsNumber(-1.00, &(NumberFormattingOptions)),	FText::FromString(TEXT("-1")));
+			Test(this, TEXT("Round a Double to a number formatted correct for en-US using HalfToZero"),	FText::AsNumber(-0.75, &(NumberFormattingOptions)),	FText::FromString(TEXT("-1")));
+			Test(this, TEXT("Round a Double to a number formatted correct for en-US using HalfToZero"),	FText::AsNumber(-0.50, &(NumberFormattingOptions)),	FText::FromString(TEXT("-0")));
+			Test(this, TEXT("Round a Double to a number formatted correct for en-US using HalfToZero"),	FText::AsNumber(-0.25, &(NumberFormattingOptions)),	FText::FromString(TEXT("-0")));
+			Test(this, TEXT("Round a Double to a number formatted correct for en-US using HalfToZero"),	FText::AsNumber(-0.00, &(NumberFormattingOptions)),	FText::FromString(TEXT("-0")));
+			Test(this, TEXT("Round a Double to a number formatted correct for en-US using HalfToZero"),	FText::AsNumber( 0.00, &(NumberFormattingOptions)),	FText::FromString(TEXT( "0")));
+			Test(this, TEXT("Round a Double to a number formatted correct for en-US using HalfToZero"),	FText::AsNumber( 0.25, &(NumberFormattingOptions)),	FText::FromString(TEXT( "0")));
+			Test(this, TEXT("Round a Double to a number formatted correct for en-US using HalfToZero"),	FText::AsNumber( 0.50, &(NumberFormattingOptions)),	FText::FromString(TEXT( "0")));
+			Test(this, TEXT("Round a Double to a number formatted correct for en-US using HalfToZero"),	FText::AsNumber( 0.75, &(NumberFormattingOptions)),	FText::FromString(TEXT( "1")));
+			Test(this, TEXT("Round a Double to a number formatted correct for en-US using HalfToZero"),	FText::AsNumber( 1.00, &(NumberFormattingOptions)),	FText::FromString(TEXT( "1")));
+
+			NumberFormattingOptions.RoundingMode = ERoundingMode::FromZero;
+
+			Test(this, TEXT("Round a Double to a number formatted correct for en-US using FromZero"),	FText::AsNumber(-1.0, &(NumberFormattingOptions)),	FText::FromString(TEXT("-1")));
+			Test(this, TEXT("Round a Double to a number formatted correct for en-US using FromZero"),	FText::AsNumber(-0.5, &(NumberFormattingOptions)),	FText::FromString(TEXT("-1")));
+			Test(this, TEXT("Round a Double to a number formatted correct for en-US using FromZero"),	FText::AsNumber(-0.0, &(NumberFormattingOptions)),	FText::FromString(TEXT("-0")));
+			Test(this, TEXT("Round a Double to a number formatted correct for en-US using FromZero"),	FText::AsNumber( 0.0, &(NumberFormattingOptions)),	FText::FromString(TEXT( "0")));
+			Test(this, TEXT("Round a Double to a number formatted correct for en-US using FromZero"),	FText::AsNumber( 0.5, &(NumberFormattingOptions)),	FText::FromString(TEXT( "1")));
+			Test(this, TEXT("Round a Double to a number formatted correct for en-US using FromZero"),	FText::AsNumber( 1.0, &(NumberFormattingOptions)),	FText::FromString(TEXT( "1")));
+
+			NumberFormattingOptions.RoundingMode = ERoundingMode::ToZero;
+
+			Test(this, TEXT("Round a Double to a number formatted correct for en-US using ToZero"),	FText::AsNumber(-1.0, &(NumberFormattingOptions)),	FText::FromString(TEXT("-1")));
+			Test(this, TEXT("Round a Double to a number formatted correct for en-US using ToZero"),	FText::AsNumber(-0.5, &(NumberFormattingOptions)),	FText::FromString(TEXT("-0")));
+			Test(this, TEXT("Round a Double to a number formatted correct for en-US using ToZero"),	FText::AsNumber(-0.0, &(NumberFormattingOptions)),	FText::FromString(TEXT("-0")));
+			Test(this, TEXT("Round a Double to a number formatted correct for en-US using ToZero"),	FText::AsNumber( 0.0, &(NumberFormattingOptions)),	FText::FromString(TEXT( "0")));
+			Test(this, TEXT("Round a Double to a number formatted correct for en-US using ToZero"),	FText::AsNumber( 0.5, &(NumberFormattingOptions)),	FText::FromString(TEXT( "0")));
+			Test(this, TEXT("Round a Double to a number formatted correct for en-US using ToZero"),	FText::AsNumber( 1.0, &(NumberFormattingOptions)),	FText::FromString(TEXT( "1")));
+
+			NumberFormattingOptions.RoundingMode = ERoundingMode::ToNegativeInfinity;
+
+			Test(this, TEXT("Round a Double to a number formatted correct for en-US using ToNegativeInfinity"),	FText::AsNumber(-1.0, &(NumberFormattingOptions)),	FText::FromString(TEXT("-1")));
+			Test(this, TEXT("Round a Double to a number formatted correct for en-US using ToNegativeInfinity"),	FText::AsNumber(-0.5, &(NumberFormattingOptions)),	FText::FromString(TEXT("-1")));
+			Test(this, TEXT("Round a Double to a number formatted correct for en-US using ToNegativeInfinity"),	FText::AsNumber(-0.0, &(NumberFormattingOptions)),	FText::FromString(TEXT("-0")));
+			Test(this, TEXT("Round a Double to a number formatted correct for en-US using ToNegativeInfinity"),	FText::AsNumber( 0.0, &(NumberFormattingOptions)),	FText::FromString(TEXT( "0")));
+			Test(this, TEXT("Round a Double to a number formatted correct for en-US using ToNegativeInfinity"),	FText::AsNumber( 0.5, &(NumberFormattingOptions)),	FText::FromString(TEXT( "0")));
+			Test(this, TEXT("Round a Double to a number formatted correct for en-US using ToNegativeInfinity"),	FText::AsNumber( 1.0, &(NumberFormattingOptions)),	FText::FromString(TEXT( "1")));
+
+			NumberFormattingOptions.RoundingMode = ERoundingMode::ToPositiveInfinity;
+
+			Test(this, TEXT("Round a Double to a number formatted correct for en-US using ToPositiveInfinity"),	FText::AsNumber(-1.0, &(NumberFormattingOptions)),	FText::FromString(TEXT("-1")));
+			Test(this, TEXT("Round a Double to a number formatted correct for en-US using ToPositiveInfinity"),	FText::AsNumber(-0.5, &(NumberFormattingOptions)),	FText::FromString(TEXT("-0")));
+			Test(this, TEXT("Round a Double to a number formatted correct for en-US using ToPositiveInfinity"),	FText::AsNumber(-0.0, &(NumberFormattingOptions)),	FText::FromString(TEXT("-0")));
+			Test(this, TEXT("Round a Double to a number formatted correct for en-US using ToPositiveInfinity"),	FText::AsNumber( 0.0, &(NumberFormattingOptions)),	FText::FromString(TEXT( "0")));
+			Test(this, TEXT("Round a Double to a number formatted correct for en-US using ToPositiveInfinity"),	FText::AsNumber( 0.5, &(NumberFormattingOptions)),	FText::FromString(TEXT( "1")));
+			Test(this, TEXT("Round a Double to a number formatted correct for en-US using ToPositiveInfinity"),	FText::AsNumber( 1.0, &(NumberFormattingOptions)),	FText::FromString(TEXT( "1")));
+		}
+
+		Test(this, TEXT("Convert a Double to a currency formatted correct for en-US"),			FText::AsCurrency(DoubleValue),			FText::FromString(TEXT("$12,345,678.90")));
+		Test(this, TEXT("Convert a Float to a currency formatted correct for en-US"),			FText::AsCurrency(FloatValue),			FText::FromString(TEXT("$1,234.57")));
+		Test(this, TEXT("Convert a Negative Double to a currency formatted correct for en-US"),	FText::AsCurrency(DoubleNegativeValue), FText::FromString(TEXT("-$12,345,678.90")));
+		Test(this, TEXT("Convert a Negative Float to a currency formatted correct for en-US"),	FText::AsCurrency(FloatNegativeValue),	FText::FromString(TEXT("-$1,234.57")));
+		Test(this, TEXT("Convert a uint8 to a currency formatted correct for en-US"),			FText::AsCurrency(Uint8Value),			FText::FromString(TEXT("$255.00")));
+		Test(this, TEXT("Convert a uint16 to a currency formatted correct for en-US"),			FText::AsCurrency(Uint16Value),			FText::FromString(TEXT("$65,535.00")));
+		Test(this, TEXT("Convert a uint32 to a currency formatted correct for en-US"),			FText::AsCurrency(Uint32Value),			FText::FromString(TEXT("$4,294,967,295.00")));
+		Test(this, TEXT("Convert a uint64 to a currency formatted correct for en-US"),			FText::AsCurrency(Uint64Value),			FText::FromString(TEXT("$1,844,674,407,370,955,161.00")));
+		Test(this, TEXT("Convert a int8 to a currency formatted correct for en-US"),			FText::AsCurrency(Int8Value),			FText::FromString(TEXT("$123.00")));
+		Test(this, TEXT("Convert a int16 to a currency formatted correct for en-US"),			FText::AsCurrency(Int16Value),			FText::FromString(TEXT("$12,345.00")));
+		Test(this, TEXT("Convert a int32 to a currency formatted correct for en-US"),			FText::AsCurrency(Int32Value),			FText::FromString(TEXT("$12,345.00")));
+		Test(this, TEXT("Convert a int64 to a currency formatted correct for en-US"),			FText::AsCurrency(Int64Value),			FText::FromString(TEXT("$12,345.00")));
+		Test(this, TEXT("Convert a Negative int8 to a currency formatted correct for en-US"),	FText::AsCurrency(Int8NegativeValue),	FText::FromString(TEXT("-$123.00")));
+		Test(this, TEXT("Convert a Negative int16 to a currency formatted correct for en-US"),	FText::AsCurrency(Int16NegativeValue),	FText::FromString(TEXT("-$12,345.00")));
+		Test(this, TEXT("Convert a Negative int32 to a currency formatted correct for en-US"),	FText::AsCurrency(Int32NegativeValue),	FText::FromString(TEXT("-$12,345.00")));
+		Test(this, TEXT("Convert a Negative int64 to a currency formatted correct for en-US"),	FText::AsCurrency(Int64NegativeValue),	FText::FromString(TEXT("-$12,345.00")));
+
+		Test(this, TEXT("Convert a Double to a percent formatted correct for en-US"),			FText::AsPercent(DoubleValue),			FText::FromString(TEXT("1,234,567,890%")));
+		Test(this, TEXT("Convert a Float to a percent formatted correct for en-US"),			FText::AsPercent(FloatValue),			FText::FromString(TEXT("123,457%")));
+		Test(this, TEXT("Convert a Negative Double to a percent formatted correct for en-US"),	FText::AsPercent(DoubleNegativeValue),	FText::FromString(TEXT("-1,234,567,890%")));
+		Test(this, TEXT("Convert a Negative Float to a percent formatted correct for en-US"),	FText::AsPercent(FloatNegativeValue),	FText::FromString(TEXT("-123,457%")));
+	}
+	else
 	{
-		FNumberFormattingOptions NumberFormattingOptions;
-		NumberFormattingOptions.MinimumIntegralDigits = 1;
-		NumberFormattingOptions.MaximumIntegralDigits = 1;
-		NumberFormattingOptions.MinimumFractionalDigits = 0;
-		NumberFormattingOptions.MaximumFractionalDigits = 0;
-		NumberFormattingOptions.UseGrouping = false;
-
-		NumberFormattingOptions.RoundingMode = ERoundingMode::HalfToEven;
-
-		Test(this, TEXT("Round a Double to a number formatted correct for en-US using HalfToEven"),	FText::AsNumber(-1.50, &(NumberFormattingOptions)),	FText::FromString(TEXT("-2")));
-		Test(this, TEXT("Round a Double to a number formatted correct for en-US using HalfToEven"),	FText::AsNumber(-1.00, &(NumberFormattingOptions)),	FText::FromString(TEXT("-1")));
-		Test(this, TEXT("Round a Double to a number formatted correct for en-US using HalfToEven"),	FText::AsNumber(-0.75, &(NumberFormattingOptions)),	FText::FromString(TEXT("-1")));
-		Test(this, TEXT("Round a Double to a number formatted correct for en-US using HalfToEven"),	FText::AsNumber(-0.50, &(NumberFormattingOptions)),	FText::FromString(TEXT("-0")));
-		Test(this, TEXT("Round a Double to a number formatted correct for en-US using HalfToEven"),	FText::AsNumber(-0.25, &(NumberFormattingOptions)),	FText::FromString(TEXT("-0")));
-		Test(this, TEXT("Round a Double to a number formatted correct for en-US using HalfToEven"),	FText::AsNumber(-0.00, &(NumberFormattingOptions)),	FText::FromString(TEXT("-0")));
-		Test(this, TEXT("Round a Double to a number formatted correct for en-US using HalfToEven"),	FText::AsNumber( 0.00, &(NumberFormattingOptions)),	FText::FromString(TEXT( "0")));
-		Test(this, TEXT("Round a Double to a number formatted correct for en-US using HalfToEven"),	FText::AsNumber( 0.25, &(NumberFormattingOptions)),	FText::FromString(TEXT( "0")));
-		Test(this, TEXT("Round a Double to a number formatted correct for en-US using HalfToEven"),	FText::AsNumber( 0.50, &(NumberFormattingOptions)),	FText::FromString(TEXT( "0")));
-		Test(this, TEXT("Round a Double to a number formatted correct for en-US using HalfToEven"),	FText::AsNumber( 0.75, &(NumberFormattingOptions)),	FText::FromString(TEXT( "1")));
-		Test(this, TEXT("Round a Double to a number formatted correct for en-US using HalfToEven"),	FText::AsNumber( 1.00, &(NumberFormattingOptions)),	FText::FromString(TEXT( "1")));
-		Test(this, TEXT("Round a Double to a number formatted correct for en-US using HalfToEven"),	FText::AsNumber( 1.50, &(NumberFormattingOptions)),	FText::FromString(TEXT( "2")));
-
-		NumberFormattingOptions.RoundingMode = ERoundingMode::HalfFromZero;
-
-		Test(this, TEXT("Round a Double to a number formatted correct for en-US using HalfFromZero"),	FText::AsNumber(-1.00, &(NumberFormattingOptions)),	FText::FromString(TEXT("-1")));
-		Test(this, TEXT("Round a Double to a number formatted correct for en-US using HalfFromZero"),	FText::AsNumber(-0.75, &(NumberFormattingOptions)),	FText::FromString(TEXT("-1")));
-		Test(this, TEXT("Round a Double to a number formatted correct for en-US using HalfFromZero"),	FText::AsNumber(-0.50, &(NumberFormattingOptions)),	FText::FromString(TEXT("-1")));
-		Test(this, TEXT("Round a Double to a number formatted correct for en-US using HalfFromZero"),	FText::AsNumber(-0.25, &(NumberFormattingOptions)),	FText::FromString(TEXT("-0")));
-		Test(this, TEXT("Round a Double to a number formatted correct for en-US using HalfFromZero"),	FText::AsNumber(-0.00, &(NumberFormattingOptions)),	FText::FromString(TEXT("-0")));
-		Test(this, TEXT("Round a Double to a number formatted correct for en-US using HalfFromZero"),	FText::AsNumber( 0.00, &(NumberFormattingOptions)),	FText::FromString(TEXT( "0")));
-		Test(this, TEXT("Round a Double to a number formatted correct for en-US using HalfFromZero"),	FText::AsNumber( 0.25, &(NumberFormattingOptions)),	FText::FromString(TEXT( "0")));
-		Test(this, TEXT("Round a Double to a number formatted correct for en-US using HalfFromZero"),	FText::AsNumber( 0.50, &(NumberFormattingOptions)),	FText::FromString(TEXT( "1")));
-		Test(this, TEXT("Round a Double to a number formatted correct for en-US using HalfFromZero"),	FText::AsNumber( 0.75, &(NumberFormattingOptions)),	FText::FromString(TEXT( "1")));
-		Test(this, TEXT("Round a Double to a number formatted correct for en-US using HalfFromZero"),	FText::AsNumber( 1.00, &(NumberFormattingOptions)),	FText::FromString(TEXT( "1")));
-
-		NumberFormattingOptions.RoundingMode = ERoundingMode::HalfToZero;
-
-		Test(this, TEXT("Round a Double to a number formatted correct for en-US using HalfToZero"),	FText::AsNumber(-1.00, &(NumberFormattingOptions)),	FText::FromString(TEXT("-1")));
-		Test(this, TEXT("Round a Double to a number formatted correct for en-US using HalfToZero"),	FText::AsNumber(-0.75, &(NumberFormattingOptions)),	FText::FromString(TEXT("-1")));
-		Test(this, TEXT("Round a Double to a number formatted correct for en-US using HalfToZero"),	FText::AsNumber(-0.50, &(NumberFormattingOptions)),	FText::FromString(TEXT("-0")));
-		Test(this, TEXT("Round a Double to a number formatted correct for en-US using HalfToZero"),	FText::AsNumber(-0.25, &(NumberFormattingOptions)),	FText::FromString(TEXT("-0")));
-		Test(this, TEXT("Round a Double to a number formatted correct for en-US using HalfToZero"),	FText::AsNumber(-0.00, &(NumberFormattingOptions)),	FText::FromString(TEXT("-0")));
-		Test(this, TEXT("Round a Double to a number formatted correct for en-US using HalfToZero"),	FText::AsNumber( 0.00, &(NumberFormattingOptions)),	FText::FromString(TEXT( "0")));
-		Test(this, TEXT("Round a Double to a number formatted correct for en-US using HalfToZero"),	FText::AsNumber( 0.25, &(NumberFormattingOptions)),	FText::FromString(TEXT( "0")));
-		Test(this, TEXT("Round a Double to a number formatted correct for en-US using HalfToZero"),	FText::AsNumber( 0.50, &(NumberFormattingOptions)),	FText::FromString(TEXT( "0")));
-		Test(this, TEXT("Round a Double to a number formatted correct for en-US using HalfToZero"),	FText::AsNumber( 0.75, &(NumberFormattingOptions)),	FText::FromString(TEXT( "1")));
-		Test(this, TEXT("Round a Double to a number formatted correct for en-US using HalfToZero"),	FText::AsNumber( 1.00, &(NumberFormattingOptions)),	FText::FromString(TEXT( "1")));
-
-		NumberFormattingOptions.RoundingMode = ERoundingMode::FromZero;
-
-		Test(this, TEXT("Round a Double to a number formatted correct for en-US using FromZero"),	FText::AsNumber(-1.0, &(NumberFormattingOptions)),	FText::FromString(TEXT("-1")));
-		Test(this, TEXT("Round a Double to a number formatted correct for en-US using FromZero"),	FText::AsNumber(-0.5, &(NumberFormattingOptions)),	FText::FromString(TEXT("-1")));
-		Test(this, TEXT("Round a Double to a number formatted correct for en-US using FromZero"),	FText::AsNumber(-0.0, &(NumberFormattingOptions)),	FText::FromString(TEXT("-0")));
-		Test(this, TEXT("Round a Double to a number formatted correct for en-US using FromZero"),	FText::AsNumber( 0.0, &(NumberFormattingOptions)),	FText::FromString(TEXT( "0")));
-		Test(this, TEXT("Round a Double to a number formatted correct for en-US using FromZero"),	FText::AsNumber( 0.5, &(NumberFormattingOptions)),	FText::FromString(TEXT( "1")));
-		Test(this, TEXT("Round a Double to a number formatted correct for en-US using FromZero"),	FText::AsNumber( 1.0, &(NumberFormattingOptions)),	FText::FromString(TEXT( "1")));
-
-		NumberFormattingOptions.RoundingMode = ERoundingMode::ToZero;
-
-		Test(this, TEXT("Round a Double to a number formatted correct for en-US using ToZero"),	FText::AsNumber(-1.0, &(NumberFormattingOptions)),	FText::FromString(TEXT("-1")));
-		Test(this, TEXT("Round a Double to a number formatted correct for en-US using ToZero"),	FText::AsNumber(-0.5, &(NumberFormattingOptions)),	FText::FromString(TEXT("-0")));
-		Test(this, TEXT("Round a Double to a number formatted correct for en-US using ToZero"),	FText::AsNumber(-0.0, &(NumberFormattingOptions)),	FText::FromString(TEXT("-0")));
-		Test(this, TEXT("Round a Double to a number formatted correct for en-US using ToZero"),	FText::AsNumber( 0.0, &(NumberFormattingOptions)),	FText::FromString(TEXT( "0")));
-		Test(this, TEXT("Round a Double to a number formatted correct for en-US using ToZero"),	FText::AsNumber( 0.5, &(NumberFormattingOptions)),	FText::FromString(TEXT( "0")));
-		Test(this, TEXT("Round a Double to a number formatted correct for en-US using ToZero"),	FText::AsNumber( 1.0, &(NumberFormattingOptions)),	FText::FromString(TEXT( "1")));
-
-		NumberFormattingOptions.RoundingMode = ERoundingMode::ToNegativeInfinity;
-
-		Test(this, TEXT("Round a Double to a number formatted correct for en-US using ToNegativeInfinity"),	FText::AsNumber(-1.0, &(NumberFormattingOptions)),	FText::FromString(TEXT("-1")));
-		Test(this, TEXT("Round a Double to a number formatted correct for en-US using ToNegativeInfinity"),	FText::AsNumber(-0.5, &(NumberFormattingOptions)),	FText::FromString(TEXT("-1")));
-		Test(this, TEXT("Round a Double to a number formatted correct for en-US using ToNegativeInfinity"),	FText::AsNumber(-0.0, &(NumberFormattingOptions)),	FText::FromString(TEXT("-0")));
-		Test(this, TEXT("Round a Double to a number formatted correct for en-US using ToNegativeInfinity"),	FText::AsNumber( 0.0, &(NumberFormattingOptions)),	FText::FromString(TEXT( "0")));
-		Test(this, TEXT("Round a Double to a number formatted correct for en-US using ToNegativeInfinity"),	FText::AsNumber( 0.5, &(NumberFormattingOptions)),	FText::FromString(TEXT( "0")));
-		Test(this, TEXT("Round a Double to a number formatted correct for en-US using ToNegativeInfinity"),	FText::AsNumber( 1.0, &(NumberFormattingOptions)),	FText::FromString(TEXT( "1")));
-
-		NumberFormattingOptions.RoundingMode = ERoundingMode::ToPositiveInfinity;
-
-		Test(this, TEXT("Round a Double to a number formatted correct for en-US using ToPositiveInfinity"),	FText::AsNumber(-1.0, &(NumberFormattingOptions)),	FText::FromString(TEXT("-1")));
-		Test(this, TEXT("Round a Double to a number formatted correct for en-US using ToPositiveInfinity"),	FText::AsNumber(-0.5, &(NumberFormattingOptions)),	FText::FromString(TEXT("-0")));
-		Test(this, TEXT("Round a Double to a number formatted correct for en-US using ToPositiveInfinity"),	FText::AsNumber(-0.0, &(NumberFormattingOptions)),	FText::FromString(TEXT("-0")));
-		Test(this, TEXT("Round a Double to a number formatted correct for en-US using ToPositiveInfinity"),	FText::AsNumber( 0.0, &(NumberFormattingOptions)),	FText::FromString(TEXT( "0")));
-		Test(this, TEXT("Round a Double to a number formatted correct for en-US using ToPositiveInfinity"),	FText::AsNumber( 0.5, &(NumberFormattingOptions)),	FText::FromString(TEXT( "1")));
-		Test(this, TEXT("Round a Double to a number formatted correct for en-US using ToPositiveInfinity"),	FText::AsNumber( 1.0, &(NumberFormattingOptions)),	FText::FromString(TEXT( "1")));
+			AddWarning(FString::Printf(TEXT("Internationalization data for %s missing - test is partially disabled."), TEXT("en-US")));
 	}
 
-	Test(this, TEXT("Convert a Double to a currency formatted correct for en-US"),			FText::AsCurrency(DoubleValue),			FText::FromString(TEXT("$12,345,678.90")));
-	Test(this, TEXT("Convert a Float to a currency formatted correct for en-US"),			FText::AsCurrency(FloatValue),			FText::FromString(TEXT("$1,234.57")));
-	Test(this, TEXT("Convert a Negative Double to a currency formatted correct for en-US"),	FText::AsCurrency(DoubleNegativeValue), FText::FromString(TEXT("-$12,345,678.90")));
-	Test(this, TEXT("Convert a Negative Float to a currency formatted correct for en-US"),	FText::AsCurrency(FloatNegativeValue),	FText::FromString(TEXT("-$1,234.57")));
-	Test(this, TEXT("Convert a uint8 to a currency formatted correct for en-US"),			FText::AsCurrency(Uint8Value),			FText::FromString(TEXT("$255.00")));
-	Test(this, TEXT("Convert a uint16 to a currency formatted correct for en-US"),			FText::AsCurrency(Uint16Value),			FText::FromString(TEXT("$65,535.00")));
-	Test(this, TEXT("Convert a uint32 to a currency formatted correct for en-US"),			FText::AsCurrency(Uint32Value),			FText::FromString(TEXT("$4,294,967,295.00")));
-	Test(this, TEXT("Convert a uint64 to a currency formatted correct for en-US"),			FText::AsCurrency(Uint64Value),			FText::FromString(TEXT("$1,844,674,407,370,955,161.00")));
-	Test(this, TEXT("Convert a int8 to a currency formatted correct for en-US"),			FText::AsCurrency(Int8Value),			FText::FromString(TEXT("$123.00")));
-	Test(this, TEXT("Convert a int16 to a currency formatted correct for en-US"),			FText::AsCurrency(Int16Value),			FText::FromString(TEXT("$12,345.00")));
-	Test(this, TEXT("Convert a int32 to a currency formatted correct for en-US"),			FText::AsCurrency(Int32Value),			FText::FromString(TEXT("$12,345.00")));
-	Test(this, TEXT("Convert a int64 to a currency formatted correct for en-US"),			FText::AsCurrency(Int64Value),			FText::FromString(TEXT("$12,345.00")));
-	Test(this, TEXT("Convert a Negative int8 to a currency formatted correct for en-US"),	FText::AsCurrency(Int8NegativeValue),	FText::FromString(TEXT("-$123.00")));
-	Test(this, TEXT("Convert a Negative int16 to a currency formatted correct for en-US"),	FText::AsCurrency(Int16NegativeValue),	FText::FromString(TEXT("-$12,345.00")));
-	Test(this, TEXT("Convert a Negative int32 to a currency formatted correct for en-US"),	FText::AsCurrency(Int32NegativeValue),	FText::FromString(TEXT("-$12,345.00")));
-	Test(this, TEXT("Convert a Negative int64 to a currency formatted correct for en-US"),	FText::AsCurrency(Int64NegativeValue),	FText::FromString(TEXT("-$12,345.00")));
-
-	Test(this, TEXT("Convert a Double to a percent formatted correct for en-US"),			FText::AsPercent(DoubleValue),			FText::FromString(TEXT("1,234,567,890%")));
-	Test(this, TEXT("Convert a Float to a percent formatted correct for en-US"),			FText::AsPercent(FloatValue),			FText::FromString(TEXT("123,457%")));
-	Test(this, TEXT("Convert a Negative Double to a percent formatted correct for en-US"),	FText::AsPercent(DoubleNegativeValue),	FText::FromString(TEXT("-1,234,567,890%")));
-	Test(this, TEXT("Convert a Negative Float to a percent formatted correct for en-US"),	FText::AsPercent(FloatNegativeValue),	FText::FromString(TEXT("-123,457%")));
-
-	I18N.SetCurrentCulture("hi-IN");
-
-	Test(this, TEXT("Convert a Double to a number formatted correct for hi-IN"),			FText::AsNumber(DoubleValue),			FText::FromString(TEXT("1,23,45,678.901")));
-	Test(this, TEXT("Convert a Float to a number formatted correct for hi-IN"),				FText::AsNumber(FloatValue),			FText::FromString(TEXT("1,234.567")));
-	Test(this, TEXT("Convert a Negative Double to a number formatted correct for hi-IN"),	FText::AsNumber(DoubleNegativeValue),	FText::FromString(TEXT("-1,23,45,678.901")));
-	Test(this, TEXT("Convert a Negative Float to a number formatted correct for hi-IN"),	FText::AsNumber(FloatNegativeValue),	FText::FromString(TEXT("-1,234.567")));
-	Test(this, TEXT("Convert a uint8 to a number formatted correct for hi-IN"),				FText::AsNumber(Uint8Value),			FText::FromString(TEXT("255")));
-	Test(this, TEXT("Convert a uint16 to a number formatted correct for hi-IN"),			FText::AsNumber(Uint16Value),			FText::FromString(TEXT("65,535")));
-	Test(this, TEXT("Convert a uint32 to a number formatted correct for hi-IN"),			FText::AsNumber(Uint32Value),			FText::FromString(TEXT("4,29,49,67,295")));
-	Test(this, TEXT("Convert a uint64 to a number formatted correct for hi-IN"),			FText::AsNumber(Uint64Value),			FText::FromString(TEXT("18,44,67,44,07,37,09,55,161")));
-	Test(this, TEXT("Convert a int8 to a number formatted correct for hi-IN"),				FText::AsNumber(Int8Value),				FText::FromString(TEXT("123")));
-	Test(this, TEXT("Convert a int16 to a number formatted correct for hi-IN"),				FText::AsNumber(Int16Value),			FText::FromString(TEXT("12,345")));
-	Test(this, TEXT("Convert a int32 to a number formatted correct for hi-IN"),				FText::AsNumber(Int32Value),			FText::FromString(TEXT("12,345")));
-	Test(this, TEXT("Convert a int64 to a number formatted correct for hi-IN"),				FText::AsNumber(Int64Value),			FText::FromString(TEXT("12,345")));
-	Test(this, TEXT("Convert a Negative int8 to a number formatted correct for hi-IN"),		FText::AsNumber(Int8NegativeValue),		FText::FromString(TEXT("-123")));
-	Test(this, TEXT("Convert a Negative int16 to a number formatted correct for hi-IN"),	FText::AsNumber(Int16NegativeValue),	FText::FromString(TEXT("-12,345")));
-	Test(this, TEXT("Convert a Negative int32 to a number formatted correct for hi-IN"),	FText::AsNumber(Int32NegativeValue),	FText::FromString(TEXT("-12,345")));
-	Test(this, TEXT("Convert a Negative int64 to a number formatted correct for hi-IN"),	FText::AsNumber(Int64NegativeValue),	FText::FromString(TEXT("-12,345")));
+	if (I18N.SetCurrentCulture("hi-IN"))
+	{
+		Test(this, TEXT("Convert a Double to a number formatted correct for hi-IN"),			FText::AsNumber(DoubleValue),			FText::FromString(TEXT("1,23,45,678.901")));
+		Test(this, TEXT("Convert a Float to a number formatted correct for hi-IN"),				FText::AsNumber(FloatValue),			FText::FromString(TEXT("1,234.567")));
+		Test(this, TEXT("Convert a Negative Double to a number formatted correct for hi-IN"),	FText::AsNumber(DoubleNegativeValue),	FText::FromString(TEXT("-1,23,45,678.901")));
+		Test(this, TEXT("Convert a Negative Float to a number formatted correct for hi-IN"),	FText::AsNumber(FloatNegativeValue),	FText::FromString(TEXT("-1,234.567")));
+		Test(this, TEXT("Convert a uint8 to a number formatted correct for hi-IN"),				FText::AsNumber(Uint8Value),			FText::FromString(TEXT("255")));
+		Test(this, TEXT("Convert a uint16 to a number formatted correct for hi-IN"),			FText::AsNumber(Uint16Value),			FText::FromString(TEXT("65,535")));
+		Test(this, TEXT("Convert a uint32 to a number formatted correct for hi-IN"),			FText::AsNumber(Uint32Value),			FText::FromString(TEXT("4,29,49,67,295")));
+		Test(this, TEXT("Convert a uint64 to a number formatted correct for hi-IN"),			FText::AsNumber(Uint64Value),			FText::FromString(TEXT("18,44,67,44,07,37,09,55,161")));
+		Test(this, TEXT("Convert a int8 to a number formatted correct for hi-IN"),				FText::AsNumber(Int8Value),				FText::FromString(TEXT("123")));
+		Test(this, TEXT("Convert a int16 to a number formatted correct for hi-IN"),				FText::AsNumber(Int16Value),			FText::FromString(TEXT("12,345")));
+		Test(this, TEXT("Convert a int32 to a number formatted correct for hi-IN"),				FText::AsNumber(Int32Value),			FText::FromString(TEXT("12,345")));
+		Test(this, TEXT("Convert a int64 to a number formatted correct for hi-IN"),				FText::AsNumber(Int64Value),			FText::FromString(TEXT("12,345")));
+		Test(this, TEXT("Convert a Negative int8 to a number formatted correct for hi-IN"),		FText::AsNumber(Int8NegativeValue),		FText::FromString(TEXT("-123")));
+		Test(this, TEXT("Convert a Negative int16 to a number formatted correct for hi-IN"),	FText::AsNumber(Int16NegativeValue),	FText::FromString(TEXT("-12,345")));
+		Test(this, TEXT("Convert a Negative int32 to a number formatted correct for hi-IN"),	FText::AsNumber(Int32NegativeValue),	FText::FromString(TEXT("-12,345")));
+		Test(this, TEXT("Convert a Negative int64 to a number formatted correct for hi-IN"),	FText::AsNumber(Int64NegativeValue),	FText::FromString(TEXT("-12,345")));
+
+		{
+			const FCulturePtr& InvariantCulture = I18N.GetInvariantCulture();
+			Test(this, TEXT("Convert a Double to a number formatted correct for hi-IN but as invariant"),				FText::AsNumber(DoubleValue, NULL, InvariantCulture),			FText::FromString(TEXT("12345678.901")));
+			Test(this, TEXT("Convert a Float to a number formatted correct for hi-IN but as invariant"),				FText::AsNumber(FloatValue, NULL, InvariantCulture),			FText::FromString(TEXT("1234.567017")));
+			Test(this, TEXT("Convert a Negative Double to a number formatted correct for hi-IN but as invariant"),		FText::AsNumber(DoubleNegativeValue, NULL, InvariantCulture),	FText::FromString(TEXT("-12345678.901")));
+			Test(this, TEXT("Convert a Negative Float to a number formatted correct for hi-IN but as invariant"),		FText::AsNumber(FloatNegativeValue, NULL, InvariantCulture),	FText::FromString(TEXT("-1234.567017")));
+			Test(this, TEXT("Convert a uint8 to a number formatted correct for hi-IN but as invariant"),				FText::AsNumber(Uint8Value, NULL, InvariantCulture),			FText::FromString(TEXT("255")));
+			Test(this, TEXT("Convert a uint16 to a number formatted correct for hi-IN but as invariant"),				FText::AsNumber(Uint16Value, NULL, InvariantCulture),			FText::FromString(TEXT("65535")));
+			Test(this, TEXT("Convert a uint32 to a number formatted correct for hi-IN but as invariant"),				FText::AsNumber(Uint32Value, NULL, InvariantCulture),			FText::FromString(TEXT("4294967295")));
+			Test(this, TEXT("Convert a uint64 to a number formatted correct for hi-IN but as invariant"),				FText::AsNumber(Uint64Value, NULL, InvariantCulture),			FText::FromString(TEXT("1844674407370955161")));
+			Test(this, TEXT("Convert a int8 to a number formatted correct for hi-IN but as invariant"),					FText::AsNumber(Int8Value, NULL, InvariantCulture),				FText::FromString(TEXT("123")));
+			Test(this, TEXT("Convert a int16 to a number formatted correct for hi-IN but as invariant"),				FText::AsNumber(Int16Value, NULL, InvariantCulture),			FText::FromString(TEXT("12345")));
+			Test(this, TEXT("Convert a int32 to a number formatted correct for hi-IN but as invariant"),				FText::AsNumber(Int32Value, NULL, InvariantCulture),			FText::FromString(TEXT("12345")));
+			Test(this, TEXT("Convert a int64 to a number formatted correct for hi-IN but as invariant"),				FText::AsNumber(Int64Value, NULL, InvariantCulture),			FText::FromString(TEXT("12345")));
+			Test(this, TEXT("Convert a Negative int8 to a number formatted correct for hi-IN but as invariant"),		FText::AsNumber(Int8NegativeValue, NULL, InvariantCulture),		FText::FromString(TEXT("-123")));
+			Test(this, TEXT("Convert a Negative int16 to a number formatted correct for hi-IN but as invariant"),		FText::AsNumber(Int16NegativeValue, NULL, InvariantCulture),	FText::FromString(TEXT("-12345")));
+			Test(this, TEXT("Convert a Negative int32 to a number formatted correct for hi-IN but as invariant"),		FText::AsNumber(Int32NegativeValue, NULL, InvariantCulture),	FText::FromString(TEXT("-12345")));
+			Test(this, TEXT("Convert a Negative int64 to a number formatted correct for hi-IN but as invariant"),		FText::AsNumber(Int64NegativeValue, NULL, InvariantCulture),	FText::FromString(TEXT("-12345")));
+		}
 
+		Test(this, TEXT("Convert a Double to a currency formatted correct for hi-IN"),				FText::AsCurrency(DoubleValue),			FText::FromString(TEXT("\x20B9") TEXT("1,23,45,678.90")));
+		Test(this, TEXT("Convert a Float to a currency formatted correct for hi-IN"),				FText::AsCurrency(FloatValue),			FText::FromString(TEXT("\x20B9") TEXT("1,234.57")));
+		Test(this, TEXT("Convert a Negative Double to a currency formatted correct for hi-IN"),		FText::AsCurrency(DoubleNegativeValue),	FText::FromString(TEXT("-") TEXT("\x20B9") TEXT("1,23,45,678.90")));
+		Test(this, TEXT("Convert a Negative Float to a currency formatted correct for hi-IN"),		FText::AsCurrency(FloatNegativeValue),	FText::FromString(TEXT("-") TEXT("\x20B9") TEXT("1,234.57")));
+		Test(this, TEXT("Convert a uint8 to a currency formatted correct for hi-IN"),				FText::AsCurrency(Uint8Value),			FText::FromString(TEXT("\x20B9") TEXT("255.00")));
+		Test(this, TEXT("Convert a uint16 to a currency formatted correct for hi-IN"),				FText::AsCurrency(Uint16Value),			FText::FromString(TEXT("\x20B9") TEXT("65,535.00")));
+		Test(this, TEXT("Convert a uint32 to a currency formatted correct for hi-IN"),				FText::AsCurrency(Uint32Value),			FText::FromString(TEXT("\x20B9") TEXT("4,29,49,67,295.00")));
+		Test(this, TEXT("Convert a uint64 to a currency formatted correct for hi-IN"),				FText::AsCurrency(Uint64Value),			FText::FromString(TEXT("\x20B9") TEXT("18,44,67,44,07,37,09,55,161.00")));
+		Test(this, TEXT("Convert a int8 to a currency formatted correct for hi-IN"),				FText::AsCurrency(Int8Value),			FText::FromString(TEXT("\x20B9") TEXT("123.00")));
+		Test(this, TEXT("Convert a int16 to a currency formatted correct for hi-IN"),				FText::AsCurrency(Int16Value),			FText::FromString(TEXT("\x20B9") TEXT("12,345.00")));
+		Test(this, TEXT("Convert a int32 to a currency formatted correct for hi-IN"),				FText::AsCurrency(Int32Value),			FText::FromString(TEXT("\x20B9") TEXT("12,345.00")));
+		Test(this, TEXT("Convert a int64 to a currency formatted correct for hi-IN"),				FText::AsCurrency(Int64Value),			FText::FromString(TEXT("\x20B9") TEXT("12,345.00")));
+		Test(this, TEXT("Convert a Negative int8 to a currency formatted correct for hi-IN"),		FText::AsCurrency(Int8NegativeValue),	FText::FromString(TEXT("-") TEXT("\x20B9") TEXT("123.00")));
+		Test(this, TEXT("Convert a Negative int16 to a currency formatted correct for hi-IN"),		FText::AsCurrency(Int16NegativeValue),	FText::FromString(TEXT("-") TEXT("\x20B9") TEXT("12,345.00")));
+		Test(this, TEXT("Convert a Negative int32 to a currency formatted correct for hi-IN"),		FText::AsCurrency(Int32NegativeValue),	FText::FromString(TEXT("-") TEXT("\x20B9") TEXT("12,345.00")));
+		Test(this, TEXT("Convert a Negative int64 to a currency formatted correct for hi-IN"),		FText::AsCurrency(Int64NegativeValue),	FText::FromString(TEXT("-") TEXT("\x20B9") TEXT("12,345.00")));
+
+		Test(this, TEXT("Convert a Double to a percent formatted correct for hi-IN"),				FText::AsPercent(DoubleValue),			FText::FromString(TEXT("1,23,45,67,890%")));
+		Test(this, TEXT("Convert a Float to a percent formatted correct for hi-IN"),				FText::AsPercent(FloatValue),			FText::FromString(TEXT("1,23,457%")));
+		Test(this, TEXT("Convert a Negative Double to a percent formatted correct for hi-IN"),		FText::AsPercent(DoubleNegativeValue),	FText::FromString(TEXT("-1,23,45,67,890%")));
+		Test(this, TEXT("Convert a Negative Float to a percent formatted correct for hi-IN"),		FText::AsPercent(FloatNegativeValue),	FText::FromString(TEXT("-1,23,457%")));
+	}
+	else
 	{
-		const FCulturePtr& InvariantCulture = I18N.GetInvariantCulture();
-		Test(this, TEXT("Convert a Double to a number formatted correct for hi-IN but as invariant"),				FText::AsNumber(DoubleValue, NULL, InvariantCulture),			FText::FromString(TEXT("12345678.901")));
-		Test(this, TEXT("Convert a Float to a number formatted correct for hi-IN but as invariant"),				FText::AsNumber(FloatValue, NULL, InvariantCulture),			FText::FromString(TEXT("1234.567017")));
-		Test(this, TEXT("Convert a Negative Double to a number formatted correct for hi-IN but as invariant"),		FText::AsNumber(DoubleNegativeValue, NULL, InvariantCulture),	FText::FromString(TEXT("-12345678.901")));
-		Test(this, TEXT("Convert a Negative Float to a number formatted correct for hi-IN but as invariant"),		FText::AsNumber(FloatNegativeValue, NULL, InvariantCulture),	FText::FromString(TEXT("-1234.567017")));
-		Test(this, TEXT("Convert a uint8 to a number formatted correct for hi-IN but as invariant"),				FText::AsNumber(Uint8Value, NULL, InvariantCulture),			FText::FromString(TEXT("255")));
-		Test(this, TEXT("Convert a uint16 to a number formatted correct for hi-IN but as invariant"),				FText::AsNumber(Uint16Value, NULL, InvariantCulture),			FText::FromString(TEXT("65535")));
-		Test(this, TEXT("Convert a uint32 to a number formatted correct for hi-IN but as invariant"),				FText::AsNumber(Uint32Value, NULL, InvariantCulture),			FText::FromString(TEXT("4294967295")));
-		Test(this, TEXT("Convert a uint64 to a number formatted correct for hi-IN but as invariant"),				FText::AsNumber(Uint64Value, NULL, InvariantCulture),			FText::FromString(TEXT("1844674407370955161")));
-		Test(this, TEXT("Convert a int8 to a number formatted correct for hi-IN but as invariant"),					FText::AsNumber(Int8Value, NULL, InvariantCulture),				FText::FromString(TEXT("123")));
-		Test(this, TEXT("Convert a int16 to a number formatted correct for hi-IN but as invariant"),				FText::AsNumber(Int16Value, NULL, InvariantCulture),			FText::FromString(TEXT("12345")));
-		Test(this, TEXT("Convert a int32 to a number formatted correct for hi-IN but as invariant"),				FText::AsNumber(Int32Value, NULL, InvariantCulture),			FText::FromString(TEXT("12345")));
-		Test(this, TEXT("Convert a int64 to a number formatted correct for hi-IN but as invariant"),				FText::AsNumber(Int64Value, NULL, InvariantCulture),			FText::FromString(TEXT("12345")));
-		Test(this, TEXT("Convert a Negative int8 to a number formatted correct for hi-IN but as invariant"),		FText::AsNumber(Int8NegativeValue, NULL, InvariantCulture),		FText::FromString(TEXT("-123")));
-		Test(this, TEXT("Convert a Negative int16 to a number formatted correct for hi-IN but as invariant"),		FText::AsNumber(Int16NegativeValue, NULL, InvariantCulture),	FText::FromString(TEXT("-12345")));
-		Test(this, TEXT("Convert a Negative int32 to a number formatted correct for hi-IN but as invariant"),		FText::AsNumber(Int32NegativeValue, NULL, InvariantCulture),	FText::FromString(TEXT("-12345")));
-		Test(this, TEXT("Convert a Negative int64 to a number formatted correct for hi-IN but as invariant"),		FText::AsNumber(Int64NegativeValue, NULL, InvariantCulture),	FText::FromString(TEXT("-12345")));
+		AddWarning(FString::Printf(TEXT("Internationalization data for %s missing - test is partially disabled."), TEXT("hi-IN")));
 	}
 
-	Test(this, TEXT("Convert a Double to a currency formatted correct for hi-IN"),				FText::AsCurrency(DoubleValue),			FText::FromString(TEXT("\x20B9") TEXT("1,23,45,678.90")));
-	Test(this, TEXT("Convert a Float to a currency formatted correct for hi-IN"),				FText::AsCurrency(FloatValue),			FText::FromString(TEXT("\x20B9") TEXT("1,234.57")));
-	Test(this, TEXT("Convert a Negative Double to a currency formatted correct for hi-IN"),		FText::AsCurrency(DoubleNegativeValue),	FText::FromString(TEXT("-") TEXT("\x20B9") TEXT("1,23,45,678.90")));
-	Test(this, TEXT("Convert a Negative Float to a currency formatted correct for hi-IN"),		FText::AsCurrency(FloatNegativeValue),	FText::FromString(TEXT("-") TEXT("\x20B9") TEXT("1,234.57")));
-	Test(this, TEXT("Convert a uint8 to a currency formatted correct for hi-IN"),				FText::AsCurrency(Uint8Value),			FText::FromString(TEXT("\x20B9") TEXT("255.00")));
-	Test(this, TEXT("Convert a uint16 to a currency formatted correct for hi-IN"),				FText::AsCurrency(Uint16Value),			FText::FromString(TEXT("\x20B9") TEXT("65,535.00")));
-	Test(this, TEXT("Convert a uint32 to a currency formatted correct for hi-IN"),				FText::AsCurrency(Uint32Value),			FText::FromString(TEXT("\x20B9") TEXT("4,29,49,67,295.00")));
-	Test(this, TEXT("Convert a uint64 to a currency formatted correct for hi-IN"),				FText::AsCurrency(Uint64Value),			FText::FromString(TEXT("\x20B9") TEXT("18,44,67,44,07,37,09,55,161.00")));
-	Test(this, TEXT("Convert a int8 to a currency formatted correct for hi-IN"),				FText::AsCurrency(Int8Value),			FText::FromString(TEXT("\x20B9") TEXT("123.00")));
-	Test(this, TEXT("Convert a int16 to a currency formatted correct for hi-IN"),				FText::AsCurrency(Int16Value),			FText::FromString(TEXT("\x20B9") TEXT("12,345.00")));
-	Test(this, TEXT("Convert a int32 to a currency formatted correct for hi-IN"),				FText::AsCurrency(Int32Value),			FText::FromString(TEXT("\x20B9") TEXT("12,345.00")));
-	Test(this, TEXT("Convert a int64 to a currency formatted correct for hi-IN"),				FText::AsCurrency(Int64Value),			FText::FromString(TEXT("\x20B9") TEXT("12,345.00")));
-	Test(this, TEXT("Convert a Negative int8 to a currency formatted correct for hi-IN"),		FText::AsCurrency(Int8NegativeValue),	FText::FromString(TEXT("-") TEXT("\x20B9") TEXT("123.00")));
-	Test(this, TEXT("Convert a Negative int16 to a currency formatted correct for hi-IN"),		FText::AsCurrency(Int16NegativeValue),	FText::FromString(TEXT("-") TEXT("\x20B9") TEXT("12,345.00")));
-	Test(this, TEXT("Convert a Negative int32 to a currency formatted correct for hi-IN"),		FText::AsCurrency(Int32NegativeValue),	FText::FromString(TEXT("-") TEXT("\x20B9") TEXT("12,345.00")));
-	Test(this, TEXT("Convert a Negative int64 to a currency formatted correct for hi-IN"),		FText::AsCurrency(Int64NegativeValue),	FText::FromString(TEXT("-") TEXT("\x20B9") TEXT("12,345.00")));
-
-	Test(this, TEXT("Convert a Double to a percent formatted correct for hi-IN"),				FText::AsPercent(DoubleValue),			FText::FromString(TEXT("1,23,45,67,890%")));
-	Test(this, TEXT("Convert a Float to a percent formatted correct for hi-IN"),				FText::AsPercent(FloatValue),			FText::FromString(TEXT("1,23,457%")));
-	Test(this, TEXT("Convert a Negative Double to a percent formatted correct for hi-IN"),		FText::AsPercent(DoubleNegativeValue),	FText::FromString(TEXT("-1,23,45,67,890%")));
-	Test(this, TEXT("Convert a Negative Float to a percent formatted correct for hi-IN"),		FText::AsPercent(FloatNegativeValue),	FText::FromString(TEXT("-1,23,457%")));
-
 	I18N.SetCurrentCulture( OriginalCulture );
 
 	FText Number = FText::AsNumber(Int64NegativeValue);
diff --git a/Engine/Source/Runtime/Core/Private/Tests/Internationalization/TextTest.cpp b/Engine/Source/Runtime/Core/Private/Tests/Internationalization/TextTest.cpp
--- a/Engine/Source/Runtime/Core/Private/Tests/Internationalization/TextTest.cpp
+++ b/Engine/Source/Runtime/Core/Private/Tests/Internationalization/TextTest.cpp
@@ -300,25 +300,31 @@ bool FTextTest::RunTest (const FString& Parameters)
 	FText::SetSuppressWarnings(true);
 
 #if UE_ENABLE_ICU
-	I18N.SetCurrentCulture("en-US");
+	if (I18N.SetCurrentCulture("en-US"))
+	{
 #define TEST( A, B, ComparisonLevel ) if( !(FText::FromString(A)).EqualTo(FText::FromString(B), (ComparisonLevel)) ) AddError(FString::Printf(TEXT("Testing comparison of equivalent characters with comparison level (%s). - A=%s B=%s"),TEXT(#ComparisonLevel),(A),(B)))
 
-	// Basic sanity checks
-	TEST( TEXT("a"), TEXT("A"), ETextComparisonLevel::Primary ); // Basic sanity check
-	TEST( TEXT("a"), TEXT("a"), ETextComparisonLevel::Tertiary ); // Basic sanity check
-	TEST( TEXT("A"), TEXT("A"), ETextComparisonLevel::Tertiary ); // Basic sanity check
+		// Basic sanity checks
+		TEST( TEXT("a"), TEXT("A"), ETextComparisonLevel::Primary ); // Basic sanity check
+		TEST( TEXT("a"), TEXT("a"), ETextComparisonLevel::Tertiary ); // Basic sanity check
+		TEST( TEXT("A"), TEXT("A"), ETextComparisonLevel::Tertiary ); // Basic sanity check
 
-	// Test equivalence
-	TEST( TEXT("ss"), TEXT("\x00DF"), ETextComparisonLevel::Primary ); // Lowercase Sharp s
-	TEST( TEXT("SS"), TEXT("\x1E9E"), ETextComparisonLevel::Primary ); // Uppercase Sharp S
-	TEST( TEXT("ae"), TEXT("\x00E6"), ETextComparisonLevel::Primary ); // Lowercase ae
-	TEST( TEXT("AE"), TEXT("\x00C6"), ETextComparisonLevel::Primary ); // Uppercase AE
+		// Test equivalence
+		TEST( TEXT("ss"), TEXT("\x00DF"), ETextComparisonLevel::Primary ); // Lowercase Sharp s
+		TEST( TEXT("SS"), TEXT("\x1E9E"), ETextComparisonLevel::Primary ); // Uppercase Sharp S
+		TEST( TEXT("ae"), TEXT("\x00E6"), ETextComparisonLevel::Primary ); // Lowercase ae
+		TEST( TEXT("AE"), TEXT("\x00C6"), ETextComparisonLevel::Primary ); // Uppercase AE
 
-	// Test accentuation
-	TEST( TEXT("u") , TEXT("\x00FC"), ETextComparisonLevel::Primary ); // Lowercase u with dieresis
-	TEST( TEXT("U") , TEXT("\x00DC"), ETextComparisonLevel::Primary ); // Uppercase U with dieresis
+		// Test accentuation
+		TEST( TEXT("u") , TEXT("\x00FC"), ETextComparisonLevel::Primary ); // Lowercase u with dieresis
+		TEST( TEXT("U") , TEXT("\x00DC"), ETextComparisonLevel::Primary ); // Uppercase U with dieresis
 
 #undef TEST
+	}
+	else
+	{
+		AddWarning(FString::Printf(TEXT("Internationalization data for %s missing - test is partially disabled."), TEXT("en-US")));
+	}
 #else
 	AddWarning("ICU is disabled thus locale-aware string comparison is disabled.");
 #endif
@@ -326,8 +332,8 @@ bool FTextTest::RunTest (const FString& Parameters)
 #if UE_ENABLE_ICU
 	// Sort Testing
 	// French
+	if (I18N.SetCurrentCulture("fr"))
 	{
-		I18N.SetCurrentCulture("fr");
 		TArray<FText> CorrectlySortedValues;
 		CorrectlySortedValues.Add( FText::FromString( TEXT("cote") ) );
 		CorrectlySortedValues.Add( FText::FromString( TEXT("coté") ) );
@@ -363,10 +369,14 @@ bool FTextTest::RunTest (const FString& Parameters)
 			}
 		}
 	}
+	else
+	{
+		AddWarning(FString::Printf(TEXT("Internationalization data for %s missing - test is partially disabled."), TEXT("fr")));
+	}
 
 	// French Canadian
+	if (I18N.SetCurrentCulture("fr-CA"))
 	{
-		I18N.SetCurrentCulture("fr-CA");
 		TArray<FText> CorrectlySortedValues;
 		CorrectlySortedValues.Add( FText::FromString( TEXT("cote") ) );
 		CorrectlySortedValues.Add( FText::FromString( TEXT("côte") ) );
@@ -399,6 +409,10 @@ bool FTextTest::RunTest (const FString& Parameters)
 			}
 		}
 	}
+	else
+	{
+		AddWarning(FString::Printf(TEXT("Internationalization data for %s missing - test is partially disabled."), TEXT("fr-CA")));
+	}
 #else
 	AddWarning("ICU is disabled thus locale-aware string collation is disabled.");
 #endif
diff --git a/Engine/Source/Runtime/Core/Public/Internationalization/Internationalization.h b/Engine/Source/Runtime/Core/Public/Internationalization/Internationalization.h
--- a/Engine/Source/Runtime/Core/Public/Internationalization/Internationalization.h
+++ b/Engine/Source/Runtime/Core/Public/Internationalization/Internationalization.h
@@ -27,7 +27,7 @@ public:
 	}
 
 	//Set the current culture by name
-	CORE_API void SetCurrentCulture(const FString& Name);
+	CORE_API bool SetCurrentCulture(const FString& Name);
 
 	//@return the current culture
 	CORE_API FCultureRef GetCurrentCulture() const
