CIS fix for HTML5 due to header move

diff --git a/Engine/Source/Runtime/Core/Private/HTML5/HTML5PlatformMisc.cpp b/Engine/Source/Runtime/Core/Private/HTML5/HTML5PlatformMisc.cpp
--- a/Engine/Source/Runtime/Core/Private/HTML5/HTML5PlatformMisc.cpp
+++ b/Engine/Source/Runtime/Core/Private/HTML5/HTML5PlatformMisc.cpp
@@ -11,6 +11,9 @@
 #include "HTML5JavaScriptFx.h"
 #endif 
 
+#include "unicode/locid.h"
+#include "GenericPlatformContext.h"
+#include <SDL.h>
 #include <ctime>
 
 void FHTML5Misc::PlatformInit()
