Fix resizing of fullscreen display using the System Preferences' Display panel on OS X. Our windows don't resize in response to system events, so in fullscreen mode we need to handle this ourselves.

diff --git a/Engine/Source/Runtime/Core/Private/Mac/CocoaWindow.cpp b/Engine/Source/Runtime/Core/Private/Mac/CocoaWindow.cpp
--- a/Engine/Source/Runtime/Core/Private/Mac/CocoaWindow.cpp
+++ b/Engine/Source/Runtime/Core/Private/Mac/CocoaWindow.cpp
@@ -292,6 +292,12 @@ NSString* NSPerformDragOperation = @"NSPerformDragOperation";
 {
 	SCOPED_AUTORELEASE_POOL;
 	NSSize Size = [self frame].size;
+	
+	if (TargetWindowMode != EWindowMode::Windowed)
+	{
+		FrameRect.size = [[self screen] frame].size;
+	}
+	
 	NSSize NewSize = FrameRect.size;
 	if(!bRenderInitialised || ([self isVisible] && [super alphaValue] > 0.0f && (Size.width > 1 || Size.height > 1 || NewSize.width > 1 || NewSize.height > 1)))
 	{
