Further simplification of Mac event handling code

diff --git a/Engine/Source/Runtime/Core/Private/Mac/CocoaWindow.cpp b/Engine/Source/Runtime/Core/Private/Mac/CocoaWindow.cpp
--- a/Engine/Source/Runtime/Core/Private/Mac/CocoaWindow.cpp
+++ b/Engine/Source/Runtime/Core/Private/Mac/CocoaWindow.cpp
@@ -174,7 +174,7 @@ NSString* NSPerformDragOperation = @"NSPerformDragOperation";
 	if(bNeedsRedraw && bForwardEvents && ([self isVisible] && [super alphaValue] > 0.0f))
 	{
 		NSNotification* Notification = [NSNotification notificationWithName:NSWindowRedrawContents object:self];
-		FMacEvent::SendToGameRunLoop(Notification, self, EMacEventSendMethod::Sync, @[ NSDefaultRunLoopMode, UE4NilEventMode, UE4ShowEventMode, UE4ResizeEventMode, UE4FullscreenEventMode, UE4CloseEventMode, UE4IMEEventMode ]);
+		FMacEvent::SendToGameRunLoop(Notification, EMacEventSendMethod::Sync, @[ NSDefaultRunLoopMode, UE4NilEventMode, UE4ShowEventMode, UE4ResizeEventMode, UE4FullscreenEventMode, UE4CloseEventMode, UE4IMEEventMode ]);
 	}
 	bNeedsRedraw = false;
 }
@@ -225,7 +225,7 @@ NSString* NSPerformDragOperation = @"NSPerformDragOperation";
 	return bAcceptsInput && ![self ignoresMouseEvents];
 }
 
-- (BOOL)validateMenuItem:(NSMenuItem *)MenuItem
+- (BOOL)validateMenuItem:(NSMenuItem*)MenuItem
 {
 	SCOPED_AUTORELEASE_POOL;
 	// Borderless windows we use do not automatically handle first responder's actions, so we force it here
@@ -311,7 +311,7 @@ NSString* NSPerformDragOperation = @"NSPerformDragOperation";
 		if(self.bForwardEvents)
 		{
 			NSNotification* Notification = [NSNotification notificationWithName:NSWindowDidResizeNotification object:self];
-			FMacEvent::SendToGameRunLoop(Notification, self, EMacEventSendMethod::Async, @[ NSDefaultRunLoopMode, UE4ResizeEventMode, UE4ShowEventMode ]);
+			FMacEvent::SendToGameRunLoop(Notification, EMacEventSendMethod::Async, @[ NSDefaultRunLoopMode, UE4ResizeEventMode, UE4ShowEventMode ]);
 		}
 	}
 }
@@ -333,21 +333,21 @@ NSString* NSPerformDragOperation = @"NSPerformDragOperation";
 		self.bDeferSetOrigin = true;
 		self.DeferFrame = NSMakeRect(Point.x, Point.y, self.DeferFrame.size.width, self.DeferFrame.size.height);
 		NSNotification* Notification = [NSNotification notificationWithName:NSWindowDidMoveNotification object:self];
-		FMacEvent::SendToGameRunLoop(Notification, self, EMacEventSendMethod::Async, @[ NSDefaultRunLoopMode, UE4ResizeEventMode, UE4ShowEventMode ]);
+		FMacEvent::SendToGameRunLoop(Notification, EMacEventSendMethod::Async, @[ NSDefaultRunLoopMode, UE4ResizeEventMode, UE4ShowEventMode ]);
 	}
 }
 
-- (void)keyDown:(NSEvent *)Event
+- (void)keyDown:(NSEvent*)Event
 {
 	// @note Deliberately empty - we don't want OS X to handle keyboard input as it will recursively re-add events we aren't handling
 }
 
-- (void)keyUp:(NSEvent *)Event
+- (void)keyUp:(NSEvent*)Event
 {
 	// @note Deliberately empty - we don't want OS X to handle keyboard input as it will recursively re-add events we aren't handling
 }
 
-- (void)windowWillEnterFullScreen:(NSNotification *)notification
+- (void)windowWillEnterFullScreen:(NSNotification*)Notification
 {
 	// Handle clicking on the titlebar fullscreen item
 	if(self.TargetWindowMode == EWindowMode::Windowed)
@@ -356,12 +356,12 @@ NSString* NSPerformDragOperation = @"NSPerformDragOperation";
 	}
 }
 
-- (void)windowDidEnterFullScreen:(NSNotification *)notification
+- (void)windowDidEnterFullScreen:(NSNotification*)Notification
 {
 	WindowMode = self.TargetWindowMode;
 	if(self.bForwardEvents)
 	{
-		FMacEvent::SendToGameRunLoop(notification, self, EMacEventSendMethod::Async, @[ NSDefaultRunLoopMode, UE4FullscreenEventMode ]);
+		FMacEvent::SendToGameRunLoop(Notification, EMacEventSendMethod::Async, @[ NSDefaultRunLoopMode, UE4FullscreenEventMode ]);
 	}
 	FMacCursor* MacCursor = (FMacCursor*)MacApplication->Cursor.Get();
 	if ( MacCursor )
@@ -374,13 +374,13 @@ NSString* NSPerformDragOperation = @"NSPerformDragOperation";
 	}
 }
 
-- (void)windowDidExitFullScreen:(NSNotification *)notification
+- (void)windowDidExitFullScreen:(NSNotification*)Notification
 {
 	WindowMode = EWindowMode::Windowed;
 	self.TargetWindowMode = EWindowMode::Windowed;
 	if(self.bForwardEvents)
 	{
-		FMacEvent::SendToGameRunLoop(notification, self, EMacEventSendMethod::Async, @[ NSDefaultRunLoopMode, UE4FullscreenEventMode ]);
+		FMacEvent::SendToGameRunLoop(Notification, EMacEventSendMethod::Async, @[ NSDefaultRunLoopMode, UE4FullscreenEventMode ]);
 	}
 	FMacCursor* MacCursor = (FMacCursor*)MacApplication->Cursor.Get();
 	if ( MacCursor )
@@ -389,7 +389,7 @@ NSString* NSPerformDragOperation = @"NSPerformDragOperation";
 	}
 }
 
-- (void)windowDidBecomeKey:(NSNotification *)Notification
+- (void)windowDidBecomeKey:(NSNotification*)Notification
 {
 	SCOPED_AUTORELEASE_POOL;
 	if([NSApp isHidden] == NO)
@@ -399,11 +399,11 @@ NSString* NSPerformDragOperation = @"NSPerformDragOperation";
 
 	if(self.bForwardEvents)
 	{
-		FMacEvent::SendToGameRunLoop(Notification, self, EMacEventSendMethod::Async, @[ NSDefaultRunLoopMode, UE4ShowEventMode, UE4CloseEventMode, UE4FullscreenEventMode ]);
+		FMacEvent::SendToGameRunLoop(Notification, EMacEventSendMethod::Async, @[ NSDefaultRunLoopMode, UE4ShowEventMode, UE4CloseEventMode, UE4FullscreenEventMode ]);
 	}
 }
 
-- (void)windowDidResignKey:(NSNotification *)Notification
+- (void)windowDidResignKey:(NSNotification*)Notification
 {
 	SCOPED_AUTORELEASE_POOL;
 	[self setMovable: YES];
@@ -411,19 +411,19 @@ NSString* NSPerformDragOperation = @"NSPerformDragOperation";
 	
 	if(self.bForwardEvents)
 	{
-		FMacEvent::SendToGameRunLoop(Notification, self, EMacEventSendMethod::Async, @[ NSDefaultRunLoopMode, UE4ShowEventMode, UE4CloseEventMode, UE4FullscreenEventMode ]);
+		FMacEvent::SendToGameRunLoop(Notification, EMacEventSendMethod::Async, @[ NSDefaultRunLoopMode, UE4ShowEventMode, UE4CloseEventMode, UE4FullscreenEventMode ]);
 	}
 }
 
-- (void)windowWillMove:(NSNotification *)Notification
+- (void)windowWillMove:(NSNotification*)Notification
 {
 	if(self.bForwardEvents)
 	{
-		FMacEvent::SendToGameRunLoop(Notification, self, EMacEventSendMethod::Async, @[ NSDefaultRunLoopMode, UE4ResizeEventMode, UE4ShowEventMode, UE4FullscreenEventMode ]);
+		FMacEvent::SendToGameRunLoop(Notification, EMacEventSendMethod::Async, @[ NSDefaultRunLoopMode, UE4ResizeEventMode, UE4ShowEventMode, UE4FullscreenEventMode ]);
 	}
 }
 
-- (void)windowDidMove:(NSNotification *)Notification
+- (void)windowDidMove:(NSNotification*)Notification
 {
 	SCOPED_AUTORELEASE_POOL;
 	bZoomed = [self isZoomed];
@@ -433,14 +433,12 @@ NSString* NSPerformDragOperation = @"NSPerformDragOperation";
 	
 	if(self.bForwardEvents)
 	{
-		FMacEvent::SendToGameRunLoop(Notification, self, EMacEventSendMethod::Async, @[ NSDefaultRunLoopMode, UE4ResizeEventMode, UE4ShowEventMode, UE4FullscreenEventMode ]);
+		FMacEvent::SendToGameRunLoop(Notification, EMacEventSendMethod::Async, @[ NSDefaultRunLoopMode, UE4ResizeEventMode, UE4ShowEventMode, UE4FullscreenEventMode ]);
 	}
 }
 
-- (void)windowDidChangeScreen:(NSNotification *)notification
+- (void)windowDidChangeScreen:(NSNotification*)Notification
 {
-	// You'd think this would be a good place to handle un/parenting...
-	// However, do that here and your dragged window will disappear when you drag it onto another screen!
 	// The windowdidChangeScreen notification only comes after you finish dragging.
 	// It does however, work fine for handling display arrangement changes that cause a window to go offscreen.
 	if(bDisplayReconfiguring)
@@ -501,23 +499,23 @@ NSString* NSPerformDragOperation = @"NSPerformDragOperation";
 	}
 }
 
-- (void)windowDidResize:(NSNotification *)Notification
+- (void)windowDidResize:(NSNotification*)Notification
 {
 	SCOPED_AUTORELEASE_POOL;
 	bZoomed = [self isZoomed];
 	if(self.bForwardEvents)
 	{
-		FMacEvent::SendToGameRunLoop(Notification, self, EMacEventSendMethod::Async, @[ NSDefaultRunLoopMode, UE4ResizeEventMode, UE4ShowEventMode, UE4FullscreenEventMode ]);
+		FMacEvent::SendToGameRunLoop(Notification, EMacEventSendMethod::Async, @[ NSDefaultRunLoopMode, UE4ResizeEventMode, UE4ShowEventMode, UE4FullscreenEventMode ]);
 	}
 	bNeedsRedraw = true;
 }
 
-- (void)windowWillClose:(NSNotification *)notification
+- (void)windowWillClose:(NSNotification*)Notification
 {
 	SCOPED_AUTORELEASE_POOL;
 	if(self.bForwardEvents && MacApplication)
 	{
-		FMacEvent::SendToGameRunLoop(notification, self, EMacEventSendMethod::Sync, @[ NSDefaultRunLoopMode, UE4CloseEventMode ]);
+		FMacEvent::SendToGameRunLoop(Notification, EMacEventSendMethod::Sync, @[ NSDefaultRunLoopMode, UE4CloseEventMode ]);
 	}
 	self.bForwardEvents = false;
 	[self setDelegate:nil];
@@ -590,7 +588,7 @@ NSString* NSPerformDragOperation = @"NSPerformDragOperation";
 	{
 		SCOPED_AUTORELEASE_POOL;
 		NSNotification* Notification = [NSNotification notificationWithName:NSDraggingExited object:Sender];
-		FMacEvent::SendToGameRunLoop(Notification, self, EMacEventSendMethod::Async);
+		FMacEvent::SendToGameRunLoop(Notification, EMacEventSendMethod::Async);
 	}
 }
 
@@ -600,7 +598,7 @@ NSString* NSPerformDragOperation = @"NSPerformDragOperation";
 	{
 		SCOPED_AUTORELEASE_POOL;
 		NSNotification* Notification = [NSNotification notificationWithName:NSDraggingUpdated object:Sender];
-		FMacEvent::SendToGameRunLoop(Notification, self, EMacEventSendMethod::Async);
+		FMacEvent::SendToGameRunLoop(Notification, EMacEventSendMethod::Async);
 	}
 	return NSDragOperationGeneric;
 }
@@ -611,7 +609,7 @@ NSString* NSPerformDragOperation = @"NSPerformDragOperation";
 	{
 		SCOPED_AUTORELEASE_POOL;
 		NSNotification* Notification = [NSNotification notificationWithName:NSPrepareForDragOperation object:Sender];
-		FMacEvent::SendToGameRunLoop(Notification, self, EMacEventSendMethod::Async);
+		FMacEvent::SendToGameRunLoop(Notification, EMacEventSendMethod::Async);
 	}
 	return YES;
 }
@@ -622,7 +620,7 @@ NSString* NSPerformDragOperation = @"NSPerformDragOperation";
 	{
 		SCOPED_AUTORELEASE_POOL;
 		NSNotification* Notification = [NSNotification notificationWithName:NSPerformDragOperation object:Sender];
-		FMacEvent::SendToGameRunLoop(Notification, self, EMacEventSendMethod::Async);
+		FMacEvent::SendToGameRunLoop(Notification, EMacEventSendMethod::Async);
 	}
 	return YES;
 }
diff --git a/Engine/Source/Runtime/Core/Private/Mac/MacApplication.cpp b/Engine/Source/Runtime/Core/Private/Mac/MacApplication.cpp
--- a/Engine/Source/Runtime/Core/Private/Mac/MacApplication.cpp
+++ b/Engine/Source/Runtime/Core/Private/Mac/MacApplication.cpp
@@ -74,7 +74,7 @@ NSEvent* FMacApplication::HandleNSEvent(NSEvent* Event)
 			if ([Event type] == NSLeftMouseUp)
 			{
 				NSNotification* Notification = [NSNotification notificationWithName:NSWindowDraggingFinished object:[Event window]];
-				FMacEvent::SendToGameRunLoop(Notification, [Event window], EMacEventSendMethod::Async);
+				FMacEvent::SendToGameRunLoop(Notification, EMacEventSendMethod::Async);
 			}
 		}
 	}
@@ -340,14 +340,14 @@ void FMacApplication::ProcessEvent(FMacEvent const* const Event)
 	if(AppKitEvent)
 	{
 		// Process a standard NSEvent as we always did
-		MacApplication->ProcessNSEvent(AppKitEvent, Event->GetMousePosition());
+		MacApplication->ProcessNSEvent(AppKitEvent);
 	}
 	else if(Notification)
 	{
 		// Notifications need to mapped to the right handler function, that's no longer handled in the
 		// call location.
 		NSString* const NotificationName = [Notification name];
-		FCocoaWindow* CocoaWindow = Event->GetWindow();
+		FCocoaWindow* CocoaWindow = [[Notification object] isKindOfClass:[FCocoaWindow class]] ? (FCocoaWindow*)[Notification object] : nullptr;
 		if (CocoaWindow)
 		{
 			if(NotificationName == NSWindowDidResizeNotification)
@@ -386,7 +386,19 @@ void FMacApplication::ProcessEvent(FMacEvent const* const Event)
 			{
 				MacApplication->OnWindowRedrawContents( CocoaWindow );
 			}
-			else if(NotificationName == NSDraggingExited)
+			else if(NotificationName == NSWindowDraggingFinished)
+			{
+				MacApplication->OnWindowDraggingFinished();
+			}
+			else
+			{
+				check(false);
+			}
+		}
+		else if ([[Notification object] conformsToProtocol:@protocol(NSDraggingInfo)])
+		{
+			CocoaWindow = (FCocoaWindow*)[(id<NSDraggingInfo>)[Notification object] draggingDestinationWindow];
+			if(NotificationName == NSDraggingExited)
 			{
 				MacApplication->OnDragOut( CocoaWindow );
 			}
@@ -402,10 +414,6 @@ void FMacApplication::ProcessEvent(FMacEvent const* const Event)
 			{
 				MacApplication->OnDragDrop( CocoaWindow );
 			}
-			else if(NotificationName == NSWindowDraggingFinished)
-			{
-				MacApplication->OnWindowDraggingFinished();
-			}
 			else
 			{
 				check(false);
@@ -497,7 +505,7 @@ void FMacApplication::HandleModifierChange(NSUInteger NewModifierFlags, NSUInteg
 	}
 }
 
-void FMacApplication::ProcessNSEvent(NSEvent* const Event, FVector2D const MousePosition)
+void FMacApplication::ProcessNSEvent(NSEvent* const Event)
 {
 	SCOPED_AUTORELEASE_POOL;
 
@@ -507,7 +515,7 @@ void FMacApplication::ProcessNSEvent(NSEvent* const Event, FVector2D const Mouse
 	if (CurrentModifierFlags != [Event modifierFlags])
 	{
 		NSUInteger ModifierFlags = [Event modifierFlags];
-		
+
 		HandleModifierChange(ModifierFlags, (1<<4), 7, MMK_RightCommand);
 		HandleModifierChange(ModifierFlags, (1<<3), 6, MMK_LeftCommand);
 		HandleModifierChange(ModifierFlags, (1<<1), 0, MMK_LeftShift);
@@ -517,7 +525,7 @@ void FMacApplication::ProcessNSEvent(NSEvent* const Event, FVector2D const Mouse
 		HandleModifierChange(ModifierFlags, (1<<2), 1, MMK_RightShift);
 		HandleModifierChange(ModifierFlags, (1<<6), 5, MMK_RightAlt);
 		HandleModifierChange(ModifierFlags, (1<<13), 3, MMK_RightControl);
-		
+
 		CurrentModifierFlags = ModifierFlags;
 	}
 
@@ -596,6 +604,14 @@ void FMacApplication::ProcessNSEvent(NSEvent* const Event, FVector2D const Mouse
 			}
 			else
 			{
+				NSPoint CursorPos = [Event locationInWindow];
+				if ([Event window])
+				{
+					CursorPos.x += [Event window].frame.origin.x;
+					CursorPos.y += [Event window].frame.origin.y;
+				}
+				CursorPos.y--; // The y coordinate in the point returned by locationInWindow starts from a base of 1
+				const FVector2D MousePosition = FVector2D(CursorPos.x, FPlatformMisc::ConvertSlateYPositionToCocoa(CursorPos.y));
 				FVector2D CurrentPosition = MousePosition * MacCursor->GetMouseScaling();
 				if (MacCursor->UpdateCursorClipping(CurrentPosition))
 				{
@@ -608,7 +624,7 @@ void FMacApplication::ProcessNSEvent(NSEvent* const Event, FVector2D const Mouse
 				MessageHandler->OnMouseMove();
 			}
 
-			if (!DraggedWindow && !GetCapture())
+			if (CurrentEventWindow.IsValid() && !DraggedWindow && !GetCapture())
 			{
 				MessageHandler->OnCursorSet();
 			}
@@ -685,7 +701,7 @@ void FMacApplication::ProcessNSEvent(NSEvent* const Event, FVector2D const Mouse
 
 			MessageHandler->OnMouseUp(Button);
 
-			if (!DraggedWindow && !GetCapture())
+			if (CurrentEventWindow.IsValid() && !DraggedWindow && !GetCapture())
 			{
 				MessageHandler->OnCursorSet();
 			}
@@ -716,7 +732,7 @@ void FMacApplication::ProcessNSEvent(NSEvent* const Event, FVector2D const Mouse
 				MessageHandler->OnMouseWheel(DeltaY);
 			}
 
-			if (!DraggedWindow && !GetCapture())
+			if (CurrentEventWindow.IsValid() && !DraggedWindow && !GetCapture())
 			{
 				MessageHandler->OnCursorSet();
 			}
@@ -774,7 +790,7 @@ void FMacApplication::ProcessNSEvent(NSEvent* const Event, FVector2D const Mouse
 					const uint32 KeyCode = [Event keyCode];
 					const bool IsPrintable = IsPrintableKey(Character);
 
-					bHandled = MessageHandler->OnKeyDown(KeyCode, TranslateCharCode( CharCode, KeyCode ), IsRepeat);
+					bHandled = MessageHandler->OnKeyDown(KeyCode, TranslateCharCode(CharCode, KeyCode), IsRepeat);
 
 					// First KeyDown, then KeyChar. This is important, as in-game console ignores first character otherwise
 					bool bCmdKeyPressed = [Event modifierFlags] & 0x18;
@@ -809,9 +825,9 @@ void FMacApplication::ProcessNSEvent(NSEvent* const Event, FVector2D const Mouse
 				const TCHAR Character = ConvertChar([Characters characterAtIndex:0]);
 				const TCHAR CharCode = [[Event charactersIgnoringModifiers] characterAtIndex:0];
 				const uint32 KeyCode = [Event keyCode];
-				const bool IsPrintable = IsPrintableKey( Character );
+				const bool IsPrintable = IsPrintableKey(Character);
 
-				bHandled = MessageHandler->OnKeyUp(KeyCode, TranslateCharCode( CharCode, KeyCode ), IsRepeat);
+				bHandled = MessageHandler->OnKeyUp(KeyCode, TranslateCharCode(CharCode, KeyCode), IsRepeat);
 			}
 			if (!bHandled)
 			{
@@ -886,6 +902,7 @@ FCocoaWindow* FMacApplication::FindEventWindow( NSEvent* Event )
 				CursorPos.x += [Event window].frame.origin.x;
 				CursorPos.y += [Event window].frame.origin.y;
 			}
+			CursorPos.y--; // The y coordinate in the point returned by locationInWindow starts from a base of 1
 			TSharedPtr<FMacWindow> WindowUnderCursor = LocateWindowUnderCursor(CursorPos);
 			if (WindowUnderCursor.IsValid())
 			{
diff --git a/Engine/Source/Runtime/Core/Private/Mac/MacEvent.cpp b/Engine/Source/Runtime/Core/Private/Mac/MacEvent.cpp
--- a/Engine/Source/Runtime/Core/Private/Mac/MacEvent.cpp
+++ b/Engine/Source/Runtime/Core/Private/Mac/MacEvent.cpp
@@ -12,29 +12,14 @@ FMacEvent::FMacEvent(NSEvent* const Event)
 	SCOPED_AUTORELEASE_POOL;
 	check(Event);
 	[Event retain];
-	
-	CocoaWindow = MacApplication->FindEventWindow(Event);
-
-	CGEventRef CGEvent = [Event CGEvent];
-	CGPoint CGMouseLocation = CGEventGetLocation(CGEvent);
-	MousePosition.X = CGMouseLocation.x;
-	MousePosition.Y = CGMouseLocation.y;
 }
 
-FMacEvent::FMacEvent(NSNotification* const Notification, NSWindow* const Window)
+FMacEvent::FMacEvent(NSNotification* const Notification)
 : EventData(Notification)
 {
 	SCOPED_AUTORELEASE_POOL;
 	check(Notification);
 	[Notification retain];
-
-	CocoaWindow = [Window isKindOfClass:[FCocoaWindow class]] ? (FCocoaWindow*)Window : nullptr;
-
-	CGEventRef Event = CGEventCreate(NULL);
-	CGPoint CursorPos = CGEventGetLocation(Event);
-	MousePosition.X = CursorPos.x;
-	MousePosition.Y = CursorPos.y;
-	CFRelease(Event);
 }
 
 void FMacEvent::SendToGameRunLoop(FMacEvent const* const Event, EMacEventSendMethod SendMethod, NSArray* SendModes)
@@ -55,16 +40,16 @@ void FMacEvent::SendToGameRunLoop(NSEvent* const Event, EMacEventSendMethod Send
 	}
 }
 
-void FMacEvent::SendToGameRunLoop(NSNotification* const Notification, NSWindow* const Window, EMacEventSendMethod SendMethod, NSArray* SendModes)
+void FMacEvent::SendToGameRunLoop(NSNotification* const Notification, EMacEventSendMethod SendMethod, NSArray* SendModes)
 {
 	if(MacApplication)
 	{
-		FMacEvent* MacEvent = new FMacEvent(Notification, Window);
+		FMacEvent* MacEvent = new FMacEvent(Notification);
 		FMacEvent::SendToGameRunLoop(MacEvent, SendMethod, SendModes);
 	}
 }
 
-FMacEvent::~FMacEvent(void)
+FMacEvent::~FMacEvent()
 {
 	if(EventData)
 	{
@@ -74,7 +59,7 @@ FMacEvent::~FMacEvent(void)
 	}
 }
 
-NSEvent* FMacEvent::GetEvent(void) const
+NSEvent* FMacEvent::GetEvent() const
 {
 	SCOPED_AUTORELEASE_POOL;
 	if(EventData && [EventData isKindOfClass:[NSEvent class]])
@@ -87,7 +72,7 @@ NSEvent* FMacEvent::GetEvent(void) const
 	}
 }
 
-NSNotification* FMacEvent::GetNotification(void) const
+NSNotification* FMacEvent::GetNotification() const
 {
 	SCOPED_AUTORELEASE_POOL;
 	if(EventData && [EventData isKindOfClass:[NSNotification class]])
@@ -99,13 +84,3 @@ NSNotification* FMacEvent::GetNotification(void) const
 		return nil;
 	}
 }
-
-FCocoaWindow* FMacEvent::GetWindow(void) const
-{
-	return CocoaWindow;
-}
-
-FVector2D FMacEvent::GetMousePosition(void) const
-{
-	return MousePosition;
-}
diff --git a/Engine/Source/Runtime/Core/Private/Mac/MacWindow.cpp b/Engine/Source/Runtime/Core/Private/Mac/MacWindow.cpp
--- a/Engine/Source/Runtime/Core/Private/Mac/MacWindow.cpp
+++ b/Engine/Source/Runtime/Core/Private/Mac/MacWindow.cpp
@@ -107,24 +107,29 @@ void FMacWindow::Initialize( FMacApplication* const Application, const TSharedRe
 			{
 				NSInteger Level = [NSApp isActive] ? NSModalPanelWindowLevel : NSNormalWindowLevel;
 				[WindowHandle setLevel: Level];
+				printf("modal - level %d\n", (int32)Level);
 			}
 			else if (Definition->IsRegularWindow)
 			{
 				if (InParent.IsValid())
 				{
+					printf("regular with parent - NSFloatingWindowLevel\n");
 					[WindowHandle setLevel: NSFloatingWindowLevel];
 				}
 				else
 				{
+					printf("regular without parent - NSNormalWindowLevel\n");
 					[WindowHandle setLevel: NSNormalWindowLevel];
 				}
 			}
 			else if (!Definition->SupportsMaximize && !Definition->SupportsMinimize)
 			{
+				printf("no maximize, no minimize - NSFloatingWindowLevel\n");
 				[WindowHandle setLevel: NSFloatingWindowLevel];
 			}
 			else
 			{
+				printf("other - NSModalPanelWindowLevel, setHidesOnDeactivate\n");
 				[WindowHandle setLevel: NSModalPanelWindowLevel];
 				[WindowHandle setHidesOnDeactivate: YES];
 			}
@@ -261,7 +266,7 @@ void FMacWindow::ReshapeWindow( int32 X, int32 Y, int32 Width, int32 Height )
 			if ( !NSEqualRects(WindowHandle.PreFullScreenRect, NewRect) )
 			{
 				WindowHandle.PreFullScreenRect = NewRect;
-				FMacEvent::SendToGameRunLoop([NSNotification notificationWithName:NSWindowDidResizeNotification object:WindowHandle], WindowHandle, EMacEventSendMethod::Sync, @[ NSDefaultRunLoopMode, UE4ResizeEventMode, UE4ShowEventMode, UE4FullscreenEventMode ]);
+				FMacEvent::SendToGameRunLoop([NSNotification notificationWithName:NSWindowDidResizeNotification object:WindowHandle], EMacEventSendMethod::Sync, @[ NSDefaultRunLoopMode, UE4ResizeEventMode, UE4ShowEventMode, UE4FullscreenEventMode ]);
 			}
 		}
 		
diff --git a/Engine/Source/Runtime/Core/Public/Mac/MacApplication.h b/Engine/Source/Runtime/Core/Public/Mac/MacApplication.h
--- a/Engine/Source/Runtime/Core/Public/Mac/MacApplication.h
+++ b/Engine/Source/Runtime/Core/Public/Mac/MacApplication.h
@@ -60,7 +60,7 @@ public:
 		return TextInputMethodSystem.Get();
 	}
 
-	void ProcessNSEvent(NSEvent* const Event, FVector2D const MousePosition);
+	void ProcessNSEvent(NSEvent* const Event);
 
 	void OnWindowDraggingFinished();
 
diff --git a/Engine/Source/Runtime/Core/Public/Mac/MacEvent.h b/Engine/Source/Runtime/Core/Public/Mac/MacEvent.h
--- a/Engine/Source/Runtime/Core/Public/Mac/MacEvent.h
+++ b/Engine/Source/Runtime/Core/Public/Mac/MacEvent.h
@@ -4,14 +4,6 @@
 #include "Core.h"
 #include "MacWindow.h"
 
-#ifndef __OBJC__
-
-class NSEvent;
-class NSNotification;
-class NSObject;
-
-#endif
-
 enum EMacEventSendMethod
 {
 	Async,
@@ -22,39 +14,29 @@ class FMacEvent
 {
 	// Constructor for an NSEvent based FMacEvent
 	FMacEvent(NSEvent* const Event);
-	
+
 	// Constructor for an NSNotification based FMacEvent
-	FMacEvent(NSNotification* const Notification, NSWindow* const Window);
-	
+	FMacEvent(NSNotification* const Notification);
+
 	// Send the event to the game run loop for processing using the specified SendMethod.
 	static void SendToGameRunLoop(FMacEvent const* const Event, EMacEventSendMethod SendMethod, NSArray* SendModes = @[ NSDefaultRunLoopMode ]);
-	
+
 public:
 	// Send an NSEvent to the Game run loop as an FMacEvent using the specified SendMethod
 	static void SendToGameRunLoop(NSEvent* const Event, EMacEventSendMethod SendMethod, NSArray* SendModes = @[ NSDefaultRunLoopMode ]);
-	
+
 	// Send an NSNotification to the Game run loop as an FMacEvent using the specified SendMethod
-	static void SendToGameRunLoop(NSNotification* const Notification, NSWindow* const Window, EMacEventSendMethod SendMethod, NSArray* SendModes = @[ NSDefaultRunLoopMode ]);
-	
+	static void SendToGameRunLoop(NSNotification* const Notification, EMacEventSendMethod SendMethod, NSArray* SendModes = @[ NSDefaultRunLoopMode ]);
+
 	// Destructor
-	~FMacEvent(void);
-	
+	~FMacEvent();
+
 	// Get the NSEvent for this FMacEvent, nil if not an NSEvent.
-	NSEvent* GetEvent(void) const;
-	
+	NSEvent* GetEvent() const;
+
 	// Get the NSNotification for this FMacEvent, nil if not an NSNotification.
-	NSNotification* GetNotification(void) const;
-	
-	// Get the FCocoaWindow for this FMacEvent, may be nullptr if not a window event, or for a non-FCocoaWindow.
-	FCocoaWindow* GetWindow(void) const;
-	
-	// Get the mouse position for this FMacEvent at the time of posting on the AppKit Main Thread.
-	FVector2D GetMousePosition(void) const;
-	
+	NSNotification* GetNotification() const;
+
 private:
-	FCocoaWindow* CocoaWindow; // The retained window for this event, or nullptr
 	NSObject* EventData; // The retained NSEvent or NSNotification for this event or nil.
-	FVector2D MousePosition; // The mouse position for the event in UE4 coordinates.
 };
-
-
