fixed comment, refactor

diff --git a/Engine/Source/Runtime/Core/Private/HAL/ConsoleManager.cpp b/Engine/Source/Runtime/Core/Private/HAL/ConsoleManager.cpp
--- a/Engine/Source/Runtime/Core/Private/HAL/ConsoleManager.cpp
+++ b/Engine/Source/Runtime/Core/Private/HAL/ConsoleManager.cpp
@@ -22,19 +22,19 @@ static const TCHAR* GetSetByTCHAR(EConsoleVariableFlags InSetBy)
 
 	switch(SetBy)
 	{
-		// CASE() macro was disliked by Mac compiler (Clang)
-		//   error: ISO C++11 does not allow conversion from string literal to 'TCHAR *' (aka 'wchar_t *') [-Werror,-Wc++11-compat-deprecated-writable-strings]
-		// so we do it with more code duplication. Even beter would be to use enum reflection instead.
-		case ECVF_SetByConstructor:			return TEXT("Constructor");
-		case ECVF_SetByScalability:			return TEXT("Scalability");
-		case ECVF_SetByGameSetting:			return TEXT("GameSetting");
-		case ECVF_SetByProjectSetting:		return TEXT("EditorSetting");
-		case ECVF_SetByDeviceProfile:		return TEXT("DeviceProfile");
-		case ECVF_SetBySystemSettingsIni:	return TEXT("SystemSettingsIni");
-		case ECVF_SetByConsoleVariablesIni: return TEXT("ConsoleVariablesIni");
-		case ECVF_SetByCommandline:			return TEXT("Commandline");
-		case ECVF_SetByCode:				return TEXT("Code");
-		case ECVF_SetByConsole:				return TEXT("Console");
+#define CASE(A) case ECVF_SetBy##A: return TEXT(#A);
+		// Could also be done with enum reflection instead
+		CASE(Constructor)
+		CASE(Scalability)
+		CASE(GameSetting)
+		CASE(ProjectSetting)
+		CASE(DeviceProfile)
+		CASE(SystemSettingsIni)
+		CASE(ConsoleVariablesIni)
+		CASE(Commandline)
+		CASE(Code)
+		CASE(Console)
+#undef CASE
 	}
 	return TEXT("<UNKNOWN>");
 }
