Save and restore windows order and levels on Mac when the app resigns and becomes active

diff --git a/Engine/Source/Runtime/Core/Private/Mac/MacApplication.cpp b/Engine/Source/Runtime/Core/Private/Mac/MacApplication.cpp
--- a/Engine/Source/Runtime/Core/Private/Mac/MacApplication.cpp
+++ b/Engine/Source/Runtime/Core/Private/Mac/MacApplication.cpp
@@ -142,37 +142,22 @@ FMacApplication::FMacApplication()
 		TextInputMethodSystem.Reset();
 	}
 
-	AppActivationObserver = [[NSNotificationCenter defaultCenter] addObserverForName:NSApplicationDidBecomeActiveNotification object:[NSApplication sharedApplication] queue:[NSOperationQueue mainQueue] usingBlock:^(NSNotification* Notification){
-								// Change modal window levels
-								NSWindow* BaseWindow = nil;
-								for( auto Window : Windows )
+	AppActivationObserver = [[NSNotificationCenter defaultCenter] addObserverForName:NSApplicationDidBecomeActiveNotification object:[NSApplication sharedApplication] queue:[NSOperationQueue mainQueue] usingBlock:^(NSNotification* Notification)
+							{
+								for (int32 Index = SavedWindowsOrder.Num() - 1; Index > 0; Index--)
 								{
-									NSWindow* CocoaWindow = (NSWindow*)Window->GetOSWindowHandle();
-									if ( CocoaWindow && Window->IsVisible() && [CocoaWindow level] == NSModalPanelWindowLevel )
+									const FSavedWindowOrderInfo& SavedWindowLevel = SavedWindowsOrder[Index];
+									NSWindow* Window = [NSApp windowWithWindowNumber:SavedWindowLevel.WindowNumber];
+									const int32 PreviousWindowNumber = SavedWindowsOrder[Index - 1].WindowNumber;
+									if (Window)
 									{
-										BaseWindow = CocoaWindow;
-										break;
+										[Window orderWindow:NSWindowBelow relativeTo:PreviousWindowNumber];
+										[Window setLevel:SavedWindowLevel.Level];
 									}
 								}
-							 
-								if (BaseWindow)
-								{
-									for( auto Window : Windows )
-									{
-										NSWindow* CocoaWindow = (NSWindow*)Window->GetOSWindowHandle();
-										if ( CocoaWindow && Window->GetDefinition().IsModalWindow )
-										{
-											[CocoaWindow setLevel:NSModalPanelWindowLevel];
-											if ( BaseWindow )
-											{
-												[CocoaWindow orderWindow:NSWindowBelow relativeTo:[BaseWindow windowNumber]];
-											}
-										}
-									}
-								}
-							 
+
 								// If editor thread doesn't have the focus, don't suck up too much CPU time.
-								if( GIsEditor )
+								if (GIsEditor)
 								{
 									// Boost our priority back to normal.
 									struct sched_param Sched;
@@ -185,61 +170,44 @@ FMacApplication::FMacApplication()
 								FApp::SetVolumeMultiplier( 1.0f );
 							}];
 
-	AppDeactivationObserver = [[NSNotificationCenter defaultCenter] addObserverForName:NSApplicationDidResignActiveNotification object:[NSApplication sharedApplication] queue:[NSOperationQueue mainQueue] usingBlock:^(NSNotification* Notification){
-									// Change modal window levels
-									if (Windows.Num() > 0)
-									{
-										NSWindow* BaseWindow = nil;
-										for( auto Window : Windows )
-										{
-											NSWindow* CocoaWindow = (NSWindow*)Window->GetOSWindowHandle();
-											if ( CocoaWindow && Window->IsVisible() && [CocoaWindow level] == NSNormalWindowLevel )
-											{
-												BaseWindow = CocoaWindow;
-											}
-										}
-										
-										if (BaseWindow)
-										{
-											for( auto Window : Windows )
-											{
-												NSWindow* CocoaWindow = (NSWindow*)Window->GetOSWindowHandle();
-												if ( CocoaWindow && Window->GetDefinition().IsModalWindow )
-												{
-													[CocoaWindow setLevel:NSNormalWindowLevel];
-													[CocoaWindow orderWindow:NSWindowAbove relativeTo:[BaseWindow windowNumber]];
-													BaseWindow = CocoaWindow;
-												}
-											}
-										}
-								   }
-
-									CFRunLoopRunInMode(kCFRunLoopDefaultMode, 0, true);
-							   
-									// If editor thread doesn't have the focus, don't suck up too much CPU time.
-									if( GIsEditor )
+	AppDeactivationObserver = [[NSNotificationCenter defaultCenter] addObserverForName:NSApplicationWillResignActiveNotification object:[NSApplication sharedApplication] queue:[NSOperationQueue mainQueue] usingBlock:^(NSNotification* Notification)
+							{
+								SavedWindowsOrder.Empty();
+
+								NSArray* OrderedWindows = [NSApp orderedWindows];
+								for (NSWindow* Window in OrderedWindows)
+								{
+									if ([Window isKindOfClass:[FCocoaWindow class]] && [Window isVisible] && ![Window hidesOnDeactivate])
 									{
-										// Drop our priority to speed up whatever is in the foreground.
-										struct sched_param Sched;
-										FMemory::Memzero(&Sched, sizeof(struct sched_param));
-										Sched.sched_priority = 5;
-										pthread_setschedparam(pthread_self(), SCHED_RR, &Sched);
-
-										// Sleep for a bit to not eat up all CPU time.
-										FPlatformProcess::Sleep(0.005f);
+										SavedWindowsOrder.Add(FSavedWindowOrderInfo([Window windowNumber], [Window level]));
+										[Window setLevel:NSNormalWindowLevel];
 									}
+								}
+
+								CFRunLoopRunInMode(kCFRunLoopDefaultMode, 0, true);
+
+								// If editor thread doesn't have the focus, don't suck up too much CPU time.
+								if (GIsEditor)
+								{
+									// Drop our priority to speed up whatever is in the foreground.
+									struct sched_param Sched;
+									FMemory::Memzero(&Sched, sizeof(struct sched_param));
+									Sched.sched_priority = 5;
+									pthread_setschedparam(pthread_self(), SCHED_RR, &Sched);
+
+									// Sleep for a bit to not eat up all CPU time.
+									FPlatformProcess::Sleep(0.005f);
+								}
+
+								// app is inactive, apply multiplier
+								FApp::SetVolumeMultiplier(FApp::GetUnfocusedVolumeMultiplier());
+							}];
 
-									// app is inactive, apply multiplier
-									FApp::SetVolumeMultiplier(FApp::GetUnfocusedVolumeMultiplier());
-								}];
-	
 	WorkspaceActivationObserver = [[[NSWorkspace sharedWorkspace] notificationCenter] addObserverForName:NSWorkspaceSessionDidBecomeActiveNotification object:[NSWorkspace sharedWorkspace] queue:[NSOperationQueue mainQueue] usingBlock:^(NSNotification* Notification){
-								   printf("workspace activated\n");
 									   bIsWorkspaceSessionActive = true;
 								   }];
 
 	WorkspaceDeactivationObserver = [[[NSWorkspace sharedWorkspace] notificationCenter] addObserverForName:NSWorkspaceSessionDidResignActiveNotification object:[NSWorkspace sharedWorkspace] queue:[NSOperationQueue mainQueue] usingBlock:^(NSNotification* Notification){
-									 printf("workspace deactivated\n");
 									   bIsWorkspaceSessionActive = false;
 								   }];
 
diff --git a/Engine/Source/Runtime/Core/Public/Mac/MacApplication.h b/Engine/Source/Runtime/Core/Public/Mac/MacApplication.h
--- a/Engine/Source/Runtime/Core/Public/Mac/MacApplication.h
+++ b/Engine/Source/Runtime/Core/Public/Mac/MacApplication.h
@@ -154,9 +154,17 @@ private:
 	EMouseButtons::Type LastPressedMouseButton;
 
 	FCriticalSection WindowsMutex;
-	TArray< TSharedRef< FMacWindow > > Windows;
+	TArray<TSharedRef<FMacWindow>> Windows;
 
-	TSharedRef< class HIDInputInterface > HIDInput;
+	struct FSavedWindowOrderInfo
+	{
+		int32 WindowNumber;
+		int32 Level;
+		FSavedWindowOrderInfo(int32 InWindowNumber, int32 InLevel) : WindowNumber(InWindowNumber), Level(InLevel) {}
+	};
+	TArray<FSavedWindowOrderInfo> SavedWindowsOrder;
+
+	TSharedRef<class HIDInputInterface> HIDInput;
 
 	FCocoaWindow* DraggedWindow;
 
