Band-aid fix for versioning of modules on Linux.

- Effectively disables version check (and fixes UT-79).
- Proper fix tracked as UEPLAT-203.

diff --git a/Engine/Source/Runtime/Core/Private/Linux/LinuxPlatformProcess.cpp b/Engine/Source/Runtime/Core/Private/Linux/LinuxPlatformProcess.cpp
--- a/Engine/Source/Runtime/Core/Private/Linux/LinuxPlatformProcess.cpp
+++ b/Engine/Source/Runtime/Core/Private/Linux/LinuxPlatformProcess.cpp
@@ -37,7 +37,7 @@ void* FLinuxPlatformProcess::GetDllExport( void* DllHandle, const TCHAR* ProcNam
 int32 FLinuxPlatformProcess::GetDllApiVersion( const TCHAR* Filename )
 {
 	check(Filename);
-	return ENGINE_VERSION;
+	return MODULE_API_VERSION;
 }
 
 const TCHAR* FLinuxPlatformProcess::GetModulePrefix()
