Removed few debug printfs

diff --git a/Engine/Source/Runtime/Core/Private/Mac/MacWindow.cpp b/Engine/Source/Runtime/Core/Private/Mac/MacWindow.cpp
--- a/Engine/Source/Runtime/Core/Private/Mac/MacWindow.cpp
+++ b/Engine/Source/Runtime/Core/Private/Mac/MacWindow.cpp
@@ -107,29 +107,24 @@ void FMacWindow::Initialize( FMacApplication* const Application, const TSharedRe
 			{
 				NSInteger Level = [NSApp isActive] ? NSModalPanelWindowLevel : NSNormalWindowLevel;
 				[WindowHandle setLevel: Level];
-				printf("modal - level %d\n", (int32)Level);
 			}
 			else if (Definition->IsRegularWindow)
 			{
 				if (InParent.IsValid())
 				{
-					printf("regular with parent - NSFloatingWindowLevel\n");
 					[WindowHandle setLevel: NSFloatingWindowLevel];
 				}
 				else
 				{
-					printf("regular without parent - NSNormalWindowLevel\n");
 					[WindowHandle setLevel: NSNormalWindowLevel];
 				}
 			}
 			else if (!Definition->SupportsMaximize && !Definition->SupportsMinimize)
 			{
-				printf("no maximize, no minimize - NSFloatingWindowLevel\n");
 				[WindowHandle setLevel: NSFloatingWindowLevel];
 			}
 			else
 			{
-				printf("other - NSModalPanelWindowLevel, setHidesOnDeactivate\n");
 				[WindowHandle setLevel: NSModalPanelWindowLevel];
 				[WindowHandle setHidesOnDeactivate: YES];
 			}
