Improvements to Linux splash (PR #555).

- Sets the context before issuing any commands and avoids unnecessary resetting in the loop.
- Fixes reported crash on Mesa 10.1.3 (shipped with Ubuntu 14.04 LTS)
- Contributed by [REDACTED_PERSON].

diff --git a/Engine/Source/Runtime/Core/Private/Linux/LinuxPlatformSplash.cpp b/Engine/Source/Runtime/Core/Private/Linux/LinuxPlatformSplash.cpp
--- a/Engine/Source/Runtime/Core/Private/Linux/LinuxPlatformSplash.cpp
+++ b/Engine/Source/Runtime/Core/Private/Linux/LinuxPlatformSplash.cpp
@@ -547,6 +547,12 @@ static int StartSplashScreenThread(void *ptr)
 	#define GET_GL_ENTRYPOINTS(Type,Func) Func = reinterpret_cast<Type>(SDL_GL_GetProcAddress(#Func));
 	ENUM_GL_ENTRYPOINTS(GET_GL_ENTRYPOINTS);
 
+	if (SDL_GL_MakeCurrent( GSplashWindow, Context ) != 0) 
+	{
+		UE_LOG(LogHAL, Error, TEXT("Splash screen SDL_GL_MakeCurrent failed: %s"), UTF8_TO_TCHAR(SDL_GetError()));
+		return -1;
+	}
+
 	// Initialize Shaders, Programs, etc.
 	GLuint VertexShader 	= 0;
 	GLuint FragmentShader 	= 0;
@@ -582,18 +588,19 @@ static int StartSplashScreenThread(void *ptr)
 	glBindVertexArray(VertexArrayObject);
 	
 	// Create Vertex Buffer and upload data.
-	GLuint VertexBuffer[2];
-	glGenBuffers(2, &VertexBuffer[0]);
-	glBindBuffer(GL_ARRAY_BUFFER, VertexBuffer[0]);
-	float screen_vertex [] = {
-			-1.0, -1.0f,
-			-1.0f, 1.0f,
-			1.0f, 1.0f,
-
-			1.0f, 1.0f,
-			1.0f, -1.0f,
-			-1.0f, -1.0f
-		};
+	GLuint VertexBuffer;
+	glGenBuffers(1, &VertexBuffer);
+	glBindBuffer(GL_ARRAY_BUFFER, VertexBuffer);
+	float screen_vertex [] = 
+	{
+		-1.0, -1.0f,
+		-1.0f, 1.0f,
+		1.0f, 1.0f,
+
+		1.0f, 1.0f,
+		1.0f, -1.0f,
+		-1.0f, -1.0f
+	};
 	glBufferData(GL_ARRAY_BUFFER, sizeof(float)*12, (GLvoid*)screen_vertex, GL_DYNAMIC_DRAW );
 
 	// create texture slot
@@ -674,8 +681,6 @@ static int StartSplashScreenThread(void *ptr)
 
 	OpenFonts();
 
-	glClearColor(0.2f, 0.3f, 0.8f, 1.0f);
-	glClear(GL_COLOR_BUFFER_BIT | GL_DEPTH_BUFFER_BIT);
 	glDisable(GL_DEPTH_TEST);
 
 	// drawing loop
@@ -685,40 +690,42 @@ static int StartSplashScreenThread(void *ptr)
 		// Poll events otherwise the window will get dark or
 		// on some Desktops it will complain that the thread
 		// does not work anymore.
-		while (SDL_PollEvent(&event)) {}
+		while (SDL_PollEvent(&event))
+		{
+			// intentionally empty
+		}
 
 		SDL_Delay(300);
 		
+		// Activate Shader Program, Texture etc.
+		glBindVertexArray(VertexArrayObject);
+		glBindBuffer(GL_ARRAY_BUFFER, VertexBuffer);
+		glUseProgram(ShaderProgram);
+		glActiveTexture(GL_TEXTURE0);
+		glUniform1i(SplashTextureID, 0);
+
 		if (ThreadState > 0)
 		{
 			RenderString(texture);
 			ThreadState--;
 		}
-		
-		if (SDL_GL_MakeCurrent( GSplashWindow, Context ) == 0)
-		{
-			// Activate Shader Program, Texture etc.
-			glUseProgram(ShaderProgram);
-			glActiveTexture(GL_TEXTURE0);
-			glUniform1i(SplashTextureID, 0);
-
-			// Set stream positions and draw.
-			glEnableVertexAttribArray(0); 
-			glVertexAttribPointer(0, 2, GL_FLOAT, GL_FALSE, 0, 0);
-			glDrawArrays(GL_TRIANGLES, 0, 6);
-
-			// Disable Shader Program, VAO, VBO etc.
-			glDisableVertexAttribArray(0); 
-			glUseProgram(0);
-			glBindBuffer(GL_ARRAY_BUFFER, 0);
-			glBindVertexArray(0);
-
-			SDL_GL_SwapWindow( GSplashWindow );
-		}
+
+		// Set stream positions and draw.
+		glEnableVertexAttribArray(0); 
+		glVertexAttribPointer(0, 2, GL_FLOAT, GL_FALSE, 0, 0);
+		glDrawArrays(GL_TRIANGLES, 0, 6);
+
+		// Disable Shader Program, VAO, VBO etc.
+		glDisableVertexAttribArray(0); 
+		glUseProgram(0);
+		glBindBuffer(GL_ARRAY_BUFFER, 0);
+		glBindVertexArray(0);
+
+		SDL_GL_SwapWindow( GSplashWindow );
 	}
 
 	// clean up
-	glDeleteBuffers(1, &VertexBuffer[0]);
+	glDeleteBuffers(1, &VertexBuffer);
 	glDeleteVertexArrays(1, &VertexArrayObject);
 	glDeleteShader(VertexShader);
 	glDeleteShader(FragmentShader);
