Core: Fixed FDateTime asserting when parsing invalid input

diff --git a/Engine/Source/Runtime/Core/Private/Misc/DateTime.cpp b/Engine/Source/Runtime/Core/Private/Misc/DateTime.cpp
--- a/Engine/Source/Runtime/Core/Private/Misc/DateTime.cpp
+++ b/Engine/Source/Runtime/Core/Private/Misc/DateTime.cpp
@@ -15,13 +15,7 @@ const int32 FDateTime::DaysToMonth[] = { 0, 31, 59, 90, 120, 151, 181, 212, 243,
 
 FDateTime::FDateTime( int32 Year, int32 Month, int32 Day, int32 Hour, int32 Minute, int32 Second, int32 Millisecond )
 {
-	check((Year >= 1) && (Year <= 9999));
-	check((Month >= 1) && (Month <= 12));
-	check((Day >= 1) && (Day <= DaysInMonth(Year, Month)));
-	check((Hour >= 0) && (Hour <= 23));
-	check((Minute >= 0) && (Minute <= 59));
-	check((Second >= 0) && (Second <= 59));
-	check((Millisecond >= 0) && (Millisecond <= 999));
+	check(Validate(Year, Month, Day, Hour, Minute, Second, Millisecond));
 
 	int32 TotalDays = 0;
 
@@ -297,18 +291,22 @@ bool FDateTime::Parse( const FString& DateTimeString, FDateTime& OutDateTime )
 		return false;
 	}
 
+	const int32 Year = FCString::Atoi(*Tokens[0]);
+	const int32 Month = FCString::Atoi(*Tokens[1]);
+	const int32 Day = FCString::Atoi(*Tokens[2]);
+	const int32 Hour = FCString::Atoi(*Tokens[3]);
+	const int32 Minute = FCString::Atoi(*Tokens[4]);
+	const int32 Second = FCString::Atoi(*Tokens[5]);
+	const int32 Millisecond = Tokens.Num() > 6 ? FCString::Atoi(*Tokens[6]) : 0;
+
+	if (!Validate(Year, Month, Day, Hour, Minute, Second, Millisecond))
+	{
+		return false;
+	}
+
 	// convert the tokens to numbers
-	OutDateTime.Ticks = FDateTime(
-		FCString::Atoi(*Tokens[0]), // year
-		FCString::Atoi(*Tokens[1]), // month
-		FCString::Atoi(*Tokens[2]), // day
-		FCString::Atoi(*Tokens[3]), // hour
-		FCString::Atoi(*Tokens[4]), // minute
-		FCString::Atoi(*Tokens[5]), // second
-		Tokens.Num() > 6 ? FCString::Atoi(*Tokens[6]) : 0 // millisecond
-	).GetTicks();
-
-	// @todo gmp: need some better validation here
+	OutDateTime.Ticks = FDateTime(Year, Month, Day, Hour, Minute, Second, Millisecond).Ticks;
+
 	return true;
 }
 
@@ -349,7 +347,7 @@ bool FDateTime::ParseIso8601( const TCHAR* DateTimeString, FDateTime& OutDateTim
 		return false;
 	}
 
-	// see if this is date+time
+	// check whether this is date and time
 	if (*Next == TCHAR('T'))
 	{
 		Ptr = Next + 1;
@@ -428,6 +426,11 @@ bool FDateTime::ParseIso8601( const TCHAR* DateTimeString, FDateTime& OutDateTim
 		return false;
 	}
 
+	if (!Validate(Year, Month, Day, Hour, Minute, Second, Millisecond))
+	{
+		return false;
+	}
+
 	FDateTime Final(Year, Month, Day, Hour, Minute, Second, Millisecond);
 
 	// adjust for the timezone (bringing the DateTime into UTC)
@@ -450,6 +453,18 @@ FDateTime FDateTime::UtcNow()
 }
 
 
+bool FDateTime::Validate( int32 Year, int32 Month, int32 Day, int32 Hour, int32 Minute, int32 Second, int32 Millisecond )
+{
+	return (Year >= 1) && (Year <= 9999) &&
+		(Month >= 1) && (Month <= 12) &&
+		(Day >= 1) && (Day <= DaysInMonth(Year, Month)) &&
+		(Hour >= 0) && (Hour <= 23) &&
+		(Minute >= 0) && (Minute <= 59) &&
+		(Second >= 0) && (Second <= 59) &&
+		(Millisecond >= 0) && (Millisecond <= 999);
+}
+
+
 /* FDateTime friend functions
  *****************************************************************************/
 
diff --git a/Engine/Source/Runtime/Core/Public/Misc/DateTime.h b/Engine/Source/Runtime/Core/Public/Misc/DateTime.h
--- a/Engine/Source/Runtime/Core/Public/Misc/DateTime.h
+++ b/Engine/Source/Runtime/Core/Public/Misc/DateTime.h
@@ -554,13 +554,6 @@ public:
 	 */
 	static CORE_API bool IsLeapYear( int32 Year );
 
-	/**
-	 * Checks whether the date and time components are valid.
-	 *
-	 * @return true if the components are valid, false otherwise.
-	 */
-	static CORE_API bool IsValid( int32 Year, int32 Month, int32 Day, int32 Hour, int32 Minute, int32 Second, int32 Millisecond );
-
 	/**
 	 * Returns the maximum date value.
 	 *
@@ -645,6 +638,22 @@ public:
 	 */
 	static CORE_API FDateTime UtcNow();
 
+	/**
+	* Validates the given components of a date and time value.
+	*
+	* The allow ranges for the components are:
+	*		Year: 1 - 9999
+	*		Month: 1 - 12
+	*		Day: 1 - DaysInMonth(Month)
+	*		Hour: 0 - 23
+	*		Minute: 0 - 59
+	*		Second: 0 - 59
+	*		Millisecond: 0 - 999
+	*
+	* @return true if the components are valid, false otherwise.
+	*/
+	static CORE_API bool Validate( int32 Year, int32 Month, int32 Day, int32 Hour, int32 Minute, int32 Second, int32 Millisecond );
+
 public:
 
 	/**
