Remove the need for skeletal mesh component to serialize its BodySetup

diff --git a/Engine/Source/Runtime/Core/Public/UObject/ObjectVersion.h b/Engine/Source/Runtime/Core/Public/UObject/ObjectVersion.h
--- a/Engine/Source/Runtime/Core/Public/UObject/ObjectVersion.h
+++ b/Engine/Source/Runtime/Core/Public/UObject/ObjectVersion.h
@@ -632,6 +632,8 @@ enum EUnrealEngineObjectUE4Version
 	VER_UE4_SLATE_COMPOSITE_FONTS,
 	// Remove UDEPRECATED_SaveGameSummary, required for UWorld::Serialize
 	VER_UE4_REMOVE_SAVEGAMESUMMARY,
+	//Remove bodyseutp serialization from skeletal mesh component
+	VER_UE4_REMOVE_SKELETALMESH_COMPONENT_BODYSETUP_SERIALIZATION,
 
 	// -----<new versions can be added before this line>-------------------------------------------------
 	// - this needs to be the last line (see note below)
diff --git a/Engine/Source/Runtime/Engine/Private/SkeletalMesh.cpp b/Engine/Source/Runtime/Engine/Private/SkeletalMesh.cpp
--- a/Engine/Source/Runtime/Engine/Private/SkeletalMesh.cpp
+++ b/Engine/Source/Runtime/Engine/Private/SkeletalMesh.cpp
@@ -4881,9 +4881,13 @@ void USkeletalMeshComponent::Serialize(FArchive& Ar)
 		}
 	}
 
-	if (bEnablePerPolyCollision)
+	if (Ar.UE4Ver() < VER_UE4_REMOVE_SKELETALMESH_COMPONENT_BODYSETUP_SERIALIZATION)
 	{
-		Ar << BodySetup;
+		//we used to serialize bodysetup of skeletal mesh component. We no longer do this, but need to not break existing content
+		if (bEnablePerPolyCollision)
+		{
+			Ar << BodySetup;
+		}
 	}
 
 	// Since we separated simulation vs blending
