Use the FMallocAnsi allocator in 'programs' such as the CrashReportClient to avoid a mismatched strdup/delete pair in CoreSymbolication when resymbolising crash logs on Yosemite. The offending call in FApplePlatformSymbolication::SymbolInfoForStrippedSymbol is only made within the CrashReportClient so the Editor & Games are unaffected.

diff --git a/Engine/Source/Runtime/Core/Private/Mac/MacPlatformMemory.cpp b/Engine/Source/Runtime/Core/Private/Mac/MacPlatformMemory.cpp
--- a/Engine/Source/Runtime/Core/Private/Mac/MacPlatformMemory.cpp
+++ b/Engine/Source/Runtime/Core/Private/Mac/MacPlatformMemory.cpp
@@ -21,7 +21,7 @@ void FMacPlatformMemory::Init()
 
 FMalloc* FMacPlatformMemory::BaseAllocator()
 {
-#if FORCE_ANSI_ALLOCATOR
+#if FORCE_ANSI_ALLOCATOR || IS_PROGRAM
 	return new FMallocAnsi();
 #elif (WITH_EDITORONLY_DATA || IS_PROGRAM) && TBB_ALLOCATOR_ALLOWED
 	return new FMallocTBB();
