Don't truncate files opened for write on Unix platforms, as we don't do this on Windows

diff --git a/Engine/Source/Runtime/Core/Private/Android/AndroidFile.cpp b/Engine/Source/Runtime/Core/Private/Android/AndroidFile.cpp
--- a/Engine/Source/Runtime/Core/Private/Android/AndroidFile.cpp
+++ b/Engine/Source/Runtime/Core/Private/Android/AndroidFile.cpp
@@ -501,10 +501,6 @@ public:
 		{
 			Flags |= O_APPEND;
 		}
-		else
-		{
-			Flags |= O_TRUNC;
-		}
 		if (bAllowRead)
 		{
 			Flags |= O_RDWR;
diff --git a/Engine/Source/Runtime/Core/Private/Apple/ApplePlatformFile.cpp b/Engine/Source/Runtime/Core/Private/Apple/ApplePlatformFile.cpp
--- a/Engine/Source/Runtime/Core/Private/Apple/ApplePlatformFile.cpp
+++ b/Engine/Source/Runtime/Core/Private/Apple/ApplePlatformFile.cpp
@@ -442,10 +442,6 @@ IFileHandle* FApplePlatformFile::OpenWrite(const TCHAR* Filename, bool bAppend,
 	{
 		Flags |= O_APPEND;
 	}
-	else
-	{
-		Flags |= O_TRUNC;
-	}
 	if (bAllowRead)
 	{
 		Flags |= O_RDWR;
diff --git a/Engine/Source/Runtime/Core/Private/HTML5/HTML5PlatformFile.cpp b/Engine/Source/Runtime/Core/Private/HTML5/HTML5PlatformFile.cpp
--- a/Engine/Source/Runtime/Core/Private/HTML5/HTML5PlatformFile.cpp
+++ b/Engine/Source/Runtime/Core/Private/HTML5/HTML5PlatformFile.cpp
@@ -312,10 +312,6 @@ public:
 		{
 			Flags |= O_APPEND;
 		}
-		else
-		{
-			Flags |= O_TRUNC;
-		}
 		if (bAllowRead)
 		{
 			Flags |= O_RDWR;
diff --git a/Engine/Source/Runtime/Core/Private/IOS/IOSPlatformFile.cpp b/Engine/Source/Runtime/Core/Private/IOS/IOSPlatformFile.cpp
--- a/Engine/Source/Runtime/Core/Private/IOS/IOSPlatformFile.cpp
+++ b/Engine/Source/Runtime/Core/Private/IOS/IOSPlatformFile.cpp
@@ -546,10 +546,6 @@ IFileHandle* FIOSPlatformFile::OpenWrite(const TCHAR* Filename, bool bAppend, bo
 	{
 		Flags |= O_APPEND;
 	}
-	else
-	{
-		Flags |= O_TRUNC;
-	}
 	if (bAllowRead)
 	{
 		Flags |= O_RDWR;
diff --git a/Engine/Source/Runtime/Core/Private/Linux/LinuxPlatformFile.cpp b/Engine/Source/Runtime/Core/Private/Linux/LinuxPlatformFile.cpp
--- a/Engine/Source/Runtime/Core/Private/Linux/LinuxPlatformFile.cpp
+++ b/Engine/Source/Runtime/Core/Private/Linux/LinuxPlatformFile.cpp
@@ -266,10 +266,9 @@ IFileHandle* FLinuxPlatformFile::OpenWrite(const TCHAR* Filename, bool bAppend,
 	// create directories if needed.
 	if (!CreateDirectoriesFromPath(Filename))
 	{
-		return NULL;
+		return nullptr;
 	}
 
-	// Caveat: cannot specify O_TRUNC in flags, as this will corrupt the file which may be "locked" by other process. We will ftruncate() it once we "lock" it
 	int32 Handle = open(TCHAR_TO_UTF8(*NormalizeFilename(Filename)), Flags, S_IRUSR | S_IWUSR);
 	if (Handle != -1)
 	{
@@ -286,19 +285,6 @@ IFileHandle* FLinuxPlatformFile::OpenWrite(const TCHAR* Filename, bool bAppend,
 			// all the other locking errors are ignored.
 		}
 
-		// truncate the file now that we locked it
-		if (!bAppend)
-		{
-			if (ftruncate(Handle, 0) != 0)
-			{
-				int ErrNo = errno;
-				UE_LOG(LogLinuxPlatformFile, Warning, TEXT( "ftruncate() failed for '%s': errno=%d (%s)" ),
-															Filename, ErrNo, ANSI_TO_TCHAR(strerror(ErrNo)));
-				close(Handle);
-				return nullptr;
-			}
-		}
-
 		FFileHandleLinux* FileHandleLinux = new FFileHandleLinux(Handle);
 		if (bAppend)
 		{
