Fixed bug in FPaths::GetBaseFilename when called with bRemovePath=false. Will not consider a '.' character in the path to be an extension separator.

diff --git a/Engine/Source/Runtime/Core/Private/Misc/Paths.cpp b/Engine/Source/Runtime/Core/Private/Misc/Paths.cpp
--- a/Engine/Source/Runtime/Core/Private/Misc/Paths.cpp
+++ b/Engine/Source/Runtime/Core/Private/Misc/Paths.cpp
@@ -414,7 +414,17 @@ FString FPaths::GetBaseFilename( const FString& InPath, bool bRemovePath )
 
 	// remove the extension
 	int32 Pos = Wk.Find(TEXT("."), ESearchCase::CaseSensitive, ESearchDir::FromEnd);
-	if ( Pos != INDEX_NONE )
+	
+	// determine the position of the path/leaf separator
+	int32 LeafPos = INDEX_NONE;
+	if (!bRemovePath)
+	{
+		LeafPos = Wk.Find(TEXT("/"), ESearchCase::CaseSensitive, ESearchDir::FromEnd);
+		// in case we are using backslashes on a platform that doesn't use backslashes
+		LeafPos = FMath::Max(LeafPos, Wk.Find(TEXT("\\"), ESearchCase::CaseSensitive, ESearchDir::FromEnd));
+	}
+
+	if (Pos != INDEX_NONE && (LeafPos == INDEX_NONE || Pos > LeafPos))
 	{
 		return Wk.Left(Pos);
 	}
