Fixing the duplicated comment for ATF_RequiredUser to be "// Test requires a user instigated session"

diff --git a/Engine/Source/Runtime/Core/Public/Misc/AutomationTest.h b/Engine/Source/Runtime/Core/Public/Misc/AutomationTest.h
--- a/Engine/Source/Runtime/Core/Public/Misc/AutomationTest.h
+++ b/Engine/Source/Runtime/Core/Public/Misc/AutomationTest.h
@@ -18,7 +18,7 @@ namespace EAutomationTestFlags
 
 		//features required for the test - not specifying means it is valid for any feature combination
 		ATF_NonNullRHI					= 0x00000100,	// Test requires a non-null RHI to run correctly
-		ATF_RequiresUser				= 0x00000200,	// Test requires a non-null RHI to run correctly
+		ATF_RequiresUser				= 0x00000200,	// Test requires a user instigated session
 		ATF_FeatureMask					= ATF_NonNullRHI | ATF_RequiresUser,
 
 		//action options - each autotest can have only one of this and this has to match
