Improve Android LocalPrint to deal with buffer limit
#android
#ue4

diff --git a/Engine/Source/Runtime/Core/Private/Android/AndroidMisc.cpp b/Engine/Source/Runtime/Core/Private/Android/AndroidMisc.cpp
--- a/Engine/Source/Runtime/Core/Private/Android/AndroidMisc.cpp
+++ b/Engine/Source/Runtime/Core/Private/Android/AndroidMisc.cpp
@@ -49,7 +49,72 @@ void FAndroidMisc::LocalPrint(const TCHAR *Message)
 	// Builds for distribution should not have logging in them:
 	// http://developer.android.com/tools/publishing/preparing.html#publishing-configure
 #if !UE_BUILD_SHIPPING
-	__android_log_print(ANDROID_LOG_DEBUG, "UE4", "%s", TCHAR_TO_ANSI(Message));
+
+	/*
+	const int MAX_LOG_LENGTH = 4096;
+	// not static since may be called by different threads
+	TCHAR MessageBuffer[MAX_LOG_LENGTH];
+
+	const TCHAR* SourcePtr = Message;
+	while (*SourcePtr)
+	{
+		TCHAR* WritePtr = MessageBuffer;
+		int32 RemainingSpace = MAX_LOG_LENGTH;
+		while (*SourcePtr && --RemainingSpace > 0)
+		{
+			if (*SourcePtr == TEXT('\r'))
+			{
+				// If next character is newline, skip it
+				if (*(++SourcePtr) == TEXT('\n'))
+					++SourcePtr;
+				break;
+			}
+			else if (*SourcePtr == TEXT('\n'))
+			{
+				++SourcePtr;
+				break;
+			}
+			else {
+				*WritePtr++ = *SourcePtr++;
+			}
+		}
+		*WritePtr = '\0';
+		__android_log_write(ANDROID_LOG_DEBUG, "UE4", TCHAR_TO_ANSI(Message));
+	}
+	*/
+
+	const int MAX_LOG_LENGTH = 4096;
+	// not static since may be called by different threads
+	ANSICHAR MessageBuffer[MAX_LOG_LENGTH];
+
+	const TCHAR* SourcePtr = Message;
+	while (*SourcePtr)
+	{
+		ANSICHAR* WritePtr = MessageBuffer;
+		int32 RemainingSpace = MAX_LOG_LENGTH;
+		while (*SourcePtr && --RemainingSpace > 0)
+		{
+			if (*SourcePtr == TEXT('\r'))
+			{
+				// If next character is newline, skip it
+				if (*(++SourcePtr) == TEXT('\n'))
+					++SourcePtr;
+				break;
+			}
+			else if (*SourcePtr == TEXT('\n'))
+			{
+				++SourcePtr;
+				break;
+			}
+			else {
+				*WritePtr++ = static_cast<ANSICHAR>(*SourcePtr++);
+			}
+		}
+		*WritePtr = '\0';
+		__android_log_write(ANDROID_LOG_DEBUG, "UE4", MessageBuffer);
+	}
+
+	//	__android_log_print(ANDROID_LOG_DEBUG, "UE4", "%s", TCHAR_TO_ANSI(Message));
 #endif
 }
 
