Back out revision 28 from //depot/UE4/Engine/Source/Runtime/Core/Private/Linux/LinuxApplication.cpp
Got caught in a rename, nothing in this file needed to be changed.

diff --git a/Engine/Source/Runtime/Core/Private/Linux/LinuxApplication.cpp b/Engine/Source/Runtime/Core/Private/Linux/LinuxApplication.cpp
--- a/Engine/Source/Runtime/Core/Private/Linux/LinuxApplication.cpp
+++ b/Engine/Source/Runtime/Core/Private/Linux/LinuxApplication.cpp
@@ -168,7 +168,7 @@ void FLinuxApplication::AddPendingEvent( SDL_Event SDLEvent )
 	}
 }
 
-bool FLinuxApplication::GeneratesKeyCharMessage(const SDL_KeyEvent & KeyDownEvent)
+bool FLinuxApplication::GeneratesKeyCharMessage(const SDL_KeyboardEvent & KeyDownEvent)
 {
 	bool bCmdKeyPressed = (KeyDownEvent.keysym.mod & (KMOD_LCTRL | KMOD_RCTRL)) != 0;
 	const SDL_Keycode Sym = KeyDownEvent.keysym.sym;
@@ -237,7 +237,7 @@ void FLinuxApplication::ProcessDeferredMessage( SDL_Event Event )
 	{
 	case SDL_KEYDOWN:
 		{
-			SDL_KeyEvent KeyEvent = Event.key;
+			SDL_KeyboardEvent KeyEvent = Event.key;
 			const SDL_Keycode KeyCode = KeyEvent.keysym.scancode;
 			const bool bIsRepeated = KeyEvent.repeat != 0;
 				
@@ -253,7 +253,7 @@ void FLinuxApplication::ProcessDeferredMessage( SDL_Event Event )
 		break;
 	case SDL_KEYUP:
 		{
-			SDL_KeyEvent keyEvent = Event.key;
+			SDL_KeyboardEvent keyEvent = Event.key;
 			const SDL_Keycode KeyCode = keyEvent.keysym.scancode;
 			const TCHAR Character = ConvertChar( keyEvent.keysym );
 			const bool IsRepeat = keyEvent.repeat != 0;
