fixes for building IOS with arm64
#ios

diff --git a/Engine/Source/Runtime/Core/Private/IOS/IOSPlatformMisc.cpp b/Engine/Source/Runtime/Core/Private/IOS/IOSPlatformMisc.cpp
--- a/Engine/Source/Runtime/Core/Private/IOS/IOSPlatformMisc.cpp
+++ b/Engine/Source/Runtime/Core/Private/IOS/IOSPlatformMisc.cpp
@@ -645,7 +645,7 @@ void GetBytesForFont(const NSString* InFontName, OUT TArray<uint8>& OutBytes)
 		{
 			size_t TableSize = 0;
 			
-			uint32 aTag = (uint32)CFArrayGetValueAtIndex(Tags, TableIndex);
+			uint64 aTag = (uint64)CFArrayGetValueAtIndex(Tags, TableIndex);
 			if (aTag == 'CFF ' && !bContainsCFFTable)
 			{
 				bContainsCFFTable = true;
@@ -700,7 +700,7 @@ void GetBytesForFont(const NSString* InFontName, OUT TArray<uint8>& OutBytes)
 
 		for (int TableIndex = 0; TableIndex < TableCount; ++TableIndex)
 		{
-			uint32 aTag = (uint32)CFArrayGetValueAtIndex(Tags, TableIndex);
+			uint64 aTag = (uint64)CFArrayGetValueAtIndex(Tags, TableIndex);
 			CFDataRef TableDataRef = CGFontCopyTableForTag(cgFont, aTag);
 			uint32 TableSize = CFDataGetLength(TableDataRef);
 
diff --git a/Engine/Source/Runtime/IOS/MetalRHI/Private/MetalCommands.cpp b/Engine/Source/Runtime/IOS/MetalRHI/Private/MetalCommands.cpp
--- a/Engine/Source/Runtime/IOS/MetalRHI/Private/MetalCommands.cpp
+++ b/Engine/Source/Runtime/IOS/MetalRHI/Private/MetalCommands.cpp
@@ -411,8 +411,8 @@ void FMetalDynamicRHI::RHISetRenderTargetsAndClear(const FRHISetRenderTargetsInf
 		const FRHIRenderTargetView& RenderTargetView = RenderTargetsInfo.ColorRenderTarget[0];
 		FMetalSurface& RenderTarget = GetMetalSurfaceFromRHITexture(RenderTargetView.Texture);
 
-		uint32 Width = FMath::Max(RenderTarget.Texture.width >> RenderTargetView.MipIndex, (uint32)1);
-		uint32 Height = FMath::Max(RenderTarget.Texture.height >> RenderTargetView.MipIndex, (uint32)1);
+		uint32 Width = FMath::Max((uint32)(RenderTarget.Texture.width >> RenderTargetView.MipIndex), (uint32)1);
+		uint32 Height = FMath::Max((uint32)(RenderTarget.Texture.height >> RenderTargetView.MipIndex), (uint32)1);
 
 		RHISetViewport(0, 0, 0.0f, Width, Height, 1.0f);
 	}
diff --git a/Engine/Source/Runtime/IOS/MetalRHI/Private/MetalTexture.cpp b/Engine/Source/Runtime/IOS/MetalRHI/Private/MetalTexture.cpp
--- a/Engine/Source/Runtime/IOS/MetalRHI/Private/MetalTexture.cpp
+++ b/Engine/Source/Runtime/IOS/MetalRHI/Private/MetalTexture.cpp
@@ -6,7 +6,6 @@
 
 #include "MetalRHIPrivate.h"
 
-TMap<MTLPixelFormat, uint8> FMetalSurface::PixelFormatKeyMap;
 uint8 FMetalSurface::NextKey = 1; // 0 is reserved for MTLPixelFormatInvalid
 
 /** Texture reference class. */
@@ -84,6 +83,7 @@ FMetalSurface::FMetalSurface(ERHIResourceType ResourceType, EPixelFormat Format,
 	, TotalTextureSize(0)
 {
 	// get a unique key for this surface's format
+	static TMap<uint32, uint8> PixelFormatKeyMap;
 	MTLPixelFormat MTLFormat = (MTLPixelFormat)GPixelFormats[Format].PlatformFormat;
 	uint8* Key = PixelFormatKeyMap.Find(MTLFormat);
 	if (Key == NULL)
diff --git a/Engine/Source/Runtime/IOS/MetalRHI/Public/MetalResources.h b/Engine/Source/Runtime/IOS/MetalRHI/Public/MetalResources.h
--- a/Engine/Source/Runtime/IOS/MetalRHI/Public/MetalResources.h
+++ b/Engine/Source/Runtime/IOS/MetalRHI/Public/MetalResources.h
@@ -186,8 +186,7 @@ public:
 	uint64 TotalTextureSize;
 
 private:
-	// a mapping of formats to small key, to fit into the PipelineState hash key
-	static TMap<MTLPixelFormat, uint8> PixelFormatKeyMap;
+	// next format for the pixel format mapping
 	static uint8 NextKey;
 };
 
