MallocCrash - Disabled for now, will be enabled later only for specific platforms.

diff --git a/Engine/Source/Runtime/Core/Private/HAL/MallocCrash.cpp b/Engine/Source/Runtime/Core/Private/HAL/MallocCrash.cpp
--- a/Engine/Source/Runtime/Core/Private/HAL/MallocCrash.cpp
+++ b/Engine/Source/Runtime/Core/Private/HAL/MallocCrash.cpp
@@ -211,6 +211,8 @@ FMallocCrash& FMallocCrash::Get( FMalloc* MainMalloc /*= nullptr*/ )
 
 void FMallocCrash::SetAsGMalloc()
 {
+	return;
+	// Disabled for now, will be enabled later only for specific platforms.
 	InternalLock.Lock();
 	GMalloc = this;
 	CrashedThreadId = FPlatformTLS::GetCurrentThreadId();
