Core: Fixed TNumericLimits<double>::Max() returning wrong value

diff --git a/Engine/Source/Runtime/Core/Public/Math/NumericLimits.h b/Engine/Source/Runtime/Core/Public/Math/NumericLimits.h
--- a/Engine/Source/Runtime/Core/Public/Math/NumericLimits.h
+++ b/Engine/Source/Runtime/Core/Public/Math/NumericLimits.h
@@ -309,7 +309,7 @@ struct TNumericLimits<double>
 
 	static NumericType Max()
 	{
-		return MIN_dbl;
+		return MAX_dbl;
 	}
 
 	static NumericType Lowest()
