Disable use of CoreSymbolication on OS X in games & the Editor - the strdup/delete mismatch occurs in more than one call & causes crashes. Programs like the CrashReportClient are already using FMallocAnsi so it isn't a problem for them.

diff --git a/Engine/Source/Runtime/Core/Private/Apple/ApplePlatformStackWalk.cpp b/Engine/Source/Runtime/Core/Private/Apple/ApplePlatformStackWalk.cpp
--- a/Engine/Source/Runtime/Core/Private/Apple/ApplePlatformStackWalk.cpp
+++ b/Engine/Source/Runtime/Core/Private/Apple/ApplePlatformStackWalk.cpp
@@ -85,7 +85,7 @@ static void AsyncSafeProgramCounterToSymbolInfo( uint64 ProgramCounter, FProgram
 		return;
 	}
 
-#if PLATFORM_MAC // On the Mac the crash report client can resymbolise
+#if PLATFORM_MAC && IS_PROGRAM // On the Mac the crash report client can resymbolise
 	if (DylibInfo.dli_sname)
 	{
 		// Mangled function name
diff --git a/Engine/Source/Runtime/Core/Private/Apple/ApplePlatformSymbolication.cpp b/Engine/Source/Runtime/Core/Private/Apple/ApplePlatformSymbolication.cpp
--- a/Engine/Source/Runtime/Core/Private/Apple/ApplePlatformSymbolication.cpp
+++ b/Engine/Source/Runtime/Core/Private/Apple/ApplePlatformSymbolication.cpp
@@ -10,7 +10,7 @@
 #include <CoreFoundation/CoreFoundation.h>
 #include <mach/mach.h>
 
-#if PLATFORM_MAC
+#if PLATFORM_MAC && IS_PROGRAM
 extern "C"
 {
 	struct CSTypeRef
@@ -61,12 +61,12 @@ extern "C"
 
 struct FApplePlatformSymbolCache
 {
-#if PLATFORM_MAC
+#if PLATFORM_MAC && IS_PROGRAM
 	TMap< FString, CSSymbolicatorRef > Symbolicators;
 #endif
 };
 
-static bool GAllowApplePlatformSymbolication = PLATFORM_MAC;
+static bool GAllowApplePlatformSymbolication = PLATFORM_MAC && IS_PROGRAM;
 
 void FApplePlatformSymbolication::SetSymbolicationAllowed(bool const bAllow)
 {
@@ -75,7 +75,7 @@ void FApplePlatformSymbolication::SetSymbolicationAllowed(bool const bAllow)
 
 FApplePlatformSymbolCache* FApplePlatformSymbolication::CreateSymbolCache(void)
 {
-#if PLATFORM_MAC
+#if PLATFORM_MAC && IS_PROGRAM
 	return GAllowApplePlatformSymbolication ? new FApplePlatformSymbolCache : nullptr;
 #else 
 	return nullptr;
@@ -84,7 +84,7 @@ FApplePlatformSymbolCache* FApplePlatformSymbolication::CreateSymbolCache(void)
 
 void FApplePlatformSymbolication::DestroySymbolCache(FApplePlatformSymbolCache* Cache)
 {
-#if PLATFORM_MAC
+#if PLATFORM_MAC && IS_PROGRAM
 	if(GAllowApplePlatformSymbolication && Cache)
 	{
 		for(auto It : Cache->Symbolicators)
@@ -98,7 +98,7 @@ void FApplePlatformSymbolication::DestroySymbolCache(FApplePlatformSymbolCache*
 
 bool FApplePlatformSymbolication::SymbolInfoForAddress(uint64 ProgramCounter, FProgramCounterSymbolInfo& out_SymbolInfo)
 {
-#if PLATFORM_MAC
+#if PLATFORM_MAC && IS_PROGRAM
 	bool bOK = false;
 	if (GAllowApplePlatformSymbolication)
 	{
@@ -137,7 +137,7 @@ bool FApplePlatformSymbolication::SymbolInfoForAddress(uint64 ProgramCounter, FP
 
 bool FApplePlatformSymbolication::SymbolInfoForFunctionFromModule(ANSICHAR const* MangledName, ANSICHAR const* ModuleName, FProgramCounterSymbolInfo& Info)
 {
-#if PLATFORM_MAC
+#if PLATFORM_MAC && IS_PROGRAM
 	bool bOK = false;
 	
 	if (GAllowApplePlatformSymbolication)
@@ -166,7 +166,7 @@ bool FApplePlatformSymbolication::SymbolInfoForFunctionFromModule(ANSICHAR const
 
 bool FApplePlatformSymbolication::SymbolInfoForStrippedSymbol(FApplePlatformSymbolCache* Cache, uint64 ProgramCounter, ANSICHAR const* ModulePath, ANSICHAR const* ModuleUUID, FProgramCounterSymbolInfo& Info)
 {
-#if PLATFORM_MAC
+#if PLATFORM_MAC && IS_PROGRAM
 	bool bOK = false;
 	
 	if(GAllowApplePlatformSymbolication && IFileManager::Get().FileSize(UTF8_TO_TCHAR(ModulePath)) > 0)
