ios now just overwrites the log each time instead of making backups to save disk space
UE-3800
#ios

diff --git a/Engine/Source/Runtime/Core/Private/IOS/IOSPlatformOutputDevices.cpp b/Engine/Source/Runtime/Core/Private/IOS/IOSPlatformOutputDevices.cpp
--- a/Engine/Source/Runtime/Core/Private/IOS/IOSPlatformOutputDevices.cpp
+++ b/Engine/Source/Runtime/Core/Private/IOS/IOSPlatformOutputDevices.cpp
@@ -18,6 +18,11 @@ class FOutputDeviceError* FIOSPlatformOutputDevices::GetError()
 	static FOutputDeviceIOSError Singleton;
 	return &Singleton;
 }
+class FOutputDevice*	FIOSPlatformOutputDevices::GetLog()
+{
+    static FOutputDeviceFile Singleton(nullptr, true);
+    return &Singleton;
+}
 
 //class FFeedbackContext*				FIOSPlatformOutputDevices::GetWarn()
 //{
diff --git a/Engine/Source/Runtime/Core/Public/IOS/IOSPlatformOutputDevices.h b/Engine/Source/Runtime/Core/Public/IOS/IOSPlatformOutputDevices.h
--- a/Engine/Source/Runtime/Core/Public/IOS/IOSPlatformOutputDevices.h
+++ b/Engine/Source/Runtime/Core/Public/IOS/IOSPlatformOutputDevices.h
@@ -10,6 +10,7 @@
 struct CORE_API FIOSPlatformOutputDevices : public FGenericPlatformOutputDevices
 {
 	static FOutputDeviceError*	GetError();
+    static FOutputDevice*		GetLog();
 };
 
 typedef FIOSPlatformOutputDevices FPlatformOutputDevices;
