Fix directory traversing on Linux.

- Not every file system supports d_type field of dirent, so a fallback to stat() is needed.
- Resolves UE-4890.
- Probably applies to Android too.

diff --git a/Engine/Source/Runtime/Core/Private/Linux/LinuxPlatformFile.cpp b/Engine/Source/Runtime/Core/Private/Linux/LinuxPlatformFile.cpp
--- a/Engine/Source/Runtime/Core/Private/Linux/LinuxPlatformFile.cpp
+++ b/Engine/Source/Runtime/Core/Private/Linux/LinuxPlatformFile.cpp
@@ -322,8 +322,8 @@ bool FLinuxPlatformFile::IterateDirectory(const TCHAR* Directory, FDirectoryVisi
 {
 	bool Result = false;
 
-	// If Directory is an empty string, assume that we want to iterate Binaries/Mac (current dir), but because we're an app bundle, iterate bundle's Contents/Frameworks instead
-	DIR* Handle = opendir(TCHAR_TO_UTF8(*NormalizeFilename(Directory)));
+	FString NormalizedDirectory = NormalizeFilename(Directory);
+	DIR* Handle = opendir(TCHAR_TO_UTF8(*NormalizedDirectory));
 	if (Handle)
 	{
 		Result = true;
@@ -332,7 +332,28 @@ bool FLinuxPlatformFile::IterateDirectory(const TCHAR* Directory, FDirectoryVisi
 		{
 			if (FCString::Strcmp(UTF8_TO_TCHAR(Entry->d_name), TEXT(".")) && FCString::Strcmp(UTF8_TO_TCHAR(Entry->d_name), TEXT("..")))
 			{
-				Result = Visitor.Visit(*(FString(Directory) / UTF8_TO_TCHAR(Entry->d_name)), Entry->d_type == DT_DIR);
+				bool bIsDirectory = false;
+				FString UnicodeEntryName = UTF8_TO_TCHAR(Entry->d_name);
+				if (Entry->d_type != DT_UNKNOWN)
+				{
+					bIsDirectory = Entry->d_type == DT_DIR;
+				}
+				else
+				{
+					// filesystem does not support d_type, fallback to stat
+					struct stat FileInfo;
+					FString AbsoluteUnicodeName = NormalizedDirectory / UnicodeEntryName;	
+					if (stat(TCHAR_TO_UTF8(*AbsoluteUnicodeName), &FileInfo) != -1)
+					{
+						bIsDirectory = ((FileInfo.st_mode & S_IFMT) == S_IFDIR);
+					}
+					else
+					{
+						int ErrNo = errno;
+						UE_LOG(LogLinuxPlatformFile, Warning, TEXT( "Cannot determine whether '%s' is a directory - d_type not supported and stat() failed with errno=%d (%s)"), *AbsoluteUnicodeName, ErrNo, UTF8_TO_TCHAR(strerror(ErrNo)));
+					}
+				}
+				Result = Visitor.Visit(*(FString(Directory) / UnicodeEntryName), bIsDirectory);
 			}
 		}
 		closedir(Handle);
