Back out changelist 2344886 - Files must be truncated if they're not opened for append.

diff --git a/Engine/Source/Runtime/Core/Private/Android/AndroidFile.cpp b/Engine/Source/Runtime/Core/Private/Android/AndroidFile.cpp
--- a/Engine/Source/Runtime/Core/Private/Android/AndroidFile.cpp
+++ b/Engine/Source/Runtime/Core/Private/Android/AndroidFile.cpp
@@ -501,6 +501,10 @@ public:
 		{
 			Flags |= O_APPEND;
 		}
+		else
+		{
+			Flags |= O_TRUNC;
+		}
 		if (bAllowRead)
 		{
 			Flags |= O_RDWR;
diff --git a/Engine/Source/Runtime/Core/Private/Apple/ApplePlatformFile.cpp b/Engine/Source/Runtime/Core/Private/Apple/ApplePlatformFile.cpp
--- a/Engine/Source/Runtime/Core/Private/Apple/ApplePlatformFile.cpp
+++ b/Engine/Source/Runtime/Core/Private/Apple/ApplePlatformFile.cpp
@@ -442,6 +442,10 @@ IFileHandle* FApplePlatformFile::OpenWrite(const TCHAR* Filename, bool bAppend,
 	{
 		Flags |= O_APPEND;
 	}
+	else
+	{
+		Flags |= O_TRUNC;
+	}
 	if (bAllowRead)
 	{
 		Flags |= O_RDWR;
diff --git a/Engine/Source/Runtime/Core/Private/HTML5/HTML5PlatformFile.cpp b/Engine/Source/Runtime/Core/Private/HTML5/HTML5PlatformFile.cpp
--- a/Engine/Source/Runtime/Core/Private/HTML5/HTML5PlatformFile.cpp
+++ b/Engine/Source/Runtime/Core/Private/HTML5/HTML5PlatformFile.cpp
@@ -312,6 +312,10 @@ public:
 		{
 			Flags |= O_APPEND;
 		}
+		else
+		{
+			Flags |= O_TRUNC;
+		}
 		if (bAllowRead)
 		{
 			Flags |= O_RDWR;
diff --git a/Engine/Source/Runtime/Core/Private/IOS/IOSPlatformFile.cpp b/Engine/Source/Runtime/Core/Private/IOS/IOSPlatformFile.cpp
--- a/Engine/Source/Runtime/Core/Private/IOS/IOSPlatformFile.cpp
+++ b/Engine/Source/Runtime/Core/Private/IOS/IOSPlatformFile.cpp
@@ -546,6 +546,10 @@ IFileHandle* FIOSPlatformFile::OpenWrite(const TCHAR* Filename, bool bAppend, bo
 	{
 		Flags |= O_APPEND;
 	}
+	else
+	{
+		Flags |= O_TRUNC;
+	}
 	if (bAllowRead)
 	{
 		Flags |= O_RDWR;
diff --git a/Engine/Source/Runtime/Core/Private/Linux/LinuxPlatformFile.cpp b/Engine/Source/Runtime/Core/Private/Linux/LinuxPlatformFile.cpp
--- a/Engine/Source/Runtime/Core/Private/Linux/LinuxPlatformFile.cpp
+++ b/Engine/Source/Runtime/Core/Private/Linux/LinuxPlatformFile.cpp
@@ -266,9 +266,10 @@ IFileHandle* FLinuxPlatformFile::OpenWrite(const TCHAR* Filename, bool bAppend,
 	// create directories if needed.
 	if (!CreateDirectoriesFromPath(Filename))
 	{
-		return nullptr;
+		return NULL;
 	}
 
+	// Caveat: cannot specify O_TRUNC in flags, as this will corrupt the file which may be "locked" by other process. We will ftruncate() it once we "lock" it
 	int32 Handle = open(TCHAR_TO_UTF8(*NormalizeFilename(Filename)), Flags, S_IRUSR | S_IWUSR);
 	if (Handle != -1)
 	{
@@ -285,6 +286,19 @@ IFileHandle* FLinuxPlatformFile::OpenWrite(const TCHAR* Filename, bool bAppend,
 			// all the other locking errors are ignored.
 		}
 
+		// truncate the file now that we locked it
+		if (!bAppend)
+		{
+			if (ftruncate(Handle, 0) != 0)
+			{
+				int ErrNo = errno;
+				UE_LOG(LogLinuxPlatformFile, Warning, TEXT( "ftruncate() failed for '%s': errno=%d (%s)" ),
+															Filename, ErrNo, ANSI_TO_TCHAR(strerror(ErrNo)));
+				close(Handle);
+				return nullptr;
+			}
+		}
+
 		FFileHandleLinux* FileHandleLinux = new FFileHandleLinux(Handle);
 		if (bAppend)
 		{
