Switch to SDL text input API for input handling.

- Part of PR #560 contributed by [REDACTED_PERSON].

diff --git a/Engine/Source/Runtime/Core/Private/Linux/LinuxApplication.cpp b/Engine/Source/Runtime/Core/Private/Linux/LinuxApplication.cpp
--- a/Engine/Source/Runtime/Core/Private/Linux/LinuxApplication.cpp
+++ b/Engine/Source/Runtime/Core/Private/Linux/LinuxApplication.cpp
@@ -240,13 +240,14 @@ void FLinuxApplication::ProcessDeferredMessage( SDL_Event Event )
 			SDL_KeyboardEvent KeyEvent = Event.key;
 			const SDL_Keycode KeyCode = KeyEvent.keysym.scancode;
 			const bool bIsRepeated = KeyEvent.repeat != 0;
-				
-			// First KeyDown, then KeyChar. This is important, as in-game console ignores first character otherwise
+
+			// Text input is now handled in SDL_TEXTINPUT: see below
 			MessageHandler->OnKeyDown(KeyCode, KeyEvent.keysym.sym, bIsRepeated);
 
-			if (GeneratesKeyCharMessage(KeyEvent))
+			// Backspace input in only caught here.
+			if (KeyCode == SDL_SCANCODE_BACKSPACE || KeyCode == SDL_SCANCODE_RETURN)
 			{
-				const TCHAR Character = ConvertChar(KeyEvent.keysym);
+				const TCHAR Character = SDL_GetKeyFromScancode(Event.key.keysym.scancode);
 				MessageHandler->OnKeyChar(Character, bIsRepeated);
 			}
 		}
@@ -255,12 +256,21 @@ void FLinuxApplication::ProcessDeferredMessage( SDL_Event Event )
 		{
 			SDL_KeyboardEvent keyEvent = Event.key;
 			const SDL_Keycode KeyCode = keyEvent.keysym.scancode;
-			const TCHAR Character = ConvertChar( keyEvent.keysym );
 			const bool IsRepeat = keyEvent.repeat != 0;
 
 			MessageHandler->OnKeyUp( KeyCode, keyEvent.keysym.sym, IsRepeat );
 		}
 		break;
+	case SDL_TEXTINPUT:
+		{
+			// Slate now gets all its text from here, I hope.
+			// Don't know if this will work with ingame text or not.
+			const bool bIsRepeated = Event.key.repeat != 0;
+			const TCHAR Character = *ANSI_TO_TCHAR(Event.text.text);
+
+			MessageHandler->OnKeyChar(Character, bIsRepeated);
+		}
+		break;
 	case SDL_MOUSEMOTION:
 		{
 			SDL_MouseMotionEvent motionEvent = Event.motion;
@@ -771,27 +781,27 @@ void FLinuxApplication::PollGameDeviceState( const float TimeDelta )
 
 TCHAR FLinuxApplication::ConvertChar( SDL_Keysym Keysym )
 {
-	if( Keysym.sym >= 128 )
+	if (SDL_GetKeyFromScancode(Keysym.scancode) >= 128)
 	{
 		return 0;
 	}
 
-    TCHAR Char = Keysym.sym;
+	TCHAR Char = SDL_GetKeyFromScancode(Keysym.scancode);
 
     if (Keysym.mod & (KMOD_LSHIFT | KMOD_RSHIFT))
     {
         // Convert to uppercase (FIXME: what about CAPS?)
-        if( Keysym.sym >= 97 && Keysym.sym <= 122)
+		if( SDL_GetKeyFromScancode(Keysym.scancode)  >= 97 && SDL_GetKeyFromScancode(Keysym.scancode)  <= 122)
         {
             return Keysym.sym - 32;
         }
-        else if( Keysym.sym >= 91 && Keysym.sym <= 93)
+		else if( SDL_GetKeyFromScancode(Keysym.scancode) >= 91 && SDL_GetKeyFromScancode(Keysym.scancode)  <= 93)
         {
-            return Keysym.sym + 32; // [ \ ] -> { | }
+			return SDL_GetKeyFromScancode(Keysym.scancode) + 32; // [ \ ] -> { | }
         }
         else
         {
-            switch(Keysym.sym)
+			switch(SDL_GetKeyFromScancode(Keysym.scancode) )
             {
                 case '`': // ` -> ~
                     Char = TEXT('`');
@@ -879,6 +889,12 @@ TSharedPtr< FLinuxWindow > FLinuxApplication::FindEventWindow( SDL_Event* Event
 	uint16 WindowID = 0;
 	switch (Event->type)
 	{
+		case SDL_TEXTINPUT:
+			WindowID = Event->text.windowID;
+			break;
+		case SDL_TEXTEDITING:
+			WindowID = Event->edit.windowID;
+			break;
 		case SDL_KEYDOWN:
 		case SDL_KEYUP:
 			WindowID = Event->key.windowID;
diff --git a/Engine/Source/Runtime/Core/Private/Linux/LinuxPlatformMisc.cpp b/Engine/Source/Runtime/Core/Private/Linux/LinuxPlatformMisc.cpp
--- a/Engine/Source/Runtime/Core/Private/Linux/LinuxPlatformMisc.cpp
+++ b/Engine/Source/Runtime/Core/Private/Linux/LinuxPlatformMisc.cpp
@@ -161,6 +161,9 @@ bool FLinuxPlatformMisc::PlatformInitMultimedia()
 			return false;
 		}
 
+		// Used to make SDL push SDL_TEXTINPUT events.
+		SDL_StartTextInput();
+
 		GInitializedSDL = true;
 
 		// needs to come after GInitializedSDL, otherwise it will recurse here
