Restoring original implementation of InterpEaseInOut - was changed when more interpolation methods were introduced. The new version did not work with FRotator.

diff --git a/Engine/Source/Runtime/Core/Public/Math/UnrealMathUtility.h b/Engine/Source/Runtime/Core/Public/Math/UnrealMathUtility.h
--- a/Engine/Source/Runtime/Core/Public/Math/UnrealMathUtility.h
+++ b/Engine/Source/Runtime/Core/Public/Math/UnrealMathUtility.h
@@ -588,9 +588,11 @@ struct FMath : public FPlatformMath
 	template< class T > 
 	static FORCEINLINE_DEBUGGABLE T InterpEaseInOut( const T& A, const T& B, float Alpha, float Exp )
 	{
-		return (Alpha < 0.5f) ?
-			InterpEaseIn(A, B, Alpha * 2.f, Exp) * 0.5f :
-			InterpEaseOut(A, B, Alpha * 2.f - 1.f, Exp) * 0.5f + 0.5f;
+		float const ModifiedAlpha = (Alpha < 0.5f) ?
+			0.5f * Pow(2.f * Alpha, Exp) :
+			1.f - 0.5f * Pow(2.f * (1.f - Alpha), Exp);
+
+		return Lerp<T>(A, B, ModifiedAlpha);
 	}
 
 	/** Interpolation between A and B, applying a step function. */
