fix for crash on devices < iOS 8 when mobile content scale is set to 0
#ios

diff --git a/Engine/Source/Runtime/Core/Private/IOS/IOSView.cpp b/Engine/Source/Runtime/Core/Private/IOS/IOSView.cpp
--- a/Engine/Source/Runtime/Core/Private/IOS/IOSView.cpp
+++ b/Engine/Source/Runtime/Core/Private/IOS/IOSView.cpp
@@ -149,11 +149,16 @@ id<MTLDevice> GMetalDevice = nil;
 		if (RequestedContentScaleFactor == 0.0f)
 		{
 #ifdef __IPHONE_8_0
-			self.contentScaleFactor = self.window.screen.nativeScale;
-			UE_LOG(LogIOS, Log, TEXT("Setting contentScaleFactor to nativeScale which is = %f"), self.contentScaleFactor);
-#else
-			UE_LOG(LogIOS, Log, TEXT("Leaving contentScaleFactor alone, with scale = %f"), NativeScale);
+            if ([self.window.screen respondsToSelector:@selector(nativeScale)])
+            {
+                self.contentScaleFactor = self.window.screen.nativeScale;
+                UE_LOG(LogIOS, Log, TEXT("Setting contentScaleFactor to nativeScale which is = %f"), self.contentScaleFactor);
+            }
+            else
 #endif
+            {
+                UE_LOG(LogIOS, Log, TEXT("Leaving contentScaleFactor alone, with scale = %f"), NativeScale);
+            }
 		}
 		else
 		{
