Fix for potential memory leak when using TArray<>::RemoveSingle

diff --git a/Engine/Source/Runtime/Core/Public/Containers/Array.h b/Engine/Source/Runtime/Core/Public/Containers/Array.h
--- a/Engine/Source/Runtime/Core/Public/Containers/Array.h
+++ b/Engine/Source/Runtime/Core/Public/Containers/Array.h
@@ -1881,7 +1881,7 @@ public:
 		// Destruct items that match the specified Item.
 		DestructItems(RemovePtr, 1);
 		const int32 NextIndex = Index + 1;
-		MoveConstructItems(RemovePtr, RemovePtr + 1, ArrayNum - (Index + 1));
+		RelocateConstructItems<ElementType>(RemovePtr, RemovePtr + 1, ArrayNum - (Index + 1));
 
 		// Update the array count
 		--ArrayNum;
