Improved saving and restoring of windows order on Mac when app activates and deactivates

diff --git a/Engine/Source/Runtime/Core/Private/Mac/MacApplication.cpp b/Engine/Source/Runtime/Core/Private/Mac/MacApplication.cpp
--- a/Engine/Source/Runtime/Core/Private/Mac/MacApplication.cpp
+++ b/Engine/Source/Runtime/Core/Private/Mac/MacApplication.cpp
@@ -144,18 +144,30 @@ FMacApplication::FMacApplication()
 
 	AppActivationObserver = [[NSNotificationCenter defaultCenter] addObserverForName:NSApplicationDidBecomeActiveNotification object:[NSApplication sharedApplication] queue:[NSOperationQueue mainQueue] usingBlock:^(NSNotification* Notification)
 							{
-								for (int32 Index = SavedWindowsOrder.Num() - 1; Index > 0; Index--)
+								for (int32 Index = 0; Index < SavedWindowsOrder.Num(); Index++)
 								{
 									const FSavedWindowOrderInfo& SavedWindowLevel = SavedWindowsOrder[Index];
 									NSWindow* Window = [NSApp windowWithWindowNumber:SavedWindowLevel.WindowNumber];
-									const int32 PreviousWindowNumber = SavedWindowsOrder[Index - 1].WindowNumber;
 									if (Window)
 									{
-										[Window orderWindow:NSWindowBelow relativeTo:PreviousWindowNumber];
 										[Window setLevel:SavedWindowLevel.Level];
 									}
 								}
 
+								if (SavedWindowsOrder.Num() > 0)
+								{
+									[[NSApp windowWithWindowNumber:SavedWindowsOrder[0].WindowNumber] orderWindow:NSWindowAbove relativeTo:0];
+									for (int32 Index = 1; Index < SavedWindowsOrder.Num(); Index++)
+									{
+										const FSavedWindowOrderInfo& SavedWindowLevel = SavedWindowsOrder[Index];
+										NSWindow* Window = [NSApp windowWithWindowNumber:SavedWindowLevel.WindowNumber];
+										if (Window)
+										{
+											[Window orderWindow:NSWindowBelow relativeTo:SavedWindowsOrder[Index - 1].WindowNumber];
+										}
+									}
+								}
+
 								// If editor thread doesn't have the focus, don't suck up too much CPU time.
 								if (GIsEditor)
 								{
@@ -184,7 +196,19 @@ FMacApplication::FMacApplication()
 									}
 								}
 
-								CFRunLoopRunInMode(kCFRunLoopDefaultMode, 0, true);
+								if (SavedWindowsOrder.Num() > 0)
+								{
+									[[NSApp windowWithWindowNumber:SavedWindowsOrder[0].WindowNumber] orderWindow:NSWindowAbove relativeTo:0];
+									for (int32 Index = 1; Index < SavedWindowsOrder.Num(); Index++)
+									{
+										const FSavedWindowOrderInfo& SavedWindowLevel = SavedWindowsOrder[Index];
+										NSWindow* Window = [NSApp windowWithWindowNumber:SavedWindowLevel.WindowNumber];
+										if (Window)
+										{
+											[Window orderWindow:NSWindowBelow relativeTo:SavedWindowsOrder[Index - 1].WindowNumber];
+										}
+									}
+								}
 
 								// If editor thread doesn't have the focus, don't suck up too much CPU time.
 								if (GIsEditor)
