Moved function into .CPP so I could get a breakpoint on it. Would like to keep function in .CPP.

diff --git a/Engine/Source/Runtime/Core/Private/Stats/StatsCommand.cpp b/Engine/Source/Runtime/Core/Private/Stats/StatsCommand.cpp
--- a/Engine/Source/Runtime/Core/Private/Stats/StatsCommand.cpp
+++ b/Engine/Source/Runtime/Core/Private/Stats/StatsCommand.cpp
@@ -428,6 +428,12 @@ static bool HandleCommandBroadcast(const FName& InStatName, bool& bOutCurrentEna
 
 #if STATS
 
+void FHUDGroupGameThreadRenderer::NewData(FGameThreadHudData* Data)
+{
+	delete Latest;
+	Latest = Data;
+}
+
 FHUDGroupGameThreadRenderer& FHUDGroupGameThreadRenderer::Get()
 {
 	static FHUDGroupGameThreadRenderer Singleton;
diff --git a/Engine/Source/Runtime/Core/Public/Stats/StatsData.h b/Engine/Source/Runtime/Core/Public/Stats/StatsData.h
--- a/Engine/Source/Runtime/Core/Public/Stats/StatsData.h
+++ b/Engine/Source/Runtime/Core/Public/Stats/StatsData.h
@@ -793,11 +793,7 @@ struct CORE_API FHUDGroupGameThreadRenderer
 		Latest = nullptr;
 	}
 
-	void NewData(FGameThreadHudData* Data)
-	{
-		delete Latest;
-		Latest = Data;
-	}
+	void NewData(FGameThreadHudData* Data);
 
 	static FHUDGroupGameThreadRenderer& Get();
 };
