Fixed a crash when using the Bing Pinyin IME

UDN - IME related editor crash

We now set the new composition range before we set the new composition text, as setting the text can have side effects that can cause the composition to end.

This also fixes an issue where some IMEs would send an empty string, rather than a cancel event, when you press escape to cancel the current composition. This is now handled correctly, rather than appearing to work, but then entering garbled text the next time you typed.

diff --git a/Engine/Source/Runtime/Core/Private/Windows/WindowsTextInputMethodSystem.cpp b/Engine/Source/Runtime/Core/Private/Windows/WindowsTextInputMethodSystem.cpp
--- a/Engine/Source/Runtime/Core/Private/Windows/WindowsTextInputMethodSystem.cpp
+++ b/Engine/Source/Runtime/Core/Private/Windows/WindowsTextInputMethodSystem.cpp
@@ -269,6 +269,26 @@ void FWindowsTextInputMethodSystem::EndIMMComposition()
 	ActiveContext->EndComposition();
 }
 
+void FWindowsTextInputMethodSystem::CancelIMMComposition()
+{
+	check(ActiveContext.IsValid());
+
+	UE_LOG(LogWindowsTextInputMethodSystem, Verbose, TEXT("WM_IME_COMPOSITION Composition Canceled"));
+
+	FInternalContext& InternalContext = ContextToInternalContextMap[ActiveContext];
+
+	const int32 CurrentCompositionBeginIndex = InternalContext.IMMContext.CompositionBeginIndex;
+	const uint32 CurrentCompositionLength = InternalContext.IMMContext.CompositionLength;
+
+	// Clear Composition
+	InternalContext.IMMContext.CompositionLength = 0;
+	ActiveContext->UpdateCompositionRange(InternalContext.IMMContext.CompositionBeginIndex, 0);
+	ActiveContext->SetSelectionRange(InternalContext.IMMContext.CompositionBeginIndex, 0, ITextInputMethodContext::ECaretPosition::Beginning);
+	ActiveContext->SetTextInRange(CurrentCompositionBeginIndex, CurrentCompositionLength, TEXT(""));
+
+	EndIMMComposition();
+}
+
 bool FWindowsTextInputMethodSystem::InitializeTSF()
 {
 	UE_LOG(LogWindowsTextInputMethodSystem, Verbose, TEXT("Initializing TSF..."));
@@ -763,15 +783,7 @@ int32 FWindowsTextInputMethodSystem::ProcessMessage(HWND hwnd, uint32 msg, WPARA
 				// Canceled, so remove the compositing string
 				if(bHasBeenCanceled)
 				{
-					UE_LOG(LogWindowsTextInputMethodSystem, Verbose, TEXT("WM_IME_COMPOSITION Composition Canceled"));
-
-					// Clear Composition
-					ActiveContext->SetTextInRange(InternalContext.IMMContext.CompositionBeginIndex, InternalContext.IMMContext.CompositionLength, TEXT(""));
-					InternalContext.IMMContext.CompositionLength = 0;
-					ActiveContext->UpdateCompositionRange(InternalContext.IMMContext.CompositionBeginIndex, 0);
-					ActiveContext->SetSelectionRange(InternalContext.IMMContext.CompositionBeginIndex, 0, ITextInputMethodContext::ECaretPosition::Beginning);
-
-					EndIMMComposition();
+					CancelIMMComposition();
 				}
 
 				// Check Result
@@ -810,6 +822,13 @@ int32 FWindowsTextInputMethodSystem::ProcessMessage(HWND hwnd, uint32 msg, WPARA
 					const FString CompositionString = GetIMMStringAsFString(IMMContext, GCS_COMPSTR);
 					UE_LOG(LogWindowsTextInputMethodSystem, Verbose, TEXT("WM_IME_COMPOSITION Composition String: %s"), *CompositionString);
 
+					// Not all IMEs send a cancel request when you press escape, but instead just set the string to empty
+					// We need to cancel out the composition string here to avoid weirdness when you start typing again
+					if(CompositionString.Len() == 0)
+					{
+						CancelIMMComposition();
+					}
+
 					// We've typed a character, so we need to clear out any currently selected text to mimic what happens when you normally type into a text input
 					uint32 SelectionBeginIndex = 0;
 					uint32 SelectionLength = 0;
@@ -827,12 +846,15 @@ int32 FWindowsTextInputMethodSystem::ProcessMessage(HWND hwnd, uint32 msg, WPARA
 						BeginIMMComposition();
 					}
 
-					// Update Composition
-					ActiveContext->SetTextInRange(InternalContext.IMMContext.CompositionBeginIndex, InternalContext.IMMContext.CompositionLength, CompositionString);
+					const int32 CurrentCompositionBeginIndex = InternalContext.IMMContext.CompositionBeginIndex;
+					const uint32 CurrentCompositionLength = InternalContext.IMMContext.CompositionLength;
 
 					// Update Composition Range
 					InternalContext.IMMContext.CompositionLength = CompositionString.Len();
 					ActiveContext->UpdateCompositionRange(InternalContext.IMMContext.CompositionBeginIndex, InternalContext.IMMContext.CompositionLength);
+
+					// Update Composition
+					ActiveContext->SetTextInRange(CurrentCompositionBeginIndex, CurrentCompositionLength, CompositionString);
 				}
 
 				// Check Cursor
diff --git a/Engine/Source/Runtime/Core/Public/Windows/WindowsTextInputMethodSystem.h b/Engine/Source/Runtime/Core/Public/Windows/WindowsTextInputMethodSystem.h
--- a/Engine/Source/Runtime/Core/Public/Windows/WindowsTextInputMethodSystem.h
+++ b/Engine/Source/Runtime/Core/Public/Windows/WindowsTextInputMethodSystem.h
@@ -81,6 +81,7 @@ private:
 	void UpdateIMMWindowPositions(HIMC IMMContext);
 	void BeginIMMComposition();
 	void EndIMMComposition();
+	void CancelIMMComposition();
 
 	// TSF Implementation
 	bool InitializeTSF();
