Slate : Slate Stats now show measurements using a monospaced font with fixed numbers of digits.
Text : FNumberFormattingOptions now provides method-chainable settings for easily creating a static const instance.

diff --git a/Engine/Source/Developer/SlateReflector/Private/Widgets/SWidgetReflector.cpp b/Engine/Source/Developer/SlateReflector/Private/Widgets/SWidgetReflector.cpp
--- a/Engine/Source/Developer/SlateReflector/Private/Widgets/SWidgetReflector.cpp
+++ b/Engine/Source/Developer/SlateReflector/Private/Widgets/SWidgetReflector.cpp
@@ -461,7 +461,16 @@ public:
 	FText GetDisplayName() const { return DisplayText; }	
 	FText GetInclusiveAvgMsText() const { return InclusiveAvgMsText; }
 	float GetInclusiveAvgMs() const { return InclusiveAvgMs; }
-	void UpdateInclusiveAvgMs( float InMs ){ InclusiveAvgMs = InMs; InclusiveAvgMsText = FText::AsNumber(InMs); }
+	void UpdateInclusiveAvgMs( float InMs )
+	{
+		static const auto FormattingOptions = FNumberFormattingOptions()
+			.SetMinimumIntegralDigits(1)
+			.SetMinimumFractionalDigits(3)
+			.SetMaximumFractionalDigits(3);
+
+		InclusiveAvgMs = InMs;
+		InclusiveAvgMsText = FText::AsNumber(InMs, &FormattingOptions);
+	}
 
 private:
 
@@ -610,6 +619,7 @@ TSharedRef<ITableRow> SWidgetReflector::GenerateStatRow( TSharedRef<FStatItem> S
 				.Padding(FMargin(5,0))
 				[
 					SNew(STextBlock)
+					.TextStyle(FCoreStyle::Get(), "MonospacedText")
 					.Text( this, &SStatTableRow::GetInclusiveAvgMsText )
 				];
 			}
diff --git a/Engine/Source/Runtime/Core/Public/Internationalization/Text.h b/Engine/Source/Runtime/Core/Public/Internationalization/Text.h
--- a/Engine/Source/Runtime/Core/Public/Internationalization/Text.h
+++ b/Engine/Source/Runtime/Core/Public/Internationalization/Text.h
@@ -95,11 +95,22 @@ struct CORE_API FNumberFormattingOptions
 	FNumberFormattingOptions();
 
 	bool UseGrouping;
+	FNumberFormattingOptions& SetUseGrouping( bool InValue ){ UseGrouping = InValue; return *this; }
+
 	ERoundingMode RoundingMode;
+	FNumberFormattingOptions& SetRoundingMode( ERoundingMode InValue ){ RoundingMode = InValue; return *this; }
+
 	int32 MinimumIntegralDigits;
+	FNumberFormattingOptions& SetMinimumIntegralDigits( int32 InValue ){ MinimumIntegralDigits = InValue; return *this; }
+
 	int32 MaximumIntegralDigits;
+	FNumberFormattingOptions& SetMaximumIntegralDigits( int32 InValue ){ MaximumIntegralDigits = InValue; return *this; }
+
 	int32 MinimumFractionalDigits;
+	FNumberFormattingOptions& SetMinimumFractionalDigits( int32 InValue ){ MinimumFractionalDigits = InValue; return *this; }
+
 	int32 MaximumFractionalDigits;
+	FNumberFormattingOptions& SetMaximumFractionalDigits( int32 InValue ){ MaximumFractionalDigits = InValue; return *this; }
 
 	friend FArchive& operator<<(FArchive& Ar, FNumberFormattingOptions& Value);
 };
