Core: Removed obsolete privat einclude path

diff --git a/Engine/Source/Runtime/Core/Core.Build.cs b/Engine/Source/Runtime/Core/Core.Build.cs
--- a/Engine/Source/Runtime/Core/Core.Build.cs
+++ b/Engine/Source/Runtime/Core/Core.Build.cs
@@ -27,7 +27,6 @@ public class Core : ModuleRules
 				"Runtime/Core/Public/Modules/Boilerplate",
 				"Runtime/Core/Public/ProfilingDebugging",
 				"Runtime/Core/Public/Serialization",
-				"Runtime/Core/Public/Serialization/Json",
 				"Runtime/Core/Public/Serialization/Csv",
 				"Runtime/Core/Public/Stats",
 				"Runtime/Core/Public/Templates",
