On OS X if the UE4_FORCE_MALLOC_ANSI is set in the current environment then use FMallocAnsi regardless, this allows GL Profiler to work without crashing. Apple's GL Profliler calls CoreSymbolication which has a strdup/delete mismatch that crashes with our custom new/delete & Apple will need to fix CoreSymbolication for our custom allocators to work.

diff --git a/Engine/Source/Runtime/Core/Private/Mac/MacPlatformMemory.cpp b/Engine/Source/Runtime/Core/Private/Mac/MacPlatformMemory.cpp
--- a/Engine/Source/Runtime/Core/Private/Mac/MacPlatformMemory.cpp
+++ b/Engine/Source/Runtime/Core/Private/Mac/MacPlatformMemory.cpp
@@ -21,13 +21,20 @@ void FMacPlatformMemory::Init()
 
 FMalloc* FMacPlatformMemory::BaseAllocator()
 {
+	if(getenv("UE4_FORCE_MALLOC_ANSI") != nullptr)
+	{
+		return new FMallocAnsi();
+	}
+	else
+	{
 #if FORCE_ANSI_ALLOCATOR || IS_PROGRAM
-	return new FMallocAnsi();
+		return new FMallocAnsi();
 #elif (WITH_EDITORONLY_DATA || IS_PROGRAM) && TBB_ALLOCATOR_ALLOWED
-	return new FMallocTBB();
+		return new FMallocTBB();
 #else
-	return new FMallocBinned((uint32)(GetConstants().PageSize&MAX_uint32), 0x100000000);
+		return new FMallocBinned((uint32)(GetConstants().PageSize&MAX_uint32), 0x100000000);
 #endif
+	}
 }
 
 FPlatformMemoryStats FMacPlatformMemory::GetStats()
