Fix another cause of deadlock in Mac UE4. -[NSEvent window] will access the global window list which requires a mutex which may be held by the opposing thread (main vs. game) causing deadlock.

To avoid this add an NSEvent category (FCachedWindowAccess) containing two selectors, one to cache the NSWindow* into the convenient NSWindow* ivar and another to access the ivar without attempting to access the system's global window list. Use this function from the game-thread to avoid contention.

diff --git a/Engine/Source/Runtime/Core/Private/Mac/MacApplication.cpp b/Engine/Source/Runtime/Core/Private/Mac/MacApplication.cpp
--- a/Engine/Source/Runtime/Core/Private/Mac/MacApplication.cpp
+++ b/Engine/Source/Runtime/Core/Private/Mac/MacApplication.cpp
@@ -597,10 +597,10 @@ void FMacApplication::ProcessNSEvent(NSEvent* const Event)
 			else
 			{
 				NSPoint CursorPos = [Event locationInWindow];
-				if ([Event window])
+				if ([Event GetWindow])
 				{
-					CursorPos.x += [Event window].frame.origin.x;
-					CursorPos.y += [Event window].frame.origin.y;
+					CursorPos.x += [Event GetWindow].frame.origin.x;
+					CursorPos.y += [Event GetWindow].frame.origin.y;
 				}
 				CursorPos.y--; // The y coordinate in the point returned by locationInWindow starts from a base of 1
 				const FVector2D MousePosition = FVector2D(CursorPos.x, FPlatformMisc::ConvertSlateYPositionToCocoa(CursorPos.y));
@@ -864,7 +864,7 @@ FCocoaWindow* FMacApplication::FindEventWindow( NSEvent* Event )
 			break;
 	}
 
-	FCocoaWindow* EventWindow = (FCocoaWindow*)[Event window];
+	FCocoaWindow* EventWindow = (FCocoaWindow*)[Event GetWindow];
 
 	if ([Event type] == NSMouseMoved)
 	{
@@ -882,17 +882,17 @@ FCocoaWindow* FMacApplication::FindEventWindow( NSEvent* Event )
 		{
 			EventWindow = DraggedWindow;
 		}
-		else if( MouseCaptureWindow && MouseCaptureWindow == [Event window] && [MouseCaptureWindow targetWindow] )
+		else if( MouseCaptureWindow && MouseCaptureWindow == [Event GetWindow] && [MouseCaptureWindow targetWindow] )
 		{
 			EventWindow = [MouseCaptureWindow targetWindow];
 		}
 		else
 		{
 			NSPoint CursorPos = [Event locationInWindow];
-			if ([Event window])
+			if ([Event GetWindow])
 			{
-				CursorPos.x += [Event window].frame.origin.x;
-				CursorPos.y += [Event window].frame.origin.y;
+				CursorPos.x += [Event GetWindow].frame.origin.x;
+				CursorPos.y += [Event GetWindow].frame.origin.y;
 			}
 			CursorPos.y--; // The y coordinate in the point returned by locationInWindow starts from a base of 1
 			TSharedPtr<FMacWindow> WindowUnderCursor = LocateWindowUnderCursor(CursorPos);
diff --git a/Engine/Source/Runtime/Core/Private/Mac/MacEvent.cpp b/Engine/Source/Runtime/Core/Private/Mac/MacEvent.cpp
--- a/Engine/Source/Runtime/Core/Private/Mac/MacEvent.cpp
+++ b/Engine/Source/Runtime/Core/Private/Mac/MacEvent.cpp
@@ -6,12 +6,24 @@
 #include "CocoaThread.h"
 #include "CocoaWindow.h"
 
+@implementation NSEvent (FCachedWindowAccess)
+-(void)CacheWindow
+{
+	_window = [self window];
+}
+-(NSWindow*)GetWindow
+{
+	return _window;
+}
+@end
+
 FMacEvent::FMacEvent(NSEvent* const Event)
 : EventData(Event)
 {
 	SCOPED_AUTORELEASE_POOL;
 	check(Event);
 	[Event retain];
+	[Event CacheWindow];
 }
 
 FMacEvent::FMacEvent(NSNotification* const Notification)
diff --git a/Engine/Source/Runtime/Core/Public/Mac/MacEvent.h b/Engine/Source/Runtime/Core/Public/Mac/MacEvent.h
--- a/Engine/Source/Runtime/Core/Public/Mac/MacEvent.h
+++ b/Engine/Source/Runtime/Core/Public/Mac/MacEvent.h
@@ -10,6 +10,11 @@ enum EMacEventSendMethod
 	Sync
 };
 
+@interface NSEvent (FCachedWindowAccess)
+-(void)CacheWindow;
+-(NSWindow*)GetWindow;
+@end
+
 class FMacEvent
 {
 	// Constructor for an NSEvent based FMacEvent
