Retain & release the cached NSEvent window so that we don't introduce a crash on OS X when the event is processed after the window is closed.

diff --git a/Engine/Source/Runtime/Core/Private/Mac/MacApplication.cpp b/Engine/Source/Runtime/Core/Private/Mac/MacApplication.cpp
--- a/Engine/Source/Runtime/Core/Private/Mac/MacApplication.cpp
+++ b/Engine/Source/Runtime/Core/Private/Mac/MacApplication.cpp
@@ -830,6 +830,7 @@ void FMacApplication::ProcessNSEvent(NSEvent* const Event)
 		}
 	}
 
+	[Event ResetWindow];
 	bIsProcessingNSEvent = bWasProcessingNSEvent;
 }
 
diff --git a/Engine/Source/Runtime/Core/Private/Mac/MacEvent.cpp b/Engine/Source/Runtime/Core/Private/Mac/MacEvent.cpp
--- a/Engine/Source/Runtime/Core/Private/Mac/MacEvent.cpp
+++ b/Engine/Source/Runtime/Core/Private/Mac/MacEvent.cpp
@@ -9,12 +9,17 @@
 @implementation NSEvent (FCachedWindowAccess)
 -(void)CacheWindow
 {
-	_window = [self window];
+	_window = [[self window] retain];
 }
 -(NSWindow*)GetWindow
 {
 	return _window;
 }
+-(void)ResetWindow
+{
+	[_window release];
+	_window = nil;
+}
 @end
 
 FMacEvent::FMacEvent(NSEvent* const Event)
diff --git a/Engine/Source/Runtime/Core/Public/Mac/MacEvent.h b/Engine/Source/Runtime/Core/Public/Mac/MacEvent.h
--- a/Engine/Source/Runtime/Core/Public/Mac/MacEvent.h
+++ b/Engine/Source/Runtime/Core/Public/Mac/MacEvent.h
@@ -13,6 +13,7 @@ enum EMacEventSendMethod
 @interface NSEvent (FCachedWindowAccess)
 -(void)CacheWindow;
 -(NSWindow*)GetWindow;
+-(void)ResetWindow;
 @end
 
 class FMacEvent
