Fix going into fullscreen mode in the editor using the system titlebar button, which should be WindowedFullscreen mode for the editor & needs to set the target window mode to Windowed when exiting.
This reports the FCocoaWindow's WindowMode to Slate as that is the canonical version of the current state.

diff --git a/Engine/Source/Runtime/Core/Private/Mac/CocoaWindow.cpp b/Engine/Source/Runtime/Core/Private/Mac/CocoaWindow.cpp
--- a/Engine/Source/Runtime/Core/Private/Mac/CocoaWindow.cpp
+++ b/Engine/Source/Runtime/Core/Private/Mac/CocoaWindow.cpp
@@ -353,6 +353,13 @@ NSString* NSPerformDragOperation = @"NSPerformDragOperation";
 	if(self.TargetWindowMode == EWindowMode::Windowed)
 	{
 		self.TargetWindowMode = EWindowMode::Fullscreen;
+#if WITH_EDITORONLY_DATA
+		if(GIsEditor)
+		{
+			self.TargetWindowMode = EWindowMode::WindowedFullscreen;
+		}
+#endif
+		self.PreFullScreenRect = [self openGLFrame];
 	}
 }
 
@@ -374,6 +381,14 @@ NSString* NSPerformDragOperation = @"NSPerformDragOperation";
 	}
 }
 
+- (void)windowWillExitFullScreen:(NSNotification *)Notification
+{
+	if(self.TargetWindowMode != EWindowMode::Windowed)
+	{
+		self.TargetWindowMode = EWindowMode::Windowed;
+	}
+}
+
 - (void)windowDidExitFullScreen:(NSNotification*)Notification
 {
 	WindowMode = EWindowMode::Windowed;
diff --git a/Engine/Source/Runtime/Core/Private/Mac/MacWindow.cpp b/Engine/Source/Runtime/Core/Private/Mac/MacWindow.cpp
--- a/Engine/Source/Runtime/Core/Private/Mac/MacWindow.cpp
+++ b/Engine/Source/Runtime/Core/Private/Mac/MacWindow.cpp
@@ -96,7 +96,7 @@ void FMacWindow::Initialize( FMacApplication* const Application, const TSharedRe
 		if( WindowHandle != nullptr )
 		{
 			[WindowHandle setReleasedWhenClosed:NO];
-			[WindowHandle setWindowMode: WindowMode];
+			[WindowHandle setWindowMode: EWindowMode::Windowed];
 			[WindowHandle setAcceptsInput: Definition->AcceptsInput];
 			[WindowHandle setDisplayReconfiguring: false];
 			[WindowHandle setAcceptsMouseMovedEvents: YES];
@@ -206,7 +206,6 @@ void FMacWindow::Initialize( FMacApplication* const Application, const TSharedRe
 
 FMacWindow::FMacWindow()
 	: WindowHandle(NULL)
-	, WindowMode(EWindowMode::Windowed)
 	, bIsVisible(false)
 	, bIsClosed(false)
 {
@@ -227,7 +226,7 @@ void FMacWindow::ReshapeWindow( int32 X, int32 Y, int32 Width, int32 Height )
 		const TSharedRef<FGenericApplicationMessageHandler> MessageHandler = OwningApplication->MessageHandler;
 		MessageHandler->BeginReshapingWindow( SharedThis( this ) );
 		
-		if(WindowMode == EWindowMode::Windowed || WindowMode == EWindowMode::WindowedFullscreen)
+		if(GetWindowMode() == EWindowMode::Windowed || GetWindowMode() == EWindowMode::WindowedFullscreen)
 		{
 			const int32 InvertedY = FPlatformMisc::ConvertSlateYPositionToCocoa(Y) - Height + 1;
 			NSRect Rect = NSMakeRect(X, InvertedY, FMath::Max(Width, 1), FMath::Max(Height, 1));
@@ -240,13 +239,13 @@ void FMacWindow::ReshapeWindow( int32 X, int32 Y, int32 Width, int32 Height )
 			{
 				MainThreadCall(^{
 					SCOPED_AUTORELEASE_POOL;
-					BOOL DisplayIfNeeded = (WindowMode == EWindowMode::Windowed);
+					BOOL DisplayIfNeeded = (GetWindowMode() == EWindowMode::Windowed);
 					
 					[WindowHandle setFrame: Rect display:DisplayIfNeeded];
 					
 					// Force resize back to screen size in fullscreen - not ideally pretty but means we don't
 					// have to subvert the OS X or UE fullscreen handling events elsewhere.
-					if(WindowMode != EWindowMode::Windowed)
+					if(GetWindowMode() != EWindowMode::Windowed)
 					{
 						[WindowHandle setFrame: [WindowHandle screen].frame display:YES];
 					}
@@ -272,7 +271,7 @@ void FMacWindow::ReshapeWindow( int32 X, int32 Y, int32 Width, int32 Height )
 bool FMacWindow::GetFullScreenInfo( int32& X, int32& Y, int32& Width, int32& Height ) const
 {
 	SCOPED_AUTORELEASE_POOL;
-	bool const bIsFullscreen = (WindowMode == EWindowMode::Fullscreen);
+	bool const bIsFullscreen = (GetWindowMode() == EWindowMode::Fullscreen);
 	const NSRect Frame = (!bIsFullscreen) ? [WindowHandle screen].frame : PreFullscreenWindowRect;
 	X = Frame.origin.x;
 	Y = Frame.origin.y;
@@ -404,7 +403,7 @@ void FMacWindow::SetWindowMode( EWindowMode::Type NewWindowMode )
 
 	// In OS X fullscreen and windowed fullscreen are the same
 	bool bMakeFullscreen = NewWindowMode != EWindowMode::Windowed;
-	bool bIsFullscreen = WindowMode != EWindowMode::Windowed;
+	bool bIsFullscreen = GetWindowMode() != EWindowMode::Windowed;
 
 	if( bIsFullscreen != bMakeFullscreen )
 	{
@@ -445,11 +444,14 @@ void FMacWindow::SetWindowMode( EWindowMode::Type NewWindowMode )
 			Behaviour &= ~(NSWindowCollectionBehaviorFullScreenPrimary);
 			Behaviour |= NSWindowCollectionBehaviorFullScreenAuxiliary;
 		}
-		
-		WindowMode = NewWindowMode;
 	}
 }
 
+EWindowMode::Type FMacWindow::GetWindowMode() const
+{
+	return [WindowHandle windowMode];
+}
+
 bool FMacWindow::IsMaximized() const
 {
 	return WindowHandle->bZoomed;
diff --git a/Engine/Source/Runtime/Core/Public/Mac/MacWindow.h b/Engine/Source/Runtime/Core/Public/Mac/MacWindow.h
--- a/Engine/Source/Runtime/Core/Public/Mac/MacWindow.h
+++ b/Engine/Source/Runtime/Core/Public/Mac/MacWindow.h
@@ -49,7 +49,7 @@ public:
 
 	virtual void SetWindowMode( EWindowMode::Type NewWindowMode ) override;
 
-	virtual EWindowMode::Type GetWindowMode() const override { return WindowMode; } 
+	virtual EWindowMode::Type GetWindowMode() const override;
 
 	virtual bool IsMaximized() const override;
 
@@ -102,9 +102,6 @@ private:
 
 	/** Mac window handle */
 	FCocoaWindow* WindowHandle;
-	
-	/** The mode that the window is in (windowed, fullscreen, windowedfullscreen ) */
-	EWindowMode::Type WindowMode;
 
 	NSRect PreFullscreenWindowRect;
 
