Add option for platform to specify local paths to skip caching.
Fix CreateDirectory to support multiple drive letters on Xbox One.
Fix some full path stomping on Xbox One.
Log to developer scratch partition by default on Xbox One.

diff --git a/Engine/Source/Runtime/Core/Private/Misc/Paths.cpp b/Engine/Source/Runtime/Core/Private/Misc/Paths.cpp
--- a/Engine/Source/Runtime/Core/Private/Misc/Paths.cpp
+++ b/Engine/Source/Runtime/Core/Private/Misc/Paths.cpp
@@ -161,7 +161,7 @@ FString FPaths::VideoCaptureDir()
 
 FString FPaths::GameLogDir()
 {
-#if PLATFORM_MAC
+#if PLATFORM_MAC || PLATFORM_XBOXONE
 	return FPlatformProcess::UserLogsDir();
 #else
 	return FPaths::GameSavedDir() + TEXT("Logs/");
diff --git a/Engine/Source/Runtime/Core/Public/GenericPlatform/GenericPlatformFile.h b/Engine/Source/Runtime/Core/Public/GenericPlatform/GenericPlatformFile.h
--- a/Engine/Source/Runtime/Core/Public/GenericPlatform/GenericPlatformFile.h
+++ b/Engine/Source/Runtime/Core/Public/GenericPlatform/GenericPlatformFile.h
@@ -117,6 +117,11 @@ public:
 	 */
 	virtual void		InitializeAfterSetActive() { }
 
+	/**
+	 * Identifies any platform specific paths that are guaranteed to be local (i.e. cache, scratch space)
+	 */
+	virtual void		AddLocalDirectories(TArray<FString> &LocalDirectories) { }
+
 	/** Gets the platform file wrapped by this file. */
 	virtual IPlatformFile* GetLowerLevel() = 0;
 		/** Gets this platform file type name. */
diff --git a/Engine/Source/Runtime/NetworkFile/Private/NetworkPlatformFile.cpp b/Engine/Source/Runtime/NetworkFile/Private/NetworkPlatformFile.cpp
--- a/Engine/Source/Runtime/NetworkFile/Private/NetworkPlatformFile.cpp
+++ b/Engine/Source/Runtime/NetworkFile/Private/NetworkPlatformFile.cpp
@@ -114,6 +114,8 @@ bool FNetworkPlatformFile::InitializeInternal(IPlatformFile* Inner, const TCHAR*
 	LocalDirectories.Add(FPaths::GameSavedDir() / TEXT("Logs"));
 	LocalDirectories.Add(FPaths::GameSavedDir() / TEXT("Sandboxes"));
 
+	InnerPlatformFile->GetLowerLevel()->AddLocalDirectories(LocalDirectories);
+
 	FNetworkFileArchive Payload(NFS_Messages::Heartbeat); 
 	FArrayReader Out;
 	if (!SendPayloadAndReceiveResponse(Payload,Out))
