added help command to console

diff --git a/Engine/Source/Runtime/Core/Private/HAL/ConsoleManager.cpp b/Engine/Source/Runtime/Core/Private/HAL/ConsoleManager.cpp
--- a/Engine/Source/Runtime/Core/Private/HAL/ConsoleManager.cpp
+++ b/Engine/Source/Runtime/Core/Private/HAL/ConsoleManager.cpp
@@ -1477,6 +1477,32 @@ void CreateConsoleVariables()
 	}
 }
 
+#if !(UE_BUILD_SHIPPING || UE_BUILD_TEST)
+static void DumpHelp(FOutputDevice& Ar)
+{
+	Ar.Logf(TEXT("Console Help:"));
+	Ar.Logf(TEXT("============="));
+	Ar.Logf(TEXT(" "));
+	Ar.Logf(TEXT("A console variable is a engine wide key value pair. The key is a string usually starting with the subsystem prefix followed"));
+	Ar.Logf(TEXT("by '.' e.g. r.BloomQuality. The value can be of different tpe (e.g. float, int, string). A console command has no state associated with"));
+	Ar.Logf(TEXT("and gets executed immediately."));
+	Ar.Logf(TEXT(" "));
+	Ar.Logf(TEXT("Console variables can be put into ini files (e.g. ConsoleVariables.ini or BaseEngine.ini) with this syntax:"));
+	Ar.Logf(TEXT("<Console variable> = <value>"));
+	Ar.Logf(TEXT(" "));
+	Ar.Logf(TEXT("DumpConsoleCommands         Lists all console variables and commands that are registered (Some are not registered)"));
+	Ar.Logf(TEXT("<Console variable>          Get the console variable state"));
+	Ar.Logf(TEXT("<Console variable> ?        Get the console variable help text"));
+	Ar.Logf(TEXT("<Console variable> <value>  Set the console variable value"));
+	Ar.Logf(TEXT("<Console command> [Params]  Execute the console command with optional parameters"));
+}
+static FAutoConsoleCommandWithOutputDevice GConsoleCommandHelp(
+	TEXT("help"),
+	TEXT("Outputs some helptext to the console and the log"),
+	FConsoleCommandWithOutputDeviceDelegate::CreateStatic(DumpHelp)
+	);
+#endif // !(UE_BUILD_SHIPPING || UE_BUILD_TEST)
+
 // Naming conventions:
 //
 // Console variable should start with (suggestion):
