Linux: fix intermittent crash on startup (UE-4869).

- Caused by thread-unsafe code in the splash: SDL_ShowCursor() on the main thread was traversing created windows, while SDL_CreateWindow() on the splash thread has not yet returned - splash window was in a bad state.
- Moved all window creation to the main thread, leaving only GL stuff on the splash one.
- Also minor fixes and clean ups along the road.

diff --git a/Engine/Source/Runtime/Core/Private/Linux/LinuxPlatformSplash.cpp b/Engine/Source/Runtime/Core/Private/Linux/LinuxPlatformSplash.cpp
--- a/Engine/Source/Runtime/Core/Private/Linux/LinuxPlatformSplash.cpp
+++ b/Engine/Source/Runtime/Core/Private/Linux/LinuxPlatformSplash.cpp
@@ -68,7 +68,7 @@ static FString GIconPath;
 
 static int32 SplashWidth = 0, SplashHeight = 0;
 static unsigned char *ScratchSpace = nullptr;
-static int32 ThreadState = 0;
+static volatile int32 ThreadState = 0;
 static int32 SplashBPP = 0;
 
 //////////////////////////////////
@@ -490,19 +490,24 @@ static int RenderString (GLuint tex_idx)
 	return 0;
 }
 
-
-/**
- * Thread function that actually creates the window and
- * renders its contents.
- */
-static int StartSplashScreenThread(void *ptr)
-{	
-	//	init the sdl here
-	if (SDL_WasInit( SDL_INIT_VIDEO ) == 0)
+/** Helper function to init resources used by the splash thread */
+bool LinuxSplash_InitSplashResources()
+{
+	if (!FPlatformMisc::PlatformInitMultimedia()) //	will not initialize more than once
 	{
-		SDL_InitSubSystem(SDL_INIT_VIDEO);
+		UE_LOG(LogInit, Warning, TEXT("LinuxSplash_InitSplashResources() : PlatformInitMultimedia() failed, there will be no splash."));
+		return false;
 	}
+
+	// load splash .bmp image
+	GSplashScreenImage = SDL_LoadBMP(TCHAR_TO_UTF8(*GSplashPath));
 	
+	if (GSplashScreenImage == nullptr)
+	{
+		UE_LOG(LogHAL, Warning, TEXT("LinuxSplash_InitSplashResources() : Could not load splash BMP! SDL_Error: %s"), UTF8_TO_TCHAR(SDL_GetError()));
+		return false;
+	}
+
 	// create splash window
 	GSplashWindow = SDL_CreateWindow(NULL,
 									SDL_WINDOWPOS_CENTERED,
@@ -511,9 +516,9 @@ static int StartSplashScreenThread(void *ptr)
 									GSplashScreenImage->h,
 									SDL_WINDOW_BORDERLESS | SDL_WINDOW_OPENGL | SDL_WINDOW_SHOWN
 									);
-	if(GSplashWindow == nullptr)
+	if (GSplashWindow == nullptr)
 	{
-		UE_LOG(LogHAL, Error, TEXT("Splash screen window could not be created! SDL_Error: %s"), UTF8_TO_TCHAR(SDL_GetError()));
+		UE_LOG(LogHAL, Error, TEXT("LinuxSplash_InitSplashResources() : Splash screen window could not be created! SDL_Error: %s"), UTF8_TO_TCHAR(SDL_GetError()));
 		return -1;
 	}
 
@@ -527,9 +532,34 @@ static int StartSplashScreenThread(void *ptr)
 	}
 	else
 	{
-		UE_LOG(LogHAL, Error, TEXT("Splash icon could not be created! SDL_Error: %s"), UTF8_TO_TCHAR(SDL_GetError()));
+		UE_LOG(LogHAL, Warning, TEXT("LinuxSplash_InitSplashResources() : Splash icon could not be created! SDL_Error: %s"), UTF8_TO_TCHAR(SDL_GetError()));
+		return false;
 	}
 
+	return true;
+}
+
+/** Helper function to tear down resources used by the splash thread */
+void LinuxSplash_TearDownSplashResources()
+{
+	SDL_DestroyWindow(GSplashWindow);
+	GSplashWindow = nullptr;
+
+	SDL_FreeSurface(GSplashIconImage);
+	GSplashIconImage = nullptr;
+
+	SDL_FreeSurface(GSplashScreenImage);
+	GSplashScreenImage = nullptr;
+
+	// do not deinit SDL here
+}
+
+/**
+ * Thread function that actually creates the window and
+ * renders its contents.
+ */
+static int StartSplashScreenThread(void *ptr)
+{	
 	SDL_GL_SetAttribute( SDL_GL_CONTEXT_MAJOR_VERSION, 3 );
 	SDL_GL_SetAttribute( SDL_GL_CONTEXT_MINOR_VERSION, 2 );
 	SDL_GL_SetAttribute( SDL_GL_CONTEXT_PROFILE_MASK, SDL_GL_CONTEXT_PROFILE_CORE);
@@ -727,6 +757,8 @@ static int StartSplashScreenThread(void *ptr)
 	// clean up
 	glDeleteBuffers(1, &VertexBuffer);
 	glDeleteVertexArrays(1, &VertexArrayObject);
+	glDetachShader(ShaderProgram, VertexShader);
+	glDetachShader(ShaderProgram, FragmentShader);
 	glDeleteShader(VertexShader);
 	glDeleteShader(FragmentShader);
 	glDeleteProgram(ShaderProgram);
@@ -755,9 +787,11 @@ static int StartSplashScreenThread(void *ptr)
 	FT_Done_FreeType(FontLibrary);
 
 	SDL_GL_DeleteContext(Context);
-	SDL_DestroyWindow(GSplashWindow);
 	FMemory::Free(ScratchSpace);
 
+	// set the thread state to 0 to let the caller know we're done (FIXME: can be done without busy-loops)
+	ThreadState = 0;
+
 	return 0;
 }
 
@@ -879,12 +913,9 @@ void FLinuxPlatformSplash::Show( )
 		}
 	}
 
-	// load splash .bmp image
-	GSplashScreenImage = SDL_LoadBMP(TCHAR_TO_UTF8(*GSplashPath));
-	
-	if (GSplashScreenImage == nullptr)
+	// init splash LinuxSplash_InitSplashResources
+	if (!LinuxSplash_InitSplashResources())
 	{
-		UE_LOG(LogHAL, Warning, TEXT("Could not load splash BMP! SDL_Error: %s"), UTF8_TO_TCHAR(SDL_GetError()));
 		return;
 	}
 
@@ -912,7 +943,20 @@ void FLinuxPlatformSplash::Hide()
 #if WITH_EDITOR
 	// signal thread it's time to quit
 	GSplashThread = nullptr;
-	ThreadState = -99;
+
+	if (ThreadState > 0)	// if there's a thread at all...
+	{
+		ThreadState = -99;
+		// wait for the thread to be done before tearing it tearing it down (it will set the ThreadState to 0)
+		while (ThreadState != 0)
+		{
+			// busy loop!
+			FPlatformProcess::Sleep(0.01f);
+		}
+
+		// tear down resources that thread used
+		LinuxSplash_TearDownSplashResources();
+	}
 #endif
 }
 
