Fixed reference casts causing construction in TOptional

diff --git a/Engine/Source/Runtime/Core/Public/Misc/Optional.h b/Engine/Source/Runtime/Core/Public/Misc/Optional.h
--- a/Engine/Source/Runtime/Core/Public/Misc/Optional.h
+++ b/Engine/Source/Runtime/Core/Public/Misc/Optional.h
@@ -39,7 +39,7 @@ public:
 	{
 		if (InValue.bIsSet)
 		{
-			new(&Value) OptionalType((const OptionalType&)InValue.Value);
+			new(&Value) OptionalType(*(const OptionalType*)&InValue.Value);
 			bIsSet = true;
 		}
 	}
@@ -48,7 +48,7 @@ public:
 	{
 		if (InValue.bIsSet)
 		{
-			new(&Value) OptionalType(MoveTemp((OptionalType&)InValue.Value));
+			new(&Value) OptionalType(MoveTemp(*(OptionalType*)&InValue.Value));
 			bIsSet = true;
 		}
 	}
@@ -60,7 +60,7 @@ public:
 			Reset();
 			if (InValue.bIsSet)
 			{
-				new(&Value) OptionalType((const OptionalType&)InValue.Value);
+				new(&Value) OptionalType(*(const OptionalType*)&InValue.Value);
 				bIsSet = true;
 			}
 		}
@@ -73,7 +73,7 @@ public:
 			Reset();
 			if (InValue.bIsSet)
 			{
-				new(&Value) OptionalType(MoveTemp((OptionalType&)InValue.Value));
+				new(&Value) OptionalType(MoveTemp(*(OptionalType*)&InValue.Value));
 				bIsSet = true;
 			}
 		}
@@ -82,7 +82,7 @@ public:
 
 	TOptional& operator=(const OptionalType& InValue)
 	{
-		if (&InValue != &(OptionalType&)Value)
+		if (&InValue != (OptionalType*)&Value)
 		{
 			Reset();
 			new(&Value) OptionalType(InValue);
@@ -92,7 +92,7 @@ public:
 	}
 	TOptional& operator=(OptionalType&& InValue)
 	{
-		if (&InValue != &(OptionalType&)Value)
+		if (&InValue != (OptionalType*)&Value)
 		{
 			Reset();
 			new(&Value) OptionalType(MoveTemp(InValue));
@@ -109,7 +109,7 @@ public:
 
 			// We need a typedef here because VC won't compile the destructor call below if OptionalType itself has a member called OptionalType
 			typedef OptionalType OptionalDestructOptionalType;
-			((OptionalType&)Value).OptionalDestructOptionalType::~OptionalDestructOptionalType();
+			(*(OptionalType*)&Value).OptionalDestructOptionalType::~OptionalDestructOptionalType();
 		}
 	}
 
@@ -172,7 +172,7 @@ public:
 		{
 			return true;
 		}
-		return (OptionalType&)lhs.Value == (OptionalType&)rhs.Value;
+		return (OptionalType*)&lhs.Value == (OptionalType*)&rhs.Value;
 	}
 	friend bool operator!=(const TOptional& lhs, const TOptional& rhs)
 	{
@@ -184,14 +184,14 @@ public:
 	FORCEINLINE_EXPLICIT_OPERATOR_BOOL() const { return bIsSet; }
 
 	/** @return The optional value; undefined when IsSet() returns false. */
-	const OptionalType& GetValue() const { checkf(IsSet(), TEXT("It is an error to call GetValue() on an unset TOptional. Please either check IsSet() or use Get(DefaultValue) instead.")); return (OptionalType&)Value; }
-	      OptionalType& GetValue()       { checkf(IsSet(), TEXT("It is an error to call GetValue() on an unset TOptional. Please either check IsSet() or use Get(DefaultValue) instead.")); return (OptionalType&)Value; }
+	const OptionalType& GetValue() const { checkf(IsSet(), TEXT("It is an error to call GetValue() on an unset TOptional. Please either check IsSet() or use Get(DefaultValue) instead.")); return *(OptionalType*)&Value; }
+	      OptionalType& GetValue()       { checkf(IsSet(), TEXT("It is an error to call GetValue() on an unset TOptional. Please either check IsSet() or use Get(DefaultValue) instead.")); return *(OptionalType*)&Value; }
 
 	const OptionalType* operator->() const { return &GetValue(); }
 	      OptionalType* operator->()       { return &GetValue(); }
 
 	/** @return The optional value when set; DefaultValue otherwise. */
-	const OptionalType& Get(const OptionalType& DefaultValue) const { return IsSet() ? (OptionalType&)Value : DefaultValue; }
+	const OptionalType& Get(const OptionalType& DefaultValue) const { return IsSet() ? *(OptionalType*)&Value : DefaultValue; }
 
 private:
 	bool bIsSet;
