Fix for UE-5042 (Mac editor does not get full focus if the user clicks on the menu bar)

diff --git a/Engine/Source/Runtime/Core/Private/Mac/MacApplication.cpp b/Engine/Source/Runtime/Core/Private/Mac/MacApplication.cpp
--- a/Engine/Source/Runtime/Core/Private/Mac/MacApplication.cpp
+++ b/Engine/Source/Runtime/Core/Private/Mac/MacApplication.cpp
@@ -44,41 +44,37 @@ FMacApplication* FMacApplication::CreateMacApplication()
 
 NSEvent* FMacApplication::HandleNSEvent(NSEvent* Event)
 {
-	NSEvent* ReturnEvent = nil;
-	
-	if ([Event windowNumber] == 0 || [Event window] != nil)
+	NSEvent* ReturnEvent = Event;
+
+	const bool bIsMouseClickOrKeyEvent = [Event type] == NSLeftMouseDown || [Event type] == NSLeftMouseUp
+		|| [Event type] == NSRightMouseDown || [Event type] == NSRightMouseUp
+		|| [Event type] == NSOtherMouseDown || [Event type] == NSOtherMouseUp;
+	const bool bIsResentEvent = [Event type] == NSApplicationDefined && (FMacApplicationEventTypes)[Event subtype] == FMacApplication::ResentEvent;
+
+	if (MacApplication)
 	{
-		ReturnEvent = Event;
-		const bool bIsMouseClickOrKeyEvent = [Event type] == NSLeftMouseDown || [Event type] == NSLeftMouseUp
-			|| [Event type] == NSRightMouseDown || [Event type] == NSRightMouseUp
-			|| [Event type] == NSOtherMouseDown || [Event type] == NSOtherMouseUp;
-		const bool bIsResentEvent = [Event type] == NSApplicationDefined && (FMacApplicationEventTypes)[Event subtype] == FMacApplication::ResentEvent;
-		
-		if (MacApplication)
+		if (bIsResentEvent)
 		{
-			if ( bIsResentEvent )
-			{
-				ReturnEvent = (NSEvent*)[Event data1];
-			}
-			
-			if ( !bIsResentEvent && ( !bIsMouseClickOrKeyEvent || [Event window] == NULL ) )
-			{
-				FMacEvent::SendToGameRunLoop(Event, EMacEventSendMethod::Async);
-				
-				if ( [Event type] == NSKeyDown || [Event type] == NSKeyUp )
-				{
-					ReturnEvent = nil;
-				}
-			}
+			ReturnEvent = (NSEvent*)[Event data1];
+		}
+
+		if (!bIsResentEvent && (!bIsMouseClickOrKeyEvent || [Event window] == NULL))
+		{
+			FMacEvent::SendToGameRunLoop(Event, EMacEventSendMethod::Async);
 			
-			if ([Event type] == NSLeftMouseUp)
+			if ([Event type] == NSKeyDown || [Event type] == NSKeyUp)
 			{
-				NSNotification* Notification = [NSNotification notificationWithName:NSWindowDraggingFinished object:[Event window]];
-				FMacEvent::SendToGameRunLoop(Notification, EMacEventSendMethod::Async);
+				ReturnEvent = nil;
 			}
 		}
+
+		if ([Event type] == NSLeftMouseUp)
+		{
+			NSNotification* Notification = [NSNotification notificationWithName:NSWindowDraggingFinished object:[Event window]];
+			FMacEvent::SendToGameRunLoop(Notification, EMacEventSendMethod::Async);
+		}
 	}
-	
+
 	return ReturnEvent;
 }
 
