Corrected comparison operator in TOptional and removed redundant dereference

diff --git a/Engine/Source/Runtime/Core/Public/Misc/Optional.h b/Engine/Source/Runtime/Core/Public/Misc/Optional.h
--- a/Engine/Source/Runtime/Core/Public/Misc/Optional.h
+++ b/Engine/Source/Runtime/Core/Public/Misc/Optional.h
@@ -109,7 +109,7 @@ public:
 
 			// We need a typedef here because VC won't compile the destructor call below if OptionalType itself has a member called OptionalType
 			typedef OptionalType OptionalDestructOptionalType;
-			(*(OptionalType*)&Value).OptionalDestructOptionalType::~OptionalDestructOptionalType();
+			((OptionalType*)&Value)->OptionalDestructOptionalType::~OptionalDestructOptionalType();
 		}
 	}
 
@@ -172,7 +172,7 @@ public:
 		{
 			return true;
 		}
-		return (OptionalType*)&lhs.Value == (OptionalType*)&rhs.Value;
+		return (*(OptionalType*)&lhs.Value) == (*(OptionalType*)&rhs.Value);
 	}
 	friend bool operator!=(const TOptional& lhs, const TOptional& rhs)
 	{
