Implement PS4 Affinity interface

diff --git a/Engine/Source/Runtime/Core/Public/HAL/PlatformAffinity.h b/Engine/Source/Runtime/Core/Public/HAL/PlatformAffinity.h
--- a/Engine/Source/Runtime/Core/Public/HAL/PlatformAffinity.h
+++ b/Engine/Source/Runtime/Core/Public/HAL/PlatformAffinity.h
@@ -5,6 +5,8 @@
 
 #if PLATFORM_XBOXONE
 #include "XboxOne/XBoxOneAffinity.h"
+#elif PLATFORM_PS4
+#include "PS4/PS4Affinity.h"
 #elif PLATFORM_ANDROID
 #include "Android/AndroidAffinity.h"
 #else
