Cache NumberOfCores*() functions on Linux.

- The number of cores is unlikely to change (barring an edge case of affinity mask changing or CPUs going offline, which arguably shouldn't be accounted for).
- Tiny performance improvement since NumberOfCores() call was taking ~0.30% time of UT4 Linux client (UEPLAT-232).

diff --git a/Engine/Source/Runtime/Core/Private/Linux/LinuxPlatformMisc.cpp b/Engine/Source/Runtime/Core/Private/Linux/LinuxPlatformMisc.cpp
--- a/Engine/Source/Runtime/Core/Private/Linux/LinuxPlatformMisc.cpp
+++ b/Engine/Source/Runtime/Core/Private/Linux/LinuxPlatformMisc.cpp
@@ -487,63 +487,80 @@ EAppReturnType::Type FLinuxPlatformMisc::MessageBoxExt(EAppMsgType::Type MsgType
 
 int32 FLinuxPlatformMisc::NumberOfCores()
 {
-	cpu_set_t AvailableCpusMask;
-	CPU_ZERO(&AvailableCpusMask);
-
-	if (0 != sched_getaffinity(0, sizeof(AvailableCpusMask), &AvailableCpusMask))
+	// WARNING: this function ignores edge cases like affinity mask changes (and even more fringe cases like CPUs going offline)
+	// in the name of performance (higher level code calls NumberOfCores() way too often...)
+	static int32 NumCoreIds = 0;
+	if (NumCoreIds == 0)
 	{
-		return 1;	// we are running on something, right?
-	}
-
-	char FileNameBuffer[1024];
-	unsigned char PossibleCores[CPU_SETSIZE] = { 0 };
+		cpu_set_t AvailableCpusMask;
+		CPU_ZERO(&AvailableCpusMask);
 
-	for(int32 CpuIdx = 0; CpuIdx < CPU_SETSIZE; ++CpuIdx)
-	{
-		if (CPU_ISSET(CpuIdx, &AvailableCpusMask))
+		if (0 != sched_getaffinity(0, sizeof(AvailableCpusMask), &AvailableCpusMask))
 		{
-			sprintf(FileNameBuffer, "/sys/devices/system/cpu/cpu%d/topology/core_id", CpuIdx);
-			
-			FILE* CoreIdFile = fopen(FileNameBuffer, "r");
-			unsigned int CoreId = 0;
-			if (CoreIdFile)
+			NumCoreIds = 1;	// we are running on something, right?
+		}
+		else
+		{
+			char FileNameBuffer[1024];
+			unsigned char PossibleCores[CPU_SETSIZE] = { 0 };
+
+			for(int32 CpuIdx = 0; CpuIdx < CPU_SETSIZE; ++CpuIdx)
 			{
-				if (1 != fscanf(CoreIdFile, "%d", &CoreId))
+				if (CPU_ISSET(CpuIdx, &AvailableCpusMask))
 				{
-					CoreId = 0;
+					sprintf(FileNameBuffer, "/sys/devices/system/cpu/cpu%d/topology/core_id", CpuIdx);
+					
+					FILE* CoreIdFile = fopen(FileNameBuffer, "r");
+					unsigned int CoreId = 0;
+					if (CoreIdFile)
+					{
+						if (1 != fscanf(CoreIdFile, "%d", &CoreId))
+						{
+							CoreId = 0;
+						}
+						fclose(CoreIdFile);
+					}
+
+					if (CoreId >= ARRAY_COUNT(PossibleCores))
+					{
+						CoreId = 0;
+					}
+					
+					PossibleCores[ CoreId ] = 1;
 				}
-				fclose(CoreIdFile);
 			}
 
-			if (CoreId >= ARRAY_COUNT(PossibleCores))
+			for(int32 Idx = 0; Idx < ARRAY_COUNT(PossibleCores); ++Idx)
 			{
-				CoreId = 0;
+				NumCoreIds += PossibleCores[Idx];
 			}
-			
-			PossibleCores[ CoreId ] = 1;
 		}
 	}
 
-	int32 NumCoreIds = 0;
-	for(int32 Idx = 0; Idx < ARRAY_COUNT(PossibleCores); ++Idx)
-	{
-		NumCoreIds += PossibleCores[Idx];
-	}
-
 	return NumCoreIds;
 }
 
 int32 FLinuxPlatformMisc::NumberOfCoresIncludingHyperthreads()
 {
-	cpu_set_t AvailableCpusMask;
-	CPU_ZERO(&AvailableCpusMask);
-
-	if (0 != sched_getaffinity(0, sizeof(AvailableCpusMask), &AvailableCpusMask))
+	// WARNING: this function ignores edge cases like affinity mask changes (and even more fringe cases like CPUs going offline)
+	// in the name of performance (higher level code calls NumberOfCores() way too often...)
+	static int32 NumCoreIds = 0;
+	if (NumCoreIds == 0)
 	{
-		return 1;	// we are running on something, right?
+		cpu_set_t AvailableCpusMask;
+		CPU_ZERO(&AvailableCpusMask);
+
+		if (0 != sched_getaffinity(0, sizeof(AvailableCpusMask), &AvailableCpusMask))
+		{
+			NumCoreIds = 1;	// we are running on something, right?
+		}
+		else
+		{
+			return CPU_COUNT(&AvailableCpusMask);
+		}
 	}
 
-	return CPU_COUNT(&AvailableCpusMask);
+	return NumCoreIds;
 }
 
 void FLinuxPlatformMisc::LoadPreInitModules()
