Added an additional optional param to OsExecute so that cmdline params can easily be passed to the ShellExecute function

#nonjira

#branch UE4

diff --git a/Engine/Source/Runtime/Core/Private/Windows/WindowsPlatformMisc.cpp b/Engine/Source/Runtime/Core/Private/Windows/WindowsPlatformMisc.cpp
--- a/Engine/Source/Runtime/Core/Private/Windows/WindowsPlatformMisc.cpp
+++ b/Engine/Source/Runtime/Core/Private/Windows/WindowsPlatformMisc.cpp
@@ -1782,12 +1782,12 @@ void FWindowsPlatformMisc::LoadStartupModules()
 #endif	//WITH_EDITOR
 }
 
-bool FWindowsPlatformMisc::OsExecute(const TCHAR* CommandType, const TCHAR* Command)
+bool FWindowsPlatformMisc::OsExecute(const TCHAR* CommandType, const TCHAR* Command, const TCHAR* CommandLine)
 {
 	HINSTANCE hApp = ShellExecute(NULL,
 		CommandType,
 		Command,
-		NULL,
+		CommandLine,
 		NULL,
 		SW_SHOWNORMAL);
 	bool bSucceeded = hApp > (HINSTANCE)32;
diff --git a/Engine/Source/Runtime/Core/Public/GenericPlatform/GenericPlatformMisc.h b/Engine/Source/Runtime/Core/Public/GenericPlatform/GenericPlatformMisc.h
--- a/Engine/Source/Runtime/Core/Public/GenericPlatform/GenericPlatformMisc.h
+++ b/Engine/Source/Runtime/Core/Public/GenericPlatform/GenericPlatformMisc.h
@@ -680,10 +680,11 @@ public:
 	 *
 	 * @param ComandType OS hint as to the type of command 
 	 * @param Command the command to execute
+	 * @param CommandLine the commands to pass to the executable
 	 *
 	 * @return whether the command was successful or not
 	 */
-	static bool OsExecute(const TCHAR* CommandType, const TCHAR* Command)
+	static bool OsExecute(const TCHAR* CommandType, const TCHAR* Command, const TCHAR* CommandLine = NULL)
 	{
 		return false;
 	}
diff --git a/Engine/Source/Runtime/Core/Public/Windows/WindowsPlatformMisc.h b/Engine/Source/Runtime/Core/Public/Windows/WindowsPlatformMisc.h
--- a/Engine/Source/Runtime/Core/Public/Windows/WindowsPlatformMisc.h
+++ b/Engine/Source/Runtime/Core/Public/Windows/WindowsPlatformMisc.h
@@ -87,10 +87,11 @@ struct CORE_API FWindowsPlatformMisc : public FGenericPlatformMisc
 	 *
 	 * @param ComandType OS hint as to the type of command 
 	 * @param Command the command to execute
+	 * @param CommandLine the commands to pass to the executable
 	 *
 	 * @return whether the command was successful or not
 	 */
-	static bool OsExecute(const TCHAR* CommandType, const TCHAR* Command);
+	static bool OsExecute(const TCHAR* CommandType, const TCHAR* Command, const TCHAR* CommandLine = NULL);
 
 	/**
 	 * Searches for a window that matches the window name or the title starts with a particular text. When
