Use the max of the left and right large/small motor values as the value to route to the motor instead of ignoring Right large and Left small motor values.

diff --git a/Engine/Source/Runtime/Core/Private/Windows/XInputInterface.cpp b/Engine/Source/Runtime/Core/Private/Windows/XInputInterface.cpp
--- a/Engine/Source/Runtime/Core/Private/Windows/XInputInterface.cpp
+++ b/Engine/Source/Runtime/Core/Private/Windows/XInputInterface.cpp
@@ -211,8 +211,11 @@ void XInputInterface::SendControllerEvents()
  				// apply force feedback
  				XINPUT_VIBRATION VibrationState;
  
- 				VibrationState.wLeftMotorSpeed = ( ::WORD ) ( ControllerState.ForceFeedback.LeftLarge * 65535.0f );
- 				VibrationState.wRightMotorSpeed = ( ::WORD ) ( ControllerState.ForceFeedback.RightSmall * 65535.0f );
+				const float LargeValue = (ControllerState.ForceFeedback.LeftLarge > ControllerState.ForceFeedback.RightLarge ? ControllerState.ForceFeedback.LeftLarge : ControllerState.ForceFeedback.RightLarge);
+				const float SmallValue = (ControllerState.ForceFeedback.LeftSmall > ControllerState.ForceFeedback.RightSmall ? ControllerState.ForceFeedback.LeftSmall : ControllerState.ForceFeedback.RightSmall);
+
+				VibrationState.wLeftMotorSpeed = ( ::WORD ) ( LargeValue * 65535.0f );
+ 				VibrationState.wRightMotorSpeed = ( ::WORD ) ( SmallValue * 65535.0f );
  
 				XInputSetState( ( ::DWORD ) ControllerState.ControllerId, &VibrationState );			
 			}
