Don't call OnWindowActivationChanged when closing non-key window on Mac

diff --git a/Engine/Source/Runtime/Core/Private/Mac/MacApplication.cpp b/Engine/Source/Runtime/Core/Private/Mac/MacApplication.cpp
--- a/Engine/Source/Runtime/Core/Private/Mac/MacApplication.cpp
+++ b/Engine/Source/Runtime/Core/Private/Mac/MacApplication.cpp
@@ -1274,7 +1274,10 @@ void FMacApplication::OnWindowDidClose( FCocoaWindow* Window )
 	if( EventWindow.IsValid() )
 	{
 		SCOPED_AUTORELEASE_POOL;
-		MessageHandler->OnWindowActivationChanged( EventWindow.ToSharedRef(), EWindowActivation::Deactivate );
+		if ([Window isKeyWindow])
+		{
+			MessageHandler->OnWindowActivationChanged( EventWindow.ToSharedRef(), EWindowActivation::Deactivate );
+		}
 		Windows.Remove( EventWindow.ToSharedRef() );
 		KeyWindows.Remove( EventWindow.ToSharedRef() );
 		MessageHandler->OnWindowClose( EventWindow.ToSharedRef() );
@@ -1286,7 +1289,10 @@ bool FMacApplication::OnWindowDestroyed( FCocoaWindow* Window )
 	TSharedPtr< FMacWindow > EventWindow = FindWindowByNSWindow( Windows, &WindowsMutex, Window );
 	if( EventWindow.IsValid() )
 	{
-		MessageHandler->OnWindowActivationChanged( EventWindow.ToSharedRef(), EWindowActivation::Deactivate );
+		if ([Window isKeyWindow])
+		{
+			MessageHandler->OnWindowActivationChanged( EventWindow.ToSharedRef(), EWindowActivation::Deactivate );
+		}
 		Windows.Remove( EventWindow.ToSharedRef() );
 		KeyWindows.Remove( EventWindow.ToSharedRef() );
 		return true;
