Fixed setting of RunningOnMavericks variable in MacApplicationInfo

diff --git a/Engine/Source/Runtime/Core/Private/Mac/MacPlatformMisc.cpp b/Engine/Source/Runtime/Core/Private/Mac/MacPlatformMisc.cpp
--- a/Engine/Source/Runtime/Core/Private/Mac/MacPlatformMisc.cpp
+++ b/Engine/Source/Runtime/Core/Private/Mac/MacPlatformMisc.cpp
@@ -55,7 +55,7 @@ struct MacApplicationInfo
 		{
 			TTypeFromString<uint8>::FromString(ComponentValues[i], *Components[i]);
 		}
-		RunningOnMavericks == ComponentValues[0] == 10 && ComponentValues[1] == 9;
+		RunningOnMavericks = ComponentValues[0] == 10 && ComponentValues[1] == 9;
 		
 		char TempSysCtlBuffer[PATH_MAX] = {};
 		size_t TempSysCtlBufferSize = PATH_MAX;
