Fix for FArchive::SerializeCompressed serializing the wrong size of the compressed buffer

PR-598, UE-5402

diff --git a/Engine/Source/Runtime/Core/Private/Serialization/Archive.cpp b/Engine/Source/Runtime/Core/Private/Serialization/Archive.cpp
--- a/Engine/Source/Runtime/Core/Private/Serialization/Archive.cpp
+++ b/Engine/Source/Runtime/Core/Private/Serialization/Archive.cpp
@@ -646,6 +646,7 @@ void FArchive::SerializeCompressed( void* V, int64 Length, ECompressionFlags Fla
 			check(CompressedSize < INT_MAX);
 			int32 CompressedSizeInt = (int32)CompressedSize;
 			verify( FCompression::CompressMemory( Flags, CompressedBuffer, CompressedSizeInt, Src, BytesToCompress ) );
+			CompressedSize = CompressedSizeInt;
 			// move to next chunk if not reading from file
 			if (!bTreatBufferAsFileReader)
 			{
