Fixes for StrategyGame fullscreen UI on HTML5.
Added missing display metrics resize.
Added compile and directory iteration fixes for HTML5-win32

diff --git a/Engine/Source/Runtime/Core/Private/HTML5/HTML5Application.cpp b/Engine/Source/Runtime/Core/Private/HTML5/HTML5Application.cpp
--- a/Engine/Source/Runtime/Core/Private/HTML5/HTML5Application.cpp
+++ b/Engine/Source/Runtime/Core/Private/HTML5/HTML5Application.cpp
@@ -104,6 +104,11 @@ void FHTML5Application::PollGameDeviceState( const float TimeDelta )
 
 						IConsoleManager::Get().ProcessUserConsoleInput(ANSI_TO_TCHAR(command), *GWarn, NULL );
 					}
+
+					// Slate needs to know when desktop size changes.
+					FDisplayMetrics DisplayMetrics;
+					FDisplayMetrics::GetDisplayMetrics(DisplayMetrics);
+					BroadcastDisplayMetricsChanged(DisplayMetrics);
 				}
 				break;
 #if PLATFORM_HTML5_WIN32
diff --git a/Engine/Source/Runtime/Core/Private/HTML5/HTML5PlatformFile.cpp b/Engine/Source/Runtime/Core/Private/HTML5/HTML5PlatformFile.cpp
--- a/Engine/Source/Runtime/Core/Private/HTML5/HTML5PlatformFile.cpp
+++ b/Engine/Source/Runtime/Core/Private/HTML5/HTML5PlatformFile.cpp
@@ -388,8 +388,9 @@ public:
 #else 
 		bool Result = false;
 		// If Directory is an empty string, assume that we want to iterate Binaries/Mac (current dir), but because we're an app bundle, iterate bundle's Contents/Frameworks instead
+		FString FilePath = FPaths::Combine(*NormalizeFilename(Directory), TEXT("*"));
 		struct _finddata_t c_file;
-		long hFile = _findfirst(TCHAR_TO_UTF8(*NormalizeFilename(Directory)),&c_file);
+		long hFile = _findfirst(TCHAR_TO_UTF8(*FilePath), &c_file);
 
 		if (hFile  != -1L )
 		{
diff --git a/Engine/Source/Runtime/Core/Private/HTML5/HTML5PlatformMisc.cpp b/Engine/Source/Runtime/Core/Private/HTML5/HTML5PlatformMisc.cpp
--- a/Engine/Source/Runtime/Core/Private/HTML5/HTML5PlatformMisc.cpp
+++ b/Engine/Source/Runtime/Core/Private/HTML5/HTML5PlatformMisc.cpp
@@ -13,7 +13,7 @@
 
 #include "unicode/locid.h"
 #include "GenericPlatformContext.h"
-#include <SDL.h>
+#include <SDL/SDL.h>
 #include <ctime>
 
 void FHTML5Misc::PlatformInit()
diff --git a/Engine/Source/Runtime/Engine/Public/SaveGameSystem.h b/Engine/Source/Runtime/Engine/Public/SaveGameSystem.h
--- a/Engine/Source/Runtime/Engine/Public/SaveGameSystem.h
+++ b/Engine/Source/Runtime/Engine/Public/SaveGameSystem.h
@@ -84,6 +84,8 @@ public:
 #elif PLATFORM_HTML5_WIN32
 		FILE *fp;
 		fp=fopen("c:\\test.sav","rb");
+		if (!fp)
+			return false;
 			// obtain file size:
 		fseek (fp, 0 , SEEK_END);
 		int size = ftell (fp);
