Minor fixes (unitialized variables).

- found by valgrind.

diff --git a/Engine/Source/Runtime/Core/Private/Linux/LinuxApplication.cpp b/Engine/Source/Runtime/Core/Private/Linux/LinuxApplication.cpp
--- a/Engine/Source/Runtime/Core/Private/Linux/LinuxApplication.cpp
+++ b/Engine/Source/Runtime/Core/Private/Linux/LinuxApplication.cpp
@@ -67,6 +67,8 @@ FLinuxApplication::FLinuxApplication() : GenericApplication( MakeShareable( new
 #if STEAM_CONTROLLER_SUPPORT
 	, SteamInput( SteamControllerInterface::Create(MessageHandler) )
 #endif // STEAM_CONTROLLER_SUPPORT
+	, bIsMouseCursorLocked(false)
+	, bIsMouseCaptureEnabled(false)
 	, bHasLoadedInputPlugins(false)
 {
 	bUsingHighPrecisionMouseInput = false;
