fix for async io thread being starved when running on a device with only one core
UE-5382
#ios

diff --git a/Engine/Source/Runtime/Core/Private/HAL/PThreadRunnableThread.h b/Engine/Source/Runtime/Core/Private/HAL/PThreadRunnableThread.h
--- a/Engine/Source/Runtime/Core/Private/HAL/PThreadRunnableThread.h
+++ b/Engine/Source/Runtime/Core/Private/HAL/PThreadRunnableThread.h
@@ -66,6 +66,7 @@ protected:
 		switch (Priority)
 		{
 			// 0 is the lowest, 31 is the highest possible priority for pthread
+            case TPri_Highest: return 30;
 			case TPri_AboveNormal: return 25;
 			case TPri_Normal: return 15;
 			case TPri_BelowNormal: return 5;
diff --git a/Engine/Source/Runtime/Core/Private/Serialization/AsyncIOSystemBase.cpp b/Engine/Source/Runtime/Core/Private/Serialization/AsyncIOSystemBase.cpp
--- a/Engine/Source/Runtime/Core/Private/Serialization/AsyncIOSystemBase.cpp
+++ b/Engine/Source/Runtime/Core/Private/Serialization/AsyncIOSystemBase.cpp
@@ -769,7 +769,7 @@ FIOSystem& FIOSystem::Get()
 			// the platform didn't have a specific need, so we just use the base class with the normal file system.
 			AsyncIOSystem = new FAsyncIOSystemBase(FPlatformFileManager::Get().GetPlatformFile());
 		}
-		AsyncIOThread = FRunnableThread::Create(AsyncIOSystem, TEXT("AsyncIOSystem"), 16384, TPri_AboveNormal, FPlatformAffinity::GetPoolThreadMask());
+		AsyncIOThread = FRunnableThread::Create(AsyncIOSystem, TEXT("AsyncIOSystem"), 16384, TPri_Highest, FPlatformAffinity::GetPoolThreadMask());
 		check(AsyncIOThread);
 	}
 	check(AsyncIOSystem);
diff --git a/Engine/Source/Runtime/Core/Public/HAL/ThreadingBase.h b/Engine/Source/Runtime/Core/Public/HAL/ThreadingBase.h
--- a/Engine/Source/Runtime/Core/Public/HAL/ThreadingBase.h
+++ b/Engine/Source/Runtime/Core/Public/HAL/ThreadingBase.h
@@ -15,7 +15,10 @@ enum EThreadPriority
 	TPri_AboveNormal,
 
 	/** Below normal priority. */
-	TPri_BelowNormal
+	TPri_BelowNormal,
+    
+    /** Highest priority. */
+    TPri_Highest
 };
 
 
