Avoid symbol collision with OpenGL funcs in Linux splash.

- UE-5505 (reported on IRC by baldurk).

diff --git a/Engine/Source/Runtime/Core/Private/Linux/LinuxPlatformSplash.cpp b/Engine/Source/Runtime/Core/Private/Linux/LinuxPlatformSplash.cpp
--- a/Engine/Source/Runtime/Core/Private/Linux/LinuxPlatformSplash.cpp
+++ b/Engine/Source/Runtime/Core/Private/Linux/LinuxPlatformSplash.cpp
@@ -125,8 +125,12 @@ static int32 SplashBPP = 0;
 	EnumMacro(PFNGLTEXPARAMETERIPROC, glTexParameteri)
 
 #define DEFINE_GL_ENTRYPOINTS(Type,Func) Type Func = NULL;
-ENUM_GL_ENTRYPOINTS(DEFINE_GL_ENTRYPOINTS);
+namespace GLFuncPointers	// see explanation in OpenGLLinux.h why we need the namespace
+{
+	ENUM_GL_ENTRYPOINTS(DEFINE_GL_ENTRYPOINTS);
+};
 
+using namespace GLFuncPointers;
 
 //////////////////////////////////
 
@@ -574,7 +578,7 @@ static int StartSplashScreenThread(void *ptr)
 	}
 
 	// Initialize all entry points required by Unreal.
-	#define GET_GL_ENTRYPOINTS(Type,Func) Func = reinterpret_cast<Type>(SDL_GL_GetProcAddress(#Func));
+	#define GET_GL_ENTRYPOINTS(Type,Func) GLFuncPointers::Func = reinterpret_cast<Type>(SDL_GL_GetProcAddress(#Func));
 	ENUM_GL_ENTRYPOINTS(GET_GL_ENTRYPOINTS);
 
 	if (SDL_GL_MakeCurrent( GSplashWindow, Context ) != 0) 
