Enable explicit operators in Visual Studio 2013+

diff --git a/Engine/Source/Runtime/Core/Public/Windows/WIndowsPlatform.h b/Engine/Source/Runtime/Core/Public/Windows/WIndowsPlatform.h
--- a/Engine/Source/Runtime/Core/Public/Windows/WIndowsPlatform.h
+++ b/Engine/Source/Runtime/Core/Public/Windows/WIndowsPlatform.h
@@ -61,11 +61,12 @@ typedef FWindowsPlatformTypes FPlatformTypes;
 #define PLATFORM_SUPPORTS_NAMED_PIPES						1
 #define PLATFORM_COMPILER_HAS_DEFAULTED_FUNCTIONS			0
 #if _MSC_VER >= 1800
-	#define PLATFORM_COMPILER_HAS_VARIADIC_TEMPLATES			1
+	#define PLATFORM_COMPILER_HAS_VARIADIC_TEMPLATES		1
+	#define PLATFORM_COMPILER_HAS_EXPLICIT_OPERATORS		1
 #else
-	#define PLATFORM_COMPILER_HAS_VARIADIC_TEMPLATES			0
+	#define PLATFORM_COMPILER_HAS_VARIADIC_TEMPLATES		0
+	#define PLATFORM_COMPILER_HAS_EXPLICIT_OPERATORS		0
 #endif
-#define PLATFORM_COMPILER_HAS_EXPLICIT_OPERATORS			0
 #define PLATFORM_COMPILER_HAS_TCHAR_WMAIN					1
 
 // Intrinsics for 128-bit atomics on Windows platform requires Windows 8 or higher (WINVER>0x0602)
