#UE4

#RememberForBUNs

- Adding "GetSignVector" to FVector2D

diff --git a/Engine/Source/Runtime/Core/Public/Math/Vector2D.h b/Engine/Source/Runtime/Core/Public/Math/Vector2D.h
--- a/Engine/Source/Runtime/Core/Public/Math/Vector2D.h
+++ b/Engine/Source/Runtime/Core/Public/Math/Vector2D.h
@@ -417,6 +417,14 @@ public:
 	 */
 	FVector2D ClampAxes( float MinAxisVal, float MaxAxisVal ) const;
 
+	/**
+	* Get a copy of the vector as sign only.
+	* Each component is set to +1 or -1, with the sign of zero treated as +1.
+	*
+	* @param A copy of the vector with each component set to +1 or -1
+	*/
+	FORCEINLINE FVector2D GetSignVector() const;
+
 	/**
 	 * Get a textual representation of the vector.
 	 *
@@ -798,6 +806,16 @@ FORCEINLINE FVector2D FVector2D::ClampAxes( float MinAxisVal, float MaxAxisVal )
 }
 
 
+FORCEINLINE FVector2D FVector2D::GetSignVector() const
+{
+	return FVector2D
+		(
+		FMath::FloatSelect(X, 1.f, -1.f),
+		FMath::FloatSelect(Y, 1.f, -1.f)
+		);
+}
+
+
 FORCEINLINE FString FVector2D::ToString() const
 {
 	return FString::Printf(TEXT("X=%3.3f Y=%3.3f"), X, Y);
