Fixed dragging of windows with bDragAnywhere on Mac

diff --git a/Engine/Source/Runtime/Core/Private/Mac/MacApplication.cpp b/Engine/Source/Runtime/Core/Private/Mac/MacApplication.cpp
--- a/Engine/Source/Runtime/Core/Private/Mac/MacApplication.cpp
+++ b/Engine/Source/Runtime/Core/Private/Mac/MacApplication.cpp
@@ -595,8 +595,8 @@ void FMacApplication::ProcessNSEvent(NSEvent* const Event)
 				NSPoint CursorPos = [Event locationInWindow];
 				if ([Event GetWindow])
 				{
-					CursorPos.x += [Event GetWindow].frame.origin.x;
-					CursorPos.y += [Event GetWindow].frame.origin.y;
+					CursorPos.x += [Event windowPosition].x;
+					CursorPos.y += [Event windowPosition].y;
 				}
 				CursorPos.y--; // The y coordinate in the point returned by locationInWindow starts from a base of 1
 				const FVector2D MousePosition = FVector2D(CursorPos.x, FPlatformMisc::ConvertSlateYPositionToCocoa(CursorPos.y));
@@ -888,8 +888,8 @@ FCocoaWindow* FMacApplication::FindEventWindow( NSEvent* Event )
 			NSPoint CursorPos = [Event locationInWindow];
 			if ([Event GetWindow])
 			{
-				CursorPos.x += [Event GetWindow].frame.origin.x;
-				CursorPos.y += [Event GetWindow].frame.origin.y;
+				CursorPos.x += [Event windowPosition].x;
+				CursorPos.y += [Event windowPosition].y;
 			}
 			CursorPos.y--; // The y coordinate in the point returned by locationInWindow starts from a base of 1
 			TSharedPtr<FMacWindow> WindowUnderCursor = LocateWindowUnderCursor(CursorPos);
diff --git a/Engine/Source/Runtime/Core/Private/Mac/MacEvent.cpp b/Engine/Source/Runtime/Core/Private/Mac/MacEvent.cpp
--- a/Engine/Source/Runtime/Core/Private/Mac/MacEvent.cpp
+++ b/Engine/Source/Runtime/Core/Private/Mac/MacEvent.cpp
@@ -5,16 +5,31 @@
 #include "MacApplication.h"
 #include "CocoaThread.h"
 #include "CocoaWindow.h"
+#import <objc/runtime.h>
+
+static char PositionPropertyKey;
 
 @implementation NSEvent (FCachedWindowAccess)
+
+@dynamic windowPosition;
+
 -(void)CacheWindow
 {
 	_window = [[self window] retain];
+	self.windowPosition = _window.frame.origin;
 }
 -(NSWindow*)GetWindow
 {
 	return _window;
 }
+-(void)setWindowPosition:(NSPoint)position
+{
+	objc_setAssociatedObject(self, &PositionPropertyKey, [NSValue valueWithPoint:position], OBJC_ASSOCIATION_RETAIN_NONATOMIC);
+}
+-(NSPoint)windowPosition
+{
+	return [(NSValue*)objc_getAssociatedObject(self, &PositionPropertyKey) pointValue];
+}
 -(void)ResetWindow
 {
 	[_window release];
diff --git a/Engine/Source/Runtime/Core/Public/Mac/MacEvent.h b/Engine/Source/Runtime/Core/Public/Mac/MacEvent.h
--- a/Engine/Source/Runtime/Core/Public/Mac/MacEvent.h
+++ b/Engine/Source/Runtime/Core/Public/Mac/MacEvent.h
@@ -11,6 +11,8 @@ enum EMacEventSendMethod
 };
 
 @interface NSEvent (FCachedWindowAccess)
+@property (readwrite, nonatomic) NSPoint windowPosition;
+
 -(void)CacheWindow;
 -(NSWindow*)GetWindow;
 -(void)ResetWindow;
diff --git a/Engine/Source/Runtime/SlateCore/Private/Widgets/SWindow.cpp b/Engine/Source/Runtime/SlateCore/Private/Widgets/SWindow.cpp
--- a/Engine/Source/Runtime/SlateCore/Private/Widgets/SWindow.cpp
+++ b/Engine/Source/Runtime/SlateCore/Private/Widgets/SWindow.cpp
@@ -1378,6 +1378,7 @@ FReply SWindow::OnMouseButtonDown( const FGeometry& MyGeometry, const FPointerEv
 #endif
 	if (bDragAnywhere && MouseEvent.GetEffectingButton() == EKeys::LeftMouseButton)
 	{
+		MoveResizeZone = WindowZone;
 		return FReply::Handled().CaptureMouse(SharedThis(this));
 	}
 	else
@@ -1397,6 +1398,7 @@ FReply SWindow::OnMouseButtonUp( const FGeometry& MyGeometry, const FPointerEven
 #endif
 	if (bDragAnywhere &&  this->HasMouseCapture() && MouseEvent.GetEffectingButton() == EKeys::LeftMouseButton)
 	{
+		MoveResizeZone =  EWindowZone::Unspecified;
 		return FReply::Handled().ReleaseMouseCapture();
 	}
 	else
@@ -1511,7 +1513,7 @@ FReply SWindow::OnMouseMove( const FGeometry& MyGeometry, const FPointerEvent& M
 		return FReply::Handled();
 	}
 #endif
-	if ( bDragAnywhere && this->HasMouseCapture() && MouseEvent.IsMouseButtonDown(EKeys::LeftMouseButton) )
+	if ( bDragAnywhere && this->HasMouseCapture() && MouseEvent.IsMouseButtonDown(EKeys::LeftMouseButton) && MoveResizeZone != EWindowZone::TitleBar )
 	{
 		this->MoveWindowTo( ScreenPosition + MouseEvent.GetCursorDelta() );
 		return FReply::Handled();
