Core: Removed obsolete file header comments; documentation cleanup pass.

diff --git a/Engine/Source/Runtime/Core/Private/Windows/WindowsPlatformFeedbackContextPrivate.h b/Engine/Source/Runtime/Core/Private/Windows/WindowsPlatformFeedbackContextPrivate.h
--- a/Engine/Source/Runtime/Core/Private/Windows/WindowsPlatformFeedbackContextPrivate.h
+++ b/Engine/Source/Runtime/Core/Private/Windows/WindowsPlatformFeedbackContextPrivate.h
@@ -1,12 +1,10 @@
 // Copyright 1998-2014 Epic Games, Inc. All Rights Reserved.
 
-/*=============================================================================
-	WindowsPlatformFeedbackContextPrivate.h: Unreal Windows user interface interaction.
-=============================================================================*/
-
 #pragma once
+
 #include "Misc/OutputDeviceConsole.h"
 
+
 /**
  * Feedback context implementation for windows.
  */
diff --git a/Engine/Source/Runtime/Core/Private/Windows/WindowsPlatformProcess.cpp b/Engine/Source/Runtime/Core/Private/Windows/WindowsPlatformProcess.cpp
--- a/Engine/Source/Runtime/Core/Private/Windows/WindowsPlatformProcess.cpp
+++ b/Engine/Source/Runtime/Core/Private/Windows/WindowsPlatformProcess.cpp
@@ -1,9 +1,5 @@
 // Copyright 1998-2014 Epic Games, Inc. All Rights Reserved.
 
-/*=============================================================================
-	WindowsPlatformProcess.cpp: Windows implementations of Process functions
-=============================================================================*/
-
 #include "CorePrivatePCH.h"
 #include "EngineVersion.h"
 #include "Resources/Windows/ModuleVersionResource.h"
@@ -28,9 +24,11 @@
 
 #pragma comment(lib, "psapi.lib")
 
+
 // static variables
 TArray<FString> FWindowsPlatformProcess::DllDirectoryStack;
 
+
 void FWindowsPlatformProcess::AddDllDirectory(const TCHAR* Directory)
 {
 	// Normalize the input directory
diff --git a/Engine/Source/Runtime/Core/Public/Windows/AllowWindowsPlatformTypes.h b/Engine/Source/Runtime/Core/Public/Windows/AllowWindowsPlatformTypes.h
--- a/Engine/Source/Runtime/Core/Public/Windows/AllowWindowsPlatformTypes.h
+++ b/Engine/Source/Runtime/Core/Public/Windows/AllowWindowsPlatformTypes.h
@@ -1,9 +1,5 @@
 // Copyright 1998-2014 Epic Games, Inc. All Rights Reserved.
 
-/*=============================================================================
-	AllowWindowsPLatformTypes.h: Defines for alowwing the use of windows types.
-=============================================================================*/
-
 #ifndef WINDOWS_PLATFORM_TYPES_GUARD
 	#define WINDOWS_PLATFORM_TYPES_GUARD
 #else
@@ -16,8 +12,9 @@
 #define FLOAT ::FLOAT
 
 #ifndef TRUE
-#define TRUE 1
+	#define TRUE 1
 #endif
+
 #ifndef FALSE
-#define FALSE 0
+	#define FALSE 0
 #endif
diff --git a/Engine/Source/Runtime/Core/Public/Windows/HideWindowsPlatformTypes.h b/Engine/Source/Runtime/Core/Public/Windows/HideWindowsPlatformTypes.h
--- a/Engine/Source/Runtime/Core/Public/Windows/HideWindowsPlatformTypes.h
+++ b/Engine/Source/Runtime/Core/Public/Windows/HideWindowsPlatformTypes.h
@@ -1,9 +1,5 @@
 // Copyright 1998-2014 Epic Games, Inc. All Rights Reserved.
 
-/*=============================================================================
-	HideWindowsPLatformTypes.h: Defines for hiding windows type names.
-=============================================================================*/
-
 #ifdef WINDOWS_PLATFORM_TYPES_GUARD
 	#undef WINDOWS_PLATFORM_TYPES_GUARD
 #else
@@ -16,8 +12,9 @@
 #undef FLOAT
 
 #ifdef TRUE
-#undef TRUE
+	#undef TRUE
 #endif
+
 #ifdef FALSE
-#undef FALSE
-#endif
\ No newline at end of file
+	#undef FALSE
+#endif
diff --git a/Engine/Source/Runtime/Core/Public/Windows/MinWindows.h b/Engine/Source/Runtime/Core/Public/Windows/MinWindows.h
--- a/Engine/Source/Runtime/Core/Public/Windows/MinWindows.h
+++ b/Engine/Source/Runtime/Core/Public/Windows/MinWindows.h
@@ -1,16 +1,12 @@
 // Copyright 1998-2014 Epic Games, Inc. All Rights Reserved.
 
-/*=============================================================================
-	MinWindows.h: Minimal includes from Windows.h
-=============================================================================*/
-
 #pragma once
 
 #ifdef _WINDOWS_
-#pragma message ( " " )
-#pragma message ( "You have included windows.h before MinWindows.h" )
-#pragma message ( "All useless stuff from the windows headers won't be excluded !!!" )
-#pragma message ( " " )
+	#pragma message ( " " )
+	#pragma message ( "You have included windows.h before MinWindows.h" )
+	#pragma message ( "All useless stuff from the windows headers won't be excluded !!!" )
+	#pragma message ( " " )
 #endif // _WINDOWS_
 
 // WIN32_LEAN_AND_MEAN excludes rarely-used services from windows headers.
@@ -64,5 +60,3 @@
 
 // Finally now we can include windows.h
 #include <windows.h>
-
-
diff --git a/Engine/Source/Runtime/Core/Public/Windows/PostWindowsApi.h b/Engine/Source/Runtime/Core/Public/Windows/PostWindowsApi.h
--- a/Engine/Source/Runtime/Core/Public/Windows/PostWindowsApi.h
+++ b/Engine/Source/Runtime/Core/Public/Windows/PostWindowsApi.h
@@ -1,9 +1,5 @@
 // Copyright 1998-2014 Epic Games, Inc. All Rights Reserved.
 
-/*=============================================================================
-	PostWindowsApi.h: The post-include part of UE4's Win32 API wrapper.
-=============================================================================*/
-
 // Undo any Windows defines.
 #undef uint8
 #undef uint16
@@ -30,11 +26,10 @@
 
 // Notify people of the windows dependency. For CRITICAL_SECTION
 #if !defined(_WINBASE_) && !defined(_XTL_)
-#error CRITICAL_SECTION relies on Windows.h/Xtl.h being included ahead of it
+	#error CRITICAL_SECTION relies on Windows.h/Xtl.h being included ahead of it
 #endif
 
 // Make sure version is high enough for API to be defined. For CRITICAL_SECTION
 #if !defined(_XTL_) && (_WIN32_WINNT < 0x0403)
-#error SetCriticalSectionSpinCount requires _WIN32_WINNT >= 0x0403
+	#error SetCriticalSectionSpinCount requires _WIN32_WINNT >= 0x0403
 #endif
-
diff --git a/Engine/Source/Runtime/Core/Public/Windows/PreWindowsApi.h b/Engine/Source/Runtime/Core/Public/Windows/PreWindowsApi.h
--- a/Engine/Source/Runtime/Core/Public/Windows/PreWindowsApi.h
+++ b/Engine/Source/Runtime/Core/Public/Windows/PreWindowsApi.h
@@ -1,9 +1,4 @@
 // Copyright 1998-2014 Epic Games, Inc. All Rights Reserved.
 
-/*=============================================================================
-	PreWindowsApi.h: The pre-include part of UE4's Win32 API wrapper.
-=============================================================================*/
-
 // Disable the warning that the pack size is changed in this header.
 #pragma warning(disable:4103)
-
diff --git a/Engine/Source/Runtime/Core/Public/Windows/WIndowsPlatform.h b/Engine/Source/Runtime/Core/Public/Windows/WIndowsPlatform.h
--- a/Engine/Source/Runtime/Core/Public/Windows/WIndowsPlatform.h
+++ b/Engine/Source/Runtime/Core/Public/Windows/WIndowsPlatform.h
@@ -1,11 +1,8 @@
 // Copyright 1998-2014 Epic Games, Inc. All Rights Reserved.
 
-/*================================================================================
-	WindowsPlatform.h: Setup for the windows platform
-==================================================================================*/
-
 #pragma once
 
+
 /**
 * Windows specific types
 **/
diff --git a/Engine/Source/Runtime/Core/Public/Windows/WindowsPlatformAtomics.h b/Engine/Source/Runtime/Core/Public/Windows/WindowsPlatformAtomics.h
--- a/Engine/Source/Runtime/Core/Public/Windows/WindowsPlatformAtomics.h
+++ b/Engine/Source/Runtime/Core/Public/Windows/WindowsPlatformAtomics.h
@@ -1,15 +1,11 @@
 // Copyright 1998-2014 Epic Games, Inc. All Rights Reserved.
 
-
-/*=============================================================================================
-	WindowsPlatformAtomics.h: Windows platform Atomics functions
-==============================================================================================*/
-
 #pragma once
 
 #include "GenericPlatform/GenericPlatformAtomics.h"
 #include "Windows/WindowsSystemIncludes.h"
 
+
 /**
  * Windows implementation of the Atomics OS functions
  */
@@ -161,5 +157,3 @@ protected:
 
 
 typedef FWindowsPlatformAtomics FPlatformAtomics;
-
-
diff --git a/Engine/Source/Runtime/Core/Public/Windows/WindowsPlatformCompilerSetup.h b/Engine/Source/Runtime/Core/Public/Windows/WindowsPlatformCompilerSetup.h
--- a/Engine/Source/Runtime/Core/Public/Windows/WindowsPlatformCompilerSetup.h
+++ b/Engine/Source/Runtime/Core/Public/Windows/WindowsPlatformCompilerSetup.h
@@ -1,10 +1,5 @@
 // Copyright 1998-2014 Epic Games, Inc. All Rights Reserved.
 
-
-/*=============================================================================================
-	WindowsPlatformCompilerSetup.h: pragmas, version checks and other things for the windows compiler
-==============================================================================================*/
-
 #pragma once
 
 #if !PLATFORM_WINDOWS
diff --git a/Engine/Source/Runtime/Core/Public/Windows/WindowsPlatformFile.h b/Engine/Source/Runtime/Core/Public/Windows/WindowsPlatformFile.h
--- a/Engine/Source/Runtime/Core/Public/Windows/WindowsPlatformFile.h
+++ b/Engine/Source/Runtime/Core/Public/Windows/WindowsPlatformFile.h
@@ -1,8 +1,3 @@
 // Copyright 1998-2014 Epic Games, Inc. All Rights Reserved.
 
-
-/*=============================================================================================
-	WindowsPlatformFile.h: Windows platform File functions
-==============================================================================================*/
-
 #pragma once
diff --git a/Engine/Source/Runtime/Core/Public/Windows/WindowsPlatformHttp.h b/Engine/Source/Runtime/Core/Public/Windows/WindowsPlatformHttp.h
--- a/Engine/Source/Runtime/Core/Public/Windows/WindowsPlatformHttp.h
+++ b/Engine/Source/Runtime/Core/Public/Windows/WindowsPlatformHttp.h
@@ -1,39 +1,37 @@
 // Copyright 1998-2014 Epic Games, Inc. All Rights Reserved.
 
 #pragma once
+
 #include "GenericPlatform/GenericPlatformHttp.h"
 
+
 /**
- * Platform specific Http implementations
+ * Platform specific HTTP implementations.
  */
 class FWindowsPlatformHttp : public FGenericPlatformHttp
 {
 public:
 
-	/**
-	 * Platform initialization step
-	 */
+	/** Platform initialization step. */
 	static void Init();
 
 	/**
 	 * Creates a platform-specific HTTP manager.
 	 *
-	 * @return NULL if default implementation is to be used
+	 * @return nullptr if default implementation is to be used.
 	 */
 	static FHttpManager* CreatePlatformHttpManager();
 
-	/**
-	 * Platform shutdown step
-	 */
+	/** Platform shutdown step. */
 	static void Shutdown();
 
 	/**
-	 * Creates a new Http request instance for the current platform
+	 * Creates a new HTTP request instance for the current platform.
 	 *
-	 * @return request object
+	 * @return The request object.
 	 */
 	static IHttpRequest* ConstructRequest();
 };
 
 
-typedef FWindowsPlatformHttp FPlatformHttp;
\ No newline at end of file
+typedef FWindowsPlatformHttp FPlatformHttp;
diff --git a/Engine/Source/Runtime/Core/Public/Windows/WindowsPlatformIncludes.h b/Engine/Source/Runtime/Core/Public/Windows/WindowsPlatformIncludes.h
--- a/Engine/Source/Runtime/Core/Public/Windows/WindowsPlatformIncludes.h
+++ b/Engine/Source/Runtime/Core/Public/Windows/WindowsPlatformIncludes.h
@@ -1,9 +1,7 @@
 // Copyright 1998-2014 Epic Games, Inc. All Rights Reserved.
 
-/*================================================================================
-	WindowsPlatformIncludes.h: Includes the platform specific headers for windows
-==================================================================================*/
 #pragma once
+
 #include "GenericPlatform/GenericPlatformAffinity.h"
 #include "GenericPlatform/GenericPlatformFramePacer.h"
 #include "Misc/Build.h"
@@ -27,21 +25,22 @@
 #include "Windows/WindowsPlatformSurvey.h"
 #include "Windows/WindowsPlatformHttp.h"
 
-typedef FGenericPlatformRHIFramePacer FPlatformRHIFramePacer;
 
+typedef FGenericPlatformRHIFramePacer FPlatformRHIFramePacer;
 typedef FGenericPlatformAffinity FPlatformAffinity;
 
-// include platform properties and typedef it for the runtime
 
+// include platform properties and typedef it for the runtime
 #include "Windows/WindowsPlatformProperties.h"
 
+
 #ifndef WITH_EDITORONLY_DATA
-#error "WITH_EDITORONLY_DATA must be defined"
+	#error "WITH_EDITORONLY_DATA must be defined"
 #endif
 
 #ifndef UE_SERVER
-#error "WITH_EDITORONLY_DATA must be defined"
+	#error "WITH_EDITORONLY_DATA must be defined"
 #endif
 
-typedef FWindowsPlatformProperties<WITH_EDITORONLY_DATA, UE_SERVER, !WITH_SERVER_CODE> FPlatformProperties;
 
+typedef FWindowsPlatformProperties<WITH_EDITORONLY_DATA, UE_SERVER, !WITH_SERVER_CODE> FPlatformProperties;
diff --git a/Engine/Source/Runtime/Core/Public/Windows/WindowsPlatformMemory.h b/Engine/Source/Runtime/Core/Public/Windows/WindowsPlatformMemory.h
--- a/Engine/Source/Runtime/Core/Public/Windows/WindowsPlatformMemory.h
+++ b/Engine/Source/Runtime/Core/Public/Windows/WindowsPlatformMemory.h
@@ -1,35 +1,35 @@
 // Copyright 1998-2014 Epic Games, Inc. All Rights Reserved.
 
-
-/*=============================================================================================
-	WindowsPlatformMemory.h: Windows platform memory functions
-==============================================================================================*/
-
 #pragma once
+
 #include "GenericPlatform/GenericPlatformMemory.h"
 #include "Windows/WindowsSystemIncludes.h"
 
+
 /**
  *	Windows implementation of the FGenericPlatformMemoryStats.
  *	At this moment it's just the same as the FGenericPlatformMemoryStats.
  *	Can be extended as shown in the following example.
  */
-struct FPlatformMemoryStats : public FGenericPlatformMemoryStats
+struct FPlatformMemoryStats
+	: public FGenericPlatformMemoryStats
 {
 	/** Default constructor, clears all variables. */
 	FPlatformMemoryStats()
 		: FGenericPlatformMemoryStats()
-		, WindowsSpecificMemoryStat( 0 )
-	{}
+		, WindowsSpecificMemoryStat(0)
+	{ }
 
 	/** Memory stat specific only for Windows. */
 	SIZE_T WindowsSpecificMemoryStat;
 };
 
+
 /**
 * Windows implementation of the memory OS functions
 **/
-struct CORE_API FWindowsPlatformMemory : public FGenericPlatformMemory
+struct CORE_API FWindowsPlatformMemory
+	: public FGenericPlatformMemory
 {
 	enum EMemoryCounterRegion
 	{
@@ -75,4 +75,5 @@ struct CORE_API FWindowsPlatformMemory : public FGenericPlatformMemory
 	// End FGenericPlatformMemory interface
 };
 
+
 typedef FWindowsPlatformMemory FPlatformMemory;
diff --git a/Engine/Source/Runtime/Core/Public/Windows/WindowsPlatformMisc.h b/Engine/Source/Runtime/Core/Public/Windows/WindowsPlatformMisc.h
--- a/Engine/Source/Runtime/Core/Public/Windows/WindowsPlatformMisc.h
+++ b/Engine/Source/Runtime/Core/Public/Windows/WindowsPlatformMisc.h
@@ -1,14 +1,11 @@
 // Copyright 1998-2014 Epic Games, Inc. All Rights Reserved.
 
-
-/*=============================================================================================
-	WindowsPlatformMisc.h: Windows platform misc functions
-==============================================================================================*/
-
 #pragma once
+
 #include "GenericPlatform/GenericPlatformMisc.h"
 #include "Windows/WindowsHWrapper.h"
 
+
 /** Helper struct used to get the string version of the Windows version. */
 struct FWindowsOSVersionHelper
 {
@@ -23,10 +20,12 @@ struct FWindowsOSVersionHelper
 	static int32 GetOSVersions( FString& out_OSVersion, FString& out_OSSubVersion );
 };
 
+
 /**
 * Windows implementation of the misc OS functions
 **/
-struct CORE_API FWindowsPlatformMisc : public FGenericPlatformMisc
+struct CORE_API FWindowsPlatformMisc
+	: public FGenericPlatformMisc
 {
 	static void PlatformPreInit();
 	static void PlatformInit();
@@ -216,4 +215,5 @@ struct CORE_API FWindowsPlatformMisc : public FGenericPlatformMisc
 	static bool IsRunningOnBattery();
 };
 
+
 typedef FWindowsPlatformMisc FPlatformMisc;
diff --git a/Engine/Source/Runtime/Core/Public/Windows/WindowsPlatformNamedPipe.h b/Engine/Source/Runtime/Core/Public/Windows/WindowsPlatformNamedPipe.h
--- a/Engine/Source/Runtime/Core/Public/Windows/WindowsPlatformNamedPipe.h
+++ b/Engine/Source/Runtime/Core/Public/Windows/WindowsPlatformNamedPipe.h
@@ -1,61 +1,41 @@
 // Copyright 1998-2014 Epic Games, Inc. All Rights Reserved.
 
-
-/*=============================================================================================
-	WindowsPlatformNamedPipe.h: Declares the FWindowsPlatformNamedPipe class
-==============================================================================================*/
-
 #pragma once
+
 #include "HAL/Platform.h"
 #include "GenericPlatform/GenericPlatformNamedPipe.h"
 #include "Windows/WindowsSystemIncludes.h"
 
+
 #if PLATFORM_SUPPORTS_NAMED_PIPES
 
 class FString;
 
+
 // Windows wrapper for named pipe communications
-class CORE_API FWindowsPlatformNamedPipe : public FGenericPlatformNamedPipe
+class CORE_API FWindowsPlatformNamedPipe
+	: public FGenericPlatformNamedPipe
 {
 public:
 
 	FWindowsPlatformNamedPipe();
 	virtual ~FWindowsPlatformNamedPipe();
 
-
 public:
 
-	// Create a named pipe, using overlapped IO if bAsync=1
-	virtual bool Create(const FString& PipeName, bool bAsServer, bool bAsync) override;
+	// FGenericPlatformNamedPipe overrides
 
-	// Destroy the pipe
+	virtual bool Create(const FString& PipeName, bool bAsServer, bool bAsync) override;
 	virtual bool Destroy() override;
-
-	// Open a connection from a client
 	virtual bool OpenConnection() override;
-
-	// Blocks if there's an IO operation in progress until it's done or errors out
 	virtual bool BlockForAsyncIO() override;
-
-	// Lets the user know if the pipe is ready to send or receive data
 	virtual bool IsReadyForRW() const override;
-
-	// Updates status of async state of the current pipe
 	virtual bool UpdateAsyncStatus() override;
-
-	// Writes a buffer out
 	virtual bool WriteBytes(int32 NumBytes, const void* Data) override;
-
-	// Reads a buffer in
 	virtual bool ReadBytes(int32 NumBytes, void* OutData) override;
-
-	// Returns true if the pipe has been created and hasn't been destroyed
 	virtual bool IsCreated() const override;
-
-	// Return true if the pipe has had any communication error
 	virtual bool HasFailed() const override;
 
-
 private:
 
 	void*		Pipe;
diff --git a/Engine/Source/Runtime/Core/Public/Windows/WindowsPlatformOutputDevices.h b/Engine/Source/Runtime/Core/Public/Windows/WindowsPlatformOutputDevices.h
--- a/Engine/Source/Runtime/Core/Public/Windows/WindowsPlatformOutputDevices.h
+++ b/Engine/Source/Runtime/Core/Public/Windows/WindowsPlatformOutputDevices.h
@@ -1,14 +1,12 @@
 // Copyright 1998-2014 Epic Games, Inc. All Rights Reserved.
 
-
-/*=============================================================================================
-	WindowsPlatformOutputDevices.h: Windows platform OutputDevices functions
-==============================================================================================*/
-
 #pragma once
+
 #include "GenericPlatform/GenericPlatformOutputDevices.h"
 
-struct CORE_API FWindowsPlatformOutputDevices : public FGenericPlatformOutputDevices
+
+struct CORE_API FWindowsPlatformOutputDevices
+	: public FGenericPlatformOutputDevices
 {
 	static FOutputDevice*			GetEventLog();
 	static FOutputDeviceConsole*	GetLogConsole();
@@ -16,5 +14,5 @@ struct CORE_API FWindowsPlatformOutputDevices : public FGenericPlatformOutputDev
 	static FFeedbackContext*		GetWarn();
 };
 
-typedef FWindowsPlatformOutputDevices FPlatformOutputDevices;
 
+typedef FWindowsPlatformOutputDevices FPlatformOutputDevices;
diff --git a/Engine/Source/Runtime/Core/Public/Windows/WindowsPlatformProcess.h b/Engine/Source/Runtime/Core/Public/Windows/WindowsPlatformProcess.h
--- a/Engine/Source/Runtime/Core/Public/Windows/WindowsPlatformProcess.h
+++ b/Engine/Source/Runtime/Core/Public/Windows/WindowsPlatformProcess.h
@@ -1,14 +1,11 @@
 // Copyright 1998-2014 Epic Games, Inc. All Rights Reserved.
 
-
-/*=============================================================================================
-	WindowsPlatformProcess.h: Windows platform Process functions
-==============================================================================================*/
-
 #pragma once
+
 #include "GenericPlatform/GenericPlatformProcess.h"
 #include "Windows/WindowsSystemIncludes.h"
 
+
 /** Windows implementation of the process handle. */
 #if PLATFORM_WINDOWS && defined(__clang__)	// @todo clang: Clang on Windows has a bug with non-type template arguments
 struct FProcHandle : public TProcHandle<HANDLE>
@@ -42,10 +39,12 @@ public:
 	}
 };
 
+
 /**
-* Windows implementation of the Process OS functions
+* Windows implementation of the Process OS functions.
 **/
-struct CORE_API FWindowsPlatformProcess : public FGenericPlatformProcess
+struct CORE_API FWindowsPlatformProcess
+	: public FGenericPlatformProcess
 {
 	/**
 	 * Windows representation of a interprocess semaphore
@@ -71,7 +70,9 @@ struct CORE_API FWindowsPlatformProcess : public FGenericPlatformProcess
 		HANDLE			Semaphore;
 	};
 
-	// Begin FGenericPlatformProcess interface
+public:
+
+	// FGenericPlatformProcess interface
 
 	static void* GetDllHandle( const TCHAR* Filename );
 	static void FreeDllHandle( void* DllHandle );
@@ -122,20 +123,17 @@ struct CORE_API FWindowsPlatformProcess : public FGenericPlatformProcess
 	static bool DeleteInterprocessSynchObject(FSemaphore * Object);
 	static bool Daemonize();
 
-	// End FGenericPlatformProcess interface
-
 protected:
 
 	/**
 	 * Reads from a collection of anonymous pipes.
 	 *
-	 * @param OutStrings - Will hold the read data.
-	 * @param InPipes - The pipes to read from.
-	 * @param PipeCount - The number of pipes.
+	 * @param OutStrings Will hold the read data.
+	 * @param InPipes The pipes to read from.
+	 * @param PipeCount The number of pipes.
 	 */
 	static void ReadFromPipes(FString* OutStrings[], HANDLE InPipes[], int32 PipeCount);
 
-
 private:
 
 	// Since Windows can only have one directory at a time, this stack is used to reset the previous DllDirectory
@@ -146,4 +144,5 @@ private:
 typedef FWindowsPlatformProcess FPlatformProcess;
 
 #include "WindowsCriticalSection.h"
+
 typedef FWindowsCriticalSection FCriticalSection;
diff --git a/Engine/Source/Runtime/Core/Public/Windows/WindowsPlatformProperties.h b/Engine/Source/Runtime/Core/Public/Windows/WindowsPlatformProperties.h
--- a/Engine/Source/Runtime/Core/Public/Windows/WindowsPlatformProperties.h
+++ b/Engine/Source/Runtime/Core/Public/Windows/WindowsPlatformProperties.h
@@ -1,13 +1,7 @@
 // Copyright 1998-2014 Epic Games, Inc. All Rights Reserved.
 
-/*================================================================================
-	WindowsPlatformProperties.h - Basic static properties of a platform 
-	These are shared between:
-		the runtime platform - via FPlatformProperties
-		the target platforms - via ITargetPlatform
-==================================================================================*/
-
 #pragma once
+
 #include "GenericPlatformProperties.h"
 
 
diff --git a/Engine/Source/Runtime/Core/Public/Windows/WindowsPlatformSplash.h b/Engine/Source/Runtime/Core/Public/Windows/WindowsPlatformSplash.h
--- a/Engine/Source/Runtime/Core/Public/Windows/WindowsPlatformSplash.h
+++ b/Engine/Source/Runtime/Core/Public/Windows/WindowsPlatformSplash.h
@@ -1,32 +1,30 @@
 // Copyright 1998-2014 Epic Games, Inc. All Rights Reserved.
-/*=============================================================================================
-	WindowsPlatformSplash.h: Windows platform splash screen...
-==============================================================================================*/
 
 #pragma once
+
 #include "GenericPlatform/GenericPlatformSplash.h"
 
+
 /**
-* Windows splash implementation
-**/
-struct CORE_API FWindowsPlatformSplash : public FGenericPlatformSplash
+ * Windows splash implementation.
+ */
+struct CORE_API FWindowsPlatformSplash
+	: public FGenericPlatformSplash
 {
-	/**
-	* Show the splash screen
-	*/
+	/** Show the splash screen. */
 	static void Show();
-	/**
-	* Hide the splash screen
-	*/
+
+	/** Hide the splash screen. */
 	static void Hide();
 
 	/**
-	 * Sets the text displayed on the splash screen (for startup/loading progress)
+	 * Sets the text displayed on the splash screen (for startup/loading progress).
 	 *
-	 * @param	InType		Type of text to change
-	 * @param	InText		Text to display
+	 * @param InType Type of text to change.
+	 * @param InText Text to display.
 	 */
 	static void SetSplashText( const SplashTextType::Type InType, const TCHAR* InText );
 };
 
+
 typedef FWindowsPlatformSplash FPlatformSplash;
diff --git a/Engine/Source/Runtime/Core/Public/Windows/WindowsPlatformStackWalk.h b/Engine/Source/Runtime/Core/Public/Windows/WindowsPlatformStackWalk.h
--- a/Engine/Source/Runtime/Core/Public/Windows/WindowsPlatformStackWalk.h
+++ b/Engine/Source/Runtime/Core/Public/Windows/WindowsPlatformStackWalk.h
@@ -1,11 +1,7 @@
 // Copyright 1998-2014 Epic Games, Inc. All Rights Reserved.
 
-
-/*=============================================================================================
-	WindowsPlatformStackWalk.h: Windows platform stack walk functions
-==============================================================================================*/
-
 #pragma once
+
 #include "GenericPlatform/GenericPlatformStackWalk.h"
 
 #undef PLATFORM_SUPPORTS_STACK_SYMBOLS
@@ -15,7 +11,8 @@
 /**
  * Windows implementation of the stack walking.
  **/
-struct CORE_API FWindowsPlatformStackWalk : public FGenericPlatformStackWalk
+struct CORE_API FWindowsPlatformStackWalk
+	: public FGenericPlatformStackWalk
 {
 	static bool InitStackWalking();
 	
@@ -29,5 +26,5 @@ struct CORE_API FWindowsPlatformStackWalk : public FGenericPlatformStackWalk
 	static void RegisterOnModulesChanged();
 };
 
-typedef FWindowsPlatformStackWalk FPlatformStackWalk;
 
+typedef FWindowsPlatformStackWalk FPlatformStackWalk;
diff --git a/Engine/Source/Runtime/Core/Public/Windows/WindowsPlatformString.h b/Engine/Source/Runtime/Core/Public/Windows/WindowsPlatformString.h
--- a/Engine/Source/Runtime/Core/Public/Windows/WindowsPlatformString.h
+++ b/Engine/Source/Runtime/Core/Public/Windows/WindowsPlatformString.h
@@ -1,18 +1,15 @@
 // Copyright 1998-2014 Epic Games, Inc. All Rights Reserved.
 
-
-/*=============================================================================================
-	WindowsPlatformString.h: Windows platform string classes, mostly implemented with ANSI C++
-==============================================================================================*/
-
 #pragma once
 
 #include "GenericPlatform/MicrosoftPlatformString.h"
 
+
 /**
-* Windows string implementation
-**/
-struct FWindowsPlatformString : public FMicrosoftPlatformString
+ * Windows string implementation.
+ */
+struct FWindowsPlatformString
+	: public FMicrosoftPlatformString
 {
 	// These should be replaced with equivalent Convert and ConvertedLength functions when we properly support variable-length encodings.
 /*	static void WideCharToMultiByte(const wchar_t *Source, uint32 LengthWM1, ANSICHAR *Dest, uint32 LengthA)
@@ -25,4 +22,5 @@ struct FWindowsPlatformString : public FMicrosoftPlatformString
 	}*/
 };
 
+
 typedef FWindowsPlatformString FPlatformString;
diff --git a/Engine/Source/Runtime/Core/Public/Windows/WindowsPlatformSurvey.h b/Engine/Source/Runtime/Core/Public/Windows/WindowsPlatformSurvey.h
--- a/Engine/Source/Runtime/Core/Public/Windows/WindowsPlatformSurvey.h
+++ b/Engine/Source/Runtime/Core/Public/Windows/WindowsPlatformSurvey.h
@@ -1,22 +1,21 @@
 // Copyright 1998-2014 Epic Games, Inc. All Rights Reserved.
 
-
-/*=============================================================================================
-	WindowsPlatformSurvey.h: Microsoft Windows platform hardware-survey classes
-==============================================================================================*/
-
 #pragma once
+
 #include "GenericPlatform/GenericPlatformSurvey.h"
 
+
 /**
-* Windows implementation of FGenericPlatformSurvey
-**/
-struct CORE_API FWindowsPlatformSurvey : public FGenericPlatformSurvey
+ * Windows implementation of FGenericPlatformSurvey.
+ */
+struct CORE_API FWindowsPlatformSurvey
+	: public FGenericPlatformSurvey
 {
 	/** Start, or check on, the hardware survey */
 	static bool GetSurveyResults( FHardwareSurveyResults& OutResults, bool bWait = false );
 
 private:
+
 	/** Start run the hardware survey now but return immediately without waiting for results **/
 	static void BeginSurveyHardware();
 
@@ -62,4 +61,5 @@ private:
 	static FHardwareSurveyResults Results;
 };
 
-typedef FWindowsPlatformSurvey FPlatformSurvey;
\ No newline at end of file
+
+typedef FWindowsPlatformSurvey FPlatformSurvey;
diff --git a/Engine/Source/Runtime/Core/Public/Windows/WindowsPlatformTLS.h b/Engine/Source/Runtime/Core/Public/Windows/WindowsPlatformTLS.h
--- a/Engine/Source/Runtime/Core/Public/Windows/WindowsPlatformTLS.h
+++ b/Engine/Source/Runtime/Core/Public/Windows/WindowsPlatformTLS.h
@@ -1,22 +1,22 @@
 // Copyright 1998-2014 Epic Games, Inc. All Rights Reserved.
 
-
-/*=============================================================================================
-	WindowsPlatformTLS.h: Windows platform TLS (Thread local storage and thread ID) functions
-==============================================================================================*/
-
 #pragma once
+
 #include "GenericPlatform/GenericPlatformTLS.h"
 #include "HAL/Platform.h"
 #include "Windows/WindowsSystemIncludes.h"
 
+
 /**
-* Windows implementation of the TLS OS functions
-**/
-struct CORE_API FWindowsPlatformTLS : public FGenericPlatformTLS
+ * Windows implementation of the TLS OS functions.
+ */
+struct CORE_API FWindowsPlatformTLS
+	: public FGenericPlatformTLS
 {
 	/**
-	 * Returns the currently executing thread's id
+	 * Returns the currently executing thread's identifier.
+	 *
+	 * @return The thread identifier.
 	 */
 	static FORCEINLINE uint32 GetCurrentThreadId(void)
 	{
@@ -24,7 +24,9 @@ struct CORE_API FWindowsPlatformTLS : public FGenericPlatformTLS
 	}
 
 	/**
-	 * Allocates a thread local store slot
+	 * Allocates a thread local store slot.
+	 *
+	 * @return The index of the allocated slot.
 	 */
 	static FORCEINLINE uint32 AllocTlsSlot(void)
 	{
@@ -32,10 +34,10 @@ struct CORE_API FWindowsPlatformTLS : public FGenericPlatformTLS
 	}
 
 	/**
-	 * Sets a value in the specified TLS slot
+	 * Sets a value in the specified TLS slot.
 	 *
-	 * @param SlotIndex the TLS index to store it in
-	 * @param Value the value to store in the slot
+	 * @param SlotIndex the TLS index to store it in.
+	 * @param Value the value to store in the slot.
 	 */
 	static FORCEINLINE void SetTlsValue(uint32 SlotIndex,void* Value)
 	{
@@ -43,9 +45,10 @@ struct CORE_API FWindowsPlatformTLS : public FGenericPlatformTLS
 	}
 
 	/**
-	 * Reads the value stored at the specified TLS slot
+	 * Reads the value stored at the specified TLS slot.
 	 *
-	 * @return the value stored in the slot
+	 * @param SlotIndex The index of the slot to read.
+	 * @return The value stored in the slot.
 	 */
 	static FORCEINLINE void* GetTlsValue(uint32 SlotIndex)
 	{
@@ -63,5 +66,5 @@ struct CORE_API FWindowsPlatformTLS : public FGenericPlatformTLS
 	}
 };
 
-typedef FWindowsPlatformTLS FPlatformTLS;
 
+typedef FWindowsPlatformTLS FPlatformTLS;
diff --git a/Engine/Source/Runtime/Core/Public/Windows/WindowsPlatformTime.h b/Engine/Source/Runtime/Core/Public/Windows/WindowsPlatformTime.h
--- a/Engine/Source/Runtime/Core/Public/Windows/WindowsPlatformTime.h
+++ b/Engine/Source/Runtime/Core/Public/Windows/WindowsPlatformTime.h
@@ -1,27 +1,21 @@
 // Copyright 1998-2014 Epic Games, Inc. All Rights Reserved.
 
-
-/*=============================================================================================
-	WindowsPlatformTime.h: Windows platform Time functions
-==============================================================================================*/
-
 #pragma once
+
 #include "GenericPlatform/GenericPlatformTime.h"
 #include "Windows/WindowsSystemIncludes.h"
 
+
 /**
+ * Windows implementation of the Time OS functions.
+ *
  * Please see following UDN post about using rdtsc on processors that support
  * result being invariant across cores.
  *
  * https://udn.epicgames.com/lists/showpost.php?id=46794&list=unprog3
  */
-
-
-
-/**
-* Windows implementation of the Time OS functions
-**/
-struct CORE_API FWindowsPlatformTime : public FGenericPlatformTime
+struct CORE_API FWindowsPlatformTime
+	: public FGenericPlatformTime
 {
 	static double InitTiming();
 
@@ -29,7 +23,8 @@ struct CORE_API FWindowsPlatformTime : public FGenericPlatformTime
 	{
 		LARGE_INTEGER Cycles;
 		QueryPerformanceCounter(&Cycles);
-		// Add big number to make bugs apparent where return value is being passed to float
+
+		// add big number to make bugs apparent where return value is being passed to float
 		return Cycles.QuadPart * GetSecondsPerCycle() + 16777216.0;
 	}
 
@@ -47,9 +42,10 @@ struct CORE_API FWindowsPlatformTime : public FGenericPlatformTime
 	static FCPUTime GetCPUTime();	
 
 protected:
+
 	/** Percentage CPU utilization for the last interval relative to one core. */
 	static float CPUTimePctRelative;
 };
 
-typedef FWindowsPlatformTime FPlatformTime;
 
+typedef FWindowsPlatformTime FPlatformTime;
diff --git a/Engine/Source/Runtime/Core/Public/Windows/WindowsSystemIncludes.h b/Engine/Source/Runtime/Core/Public/Windows/WindowsSystemIncludes.h
--- a/Engine/Source/Runtime/Core/Public/Windows/WindowsSystemIncludes.h
+++ b/Engine/Source/Runtime/Core/Public/Windows/WindowsSystemIncludes.h
@@ -1,10 +1,8 @@
 // Copyright 1998-2014 Epic Games, Inc. All Rights Reserved.
+
 #pragma once
 
-// Set up compiler pragmas, etc
 #include "Windows/WindowsPlatformCompilerSetup.h"
-
-// include windows.h
 #include "Windows/WindowsHWrapper.h"
 
 // Macro for releasing COM objects
@@ -38,4 +36,4 @@ extern "C" CORE_API HINSTANCE hInstance;
 // We import the vc_attributes namespace so we can use annotations more easily.  It only defines
 // MSVC-specific attributes so there should never be collisions.
 using namespace vc_attributes;
-#endif
\ No newline at end of file
+#endif
diff --git a/Engine/Source/Runtime/Core/Public/Windows/WindowsTextInputMethodSystem.h b/Engine/Source/Runtime/Core/Public/Windows/WindowsTextInputMethodSystem.h
--- a/Engine/Source/Runtime/Core/Public/Windows/WindowsTextInputMethodSystem.h
+++ b/Engine/Source/Runtime/Core/Public/Windows/WindowsTextInputMethodSystem.h
@@ -9,10 +9,14 @@
 #include <msctf.h>
 #include "COMPointer.h"
 
+
 class FTextStoreACP;
 class FWindowsTextInputMethodSystem;
 
-class FTSFActivationProxy : public ITfInputProcessorProfileActivationSink, public ITfActiveLanguageProfileNotifySink
+
+class FTSFActivationProxy
+	: public ITfInputProcessorProfileActivationSink
+	, public ITfActiveLanguageProfileNotifySink
 {
 public:
 	FTSFActivationProxy(FWindowsTextInputMethodSystem* InOwner) 
@@ -20,12 +24,10 @@ public:
 		, TSFLanguageCookie(TF_INVALID_COOKIE)
 		, Owner(InOwner)
 		, ReferenceCount(1)
-	{
-	}
+	{ }
 
 	virtual ~FTSFActivationProxy()
-	{
-	}
+	{ }
 
 	// IUnknown Interface Begin
 	STDMETHODIMP QueryInterface(REFIID riid, void **ppvObj) override;
@@ -42,39 +44,48 @@ public:
 	// ITfActiveLanguageProfileNotifySink Interface End
 
 public:
+
 	DWORD TSFProfileCookie;
 	DWORD TSFLanguageCookie;
 
 private:
+
 	FWindowsTextInputMethodSystem* Owner;
 
 	// Reference count for IUnknown Implementation
 	ULONG ReferenceCount;
 };
 
-class FWindowsTextInputMethodSystem : public ITextInputMethodSystem
+
+class FWindowsTextInputMethodSystem
+	: public ITextInputMethodSystem
 {
 	friend class FTSFActivationProxy;
 
 public:
+
 	virtual ~FWindowsTextInputMethodSystem() {}
 
 	bool Initialize();
 	void Terminate();
 
-	// ITextInputMethodSystem Interface Begin
+	int32 ProcessMessage(HWND hwnd, uint32 msg, WPARAM wParam, LPARAM lParam);
+
+public:
+
+	// ITextInputMethodSystem interface
+
 	virtual void ApplyDefaults(const TSharedRef<FGenericWindow>& InWindow) override;
 	virtual TSharedPtr<ITextInputMethodChangeNotifier> RegisterContext(const TSharedRef<ITextInputMethodContext>& Context) override;
 	virtual void UnregisterContext(const TSharedRef<ITextInputMethodContext>& Context) override;
 	virtual void ActivateContext(const TSharedRef<ITextInputMethodContext>& Context) override;
 	virtual void DeactivateContext(const TSharedRef<ITextInputMethodContext>& Context) override;
 	virtual bool IsActiveContext(const TSharedRef<ITextInputMethodContext>& Context) const override;
-	// ITextInputMethodSystem Interface End
-
-	int32 ProcessMessage(HWND hwnd, uint32 msg, WPARAM wParam, LPARAM lParam);
 
 private:
+
 	// IMM Implementation
+
 	bool InitializeIMM();
 	void UpdateIMMProperty(HKL KeyboardLatoutHandle);
 	bool ShouldDrawIMMCompositionString() const;
@@ -84,14 +95,15 @@ private:
 	void CancelIMMComposition();
 
 	// TSF Implementation
-	bool InitializeTSF();
 
+	bool InitializeTSF();
 	void OnIMEActivationStateChanged(const bool bIsEnabled);
-
 	void ClearStaleWindowHandles();
 
 private:
+
 	TSharedPtr<ITextInputMethodContext> ActiveContext;
+
 	enum class EAPI
 	{
 		Unknown,
@@ -130,6 +142,7 @@ private:
 			uint32 CompositionLength;
 		} IMMContext;
 	};
+
 	TMap< TWeakPtr<ITextInputMethodContext>, FInternalContext > ContextToInternalContextMap;
 
 	// IMM Implementation
@@ -139,4 +152,5 @@ private:
 	TSet<TWeakPtr<FGenericWindow>> KnownWindows;
 };
 
-#include "HideWindowsPlatformTypes.h"
\ No newline at end of file
+
+#include "HideWindowsPlatformTypes.h"
diff --git a/Engine/Source/Runtime/Core/Public/Windows/WindowsWindow.h b/Engine/Source/Runtime/Core/Public/Windows/WindowsWindow.h
--- a/Engine/Source/Runtime/Core/Public/Windows/WindowsWindow.h
+++ b/Engine/Source/Runtime/Core/Public/Windows/WindowsWindow.h
@@ -4,141 +4,126 @@
 
 #include "GenericWindow.h"
 #include "SharedPointer.h"
+
 #include "AllowWindowsPlatformTypes.h"
-	#include "Ole2.h"
-	#include "OleIdl.h"
+#include "Ole2.h"
+#include "OleIdl.h"
 #include "HideWindowsPlatformTypes.h"
 
+
 /**
  * A platform specific implementation of FNativeWindow.
+ *
  * Native Windows provide platform-specific backing for and are always owned by an SWindow.
  */
-class CORE_API FWindowsWindow : public FGenericWindow, public IDropTarget
+class CORE_API FWindowsWindow
+	: public FGenericWindow
+	, public IDropTarget
 {
 public:
-	~FWindowsWindow();
 
 	/** Win32 requirement: see CreateWindowEx and RegisterClassEx. */
 	static const TCHAR AppWindowClass[];
 
+public:
+
+	/** Destructor. */
+	~FWindowsWindow();
+
 	/** Create a new FWin32Window.
 	 *
-	 * @param SlateWindows		List of all top level Slate windows.  This function will add the owner window to this list.
-	 * @param OwnerWindow		The SlateWindow for which we are crating a backing Win32Window
-	 * @param InHInstance		Win32 application instance handle
-	 * @param InParent			Parent Win32 window; usually NULL.
-	 * @param bShowImmediately	True to show this window as soon as its initialized
+	 * @param SlateWindows List of all top level Slate windows.  This function will add the owner window to this list.
+	 * @param OwnerWindow The SlateWindow for which we are crating a backing Win32Window.
+	 * @param InHInstance Win32 application instance handle.
+	 * @param InParent Parent Win32 window; usually NULL.
+	 * @param bShowImmediately True to show this window as soon as its initialized.
 	 */
-	static TSharedRef< FWindowsWindow > Make();
+	static TSharedRef<FWindowsWindow> Make();
 
+	/**
+	 * Gets the Window's handle.
+	 *
+	 * @return The window's HWND handle.
+	 */
 	HWND GetHWnd() const;
 
-	void Initialize( class FWindowsApplication* const Application, const TSharedRef< FGenericWindowDefinition >& InDefinition, HINSTANCE InHInstance, const TSharedPtr< FWindowsWindow >& InParent, const bool bShowImmediately );
+	void Initialize( class FWindowsApplication* const Application, const TSharedRef<FGenericWindowDefinition>& InDefinition, HINSTANCE InHInstance, const TSharedPtr<FWindowsWindow>& InParent, const bool bShowImmediately );
+
+	/**
+	* Sets the window text (usually the title but can also be text content for things like controls).
+	*
+	* @param Text The window's title or content text
+	*/
+	bool IsRegularWindow() const;
+
+	/**
+     * Sets the window region to specified dimensions.
+	 *
+	 * @param Width The width of the window region (in pixels).
+	 * @param Height The height of the window region (in pixels).
+	 */
+	void AdjustWindowRegion(int32 Width, int32 Height);
+
+	/** @return	Gives the native window a chance to adjust our stored window size before we cache it off. */
+	virtual void AdjustCachedSize(FVector2D& Size) const override;
 
+	/** Called when our parent window is minimized (which will in turn cause us to become minimized). */
+	void OnParentWindowMinimized();
+
+	/** Called when our parent window is restored (which will in turn cause us to become restored). */
+	void OnParentWindowRestored();
 
 public:
 
-	virtual void ReshapeWindow( int32 X, int32 Y, int32 Width, int32 Height ) override;
+	// FGenericWindow interface
 
+	virtual void ReshapeWindow( int32 X, int32 Y, int32 Width, int32 Height ) override;
 	virtual bool GetFullScreenInfo( int32& X, int32& Y, int32& Width, int32& Height ) const override;
-
 	virtual void MoveWindowTo ( int32 X, int32 Y ) override;
-
 	virtual void BringToFront( bool bForce = false ) override;
-
 	virtual void HACK_ForceToFront() override;
-
 	virtual void Destroy() override;
-
 	virtual void Minimize() override;
-
 	virtual void Maximize() override;
-
 	virtual void Restore() override;
-
 	virtual void Show() override;
-
 	virtual void Hide() override;
-
 	virtual void SetWindowMode( EWindowMode::Type NewWindowMode ) override;
-
 	virtual EWindowMode::Type GetWindowMode() const override { return WindowMode; } 
-
 	virtual bool IsMaximized() const override;
-
 	virtual bool IsMinimized() const override;
-
 	virtual bool IsVisible() const override;
-
 	virtual bool GetRestoredDimensions(int32& X, int32& Y, int32& Width, int32& Height) override;
-
 	virtual void SetWindowFocus() override;
-
 	virtual void SetOpacity( const float InOpacity ) override;
-
 	virtual void Enable( bool bEnable ) override;
-
 	virtual bool IsPointInWindow( int32 X, int32 Y ) const override;
-
 	virtual int32 GetWindowBorderSize() const override;
-
 	virtual int32 GetWindowTitleBarSize() const override;
-
 	virtual void* GetOSWindowHandle() const  override { return HWnd; }
-
 	virtual bool IsForegroundWindow() const override;
-
 	virtual void SetText(const TCHAR* const Text) override;
 
-	/**
-	 * Slate Win32Windows implement the IDropTarget and IUnknown OLE interfaces.
-	 */
+public:
+
+	// IUnknown interface
 
-	/** IUnknown : attempt to get a reference to this instance as a pointer  */
-	HRESULT STDCALL QueryInterface( REFIID iid, void ** ppvObject ) override;
+	HRESULT STDCALL QueryInterface(REFIID iid, void ** ppvObject) override;
+	ULONG STDCALL AddRef(void) override;
+	ULONG STDCALL Release(void) override;
 
-	/** IUnknown : Invoked by OLE when a new reference to this instance */
-	ULONG STDCALL AddRef( void ) override;
+public:
 
-	/** IUnknown : Invoked by OLE when a reference to this instance is released */
-	ULONG STDCALL Release( void ) override;
+	// IDropTarget interface
 
-	/** IDropTarget : Invoked by OLE when someone drag data our window initially */
 	virtual HRESULT STDCALL DragEnter( __RPC__in_opt IDataObject *DataObjectPointer, ::DWORD KeyState, POINTL CursorPosition, __RPC__inout ::DWORD *CursorEffect) override;
-
-	/** IDropTarget : Invoked by OLE when someone drag data over our window */
 	virtual HRESULT STDCALL DragOver( ::DWORD KeyState, POINTL CursorPosition, __RPC__inout ::DWORD *CursorEffect) override;
-
-	/** IDropTarget : Invoked by OLE when someone drag data and it exits our window or hits ESC to cancel the drag and drop. */
 	virtual HRESULT STDCALL DragLeave( void ) override;
-
-	/** IDropTarget : Invoked by OLE when someone drag data our window initially */
 	virtual HRESULT STDCALL Drop( __RPC__in_opt IDataObject *DataObjectPointer, ::DWORD KeyState, POINTL CursorPosition, __RPC__inout ::DWORD *CursorEffect) override;
 
-	/**
-	 * Sets the window text - usually the title but can also be text content for things like controls
-	 *
-	 * @param Text	The window's title or content text
-	 */
-	bool IsRegularWindow() const;
-
-	/**	Sets the window region to specified dimensions */
-	void AdjustWindowRegion( int32 Width, int32 Height );
-
-	/** @return	Gives the native window a chance to adjust our stored window size before we cache it off */
-	virtual void AdjustCachedSize( FVector2D& Size ) const override;
-
-	/** Called when our parent window is minimized (which will in turn cause us to become minimized) */
-	void OnParentWindowMinimized();
-
-	/** Called when our parent window is restored (which will in turn cause us to become restored) */
-	void OnParentWindowRestored();
-
 private:
 
-	/**
-	 * Protect the constructor; only TSharedRefs of this class can be made.
-	 */
+	/** Protect the constructor; only TSharedRefs of this class can be made. */
 	FWindowsWindow();
 
 	void UpdateVisibility();
@@ -149,24 +134,26 @@ private:
 
 private:
 
+	/** The application that owns this window. */
 	FWindowsApplication* OwningApplication;
 
-	/** WIN32 window handle */
+	/** The window's handle. */
 	HWND HWnd;
 
-	/** Store the window region size for querying whether a point lies within the window */
+	/** Store the window region size for querying whether a point lies within the window. */
 	int32 RegionWidth;
 	int32 RegionHeight;
 	
 	/** The mode that the window is in (windowed, fullscreen, windowedfullscreen ) */
 	EWindowMode::Type WindowMode;
 
+	/** This object's reference count (for the IUnknown interface). */
 	int32 OLEReferenceCount;
 		
-	/** The placement of the window before it entered a fullscreen state */
+	/** The placement of the window before it entered a fullscreen state. */
 	WINDOWPLACEMENT PreFullscreenWindowPlacement;
 
-	/** The placement of the window before it entered a minimized state due to its parent window being minimized */
+	/** The placement of the window before it entered a minimized state due to its parent window being minimized. */
 	WINDOWPLACEMENT PreParentMinimizedWindowPlacement;
 
 	/** Virtual width and height of the window.  This is only different than the actual width and height for
