Linux: maintain a "shadow" global cursor position.

- PR #611 with changes.
- Partially addresses UE-5160 (moving objects in editor).

diff --git a/Engine/Source/Runtime/Core/Private/Linux/LinuxApplication.cpp b/Engine/Source/Runtime/Core/Private/Linux/LinuxApplication.cpp
--- a/Engine/Source/Runtime/Core/Private/Linux/LinuxApplication.cpp
+++ b/Engine/Source/Runtime/Core/Private/Linux/LinuxApplication.cpp
@@ -314,7 +314,8 @@ void FLinuxApplication::ProcessDeferredMessage( SDL_Event Event )
 
 			if(bUsingHighPrecisionMouseInput)
 			{
-				MessageHandler->OnRawMouseMove( motionEvent.xrel, motionEvent.yrel );
+					LinuxCursor->AddOffset(motionEvent.xrel, motionEvent.yrel);
+ 					MessageHandler->OnRawMouseMove(motionEvent.xrel, motionEvent.yrel);
 			}
 			else
 			{
@@ -1025,6 +1026,7 @@ void FLinuxApplication::SetHighPrecisionMouseMode( const bool Enable, const TSha
 {
 	MessageHandler->OnCursorSet();
 	bUsingHighPrecisionMouseInput = Enable;
+	((FLinuxCursor*)Cursor.Get())->ResetOffset();
 }
 
 
diff --git a/Engine/Source/Runtime/Core/Private/Linux/LinuxCursor.cpp b/Engine/Source/Runtime/Core/Private/Linux/LinuxCursor.cpp
--- a/Engine/Source/Runtime/Core/Private/Linux/LinuxCursor.cpp
+++ b/Engine/Source/Runtime/Core/Private/Linux/LinuxCursor.cpp
@@ -5,7 +5,10 @@
 #include "LinuxCursor.h"
 #include "LinuxWindow.h"
 
-FLinuxCursor::FLinuxCursor() : bHidden(false)
+FLinuxCursor::FLinuxCursor()
+	: 	bHidden(false)
+	,	AccumulatedOffsetX(0)
+	,	AccumulatedOffsetY(0)
 {
 	if (!FPlatformMisc::PlatformInitMultimedia()) //	will not initialize more than once
 	{
@@ -162,6 +165,8 @@ FVector2D FLinuxCursor::GetPosition() const
 	int CursorX, CursorY;
 
 	SDL_GetGlobalMouseState(&CursorX, &CursorY);
+	CursorX += AccumulatedOffsetX;
+	CursorY += AccumulatedOffsetY;
 
 	return FVector2D( CursorX, CursorY );
 }
@@ -176,6 +181,17 @@ void FLinuxCursor::SetPosition( const int32 X, const int32 Y )
 	SDL_WarpMouseInWindow( NULL, X - WndX, Y - WndY );
 }
 
+void FLinuxCursor::AddOffset(const int32 DX, const int32 DY)
+{
+	AccumulatedOffsetX += DX; 
+	AccumulatedOffsetY += DY;
+}
+
+void FLinuxCursor::ResetOffset()
+{
+	AccumulatedOffsetX = AccumulatedOffsetY = 0;
+}
+
 void FLinuxCursor::SetType( const EMouseCursor::Type InNewCursor )
 {	
 	checkf( InNewCursor < EMouseCursor::TotalCursorCount, TEXT("Invalid cursor(%d) supplied"), InNewCursor );
diff --git a/Engine/Source/Runtime/Core/Public/Linux/LinuxCursor.h b/Engine/Source/Runtime/Core/Public/Linux/LinuxCursor.h
--- a/Engine/Source/Runtime/Core/Public/Linux/LinuxCursor.h
+++ b/Engine/Source/Runtime/Core/Public/Linux/LinuxCursor.h
@@ -40,7 +40,13 @@ public:
 	bool IsHidden();
 
 	uint32 GetCursorEvent();
-	
+
+	/** Updates accumulated offset which is used for maintaining a "shadow" cursor position */
+	void AddOffset(const int32 DX, const int32 DY);
+
+	/** Resets accumulated offsets to (0,0), */
+	void ResetOffset();
+
 private:
 
 	EMouseCursor::Type CurrentType;
@@ -52,5 +58,10 @@ private:
 	FIntRect CursorClipRect;
 
 	uint32 CursorEvent;
-};
 
+	/** Accumulated X offset (from relative mouse movement events) used for reconstructing global cursor position */
+	int32 AccumulatedOffsetX;
+
+	/** Accumulated Y offset (from relative mouse movement events) used for reconstructing global cursor position */
+	int32 AccumulatedOffsetY;
+};
