Core: Code and documentation cleanup pass

diff --git a/Engine/Source/Runtime/Core/Private/Windows/WindowsPlatformProcess.cpp b/Engine/Source/Runtime/Core/Private/Windows/WindowsPlatformProcess.cpp
--- a/Engine/Source/Runtime/Core/Private/Windows/WindowsPlatformProcess.cpp
+++ b/Engine/Source/Runtime/Core/Private/Windows/WindowsPlatformProcess.cpp
@@ -280,76 +280,90 @@ FProcHandle FWindowsPlatformProcess::CreateProc( const TCHAR* URL, const TCHAR*
 {
 	//UE_LOG(LogWindows, Log,  TEXT("CreateProc %s %s"), URL, Parms );
 
-	FString CommandLine = FString::Printf(TEXT("\"%s\" %s"), URL, Parms);
-
-	PROCESS_INFORMATION ProcInfo;
+	// initialize process attributes
 	SECURITY_ATTRIBUTES Attr;
-	Attr.nLength = sizeof(SECURITY_ATTRIBUTES);
-	Attr.lpSecurityDescriptor = NULL;
-	Attr.bInheritHandle = true;
+	{
+		Attr.nLength = sizeof(SECURITY_ATTRIBUTES);
+		Attr.lpSecurityDescriptor = NULL;
+		Attr.bInheritHandle = true;
+	}
 
+	// initialize process creation flags
 	uint32 CreateFlags = NORMAL_PRIORITY_CLASS;
-	if (PriorityModifier < 0)
 	{
-		if (PriorityModifier == -1)
-		{
-			CreateFlags = BELOW_NORMAL_PRIORITY_CLASS;
-		}
-		else
+		if (PriorityModifier < 0)
 		{
-			CreateFlags = IDLE_PRIORITY_CLASS;
+			CreateFlags = (PriorityModifier == -1) ? BELOW_NORMAL_PRIORITY_CLASS : IDLE_PRIORITY_CLASS;
 		}
-	}
-	else if (PriorityModifier > 0)
-	{
-		if (PriorityModifier == 1)
+		else if (PriorityModifier > 0)
 		{
-			CreateFlags = ABOVE_NORMAL_PRIORITY_CLASS;
+			CreateFlags = (PriorityModifier == 1) ? ABOVE_NORMAL_PRIORITY_CLASS : HIGH_PRIORITY_CLASS;
 		}
-		else
+
+		if (bLaunchDetached)
 		{
-			CreateFlags = HIGH_PRIORITY_CLASS;
+			CreateFlags |= DETACHED_PROCESS;
 		}
 	}
-	if (bLaunchDetached)
-	{
-		CreateFlags |= DETACHED_PROCESS;
-	}
+
+	// initialize window flags
 	uint32 dwFlags = 0;
 	uint16 ShowWindowFlags = SW_HIDE;
-	if (bLaunchReallyHidden)
-	{
-		dwFlags = STARTF_USESHOWWINDOW;
-	}
-	else if (bLaunchHidden)
 	{
-		dwFlags = STARTF_USESHOWWINDOW;
-		ShowWindowFlags = SW_SHOWMINNOACTIVE;
-	}
-	if (PipeWrite)
-	{
-		dwFlags |= STARTF_USESTDHANDLES;
+		if (bLaunchReallyHidden)
+		{
+			dwFlags = STARTF_USESHOWWINDOW;
+		}
+		else if (bLaunchHidden)
+		{
+			dwFlags = STARTF_USESHOWWINDOW;
+			ShowWindowFlags = SW_SHOWMINNOACTIVE;
+		}
+
+		if (PipeWrite != nullptr)
+		{
+			dwFlags |= STARTF_USESTDHANDLES;
+		}
 	}
 
-	STARTUPINFO StartupInfo = { sizeof(STARTUPINFO), NULL, NULL, NULL,
-		CW_USEDEFAULT, CW_USEDEFAULT, CW_USEDEFAULT, CW_USEDEFAULT,
-		0, 0, 0, dwFlags, ShowWindowFlags, 0, NULL,
-		::GetStdHandle(ProcessConstants::WIN_STD_INPUT_HANDLE), HANDLE(PipeWrite), HANDLE(PipeWrite) };
-	if (!CreateProcess(NULL, CommandLine.GetCharArray().GetData(), &Attr, &Attr, true, CreateFlags,
-		NULL, OptionalWorkingDirectory, &StartupInfo, &ProcInfo ) )
+	// initialize startup info
+	STARTUPINFO StartupInfo = {
+		sizeof(STARTUPINFO),
+		NULL, NULL, NULL,
+		CW_USEDEFAULT,
+		CW_USEDEFAULT,
+		CW_USEDEFAULT,
+		CW_USEDEFAULT,
+		0, 0, 0,
+		dwFlags,
+		ShowWindowFlags,
+		0, NULL,
+		::GetStdHandle(ProcessConstants::WIN_STD_INPUT_HANDLE),
+		HANDLE(PipeWrite),
+		HANDLE(PipeWrite)
+	};
+
+	// create the child process
+	FString CommandLine = FString::Printf(TEXT("\"%s\" %s"), URL, Parms);
+	PROCESS_INFORMATION ProcInfo;
+
+	if (!CreateProcess(NULL, CommandLine.GetCharArray().GetData(), &Attr, &Attr, true, CreateFlags, NULL, OptionalWorkingDirectory, &StartupInfo, &ProcInfo))
 	{
-		if (OutProcessID)
+		if (OutProcessID != nullptr)
 		{
 			*OutProcessID = 0;
 		}
+
 		return FProcHandle();
 	}
-	if (OutProcessID)
+
+	if (OutProcessID != nullptr)
 	{
 		*OutProcessID = ProcInfo.dwProcessId;
 	}
-	// Close the thread handle.
-	::CloseHandle( ProcInfo.hThread );
+
+	::CloseHandle(ProcInfo.hThread);
+
 	return FProcHandle(ProcInfo.hProcess);
 }
 
diff --git a/Engine/Source/Runtime/Core/Public/Windows/WindowsPlatformProcess.h b/Engine/Source/Runtime/Core/Public/Windows/WindowsPlatformProcess.h
--- a/Engine/Source/Runtime/Core/Public/Windows/WindowsPlatformProcess.h
+++ b/Engine/Source/Runtime/Core/Public/Windows/WindowsPlatformProcess.h
@@ -136,7 +136,10 @@ protected:
 
 private:
 
-	// Since Windows can only have one directory at a time, this stack is used to reset the previous DllDirectory
+	/**
+	 * Since Windows can only have one directory at a time,
+	 * this stack is used to reset the previous DllDirectory.
+	 */
 	static TArray<FString> DllDirectoryStack;
 };
 
