FParse::Value parses options with arguments properly when fully quoted.

diff --git a/Engine/Source/Runtime/Core/Private/Misc/Parse.cpp b/Engine/Source/Runtime/Core/Private/Misc/Parse.cpp
--- a/Engine/Source/Runtime/Core/Private/Misc/Parse.cpp
+++ b/Engine/Source/Runtime/Core/Private/Misc/Parse.cpp
@@ -110,46 +110,68 @@ void ConsoleCommandLibrary_DumpLibrary(UWorld* InWorld, FExec& SubSystem, const
 // Get a string from a text string.
 //
 bool FParse::Value(
-	const TCHAR* Stream, 
-	const TCHAR* Match,
-	TCHAR*		 Value,
-	int32			 MaxLen,
-	bool		 bShouldStopOnComma
+	const TCHAR*	Stream,
+	const TCHAR*	Match,
+	TCHAR*			Value,
+	int32			MaxLen,
+	bool			bShouldStopOnComma
 )
 {
 	const TCHAR* Found = FCString::Strfind(Stream,Match);
-	const TCHAR* Start;
 
-	if( Found )
+	if (!Found)
 	{
-		Start = Found + FCString::Strlen(Match);
-		if( *Start == '\x22' )
+		return false;
+	}
+
+	const TCHAR* Start = Found + FCString::Strlen(Match);
+
+	// Check for quoted arguments' string with spaces
+	// -Option="Value1 Value2"
+	//         ^~~~Start
+	bool bArgumentsQuoted = *Start == '"';
+
+	// Number of characters we can look back from found looking for first parenthesis.
+	uint32 AllowedBacktraceCharactersCount = Found - Stream;
+
+	// Check for fully quoted string with spaces
+	bool bFullyQuoted = 
+		// "Option=Value1 Value2"
+		//  ^~~~Found
+		(AllowedBacktraceCharactersCount > 0 && (*(Found - 1) == '"'))
+		// "-Option=Value1 Value2"
+		//   ^~~~Found
+		|| (AllowedBacktraceCharactersCount > 1 && ((*(Found - 1) == '-') && (*(Found - 2) == '"')));
+
+	if (bArgumentsQuoted || bFullyQuoted)
+	{
+		// Skip quote character if only params were quoted.
+		int32 QuoteCharactersToSkip = bArgumentsQuoted ? 1 : 0;
+		FCString::Strncpy(Value, Start + QuoteCharactersToSkip, MaxLen);
+
+		Value[MaxLen-1]=0;
+		TCHAR* Temp = FCString::Strstr( Value, TEXT("\x22") );
+		if (Temp != nullptr)
 		{
-			// Quoted string with spaces.
-			FCString::Strncpy( Value, Start+1, MaxLen );
-			Value[MaxLen-1]=0;
-			TCHAR* Temp = FCString::Strstr( Value, TEXT("\x22") );
-			if( Temp != NULL )
-				*Temp=0;
+			*Temp = 0;
 		}
-		else
+	}
+	else
+	{
+		// Non-quoted string without spaces.
+		FCString::Strncpy( Value, Start, MaxLen );
+		Value[MaxLen-1]=0;
+		TCHAR* Temp;
+		Temp = FCString::Strstr( Value, TEXT(" ")  ); if( Temp ) *Temp=0;
+		Temp = FCString::Strstr( Value, TEXT("\r") ); if( Temp ) *Temp=0;
+		Temp = FCString::Strstr( Value, TEXT("\n") ); if( Temp ) *Temp=0;
+		Temp = FCString::Strstr( Value, TEXT("\t") ); if( Temp ) *Temp=0;
+		if (bShouldStopOnComma)
 		{
-			// Non-quoted string without spaces.
-			FCString::Strncpy( Value, Start, MaxLen );
-			Value[MaxLen-1]=0;
-			TCHAR* Temp;
-			Temp = FCString::Strstr( Value, TEXT(" ")  ); if( Temp ) *Temp=0;
-			Temp = FCString::Strstr( Value, TEXT("\r") ); if( Temp ) *Temp=0;
-			Temp = FCString::Strstr( Value, TEXT("\n") ); if( Temp ) *Temp=0;
-			Temp = FCString::Strstr( Value, TEXT("\t") ); if( Temp ) *Temp=0;
-			if (bShouldStopOnComma)
-			{
-				Temp = FCString::Strstr( Value, TEXT(",")  ); if( Temp ) *Temp=0;
-			}
+			Temp = FCString::Strstr( Value, TEXT(",")  ); if( Temp ) *Temp=0;
 		}
-		return 1;
 	}
-	else return 0;
+	return true;
 }
 
 //
