Do not raise windows when bringing to front.

- Implements the intended behavior of the function.
- PR #614 contributed by [REDACTED_PERSON].

diff --git a/Engine/Source/Runtime/Core/Private/Linux/LinuxWindow.cpp b/Engine/Source/Runtime/Core/Private/Linux/LinuxWindow.cpp
--- a/Engine/Source/Runtime/Core/Private/Linux/LinuxWindow.cpp
+++ b/Engine/Source/Runtime/Core/Private/Linux/LinuxWindow.cpp
@@ -184,7 +184,15 @@ void FLinuxWindow::MoveWindowTo( int32 X, int32 Y )
  */
 void FLinuxWindow::BringToFront( bool bForce )
 {
-	SDL_RaiseWindow( HWnd );
+	if (bForce)
+	{
+		SDL_RaiseWindow(HWnd);
+	}
+	else
+	{
+		// FIXME: we don't bring it to front here now
+		SDL_ShowWindow(HWnd);
+	}
 }
 
 /** Native windows should implement this function by asking the OS to destroy OS-specific resource associated with the window (e.g. Win32 window handle) */
