Fix runnable threads trashing the main thread's affinities.

diff --git a/Engine/Source/Runtime/Core/Private/HAL/PThreadRunnableThread.h b/Engine/Source/Runtime/Core/Private/HAL/PThreadRunnableThread.h
--- a/Engine/Source/Runtime/Core/Private/HAL/PThreadRunnableThread.h
+++ b/Engine/Source/Runtime/Core/Private/HAL/PThreadRunnableThread.h
@@ -39,6 +39,11 @@ protected:
 	 */
 	EThreadPriority ThreadPriority;
 
+	/**
+	 * The CPU Mask to indicate which CPUs the thread will run on.
+	 */
+	uint64 ThreadAffinityMask;
+
 	/**
 	* ID set during thread creation
 	*/
@@ -184,6 +189,9 @@ protected:
 		// cache the thread ID for this thread (defined by the platform)
 		ThisThread->ThreadID = FPlatformTLS::GetCurrentThreadId();
 
+		// set the affinity.  This function sets affinity on the current thread, so don't call in the Create function which will trash the main thread affinity.
+		FPlatformProcess::SetThreadAffinityMask(ThisThread->ThreadAffinityMask);		
+
 		// run the thread!
 		ThisThread->PreRun();
 		ThisThread->Run();
@@ -250,6 +258,7 @@ public:
 		, Runnable(NULL)
 		, ThreadInitSyncEvent(NULL)
 		, ThreadPriority(TPri_Normal)
+		, ThreadAffinityMask(FPlatformAffinity::GetNoAffinityMask())
 		, ThreadID(0)
 		, ThreadIsRunning(false)
 	{
@@ -340,6 +349,8 @@ protected:
 		ThreadInitSyncEvent	= FPlatformProcess::CreateSynchEvent(true);
 		// A name for the thread in for debug purposes. _ThreadProc will set it.
 		ThreadName = InThreadName ? InThreadName : TEXT("Unnamed UE4");
+		ThreadAffinityMask = InThreadAffinityMask;
+
 		// Create the new thread
 		bool ThreadCreated = SpinPthread(&Thread, GetThreadEntryPoint(), InStackSize, this);
 		// If it fails, clear all the vars
@@ -351,10 +362,7 @@ protected:
 			ThreadInitSyncEvent->Wait((uint32)-1); // infinite wait
 
 			// set the priority
-			SetThreadPriority(InThreadPri);
-
-			// set the affinity
-			FPlatformProcess::SetThreadAffinityMask( InThreadAffinityMask );
+			SetThreadPriority(InThreadPri);			
 		}
 		else // If it fails, clear all the vars
 		{
