Fix non-ANSI characters in source file, which corrupts doxygen output.

diff --git a/Engine/Source/Runtime/Core/Private/Windows/WindowsPlatformMisc.cpp b/Engine/Source/Runtime/Core/Private/Windows/WindowsPlatformMisc.cpp
--- a/Engine/Source/Runtime/Core/Private/Windows/WindowsPlatformMisc.cpp
+++ b/Engine/Source/Runtime/Core/Private/Windows/WindowsPlatformMisc.cpp
@@ -2308,13 +2308,13 @@ bool FWindowsPlatformMisc::IsRunningOnBattery()
 	GetSystemPowerStatus(&status);
 	switch(status.BatteryFlag)
 	{
-	case 4://	"Critical—the battery capacity is at less than five percent"
-	case 2://	"Low—the battery capacity is at less than 33 percent"
-	case 1://	"High—the battery capacity is at more than 66 percent"
+	case 4://	"Critical-the battery capacity is at less than five percent"
+	case 2://	"Low-the battery capacity is at less than 33 percent"
+	case 1://	"High-the battery capacity is at more than 66 percent"
 	case 8://	"Charging"
 		return true;
 	case 128://	"No system battery" - desktop, NB: UPS don't count as batteries under Windows
-	case 255://	"Unknown status—unable to read the battery flag information"
+	case 255://	"Unknown status-unable to read the battery flag information"
 	default:
 		return false;
 	}
