- Add an iPad Air 2 DeviceProfile

diff --git a/Engine/Config/BaseDeviceProfiles.ini b/Engine/Config/BaseDeviceProfiles.ini
--- a/Engine/Config/BaseDeviceProfiles.ini
+++ b/Engine/Config/BaseDeviceProfiles.ini
@@ -68,6 +68,13 @@ DeviceType=IOS
 BaseProfileName=IOS
 +CVars=r.BloomQuality=1
 
+[iPadAir2 DeviceProfile]
+DeviceType=IOS
+BaseProfileName=IOS
++CVars=r.BloomQuality=1
++CVars=r.DepthOfFieldQuality=1
++CVars=r.LightShaftQuality=1
+
 [iPadMini DeviceProfile]
 DeviceType=IOS
 BaseProfileName=IOS
diff --git a/Engine/Source/Runtime/Core/Private/IOS/IOSPlatformMisc.cpp b/Engine/Source/Runtime/Core/Private/IOS/IOSPlatformMisc.cpp
--- a/Engine/Source/Runtime/Core/Private/IOS/IOSPlatformMisc.cpp
+++ b/Engine/Source/Runtime/Core/Private/IOS/IOSPlatformMisc.cpp
@@ -427,12 +427,16 @@ FIOSPlatformMisc::EIOSDevice FIOSPlatformMisc::GetIOSDeviceType()
 			{
 				DeviceType = IOS_IPadAir;
 			}
-
+		}
+		// iPad Air 2
+		else if (Major == 5)
+		{
+			DeviceType = IOS_IPadAir2;
 		}
 		// Default to highest settings currently available for any future device
-		else if (Major > 4)
+		else if (Major > 5)
 		{
-			DeviceType = IOS_IPadAir;
+			DeviceType = IOS_IPadAir2;
 		}
 	}
 	// iPhones
diff --git a/Engine/Source/Runtime/Core/Public/IOS/IOSPlatformMisc.h b/Engine/Source/Runtime/Core/Public/IOS/IOSPlatformMisc.h
--- a/Engine/Source/Runtime/Core/Public/IOS/IOSPlatformMisc.h
+++ b/Engine/Source/Runtime/Core/Public/IOS/IOSPlatformMisc.h
@@ -95,6 +95,7 @@ struct CORE_API FIOSPlatformMisc : public FGenericPlatformMisc
 		IOS_IPadMini,
 		IOS_IPadMini2,
 		IOS_IPadAir,
+		IOS_IPadAir2,
 		IOS_IPhone6,
 		IOS_IPhone6Plus,
 		IOS_Unknown,
@@ -117,6 +118,7 @@ struct CORE_API FIOSPlatformMisc : public FGenericPlatformMisc
 			L"IPadMini",
 			L"IPadMini2",
 			L"IPadAir",
+			L"IPadAir2",
 			L"IPhone6",
 			L"IPhone6Plus",
 			L"Unknown",
