cleanup cvar help, dump current state when using "cvar ?"

diff --git a/Engine/Source/Runtime/Core/Private/HAL/ConsoleManager.cpp b/Engine/Source/Runtime/Core/Private/HAL/ConsoleManager.cpp
--- a/Engine/Source/Runtime/Core/Private/HAL/ConsoleManager.cpp
+++ b/Engine/Source/Runtime/Core/Private/HAL/ConsoleManager.cpp
@@ -14,6 +14,31 @@ DEFINE_LOG_CATEGORY_STATIC(LogConsoleManager, Log, All);
 
 static inline bool IsWhiteSpace(TCHAR Value) { return Value == TCHAR(' '); }
 
+// @param In must not be 0
+bool IsGoodHelpString(const TCHAR* In)
+{
+	check(In);
+
+	if(*In == 0)
+	{
+		return false;
+	}
+
+	bool bGoodEndChar = true;
+
+	while(TCHAR c = *In++)
+	{
+		bGoodEndChar = true;
+
+		if(c == L'\n' || c == L'\t' || c == L' ' || c == L'\r')
+		{
+			bGoodEndChar = false;
+		}
+	}
+
+	return bGoodEndChar;
+}
+
 // Get human readable string
 // @return never 0
 static const TCHAR* GetSetByTCHAR(EConsoleVariableFlags InSetBy)
@@ -47,10 +72,9 @@ public:
 	 * @param InHelp must not be 0, must not be empty
 	 */
 	FConsoleVariableBase(const TCHAR* InHelp, EConsoleVariableFlags InFlags)
-		:Help(InHelp), Flags(InFlags), bWarnedAboutThreadSafety(false)
+		:Flags(InFlags), bWarnedAboutThreadSafety(false)
 	{
-		check(InHelp);
-		//check(*Help != 0); for now disabled as there callstack when we crash early during engine init
+		SetHelp(InHelp);
 	}
 
 	// interface IConsoleVariable -----------------------------------
@@ -62,9 +86,11 @@ public:
 	virtual void SetHelp(const TCHAR* Value)
 	{
 		check(Value);
-		//check(*Value != 0);
 
 		Help = Value;
+
+		// for now disabled as there is no good callstack when we crash early during engine init
+//		ensure(IsGoodHelpString(Value));
 	}
 	virtual EConsoleVariableFlags GetFlags() const
 	{
@@ -165,8 +191,6 @@ protected: // -----------------------------------------
 	}
 };
 
-
-
 class FConsoleCommandBase : public IConsoleCommand
 {
 public:
@@ -942,7 +966,7 @@ bool FConsoleManager::ProcessUserConsoleInput(const TCHAR* InInput, FOutputDevic
 		TArray< FString > Args;
 		FString( It ).ParseIntoArrayWS( &Args );
 
-		const bool bShowHelp = Args.Num() == 1 && Args[0] == TEXT("?") ? true : false;
+		const bool bShowHelp = Args.Num() == 1 && Args[0] == TEXT("?");
 		if( bShowHelp )
 		{
 			// get help
@@ -959,10 +983,11 @@ bool FConsoleManager::ProcessUserConsoleInput(const TCHAR* InInput, FOutputDevic
 	{
 		// Process variable
 
+		bool bShowCurrentState = false;
+
 		if(*It == 0)
 		{
-			// get current state
-			Ar.Logf(TEXT("%s = \"%s\"      LastSetBy: %s"), *Param1, *CVar->GetString(), GetSetByTCHAR(CVar->GetFlags()));
+			bShowCurrentState = true;
 		}
 		else
 		{
@@ -982,6 +1007,7 @@ bool FConsoleManager::ProcessUserConsoleInput(const TCHAR* InInput, FOutputDevic
 			{
 				// get help
 				Ar.Logf(TEXT("HELP for '%s'%s:\n%s"), *Param1, bReadOnly ? TEXT("(ReadOnly)") : TEXT(""), CVar->GetHelp());
+				bShowCurrentState = true;
 			}
 			else
 			{
@@ -1000,6 +1026,11 @@ bool FConsoleManager::ProcessUserConsoleInput(const TCHAR* InInput, FOutputDevic
 				}
 			}
 		}
+
+		if(bShowCurrentState)
+		{
+			Ar.Logf(TEXT("%s = \"%s\"      LastSetBy: %s"), *Param1, *CVar->GetString(), GetSetByTCHAR(CVar->GetFlags()));
+		}
 	}
 
 	return true;
diff --git a/Engine/Source/Runtime/Engine/Private/ContentStreaming.cpp b/Engine/Source/Runtime/Engine/Private/ContentStreaming.cpp
--- a/Engine/Source/Runtime/Engine/Private/ContentStreaming.cpp
+++ b/Engine/Source/Runtime/Engine/Private/ContentStreaming.cpp
@@ -150,7 +150,7 @@ static TAutoConsoleVariable<int32> CVarStreamingUseFixedPoolSize(
 static TAutoConsoleVariable<int32> CVarStreamingPoolSize(
 	TEXT("r.Streaming.PoolSize"),
 	-1,
-	TEXT("-1: Default texture pool size, otherwise the size in MB\n"),
+	TEXT("-1: Default texture pool size, otherwise the size in MB"),
 	ECVF_Scalability);
 
 /** Streaming priority: Linear distance factor from 0 to MAX_STREAMINGDISTANCE. */
diff --git a/Engine/Source/Runtime/Engine/Private/PrimitiveComponent.cpp b/Engine/Source/Runtime/Engine/Private/PrimitiveComponent.cpp
--- a/Engine/Source/Runtime/Engine/Private/PrimitiveComponent.cpp
+++ b/Engine/Source/Runtime/Engine/Private/PrimitiveComponent.cpp
@@ -37,8 +37,10 @@
 DEFINE_LOG_CATEGORY_STATIC(LogPrimitiveComponent, Log, All);
 
 static FAutoConsoleVariable CVarAllowCachedOverlaps(
-	TEXT("AllowCachedOverlaps"), 1,
-	TEXT("0: disable cached overlaps, 1: enable\n"));
+	TEXT("p.AllowCachedOverlaps"), 
+	1,
+	TEXT("Primitive Component physics\n")
+	TEXT("0: disable cached overlaps, 1: enable (default)"));
 
 static TAutoConsoleVariable<float> CVarInitialOverlapTolerance(
 	TEXT("p.InitialOverlapTolerance"),
diff --git a/Engine/Source/Runtime/RHI/Private/RHICommandList.cpp b/Engine/Source/Runtime/RHI/Private/RHICommandList.cpp
--- a/Engine/Source/Runtime/RHI/Private/RHICommandList.cpp
+++ b/Engine/Source/Runtime/RHI/Private/RHICommandList.cpp
@@ -38,7 +38,7 @@ static TAutoConsoleVariable<int32> CVarRHICmdBypass(
 static TAutoConsoleVariable<int32> CVarRHICmdUseParallelAlgorithms(
 	TEXT("r.RHICmdUseParallelAlgorithms"),
 	1,
-	TEXT("True to use parallel algorithms. Ignored if r.RHICmdBypass is 1.\n"));
+	TEXT("True to use parallel algorithms. Ignored if r.RHICmdBypass is 1."));
 
 TAutoConsoleVariable<int32> CVarRHICmdWidth(
 	TEXT("r.RHICmdWidth"), 
@@ -49,11 +49,11 @@ TAutoConsoleVariable<int32> CVarRHICmdWidth(
 static TAutoConsoleVariable<int32> CVarRHICmdUseDeferredContexts(
 	TEXT("r.RHICmdUseDeferredContexts"),
 	1,
-	TEXT("True to use deferred contexts to parallelize command list execution.\n"));
+	TEXT("True to use deferred contexts to parallelize command list execution."));
 static TAutoConsoleVariable<int32> CVarRHIDeferredContextWidth(
 	TEXT("r.RHIDeferredContextWidth"),
 	99999,
-	TEXT("Number of pieces to break a parallel translate task into. A very large number is one task per original cmd list.\n"));
+	TEXT("Number of pieces to break a parallel translate task into. A very large number is one task per original cmd list."));
 #endif
 
 RHI_API FRHICommandListExecutor GRHICommandList;
diff --git a/Engine/Source/Runtime/Renderer/Private/CompositionLighting/PostProcessDeferredDecals.h b/Engine/Source/Runtime/Renderer/Private/CompositionLighting/PostProcessDeferredDecals.h
--- a/Engine/Source/Runtime/Renderer/Private/CompositionLighting/PostProcessDeferredDecals.h
+++ b/Engine/Source/Runtime/Renderer/Private/CompositionLighting/PostProcessDeferredDecals.h
@@ -8,7 +8,7 @@
 
 #include "RenderingCompositionGraph.h"
 
-// ePId_Input0: SceneColor
+// ePId_Input0: SceneColor (not needed for DBuffer decals)
 // derives from TRenderingCompositePassBase<InputCount, OutputCount> 
 class FRCPassPostProcessDeferredDecals : public TRenderingCompositePassBase<1, 1>
 {
diff --git a/Engine/Source/Runtime/Renderer/Private/DeferredShadingRenderer.cpp b/Engine/Source/Runtime/Renderer/Private/DeferredShadingRenderer.cpp
--- a/Engine/Source/Runtime/Renderer/Private/DeferredShadingRenderer.cpp
+++ b/Engine/Source/Runtime/Renderer/Private/DeferredShadingRenderer.cpp
@@ -28,7 +28,7 @@ TAutoConsoleVariable<int32> CVarEarlyZPass(
 	TEXT("  0: off\n")
 	TEXT("  1: only if not masked, and only if large on the screen\n")
 	TEXT("  2: all opaque (including masked)\n")
-	TEXT("  x: use built in heuristic (default is 3)\n"),
+	TEXT("  x: use built in heuristic (default is 3)"),
 	ECVF_Default);
 
 
@@ -532,7 +532,7 @@ public:
 static TAutoConsoleVariable<int32> CVarRHICmdBasePassDeferredContexts(
 	TEXT("r.RHICmdBasePassDeferredContexts"),
 	1,
-	TEXT("True to use deferred contexts to parallelize base pass command list execution.\n"));
+	TEXT("True to use deferred contexts to parallelize base pass command list execution."));
 
 void FDeferredShadingSceneRenderer::RenderBasePassViewParallel(FViewInfo& View, FRHICommandList& ParentCmdList, bool& OutDirty)
 {
@@ -1248,7 +1248,7 @@ public:
 static TAutoConsoleVariable<int32> CVarRHICmdPrePassDeferredContexts(
 	TEXT("r.RHICmdPrePassDeferredContexts"),
 	1,
-	TEXT("True to use deferred contexts to parallelize prepass command list execution.\n"));
+	TEXT("True to use deferred contexts to parallelize prepass command list execution."));
 
 void FDeferredShadingSceneRenderer::RenderPrePassViewParallel(const FViewInfo& View, FRHICommandList& ParentCmdList, bool& OutDirty)
 {
diff --git a/Engine/Source/Runtime/Renderer/Private/PostProcess/SceneRenderTargets.cpp b/Engine/Source/Runtime/Renderer/Private/PostProcess/SceneRenderTargets.cpp
--- a/Engine/Source/Runtime/Renderer/Private/PostProcess/SceneRenderTargets.cpp
+++ b/Engine/Source/Runtime/Renderer/Private/PostProcess/SceneRenderTargets.cpp
@@ -77,7 +77,7 @@ static TAutoConsoleVariable<int32> CVarMobileMSAA(
 	TEXT("Use MSAA instead of Temporal AA on mobile:\n")
 	TEXT("1: Use Temporal AA (MSAA disabled)\n")
 	TEXT("2: Use 2x MSAA (Temporal AA disabled)\n")
-	TEXT("4: Use 4x MSAA (Temporal AA disabled)\n"),
+	TEXT("4: Use 4x MSAA (Temporal AA disabled)"),
 	ECVF_RenderThreadSafe
 	);
 
diff --git a/Engine/Source/Runtime/Renderer/Private/SceneVisibility.cpp b/Engine/Source/Runtime/Renderer/Private/SceneVisibility.cpp
--- a/Engine/Source/Runtime/Renderer/Private/SceneVisibility.cpp
+++ b/Engine/Source/Runtime/Renderer/Private/SceneVisibility.cpp
@@ -82,7 +82,7 @@ static TAutoConsoleVariable<int32> CVarLightShaftQuality(
 	1,
 	TEXT("Defines the light shaft quality.\n")
 	TEXT("  0: off\n")
-	TEXT("  1: on (default)\n"),
+	TEXT("  1: on (default)"),
 	ECVF_Scalability | ECVF_RenderThreadSafe);
 
 static TAutoConsoleVariable<float> CVarStaticMeshLODDistanceScale(
@@ -1176,7 +1176,7 @@ static TAutoConsoleVariable<int32> CVarParallelInitViews(
 #else
 	1,  
 #endif
-	TEXT("Toggles parallel init views.\n"),
+	TEXT("Toggles parallel init views."),
 	ECVF_RenderThreadSafe
 	);
 
@@ -1835,7 +1835,7 @@ void FSceneRenderer::PreVisibilityFrameSetup(FRHICommandListImmediate& RHICmdLis
 static TAutoConsoleVariable<int32> CVarLegacySingleThreadedRelevance(
 	TEXT("r.LegacySingleThreadedRelevance"),
 	0,  
-	TEXT("Toggles the lgacy codepath for view relevance.\n"),
+	TEXT("Toggles the legacy codepath for view relevance."),
 	ECVF_RenderThreadSafe
 	);
 
diff --git a/Engine/Source/Runtime/Renderer/Private/ShadowRendering.cpp b/Engine/Source/Runtime/Renderer/Private/ShadowRendering.cpp
--- a/Engine/Source/Runtime/Renderer/Private/ShadowRendering.cpp
+++ b/Engine/Source/Runtime/Renderer/Private/ShadowRendering.cpp
@@ -1501,14 +1501,14 @@ void FProjectedShadowInfo::SetStateForDepth(FRHICommandList& RHICmdList)
 static TAutoConsoleVariable<int32> CVarParallelShadows(
 	TEXT("r.ParallelShadows"),
 	0,
-	TEXT("Toggles parallel shadow rendering. Parallel rendering must be enabled for this to have an effect.\n"),
+	TEXT("Toggles parallel shadow rendering. Parallel rendering must be enabled for this to have an effect."),
 	ECVF_RenderThreadSafe
 	);
 
 static TAutoConsoleVariable<int32> CVarRHICmdShadowDeferredContexts(
 	TEXT("r.RHICmdShadowDeferredContexts"),
 	0,
-	TEXT("True to use deferred contexts to parallelize shadow command list execution.\n"));
+	TEXT("True to use deferred contexts to parallelize shadow command list execution."));
 
 class FShadowParallelCommandListSet : public FParallelCommandListSet
 {
diff --git a/Engine/Source/Runtime/Renderer/Private/TranslucentRendering.cpp b/Engine/Source/Runtime/Renderer/Private/TranslucentRendering.cpp
--- a/Engine/Source/Runtime/Renderer/Private/TranslucentRendering.cpp
+++ b/Engine/Source/Runtime/Renderer/Private/TranslucentRendering.cpp
@@ -887,7 +887,7 @@ public:
 static TAutoConsoleVariable<int32> CVarRHICmdTranslucencyPassDeferredContexts(
 	TEXT("r.RHICmdTranslucencyPassDeferredContexts"),
 	1,
-	TEXT("True to use deferred contexts to parallelize base pass command list execution.\n"));
+	TEXT("True to use deferred contexts to parallelize base pass command list execution."));
 
 void FDeferredShadingSceneRenderer::RenderTranslucencyParallel(FRHICommandListImmediate& RHICmdList)
 {
@@ -999,7 +999,7 @@ void FDeferredShadingSceneRenderer::RenderTranslucencyParallel(FRHICommandListIm
 static TAutoConsoleVariable<int32> CVarParallelTranslucency(
 	TEXT("r.ParallelTranslucency"),
 	1,
-	TEXT("Toggles parallel translucency rendering. Parallel rendering must be enabled for this to have an effect.\n"),
+	TEXT("Toggles parallel translucency rendering. Parallel rendering must be enabled for this to have an effect."),
 	ECVF_RenderThreadSafe
 	);
 
diff --git a/Engine/Source/Runtime/Renderer/Private/VelocityRendering.cpp b/Engine/Source/Runtime/Renderer/Private/VelocityRendering.cpp
--- a/Engine/Source/Runtime/Renderer/Private/VelocityRendering.cpp
+++ b/Engine/Source/Runtime/Renderer/Private/VelocityRendering.cpp
@@ -572,7 +572,7 @@ public:
 static TAutoConsoleVariable<int32> CVarRHICmdVelocityPassDeferredContexts(
 	TEXT("r.RHICmdVelocityPassDeferredContexts"),
 	1,
-	TEXT("True to use deferred contexts to parallelize velocity pass command list execution.\n"));
+	TEXT("True to use deferred contexts to parallelize velocity pass command list execution."));
 
 void FDeferredShadingSceneRenderer::RenderVelocitiesInnerParallel(FRHICommandListImmediate& RHICmdList, TRefCountPtr<IPooledRenderTarget>& VelocityRT)
 {
@@ -645,7 +645,7 @@ void FDeferredShadingSceneRenderer::RenderVelocitiesInner(FRHICommandListImmedia
 static TAutoConsoleVariable<int32> CVarParallelVelocity(
 	TEXT("r.ParallelVelocity"),
 	1,  
-	TEXT("Toggles parallel velocity rendering. Parallel rendering must be enabled for this to have an effect.\n"),
+	TEXT("Toggles parallel velocity rendering. Parallel rendering must be enabled for this to have an effect."),
 	ECVF_RenderThreadSafe
 	);
 
