Fix Mac compilation.

diff --git a/Engine/Source/Runtime/Core/Private/Apple/ApplePlatformSymbolication.cpp b/Engine/Source/Runtime/Core/Private/Apple/ApplePlatformSymbolication.cpp
--- a/Engine/Source/Runtime/Core/Private/Apple/ApplePlatformSymbolication.cpp
+++ b/Engine/Source/Runtime/Core/Private/Apple/ApplePlatformSymbolication.cpp
@@ -66,7 +66,7 @@ struct FApplePlatformSymbolCache
 #endif
 };
 
-static bool GAllowApplePlatformSymbolication = PLATFORM_MAC && IS_PROGRAM;
+static bool GAllowApplePlatformSymbolication = (PLATFORM_MAC != 0) && (IS_PROGRAM != 0);
 
 void FApplePlatformSymbolication::SetSymbolicationAllowed(bool const bAllow)
 {
