Some missing files for the Android movie playback.

diff --git a/Engine/Source/Runtime/Core/Private/Android/AndroidJavaMediaPlayer.cpp b/Engine/Source/Runtime/Core/Private/Android/AndroidJavaMediaPlayer.cpp
new file mode 100644
--- /dev/null
+++ b/Engine/Source/Runtime/Core/Private/Android/AndroidJavaMediaPlayer.cpp
@@ -0,0 +1,182 @@
+// Copyright 1998-2014 Epic Games, Inc. All Rights Reserved.
+
+#include "CorePrivatePCH.h"
+#include "AndroidJavaMediaPlayer.h"
+#include "AndroidApplication.h"
+
+
+FJavaAndroidMediaPlayer::FJavaAndroidMediaPlayer(bool swizzlePixels)
+	: FJavaClassObject(GetClassName(), "(Z)V", swizzlePixels)
+	, GetDurationMethod(GetClassMethod("getDuration", "()I"))
+	, ResetMethod(GetClassMethod("reset", "()V"))
+	, GetCurrentPositionMethod(GetClassMethod("getCurrentPosition", "()I"))
+	, IsLoopingMethod(GetClassMethod("isLooping", "()Z"))
+	, IsPlayingMethod(GetClassMethod("isPlaying", "()Z"))
+	, SetDataSourceURLMethod(GetClassMethod("setDataSource", "(Ljava/lang/String;)V"))
+	, SetDataSourceFileMethod(GetClassMethod("setDataSource", "(Ljava/lang/String;JJ)Z"))
+	, SetDataSourceAssetMethod(GetClassMethod("setDataSource", "(Landroid/content/res/AssetManager;Ljava/lang/String;JJ)Z"))
+	, PrepareMethod(GetClassMethod("prepare", "()V"))
+	, SeekToMethod(GetClassMethod("seekTo", "(I)V"))
+	, SetLoopingMethod(GetClassMethod("setLooping", "(Z)V"))
+	, ReleaseMethod(GetClassMethod("release", "()V"))
+	, GetVideoHeightMethod(GetClassMethod("getVideoHeight", "()I"))
+	, GetVideoWidthMethod(GetClassMethod("getVideoWidth", "()I"))
+	, SetVideoEnabledMethod(GetClassMethod("setVideoEnabled", "(Z)V"))
+	, SetAudioEnabledMethod(GetClassMethod("setAudioEnabled", "(Z)V"))
+	, GetVideoLastFrameDataMethod(GetClassMethod("getVideoLastFrameData", "()Ljava/nio/Buffer;"))
+	, StartMethod(GetClassMethod("start", "()V"))
+	, PauseMethod(GetClassMethod("pause", "()V"))
+	, StopMethod(GetClassMethod("stop", "()V"))
+	, GetVideoLastFrameMethod(GetClassMethod("getVideoLastFrame", "(I)Z"))
+{
+}
+
+int32 FJavaAndroidMediaPlayer::GetDuration()
+{
+	BindObjectToThread();
+	return CallMethod<int32>(GetDurationMethod);
+}
+
+void FJavaAndroidMediaPlayer::Reset()
+{
+	BindObjectToThread();
+	CallMethod<void>(ResetMethod);
+}
+
+void FJavaAndroidMediaPlayer::Stop()
+{
+	BindObjectToThread();
+	CallMethod<void>(StopMethod);
+}
+
+int32 FJavaAndroidMediaPlayer::GetCurrentPosition()
+{
+	BindObjectToThread();
+	int32 position = CallMethod<int32>(GetCurrentPositionMethod);
+	return position;
+}
+
+bool FJavaAndroidMediaPlayer::IsLooping()
+{
+	BindObjectToThread();
+	return CallMethod<bool>(IsLoopingMethod);
+}
+
+bool FJavaAndroidMediaPlayer::IsPlaying()
+{
+	BindObjectToThread();
+	return CallMethod<bool>(IsPlayingMethod);
+}
+
+void FJavaAndroidMediaPlayer::SetDataSource(const FString & Url)
+{
+	BindObjectToThread();
+	CallMethod<void>(SetDataSourceURLMethod, GetJString(Url));
+}
+
+bool FJavaAndroidMediaPlayer::SetDataSource(const FString& MoviePathOnDevice, int64 offset, int64 size)
+{
+	BindObjectToThread();
+	return CallMethod<bool>(SetDataSourceFileMethod, GetJString(MoviePathOnDevice), offset, size);
+}
+
+bool FJavaAndroidMediaPlayer::SetDataSource(jobject AssetMgr, const FString& AssetPath, int64 offset, int64 size)
+{
+	BindObjectToThread();
+	return CallMethod<bool>(SetDataSourceAssetMethod, AssetMgr, GetJString(AssetPath), offset, size);
+}
+
+void FJavaAndroidMediaPlayer::Prepare()
+{
+	BindObjectToThread();
+	CallMethod<void>(PrepareMethod);
+}
+
+void FJavaAndroidMediaPlayer::SeekTo(int32 Milliseconds)
+{
+	BindObjectToThread();
+	CallMethod<void>(SeekToMethod, Milliseconds);
+}
+
+void FJavaAndroidMediaPlayer::SetLooping(bool Looping)
+{
+	BindObjectToThread();
+	CallMethod<void>(SetLoopingMethod, Looping);
+}
+
+void FJavaAndroidMediaPlayer::Release()
+{
+	BindObjectToThread();
+	CallMethod<void>(ReleaseMethod);
+}
+
+int32 FJavaAndroidMediaPlayer::GetVideoHeight()
+{
+	BindObjectToThread();
+	return CallMethod<int32>(GetVideoHeightMethod);
+}
+
+int32 FJavaAndroidMediaPlayer::GetVideoWidth()
+{
+	BindObjectToThread();
+	return CallMethod<int32>(GetVideoWidthMethod);
+}
+
+void FJavaAndroidMediaPlayer::SetVideoEnabled(bool enabled /*= true*/)
+{
+	BindObjectToThread();
+	CallMethod<void>(SetVideoEnabledMethod, enabled);
+}
+
+void FJavaAndroidMediaPlayer::SetAudioEnabled(bool enabled /*= true*/)
+{
+	BindObjectToThread();
+	CallMethod<void>(SetAudioEnabledMethod, enabled);
+}
+
+bool FJavaAndroidMediaPlayer::GetVideoLastFrameData(void* & outPixels, int64 & outCount)
+{
+	BindObjectToThread();
+	jobject buffer = CallMethod<jobject>(GetVideoLastFrameDataMethod);
+	if (nullptr != buffer)
+	{
+		outPixels = JEnv->GetDirectBufferAddress(buffer);
+		outCount = JEnv->GetDirectBufferCapacity(buffer);
+		JEnv->DeleteGlobalRef(buffer);
+	}
+	if (nullptr == buffer || nullptr == outPixels || 0 == outCount)
+	{
+		return false;
+	}
+	return true;
+}
+
+void FJavaAndroidMediaPlayer::Start()
+{
+	BindObjectToThread();
+	CallMethod<void>(StartMethod);
+}
+
+void FJavaAndroidMediaPlayer::Pause()
+{
+	BindObjectToThread();
+	CallMethod<void>(PauseMethod);
+}
+
+bool FJavaAndroidMediaPlayer::GetVideoLastFrame(int32 destTexture)
+{
+	BindObjectToThread();
+	return CallMethod<bool>(GetVideoLastFrameMethod, destTexture);
+}
+
+FName FJavaAndroidMediaPlayer::GetClassName()
+{
+	if (FAndroidMisc::GetAndroidBuildVersion() >= 14)
+	{
+		return FName("com/epicgames/ue4/MediaPlayer14");
+	}
+	else
+	{
+		return FName("");
+	}
+}
diff --git a/Engine/Source/Runtime/Core/Public/Android/AndroidJavaMediaPlayer.h b/Engine/Source/Runtime/Core/Public/Android/AndroidJavaMediaPlayer.h
new file mode 100644
--- /dev/null
+++ b/Engine/Source/Runtime/Core/Public/Android/AndroidJavaMediaPlayer.h
@@ -0,0 +1,59 @@
+// Copyright 1998-2014 Epic Games, Inc. All Rights Reserved.
+
+#pragma once
+
+#include "Core.h"
+#include "AndroidJava.h"
+
+// Wrapper for com/epicgames/ue4/MediaPlayer*.java.
+class FJavaAndroidMediaPlayer : public FJavaClassObject
+{
+public:
+	FJavaAndroidMediaPlayer(bool swizzlePixels = true);
+	int32 GetDuration();
+	void Reset();
+	void Stop();
+	int32 GetCurrentPosition();
+	bool IsLooping();
+	bool IsPlaying();
+	void SetDataSource(const FString & Url);
+	bool SetDataSource(const FString& MoviePathOnDevice, int64 offset, int64 size);
+	bool SetDataSource(jobject AssetMgr, const FString& AssetPath, int64 offset, int64 size);
+	void Prepare();
+	void SeekTo(int32 Milliseconds);
+	void SetLooping(bool Looping);
+	void Release();
+	int32 GetVideoHeight();
+	int32 GetVideoWidth();
+	void SetVideoEnabled(bool enabled = true);
+	void SetAudioEnabled(bool enabled = true);
+	bool GetVideoLastFrameData(void* & outPixels, int64 & outCount);
+	void Start();
+	void Pause();
+	bool GetVideoLastFrame(int32 destTexture);
+
+private:
+	static FName GetClassName();
+
+	FJavaClassMethod GetDurationMethod;
+	FJavaClassMethod ResetMethod;
+	FJavaClassMethod StopMethod;
+	FJavaClassMethod GetCurrentPositionMethod;
+	FJavaClassMethod IsLoopingMethod;
+	FJavaClassMethod IsPlayingMethod;
+	FJavaClassMethod SetDataSourceURLMethod;
+	FJavaClassMethod SetDataSourceFileMethod;
+	FJavaClassMethod SetDataSourceAssetMethod;
+	FJavaClassMethod PrepareMethod;
+	FJavaClassMethod SeekToMethod;
+	FJavaClassMethod SetLoopingMethod;
+	FJavaClassMethod ReleaseMethod;
+	FJavaClassMethod GetVideoHeightMethod;
+	FJavaClassMethod GetVideoWidthMethod;
+	FJavaClassMethod SetVideoEnabledMethod;
+	FJavaClassMethod SetAudioEnabledMethod;
+	FJavaClassMethod GetVideoLastFrameDataMethod;
+	FJavaClassMethod StartMethod;
+	FJavaClassMethod PauseMethod;
+	FJavaClassMethod GetVideoLastFrameMethod;
+};
