Implementing FThreadSafeBool using FThreadSafeCounter and moving to Core

diff --git a/Engine/Source/Runtime/Core/Public/HAL/ThreadingBase.h b/Engine/Source/Runtime/Core/Public/HAL/ThreadingBase.h
--- a/Engine/Source/Runtime/Core/Public/HAL/ThreadingBase.h
+++ b/Engine/Source/Runtime/Core/Public/HAL/ThreadingBase.h
@@ -764,6 +764,39 @@ private:
 };
 
 
+/**
+ * Thread safe bool, wraps FThreadSafeCounter
+ */
+struct FThreadSafeBool
+	: private FThreadSafeCounter
+{
+public:
+	/**
+	 * Constructor optionally takes value to initialize with, otherwise initializes false
+	 * @param bValue	Value to initialize with
+	 */
+	FThreadSafeBool(bool bValue = false)
+		: FThreadSafeCounter(bValue ? 1 : 0)
+	{}
+
+	/**
+	 * Operator to use this struct as a bool with thread safety
+	 */
+	FORCEINLINE operator bool() const
+	{
+		return GetValue() != 0;
+	}
+
+	/**
+	 * Operator to set the bool value with thread safety
+	 */
+	FORCEINLINE const bool operator=(const bool& bNewValue)
+	{
+		return Set(bNewValue ? 1 : 0) != 0;
+	}
+};
+
+
 /** Fake Thread safe counter, used to avoid cluttering code with ifdefs when counters are only used for debugging. */
 class FNoopCounter
 {
