Fixes for FThreadSafeBool

diff --git a/Engine/Source/Runtime/Core/Public/HAL/ThreadingBase.h b/Engine/Source/Runtime/Core/Public/HAL/ThreadingBase.h
--- a/Engine/Source/Runtime/Core/Public/HAL/ThreadingBase.h
+++ b/Engine/Source/Runtime/Core/Public/HAL/ThreadingBase.h
@@ -790,9 +790,10 @@ public:
 	/**
 	 * Operator to set the bool value with thread safety
 	 */
-	FORCEINLINE const bool operator=(const bool& bNewValue)
+	FORCEINLINE const bool operator=(const bool bNewValue)
 	{
-		return Set(bNewValue ? 1 : 0) != 0;
+		Set(bNewValue ? 1 : 0);
+		return bNewValue;
 	}
 };
 
