UE-5694 Fix divide-by-zero in FindBounds

diff --git a/Engine/Source/Runtime/Core/Private/Math/UnrealMath.cpp b/Engine/Source/Runtime/Core/Private/Math/UnrealMath.cpp
--- a/Engine/Source/Runtime/Core/Private/Math/UnrealMath.cpp
+++ b/Engine/Source/Runtime/Core/Private/Math/UnrealMath.cpp
@@ -1054,7 +1054,7 @@ static void FindBounds( float& OutMin, float& OutMax,  float Start, float StartL
 		const float c = StartLeaveTan;
 
 		const float Discriminant = (b*b) - (4.f*a*c);
-		if(Discriminant > 0.f)
+		if(Discriminant > 0.f && a > 0.f) // Solving doesn't work if a is zero, which usually indicates co-incident start and end, and zero tangents anyway
 		{
 			const float SqrtDisc = FMath::Sqrt( Discriminant );
 
