Adding AtomicSet to FThreadSafeBool

diff --git a/Engine/Source/Runtime/Core/Public/HAL/ThreadingBase.h b/Engine/Source/Runtime/Core/Public/HAL/ThreadingBase.h
--- a/Engine/Source/Runtime/Core/Public/HAL/ThreadingBase.h
+++ b/Engine/Source/Runtime/Core/Public/HAL/ThreadingBase.h
@@ -767,7 +767,7 @@ private:
 /**
  * Thread safe bool, wraps FThreadSafeCounter
  */
-struct FThreadSafeBool
+class FThreadSafeBool
 	: private FThreadSafeCounter
 {
 public:
@@ -790,11 +790,22 @@ public:
 	/**
 	 * Operator to set the bool value with thread safety
 	 */
-	FORCEINLINE const bool operator=(const bool bNewValue)
+	FORCEINLINE bool operator=(bool bNewValue)
 	{
 		Set(bNewValue ? 1 : 0);
 		return bNewValue;
 	}
+
+	/**
+	 * Sets a new value atomically, and returns the old value.
+	 *
+	 * @param bNewValue   Value to set
+	 * @return The old value
+	 */
+	FORCEINLINE bool AtomicSet(bool bNewValue)
+	{
+		return Set(bNewValue ? 1 : 0) != 0;
+	}
 };
 
 
