On Mac, when user double clicks the title bar, minimize or zoom the window, depending on Dock settings and OS X version

diff --git a/Engine/Source/Runtime/Core/Private/Mac/MacApplication.cpp b/Engine/Source/Runtime/Core/Private/Mac/MacApplication.cpp
--- a/Engine/Source/Runtime/Core/Private/Mac/MacApplication.cpp
+++ b/Engine/Source/Runtime/Core/Private/Mac/MacApplication.cpp
@@ -647,6 +647,21 @@ void FMacApplication::ProcessNSEvent(NSEvent* const Event)
 				if (Button == LastPressedMouseButton && ([Event clickCount] % 2) == 0)
 				{
 					MessageHandler->OnMouseDoubleClick(CurrentEventWindow, Button);
+
+					bool IsMouseOverTitleBar = false;
+					const bool IsMovable = IsWindowMovable(NativeWindow, &IsMouseOverTitleBar);
+					if (IsMouseOverTitleBar)
+					{
+						const bool bShouldMinimize = [[NSUserDefaults standardUserDefaults] boolForKey:@"AppleMiniaturizeOnDoubleClick"];
+						if (bShouldMinimize)
+						{
+							MainThreadCall(^{ [NativeWindow performMiniaturize:nil]; }, NSDefaultRunLoopMode, true);
+						}
+						else if (!FPlatformMisc::IsRunningOnMavericks())
+						{
+							MainThreadCall(^{ [NativeWindow performZoom:nil]; }, NSDefaultRunLoopMode, true);
+						}
+					}
 				}
 				else
 				{
