Reimplemented FPaths::IsRelative based on whether or not a path is rooted.

diff --git a/Engine/Source/Runtime/Core/Private/Misc/Paths.cpp b/Engine/Source/Runtime/Core/Private/Misc/Paths.cpp
--- a/Engine/Source/Runtime/Core/Private/Misc/Paths.cpp
+++ b/Engine/Source/Runtime/Core/Private/Misc/Paths.cpp
@@ -553,14 +553,15 @@ bool FPaths::IsDrive(const FString& InPath)
 
 bool FPaths::IsRelative(const FString& InPath)
 {
-	return
-		InPath.StartsWith( TEXT("./") ) ||
-		InPath.StartsWith( TEXT(".\\") ) || 
-		InPath.StartsWith( TEXT("../") ) ||
-		InPath.StartsWith( TEXT("..\\") ) || 
-		InPath.IsEmpty() ||
-		!(InPath.Contains("/")
-		|| InPath.Contains("\\"));
+	// The previous implementation of this function seemed to handle normalized and unnormalized paths, so this one does too for legacy reasons.
+
+	const bool IsRooted =	InPath.StartsWith(TEXT("\\\\"))	||												// "\\" for UNC or "network" paths.
+							InPath.StartsWith(TEXT("//"))	||												// Equivalent to "\\", considering normalization replaces "\\" with "//".
+							InPath.StartsWith(TEXT("\\"))	||												// Root of the current directory on Windows
+							InPath.StartsWith(TEXT("/"))	||												// Root of the current directory on Windows, root on UNIX-likes.
+							(InPath.Len() >= 2 && FChar::IsAlpha(InPath[0]) && InPath[1] == TEXT(':'));	// Starts with "<DriveLetter>:"
+
+	return !IsRooted;
 }
 
 void FPaths::NormalizeFilename(FString& InPath)
diff --git a/Engine/Source/Runtime/Core/Public/Misc/Paths.h b/Engine/Source/Runtime/Core/Public/Misc/Paths.h
--- a/Engine/Source/Runtime/Core/Public/Misc/Paths.h
+++ b/Engine/Source/Runtime/Core/Public/Misc/Paths.h
@@ -315,7 +315,7 @@ public:
 	/** @return true if this path represents a drive */
 	static bool IsDrive(const FString& InPath);
 
-	/** @return true if this path is relative */
+	/** @return true if this path is relative (not rooted) */
 	static bool IsRelative(const FString& InPath);
 
 	/** Convert all / and \ to TEXT("/") */
