PR #625 - fix for cursor position during terrain painting.

- Contributed by [REDACTED_PERSON].

diff --git a/Engine/Source/Runtime/Core/Private/Linux/LinuxApplication.cpp b/Engine/Source/Runtime/Core/Private/Linux/LinuxApplication.cpp
--- a/Engine/Source/Runtime/Core/Private/Linux/LinuxApplication.cpp
+++ b/Engine/Source/Runtime/Core/Private/Linux/LinuxApplication.cpp
@@ -314,7 +314,11 @@ void FLinuxApplication::ProcessDeferredMessage( SDL_Event Event )
 
 			if(bUsingHighPrecisionMouseInput)
 			{
-					LinuxCursor->AddOffset(motionEvent.xrel, motionEvent.yrel);
+					// maintain "shadow" global position
+					if (LinuxCursor->IsHidden())
+					{
+						LinuxCursor->AddOffset(motionEvent.xrel, motionEvent.yrel);
+					}
  					MessageHandler->OnRawMouseMove(motionEvent.xrel, motionEvent.yrel);
 			}
 			else
