Added support for PlayStation 4 controller and fixed DPad handling on controllers that actually report a hatswitch instead of treating it as a set of buttons

diff --git a/Engine/Source/Runtime/Core/Private/Mac/HIDInputInterface.cpp b/Engine/Source/Runtime/Core/Private/Mac/HIDInputInterface.cpp
--- a/Engine/Source/Runtime/Core/Private/Mac/HIDInputInterface.cpp
+++ b/Engine/Source/Runtime/Core/Private/Mac/HIDInputInterface.cpp
@@ -153,6 +153,28 @@ void HIDInputInterface::FHIDDeviceInfo::SetupMappings()
 		RightAnalogYMapping = kHIDUsage_GD_Rz;
 		RightTriggerAnalogMapping = kHIDUsage_GD_Ry;
 	}
+	else if (VendorID == 0x54c && ProductID == 0x5c4)
+	{
+		ButtonsMapping[1]	= 2;	// Square		->	X
+		ButtonsMapping[2]	= 0;	// Cross		->	A
+		ButtonsMapping[3]	= 1;	// Circle		->	B
+		ButtonsMapping[4]	= 3;	// Triangle		->	Y
+		ButtonsMapping[5]	= 4;	// L1			->	Left Shoulder
+		ButtonsMapping[6]	= 5;	// R1			->	Right Shoulder
+		ButtonsMapping[7]	= 10;	// L2			->	Left Trigger
+		ButtonsMapping[8]	= 11;	// R2			->	Right Trigger
+		ButtonsMapping[9]	= 7;	// Share		->	Back
+		ButtonsMapping[10]	= 6;	// Options		->	Start
+		ButtonsMapping[11]	= 8;	// L3			->	Left Thumbstick
+		ButtonsMapping[12]	= 9;	// R3			->	Right Thumbstick
+
+		LeftAnalogXMapping = kHIDUsage_GD_X;
+		LeftAnalogYMapping = kHIDUsage_GD_Y;
+		LeftTriggerAnalogMapping = kHIDUsage_GD_Rx;
+		RightAnalogXMapping = kHIDUsage_GD_Z;
+		RightAnalogYMapping = kHIDUsage_GD_Rz;
+		RightTriggerAnalogMapping = kHIDUsage_GD_Ry;
+	}
 	else if ((VendorID == 0x45e && (ProductID == 0x28e || ProductID == 0x719)) // Original Microsoft controller
 			 || (VendorID == 0xe6f && ProductID == 0x401)) // GameStop version
 	{
@@ -338,13 +360,13 @@ void HIDInputInterface::SendControllerEvents()
 							switch (NewValue)
 							{
 								case 0: CurrentButtonStates[12] = true;									break; // Up
-								case 1: CurrentButtonStates[12] = true;	CurrentButtonStates[13] = true;	break; // Up + right
-								case 2: CurrentButtonStates[13] = true;									break; // Right
-								case 3: CurrentButtonStates[13] = true;	CurrentButtonStates[14] = true;	break; // Right + down
-								case 4: CurrentButtonStates[14] = true;									break; // Down
-								case 5: CurrentButtonStates[14] = true;	CurrentButtonStates[15] = true;	break; // Down + left
-								case 6: CurrentButtonStates[15] = true;									break; // Left
-								case 7: CurrentButtonStates[15] = true;	CurrentButtonStates[12] = true;	break; // Left + up
+								case 1: CurrentButtonStates[12] = true;	CurrentButtonStates[15] = true;	break; // Up + right
+								case 2: CurrentButtonStates[15] = true;									break; // Right
+								case 3: CurrentButtonStates[15] = true;	CurrentButtonStates[13] = true;	break; // Right + down
+								case 4: CurrentButtonStates[13] = true;									break; // Down
+								case 5: CurrentButtonStates[13] = true;	CurrentButtonStates[14] = true;	break; // Down + left
+								case 6: CurrentButtonStates[14] = true;									break; // Left
+								case 7: CurrentButtonStates[14] = true;	CurrentButtonStates[12] = true;	break; // Left + up
 							}
 						}
 					}
