Revised FPaths::IsRelative to support paths used by feature pack

#codereivew Saul.Abreu, Marc.Audy

diff --git a/Engine/Source/Runtime/Core/Private/Misc/Paths.cpp b/Engine/Source/Runtime/Core/Private/Misc/Paths.cpp
--- a/Engine/Source/Runtime/Core/Private/Misc/Paths.cpp
+++ b/Engine/Source/Runtime/Core/Private/Misc/Paths.cpp
@@ -554,13 +554,13 @@ bool FPaths::IsDrive(const FString& InPath)
 bool FPaths::IsRelative(const FString& InPath)
 {
 	// The previous implementation of this function seemed to handle normalized and unnormalized paths, so this one does too for legacy reasons.
-
 	const bool IsRooted =	InPath.StartsWith(TEXT("\\\\"))	||												// "\\" for UNC or "network" paths.
 							InPath.StartsWith(TEXT("//"))	||												// Equivalent to "\\", considering normalization replaces "\\" with "//".
 							InPath.StartsWith(TEXT("\\"))	||												// Root of the current directory on Windows
 							InPath.StartsWith(TEXT("/"))	||												// Root of the current directory on Windows, root on UNIX-likes.
+							InPath.StartsWith(TEXT("root:/")) ||											// Feature packs use this
 							(InPath.Len() >= 2 && FChar::IsAlpha(InPath[0]) && InPath[1] == TEXT(':'));	// Starts with "<DriveLetter>:"
-
+	 
 	return !IsRooted;
 }
 
