PS4 - Function level optimization control.
* Use DISABLE_FUNCTION_OPTIMIZATION when declaring a function, eg: DISABLE_FUNCTION_OPTIMIZATION void TestFunc();

diff --git a/Engine/Source/Runtime/Core/Public/HAL/Platform.h b/Engine/Source/Runtime/Core/Public/HAL/Platform.h
--- a/Engine/Source/Runtime/Core/Public/HAL/Platform.h
+++ b/Engine/Source/Runtime/Core/Public/HAL/Platform.h
@@ -322,6 +322,11 @@
 	#define PRAGMA_ENABLE_OPTIMIZATION_ACTUAL
 #endif
 
+// Disable optimization of a specific function
+#ifndef DISABLE_FUNCTION_OPTIMIZATION
+	#define DISABLE_FUNCTION_OPTIMIZATION
+#endif
+
 #ifndef FORCEINLINE_DEBUGGABLE_ACTUAL
 	#define FORCEINLINE_DEBUGGABLE_ACTUAL inline
 #endif
