Log DisplayMetrics in platform-independent code.

diff --git a/Engine/Source/Runtime/Core/Private/GenericPlatform/GenericApplication.cpp b/Engine/Source/Runtime/Core/Private/GenericPlatform/GenericApplication.cpp
new file mode 100644
--- /dev/null
+++ b/Engine/Source/Runtime/Core/Private/GenericPlatform/GenericApplication.cpp
@@ -0,0 +1,34 @@
+// Copyright 1998-2014 Epic Games, Inc. All Rights Reserved.
+
+#include "CorePrivatePCH.h"
+#include "GenericApplication.h"
+
+void FDisplayMetrics::PrintToLog() const
+{
+	UE_LOG(LogInit, Log, TEXT("Display metrics:"));
+	UE_LOG(LogInit, Log, TEXT("  PrimaryDisplayWidth: %d"), PrimaryDisplayWidth);
+	UE_LOG(LogInit, Log, TEXT("  PrimaryDisplayHeight: %d"), PrimaryDisplayHeight);
+	UE_LOG(LogInit, Log, TEXT("  PrimaryDisplayWorkAreaRect:"));
+	UE_LOG(LogInit, Log, TEXT("    Left=%d, Top=%d, Right=%d, Bottom=%d"), 
+		PrimaryDisplayWorkAreaRect.Left, PrimaryDisplayWorkAreaRect.Top, 
+		PrimaryDisplayWorkAreaRect.Right, PrimaryDisplayWorkAreaRect.Bottom);
+	UE_LOG(LogInit, Log, TEXT("  VirtualDisplayRect:"));
+	UE_LOG(LogInit, Log, TEXT("    Left=%d, Top=%d, Right=%d, Bottom=%d"), 
+		VirtualDisplayRect.Left, VirtualDisplayRect.Top, 
+		VirtualDisplayRect.Right, VirtualDisplayRect.Bottom);
+	UE_LOG(LogInit, Log, TEXT("  TitleSafePaddingSize: %s"), *TitleSafePaddingSize.ToString());
+	UE_LOG(LogInit, Log, TEXT("  ActionSafePaddingSize: %s"), *ActionSafePaddingSize.ToString());
+
+	const int NumMonitors = MonitorInfo.Num();
+	UE_LOG(LogInit, Log, TEXT("  Number of monitors: %d"), NumMonitors);
+	for (int MonitorIdx = 0; MonitorIdx < NumMonitors; ++MonitorIdx)
+	{
+		const FMonitorInfo & Info = MonitorInfo[MonitorIdx];
+		UE_LOG(LogInit, Log, TEXT("    Monitor %d"), MonitorIdx);
+		UE_LOG(LogInit, Log, TEXT("      Name: %s"), *Info.Name);
+		UE_LOG(LogInit, Log, TEXT("      ID: %s"), *Info.ID);
+		UE_LOG(LogInit, Log, TEXT("      NativeWidth: %d"), Info.NativeWidth);
+		UE_LOG(LogInit, Log, TEXT("      NativeHeight: %d"), Info.NativeHeight);
+		UE_LOG(LogInit, Log, TEXT("      bIsPrimary: %s"), Info.bIsPrimary ? TEXT("true") : TEXT("false"));
+	}
+}
diff --git a/Engine/Source/Runtime/Core/Private/Linux/LinuxPlatformMisc.cpp b/Engine/Source/Runtime/Core/Private/Linux/LinuxPlatformMisc.cpp
--- a/Engine/Source/Runtime/Core/Private/Linux/LinuxPlatformMisc.cpp
+++ b/Engine/Source/Runtime/Core/Private/Linux/LinuxPlatformMisc.cpp
@@ -180,33 +180,7 @@ bool FLinuxPlatformMisc::PlatformInitMultimedia()
 			// dump information about screens for debug
 			FDisplayMetrics DisplayMetrics;
 			FDisplayMetrics::GetDisplayMetrics(DisplayMetrics);
-			
-			UE_LOG(LogInit, Log, TEXT("Display metrics:"));
-			UE_LOG(LogInit, Log, TEXT("  PrimaryDisplayWidth: %d"), DisplayMetrics.PrimaryDisplayWidth);
-			UE_LOG(LogInit, Log, TEXT("  PrimaryDisplayHeight: %d"), DisplayMetrics.PrimaryDisplayHeight);
-			UE_LOG(LogInit, Log, TEXT("  PrimaryDisplayWorkAreaRect:"));
-			UE_LOG(LogInit, Log, TEXT("    Left=%d, Top=%d, Right=%d, Bottom=%d"), 
-				DisplayMetrics.PrimaryDisplayWorkAreaRect.Left, DisplayMetrics.PrimaryDisplayWorkAreaRect.Top, 
-				DisplayMetrics.PrimaryDisplayWorkAreaRect.Right, DisplayMetrics.PrimaryDisplayWorkAreaRect.Bottom);
-			UE_LOG(LogInit, Log, TEXT("  VirtualDisplayRect:"));
-			UE_LOG(LogInit, Log, TEXT("    Left=%d, Top=%d, Right=%d, Bottom=%d"), 
-				DisplayMetrics.VirtualDisplayRect.Left, DisplayMetrics.VirtualDisplayRect.Top, 
-				DisplayMetrics.VirtualDisplayRect.Right, DisplayMetrics.VirtualDisplayRect.Bottom);
-			UE_LOG(LogInit, Log, TEXT("  TitleSafePaddingSize: %s"), *DisplayMetrics.TitleSafePaddingSize.ToString());
-			UE_LOG(LogInit, Log, TEXT("  ActionSafePaddingSize: %s"), *DisplayMetrics.ActionSafePaddingSize.ToString());
-
-			const int NumMonitors = DisplayMetrics.MonitorInfo.Num();
-			UE_LOG(LogInit, Log, TEXT("  Number of monitors: %d"), NumMonitors);
-			for (int MonitorIdx = 0; MonitorIdx < NumMonitors; ++MonitorIdx)
-			{
-				const FMonitorInfo & MonitorInfo = DisplayMetrics.MonitorInfo[MonitorIdx];
-				UE_LOG(LogInit, Log, TEXT("    Monitor %d"), MonitorIdx);
-				UE_LOG(LogInit, Log, TEXT("      Name: %s"), *MonitorInfo.Name);
-				UE_LOG(LogInit, Log, TEXT("      ID: %s"), *MonitorInfo.ID);
-				UE_LOG(LogInit, Log, TEXT("      NativeWidth: %d"), MonitorInfo.NativeWidth);
-				UE_LOG(LogInit, Log, TEXT("      NativeHeight: %d"), MonitorInfo.NativeHeight);
-				UE_LOG(LogInit, Log, TEXT("      bIsPrimary: %s"), MonitorInfo.bIsPrimary ? TEXT("true") : TEXT("false"));
-			}
+			DisplayMetrics.PrintToLog();
 		}
 	}
 
diff --git a/Engine/Source/Runtime/Core/Public/GenericPlatform/GenericApplication.h b/Engine/Source/Runtime/Core/Public/GenericPlatform/GenericApplication.h
--- a/Engine/Source/Runtime/Core/Public/GenericPlatform/GenericApplication.h
+++ b/Engine/Source/Runtime/Core/Public/GenericPlatform/GenericApplication.h
@@ -340,6 +340,9 @@ struct FDisplayMetrics
 	FVector2D ActionSafePaddingSize;
 
 	CORE_API static void GetDisplayMetrics(struct FDisplayMetrics& OutDisplayMetrics);
+
+	/** Logs out display metrics */
+	CORE_API void PrintToLog() const;
 };
 
 
