Populate -o //depot/UE4/Engine/Source/Runtime/Core/Public/UObject/UnrealNames.h //depot/UE4/Engine/Source/Runtime/Core/Public/UObject/UnrealNames.inl.

diff --git a/Engine/Source/Runtime/Core/Public/UObject/UnrealNames.inl b/Engine/Source/Runtime/Core/Public/UObject/UnrealNames.inl
new file mode 100644
--- /dev/null
+++ b/Engine/Source/Runtime/Core/Public/UObject/UnrealNames.inl
@@ -0,0 +1,211 @@
+// Copyright 1998-2014 Epic Games, Inc. All Rights Reserved.
+#ifndef _UNREAL_NAMES_H_
+#define _UNREAL_NAMES_H_
+
+//
+// Macros.
+//
+
+/** Index of highest hardcoded name to be replicated by index by the networking code
+ * @warning: changing this number or making any change to the list of hardcoded names with index
+ * less than this value breaks network compatibility (update GEngineMinNetVersion)
+ * @note: names with a greater value than this can still be replicated, but they are sent as
+ * strings instead of an efficient index
+ */
+#define MAX_NETWORKED_HARDCODED_NAME 410
+
+// Define a message as an enumeration.
+#ifndef REGISTER_NAME
+	#define REGISTER_NAME(num,name) NAME_##name = num,
+	#define REGISTERING_ENUM
+	enum EName {
+#endif
+
+//
+// Hardcoded names.
+//
+
+// Special zero value, meaning no name.
+REGISTER_NAME(0,None)
+
+// Class property types (name indices are significant for serialization).
+REGISTER_NAME(1,ByteProperty)
+REGISTER_NAME(2,IntProperty)
+REGISTER_NAME(3,BoolProperty)
+REGISTER_NAME(4,FloatProperty)
+REGISTER_NAME(5,ObjectProperty)
+REGISTER_NAME(6,NameProperty)
+REGISTER_NAME(7,DelegateProperty)
+REGISTER_NAME(8,ClassProperty)
+REGISTER_NAME(9,ArrayProperty)
+REGISTER_NAME(10,StructProperty)
+REGISTER_NAME(11,VectorProperty)
+REGISTER_NAME(12,RotatorProperty)
+REGISTER_NAME(13,StrProperty)
+REGISTER_NAME(14,TextProperty)
+REGISTER_NAME(15,InterfaceProperty)
+REGISTER_NAME(16,MulticastDelegateProperty )
+REGISTER_NAME(17,WeakObjectProperty )
+REGISTER_NAME(18,LazyObjectProperty )
+REGISTER_NAME(19,AssetObjectProperty )
+REGISTER_NAME(20,UInt64Property)
+REGISTER_NAME(21,UInt32Property)
+REGISTER_NAME(22,UInt16Property)
+REGISTER_NAME(23,Int64Property)
+REGISTER_NAME(25,Int16Property)
+REGISTER_NAME(26,Int8Property)
+REGISTER_NAME(27,AssetSubclassOfProperty)
+
+// Special packages.
+REGISTER_NAME(30,Core)
+REGISTER_NAME(31,Engine)
+REGISTER_NAME(32,Editor)
+REGISTER_NAME(33,CoreUObject)
+
+// Special types.
+REGISTER_NAME(50,Cylinder)
+REGISTER_NAME(51,BoxSphereBounds)
+REGISTER_NAME(52,Sphere)
+REGISTER_NAME(53,Box)
+REGISTER_NAME(54,Vector2D)
+REGISTER_NAME(55,IntRect)
+REGISTER_NAME(56,IntPoint)
+REGISTER_NAME(57,Vector4)
+REGISTER_NAME(58,Name)
+REGISTER_NAME(59,Vector)
+REGISTER_NAME(60,Rotator)
+REGISTER_NAME(61,SHVector)
+REGISTER_NAME(62,Color)
+REGISTER_NAME(63,Plane)
+REGISTER_NAME(64,Matrix)
+REGISTER_NAME(65,LinearColor)
+REGISTER_NAME(66,AdvanceFrame)
+REGISTER_NAME(67,Pointer)
+REGISTER_NAME(68,Double)
+REGISTER_NAME(69,Quat)
+REGISTER_NAME(70,Self)
+REGISTER_NAME(71,Transform)
+
+// Object class names.
+REGISTER_NAME(100,Object)
+REGISTER_NAME(101,Camera)
+REGISTER_NAME(102,Actor)
+REGISTER_NAME(103,ObjectRedirector)
+REGISTER_NAME(104,ObjectArchetype)
+REGISTER_NAME(105,Class)
+
+// Misc.
+REGISTER_NAME(200,State)
+REGISTER_NAME(201,TRUE)
+REGISTER_NAME(202,FALSE)
+REGISTER_NAME(203,Enum)
+REGISTER_NAME(204,Default)
+REGISTER_NAME(205,Skip)
+REGISTER_NAME(206,Input)
+REGISTER_NAME(207,Package)
+REGISTER_NAME(208,Groups)
+REGISTER_NAME(209,Interface)
+REGISTER_NAME(210,Components)
+REGISTER_NAME(211,Global)
+REGISTER_NAME(212,Super)
+REGISTER_NAME(213,Outer)
+REGISTER_NAME(214,Map)
+REGISTER_NAME(215,Role)
+REGISTER_NAME(216,RemoteRole)
+REGISTER_NAME(217,PersistentLevel)
+REGISTER_NAME(218,TheWorld)
+REGISTER_NAME(219,PackageMetaData)
+REGISTER_NAME(220,InitialState)
+REGISTER_NAME(221,Game)
+REGISTER_NAME(222,SelectionColor)
+REGISTER_NAME(223,UI)
+REGISTER_NAME(224,ExecuteUbergraph)
+REGISTER_NAME(225,DeviceID)
+REGISTER_NAME(226,RootStat)
+REGISTER_NAME(227,MoveActor)
+REGISTER_NAME(230,All) 
+REGISTER_NAME(231,MeshEmitterVertexColor)
+REGISTER_NAME(232,TextureOffsetParameter)
+REGISTER_NAME(233,TextureScaleParameter)
+REGISTER_NAME(234,ImpactVel)
+REGISTER_NAME(235,SlideVel)
+REGISTER_NAME(236,TextureOffset1Parameter)
+REGISTER_NAME(237,MeshEmitterDynamicParameter)
+REGISTER_NAME(238,ExpressionInput)
+REGISTER_NAME(239,Untitled)
+REGISTER_NAME(240,Timer)
+REGISTER_NAME(241,Team)
+REGISTER_NAME(242,Low)
+REGISTER_NAME(243,High)
+REGISTER_NAME(244,NetworkGUID)
+REGISTER_NAME(245,GameThread)
+REGISTER_NAME(246,RenderThread)
+REGISTER_NAME(247,OtherChildren)
+REGISTER_NAME(248,Location)
+REGISTER_NAME(249,Rotation)
+REGISTER_NAME(250,BSP)
+REGISTER_NAME(251,EditorGameAgnostic)
+
+// Online
+REGISTER_NAME(280,DGram)
+REGISTER_NAME(281,Stream)
+REGISTER_NAME(282,GameNetDriver)
+REGISTER_NAME(283,PendingNetDriver)
+REGISTER_NAME(284,BeaconNetDriver)
+REGISTER_NAME(285,FlushNetDormancy)
+
+// Texture settings.
+REGISTER_NAME(300,Linear)
+REGISTER_NAME(301,Point)
+REGISTER_NAME(302,Aniso)
+REGISTER_NAME(303,LightMapResolution)
+
+// Sound.
+//REGISTER_NAME(310,)
+REGISTER_NAME(311,UnGrouped)
+REGISTER_NAME(312,VoiceChat)
+
+// Optimized replication.
+REGISTER_NAME(320,Playing)
+REGISTER_NAME(322,Spectating)
+REGISTER_NAME(325,Inactive)
+
+// Log messages.
+REGISTER_NAME(350,PerfWarning)
+REGISTER_NAME(351,Info)
+REGISTER_NAME(352,Init)
+REGISTER_NAME(353,Exit)
+REGISTER_NAME(354,Cmd)
+REGISTER_NAME(355,Warning)
+REGISTER_NAME(356,Error)
+
+// File format backwards-compatibility.
+REGISTER_NAME(400,FontCharacter)
+REGISTER_NAME(401,InitChild2StartBone)
+REGISTER_NAME(402,SoundCueLocalized)
+REGISTER_NAME(403,SoundCue)
+REGISTER_NAME(404,RawDistributionFloat)
+REGISTER_NAME(405,RawDistributionVector)
+REGISTER_NAME(406,InterpCurveFloat)
+REGISTER_NAME(407,InterpCurveVector2D)
+REGISTER_NAME(408,InterpCurveVector)
+REGISTER_NAME(409,InterpCurveTwoVectors)
+REGISTER_NAME(410,InterpCurveQuat)
+
+REGISTER_NAME(450,AI)
+REGISTER_NAME(451,NavMesh)
+
+REGISTER_NAME(500,PerformanceCapture)
+
+//
+// Closing.
+//
+
+#ifdef REGISTERING_ENUM
+	NAME_MaxHardcodedNameIndex,
+	};
+	#undef REGISTER_NAME
+	#undef REGISTERING_ENUM
+#endif
+
+#endif // _UNREAL_NAMES_H_
\ No newline at end of file
