added ComponentMin and ComponentMax to FIntPoint like FVector

diff --git a/Engine/Source/Runtime/Core/Public/Math/IntPoint.h b/Engine/Source/Runtime/Core/Public/Math/IntPoint.h
--- a/Engine/Source/Runtime/Core/Public/Math/IntPoint.h
+++ b/Engine/Source/Runtime/Core/Public/Math/IntPoint.h
@@ -189,6 +189,12 @@ public:
 	*/
 	int32 operator[](int32 Index) const;
 
+	/** Gets the component-wise min of two vectors. */
+	FORCEINLINE FIntPoint ComponentMin(const FIntPoint& Other) const;
+
+	/** Gets the component-wise max of two vectors. */
+	FORCEINLINE FIntPoint ComponentMax(const FIntPoint& Other) const;
+
 public:
 
 	/**
@@ -399,6 +405,16 @@ FORCEINLINE int32 FIntPoint::operator[](int32 Index) const
 	return ((Index == 0) ? X : Y);
 }
 
+FORCEINLINE FIntPoint FIntPoint::ComponentMin(const FIntPoint& Other) const
+{
+	return FIntPoint(FMath::Min(X, Other.X), FMath::Min(Y, Other.Y));
+}
+
+FORCEINLINE FIntPoint FIntPoint::ComponentMax(const FIntPoint& Other) const
+{
+	return FIntPoint(FMath::Max(X, Other.X), FMath::Max(Y, Other.Y));
+}
+
 FORCEINLINE FIntPoint FIntPoint::DivideAndRoundUp( FIntPoint lhs, int32 Divisor )
 {
 	return FIntPoint(FMath::DivideAndRoundUp(lhs.X, Divisor), FMath::DivideAndRoundUp(lhs.Y, Divisor));
