Added missing constructor paramater

diff --git a/Engine/Source/Runtime/Core/Private/Misc/AutomationTest.cpp b/Engine/Source/Runtime/Core/Private/Misc/AutomationTest.cpp
--- a/Engine/Source/Runtime/Core/Private/Misc/AutomationTest.cpp
+++ b/Engine/Source/Runtime/Core/Private/Misc/AutomationTest.cpp
@@ -646,9 +646,10 @@ FAutomationTestFramework::FAutomationTestFramework()
 ,	CurrentTest(NULL)
 ,	bDeveloperDirectoryIncluded(false)
 ,	bVisualCommandletFilterOn(false)
-,	bUseFullSizeScreenShots(false)
+,	bUseFullSizeScreenShots(true)
 ,	NetworkRoleIndex(0)
 ,	bForceSmokeTests(false)
+,	bScreenshotsEnabled(true)
 { }
 
 
