Core - Misc code clean-up

diff --git a/Engine/Source/Runtime/Core/Private/HAL/MallocTBB.cpp b/Engine/Source/Runtime/Core/Private/HAL/MallocTBB.cpp
--- a/Engine/Source/Runtime/Core/Private/HAL/MallocTBB.cpp
+++ b/Engine/Source/Runtime/Core/Private/HAL/MallocTBB.cpp
@@ -33,29 +33,29 @@ void* FMallocTBB::Malloc( SIZE_T Size, uint32 Alignment )
 
 	MEM_TIME(MemTime -= FPlatformTime::Seconds());
 
-	void* Free = NULL;
+	void* NewPtr = NULL;
 	if( Alignment != DEFAULT_ALIGNMENT )
 	{
 		Alignment = FMath::Max(Size >= 16 ? (uint32)16 : (uint32)8, Alignment);
-		Free = scalable_aligned_malloc( Size, Alignment );
+		NewPtr = scalable_aligned_malloc( Size, Alignment );
 	}
 	else
 	{
-		Free = scalable_malloc( Size );
+		NewPtr = scalable_malloc( Size );
 	}
 
-	if( !Free )
+	if( !NewPtr )
 	{
 		OutOfMemory(Size, Alignment);
 	}
 #if UE_BUILD_DEBUG || UE_BUILD_DEVELOPMENT
 	else if (Size)
 	{
-		FMemory::Memset(Free, DEBUG_FILL_NEW, Size); 
+		FMemory::Memset(NewPtr, DEBUG_FILL_NEW, Size); 
 	}
 #endif
 	MEM_TIME(MemTime += FPlatformTime::Seconds());
-	return Free;
+	return NewPtr;
 }
 
 void* FMallocTBB::Realloc( void* Ptr, SIZE_T NewSize, uint32 Alignment )
