Core - Fixed GetAllocatedSize in TIndirectArray

diff --git a/Engine/Source/Runtime/Core/Public/Containers/Array.h b/Engine/Source/Runtime/Core/Public/Containers/Array.h
--- a/Engine/Source/Runtime/Core/Public/Containers/Array.h
+++ b/Engine/Source/Runtime/Core/Public/Containers/Array.h
@@ -3238,7 +3238,7 @@ public:
 	 */
 	uint32 GetAllocatedSize() const
 	{
-		return Array.Max() * (sizeof(T)+sizeof(T*));
+		return Array.Max() * sizeof(T*) + Array.Num() * sizeof(T);
 	}
 
 	// Iterators
