Core - Fixed Reserver/Empty in TChunkedArray

diff --git a/Engine/Source/Runtime/Core/Public/Containers/ChunkedArray.h b/Engine/Source/Runtime/Core/Public/Containers/ChunkedArray.h
--- a/Engine/Source/Runtime/Core/Public/Containers/ChunkedArray.h
+++ b/Engine/Source/Runtime/Core/Public/Containers/ChunkedArray.h
@@ -171,7 +171,9 @@ public:
 
 	void Empty( int32 Slack=0 ) 
 	{
-		Chunks.Empty(Slack % NumElementsPerChunk + 1);
+		// Compute the number of chunks needed.
+		const int32 NumChunks = (Slack + NumElementsPerChunk - 1) / NumElementsPerChunk;
+		Chunks.Empty(NumChunks);
 		NumElements = 0;
 	}
 
@@ -183,7 +185,9 @@ public:
 	 */
 	void Reserve(int32 Number)
 	{
-		Chunks.Reserve(Number % NumElementsPerChunk + 1);
+		// Compute the number of chunks needed.
+		const int32 NumChunks = (Number + NumElementsPerChunk - 1) / NumElementsPerChunk;
+		Chunks.Reserve(NumChunks);
 	}
 
 	void Shrink()
@@ -224,4 +228,17 @@ template <typename T,uint32 TargetBytesPerChunk> void* operator new( size_t Size
 	check(Size == sizeof(T));
 	const int32 Index = ChunkedArray.Add(1);
 	return &ChunkedArray(Index);
+}
+
+/**
+ * A specialization of the exchange macro that avoids reallocating when
+ * exchanging two arrays.
+ *
+ * @param FirstArrayToExchange First array to exchange.
+ * @param SecondArrayToExchange Second array to exchange.
+ */
+template <typename T, uint32 TargetBytesPerChunk>
+inline void Exchange(TChunkedArray<T,TargetBytesPerChunk>& FirstArrayToExchange, TChunkedArray<T,TargetBytesPerChunk>& SecondArrayToExchange)
+{
+	FMemory::Memswap(&FirstArrayToExchange, &SecondArrayToExchange, sizeof(TChunkedArray<T,TargetBytesPerChunk>));
 }
\ No newline at end of file
