Stats - Memory profiling using stats system (preliminary version)

diff --git a/Engine/Source/Developer/Profiler/Private/ProfilerSession.cpp b/Engine/Source/Developer/Profiler/Private/ProfilerSession.cpp
--- a/Engine/Source/Developer/Profiler/Private/ProfilerSession.cpp
+++ b/Engine/Source/Developer/Profiler/Private/ProfilerSession.cpp
@@ -803,6 +803,11 @@ void FProfilerStatMetaData::UpdateFromStatsState( const FStatsThreadState& Stats
 		{
 			StatType = STATTYPE_AccumulatorFLOAT;
 		}
+		else if( LongName.NameAndInfo.GetField<EStatDataType>() == EStatDataType::ST_Ptr )
+		{
+			// Not supported at this moment.
+			continue;
+		}
 
 		check( StatType != STATTYPE_Error );
 
diff --git a/Engine/Source/Runtime/Core/Private/HAL/MallocBinned.cpp b/Engine/Source/Runtime/Core/Private/HAL/MallocBinned.cpp
--- a/Engine/Source/Runtime/Core/Private/HAL/MallocBinned.cpp
+++ b/Engine/Source/Runtime/Core/Private/HAL/MallocBinned.cpp
@@ -65,4 +65,21 @@ void FMallocBinned::GetAllocatorStats( FGenericMemoryStats& out_Stats )
 	out_Stats.Add( GET_STATDESCRIPTION( STAT_Binned_TotalAllocs ), LocalTotalAllocs );
 	out_Stats.Add( GET_STATDESCRIPTION( STAT_Binned_SlackCurrent ), LocalSlackCurrent );
 #endif // STATS
-}
\ No newline at end of file
+}
+
+void FMallocBinned::InitializeStatsMetadata()
+{
+	FMalloc::InitializeStatsMetadata();
+
+	// Initialize stats metadata here instead of UpdateStats.
+	// Mostly to avoid dead-lock when stats malloc profiler is enabled.
+	GET_STATFNAME(STAT_Binned_OsCurrent);
+	GET_STATFNAME(STAT_Binned_OsPeak);
+	GET_STATFNAME(STAT_Binned_WasteCurrent);
+	GET_STATFNAME(STAT_Binned_WastePeak);
+	GET_STATFNAME(STAT_Binned_UsedCurrent);
+	GET_STATFNAME(STAT_Binned_UsedPeak);
+	GET_STATFNAME(STAT_Binned_CurrentAllocs);
+	GET_STATFNAME(STAT_Binned_TotalAllocs);
+	GET_STATFNAME(STAT_Binned_SlackCurrent);
+}
diff --git a/Engine/Source/Runtime/Core/Private/HAL/MemoryBase.cpp b/Engine/Source/Runtime/Core/Private/HAL/MemoryBase.cpp
--- a/Engine/Source/Runtime/Core/Private/HAL/MemoryBase.cpp
+++ b/Engine/Source/Runtime/Core/Private/HAL/MemoryBase.cpp
@@ -53,6 +53,17 @@ static FCurrentFrameCalls& GetCurrentFrameCalls()
 	return CurrentFrameCalls;
 }
 
+void FMalloc::InitializeStatsMetadata()
+{
+	// Initialize stats metadata here instead of UpdateStats.
+	// Mostly to avoid dead-lock when stats malloc profiler is enabled.
+	GET_STATFNAME(STAT_MallocCalls);
+	GET_STATFNAME(STAT_ReallocCalls);
+	GET_STATFNAME(STAT_FreeCalls);
+	GET_STATFNAME(STAT_TotalAllocatorCalls);
+}
+
+
 void FMalloc::UpdateStats()
 {
 	GetCurrentFrameCalls().Update();
@@ -72,4 +83,4 @@ void FMalloc::GetAllocatorStats( FGenericMemoryStats& out_Stats )
 	out_Stats.Add( GET_STATDESCRIPTION( STAT_FreeCalls ), GetCurrentFrameCalls().FreeCalls );
 	out_Stats.Add( GET_STATDESCRIPTION( STAT_TotalAllocatorCalls ), GetCurrentFrameCalls().AllocatorCalls );
 #endif // STATS
-}
\ No newline at end of file
+}
diff --git a/Engine/Source/Runtime/Core/Private/HAL/UnrealMemory.cpp b/Engine/Source/Runtime/Core/Private/HAL/UnrealMemory.cpp
--- a/Engine/Source/Runtime/Core/Private/HAL/UnrealMemory.cpp
+++ b/Engine/Source/Runtime/Core/Private/HAL/UnrealMemory.cpp
@@ -14,6 +14,7 @@
 #include "MallocProfiler.h"
 #include "MallocThreadSafeProxy.h"
 #include "MallocCrash.h"
+#include "StatsMallocProfilerProxy.h"
 
 
 /** Helper function called on first allocation to create and initialize GMalloc */
@@ -39,6 +40,11 @@ void GCreateMalloc()
 	{
 		GMalloc = new FMallocThreadSafeProxy( GMalloc );
 	}
+
+#if	STATS
+	// Create the stats malloc profiler proxy, needs to be enabled separately.
+	//GMalloc = FStatsMallocProfilerProxy::Get();
+#endif // STATS
 }
 
 
diff --git a/Engine/Source/Runtime/Core/Private/Stats/Stats2.cpp b/Engine/Source/Runtime/Core/Private/Stats/Stats2.cpp
--- a/Engine/Source/Runtime/Core/Private/Stats/Stats2.cpp
+++ b/Engine/Source/Runtime/Core/Private/Stats/Stats2.cpp
@@ -1,12 +1,11 @@
 // Copyright 1998-2014 Epic Games, Inc. All Rights Reserved.
 
-/*=============================================================================
-	UnStats.cpp: Performance stats framework.
-=============================================================================*/
-
-
 #include "CorePrivatePCH.h"
 
+/*-----------------------------------------------------------------------------
+	Global
+-----------------------------------------------------------------------------*/
+
 DECLARE_THREAD_SINGLETON( FThreadIdleStats );
 
 struct FStats2Globals
@@ -23,7 +22,20 @@ struct FStats2Globals
 static struct FForceInitAtBootFStats2 : public TForceInitAtBoot<FStats2Globals>
 {} FForceInitAtBootFStats2;
 
-DECLARE_FLOAT_COUNTER_STAT( TEXT( "Seconds Per Cycle" ), STAT_SecondsPerCycle, STATGROUP_Engine );
+DECLARE_FLOAT_COUNTER_STAT( TEXT("Seconds Per Cycle"), STAT_SecondsPerCycle, STATGROUP_Engine );
+DECLARE_DWORD_COUNTER_STAT( TEXT("Frame Packets Received"),STAT_StatFramePacketsRecv,STATGROUP_StatSystem);
+
+DECLARE_CYCLE_STAT(TEXT("WaitForStats"),STAT_WaitForStats,STATGROUP_Engine);
+DECLARE_CYCLE_STAT(TEXT("StatsNew Tick"),STAT_StatsNewTick,STATGROUP_StatSystem);
+DECLARE_CYCLE_STAT(TEXT("Parse Meta"),STAT_StatsNewParseMeta,STATGROUP_StatSystem);
+DECLARE_CYCLE_STAT(TEXT("Add To History"),STAT_StatsNewAddToHistory,STATGROUP_StatSystem);
+
+DECLARE_MEMORY_STAT( TEXT( "Stats Descriptions" ), STAT_StatDescMemory, STATGROUP_StatSystem );
+DECLARE_MEMORY_STAT( TEXT( "Memory Map" ), STAT_MemoryAllocationMap, STATGROUP_StatSystem );
+
+DEFINE_STAT(STAT_FrameTime);
+DEFINE_STAT(STAT_FPS);
+DEFINE_STAT(STAT_DrawStats);
 DEFINE_STAT(STAT_StatMessagesMemory);
 
 /*-----------------------------------------------------------------------------
@@ -119,22 +131,17 @@ check(InStatId != NAME_None);
 
 */
 
-
 #if STATS
 
 #include "TaskGraphInterfaces.h"
 #include "StatsData.h"
 
-DEFINE_STAT(STAT_FrameTime);
-DEFINE_STAT(STAT_FPS);
-DEFINE_STAT(STAT_DrawStats);
-
-DECLARE_DWORD_COUNTER_STAT(TEXT("Frame Packets Received"),STAT_StatFramePacketsRecv,STATGROUP_StatSystem);
-DECLARE_MEMORY_STAT( TEXT( "Stats descriptions (ansi+wide)" ), STAT_StatDescMemory, STATGROUP_StatSystem );
-
-
 TStatIdData TStatId::TStatId_NAME_None;
 
+/*-----------------------------------------------------------------------------
+	FStartupMessages
+-----------------------------------------------------------------------------*/
+
 void FStartupMessages::AddThreadMetadata( const FName InThreadName, uint32 InThreadID )
 {
 	// Make unique name.
@@ -164,6 +171,10 @@ FStartupMessages& FStartupMessages::Get()
 	return *Messages;
 }
 
+/*-----------------------------------------------------------------------------
+	FThreadSafeStaticStatBase
+-----------------------------------------------------------------------------*/
+
 void FThreadSafeStaticStatBase::DoSetup(const char* InStatName, const TCHAR* InStatDesc, const char* InGroupName, const char* InGroupCategory, const TCHAR* InGroupDesc, bool bDefaultEnable, bool bShouldClearEveryFrame, EStatDataType::Type InStatType, bool bCycleStat, FPlatformMemory::EMemoryCounterRegion InMemoryRegion) const
 {
 	FName TempName(InStatName);
@@ -176,6 +187,10 @@ void FThreadSafeStaticStatBase::DoSetup(const char* InStatName, const TCHAR* InS
 	check(!OldHighPerformanceEnable || HighPerformanceEnable == OldHighPerformanceEnable); // we are assigned two different groups?
 }
 
+/*-----------------------------------------------------------------------------
+	FStatGroupEnableManager
+-----------------------------------------------------------------------------*/
+
 DEFINE_LOG_CATEGORY_STATIC(LogStatGroupEnableManager, Log, All);
 
 class FStatGroupEnableManager : public IStatGroupEnableManager
@@ -506,17 +521,9 @@ IStatGroupEnableManager& IStatGroupEnableManager::Get()
 	return *Singleton;
 }
 
-uint32 FThreadStats::TlsSlot = 0;
-FThreadSafeCounter FThreadStats::MasterEnableCounter;
-FThreadSafeCounter FThreadStats::MasterEnableUpdateNumber;
-FThreadSafeCounter FThreadStats::MasterDisableChangeTagLock;
-bool FThreadStats::bMasterEnable = false;
-bool FThreadStats::bMasterDisableForever = false;
-
-
-DECLARE_CYCLE_STAT(TEXT("StatsNew Tick"),STAT_StatsNewTick,STATGROUP_StatSystem);
-DECLARE_CYCLE_STAT(TEXT("Parse Meta"),STAT_StatsNewParseMeta,STATGROUP_StatSystem);
-DECLARE_CYCLE_STAT(TEXT("Add To History"),STAT_StatsNewAddToHistory,STATGROUP_StatSystem);
+/*-----------------------------------------------------------------------------
+	FStatNameAndInfo
+-----------------------------------------------------------------------------*/
 
 FName FStatNameAndInfo::ToLongName(FName InStatName, char const* InGroup, char const* InCategory, TCHAR const* InDescription)
 {
@@ -618,6 +625,10 @@ FName FStatNameAndInfo::GetGroupCategoryFrom(FName InLongName)
 	return NAME_None;
 }
 
+/*-----------------------------------------------------------------------------
+	FStatsThread
+-----------------------------------------------------------------------------*/
+
 static TAutoConsoleVariable<int32> CVarDumpStatPackets(	TEXT("DumpStatPackets"),0,	TEXT("If true, dump stat packets."));
 
 /** The rendering thread runnable object. */
@@ -654,12 +665,16 @@ public:
 
 	virtual void Tick() override
 	{
+		// Ignore all memory stack inside the stats thread tick, not really sure about that.
+		FThreadStats* ThreadStats = FThreadStats::GetThreadStats();
+		FStatMessageLock MessageLock(ThreadStats->MemoryMessageScope);
+
 		static double LastTime = -1.0;
 		if (bReadyToProcess && FPlatformTime::Seconds() - LastTime > .005) // we won't process more than every 5ms
 		{
 			// Update the seconds per cycle.
 			SET_FLOAT_STAT( STAT_SecondsPerCycle, FPlatformTime::GetSecondsPerCycle() ); 
-
+			SET_MEMORY_STAT( STAT_MemoryAllocationMap, State.GetAllocations().GetAllocatedSize()+State.GetNonSequentialAllocations().GetAllocatedSize() );
 			IStatGroupEnableManager::Get().UpdateMemoryUsage();
 
 			SCOPE_CYCLE_COUNTER(STAT_StatsNewTick);
@@ -700,15 +715,16 @@ public:
 
 	void StatMessage(FStatPacket* Packet)
 	{
-		// the parallel case doesn't ever seem to be faster, so it disabled by default.
 		if (CVarDumpStatPackets.GetValueOnAnyThread())
 		{
-			UE_LOG(LogTemp, Log, TEXT("Packet from %x"), Packet->ThreadId);
+			UE_LOG(LogStats, Log, TEXT("Packet from %x with %d messages"), Packet->ThreadId, Packet->StatMessages.Num());
 		}
 
 		bReadyToProcess = Packet->ThreadType != EThreadType::Other;
 		IncomingData.Packets.Add(Packet);
-		INC_MEMORY_STAT_BY(STAT_StatMessagesMemory,Packet->StatMessages.GetAllocatedSize());
+		const uint32 PacketMemory = Packet->StatMessages.GetAllocatedSize();
+		INC_MEMORY_STAT_BY(STAT_StatMessagesMemory, PacketMemory);
+
 		Tick();
 	}
 
@@ -720,6 +736,10 @@ public:
 
 };
 
+/*-----------------------------------------------------------------------------
+	FStatMessagesTask
+-----------------------------------------------------------------------------*/
+
 // not using a delegate here to allow higher performance since we may end up sending a lot of small message arrays to the thread.
 class FStatMessagesTask
 {
@@ -746,10 +766,48 @@ public:
 	}
 };
 
+/*-----------------------------------------------------------------------------
+	FThreadStatsPool
+-----------------------------------------------------------------------------*/
+
+FThreadStatsPool::FThreadStatsPool()
+{
+	for( int32 Index = 0; Index < NUM_ELEMENTS_IN_POOL; ++Index )
+	{
+		Pool.Push( new FThreadStats(true) );
+	}
+}
+
+FThreadStats* FThreadStatsPool::GetFromPool()
+{
+	FThreadStats* Result = new(Pool.Pop()) FThreadStats();
+	check(Result && "Increase NUM_ELEMENTS_IN_POOL");
+	return Result;
+}
+
+void FThreadStatsPool::ReturnToPool( FThreadStats* Instance )
+{
+	check(Instance);
+	Instance->~FThreadStats();
+	Pool.Push(Instance);
+}
+
+/*-----------------------------------------------------------------------------
+	FThreadStats
+-----------------------------------------------------------------------------*/
+
+uint32 FThreadStats::TlsSlot = 0;
+FThreadSafeCounter FThreadStats::MasterEnableCounter;
+FThreadSafeCounter FThreadStats::MasterEnableUpdateNumber;
+FThreadSafeCounter FThreadStats::MasterDisableChangeTagLock;
+bool FThreadStats::bMasterEnable = false;
+bool FThreadStats::bMasterDisableForever = false;
+
 FThreadStats::FThreadStats():
 	CurrentGameFrame(FStatsThreadState::GetLocalState().CurrentGameFrame),
 	ScopeCount(0), 
 	bWaitForExplicitFlush(0),
+	MemoryMessageScope(0),
 	bSawExplicitFlush(false)
 {
 	Packet.ThreadId = FPlatformTLS::GetCurrentThreadId();
@@ -769,6 +827,14 @@ FThreadStats::FThreadStats():
 	FPlatformTLS::SetTlsValue(TlsSlot, this);
 }
 
+CORE_API FThreadStats::FThreadStats( bool ):
+	CurrentGameFrame(0),
+	ScopeCount(0), 
+	bWaitForExplicitFlush(0),
+	MemoryMessageScope(0),
+	bSawExplicitFlush(false)
+{}
+
 void FThreadStats::CheckEnable()
 {
 	bool bOldMasterEnable(bMasterEnable);
@@ -781,9 +847,10 @@ void FThreadStats::CheckEnable()
 	}
 }
 
-
 void FThreadStats::Flush(bool bHasBrokenCallstacks)
 {
+	FStatMessageLock MessageLock(MemoryMessageScope);
+
 	if (bMasterDisableForever)
 	{
 		Packet.StatMessages.Empty();
@@ -878,6 +945,8 @@ void FThreadStats::StartThread()
 	}
 	check(IsThreadingReady());
 	CheckEnable();
+	// Preallocate a bunch of FThreadStats to avoid dynamic memory allocation.
+	FThreadStatsPool::Get();
 
 	{
 		FString CmdLine(FCommandLine::Get());
@@ -906,21 +975,24 @@ void FThreadStats::StartThread()
 			CmdLine = CmdLine.Mid(Index + StatCmds.Len());
 		}
 	}
+
+	if (FParse::Param(FCommandLine::Get(), TEXT("LoadTimeStats")))
 	{
-		if (FParse::Param(FCommandLine::Get(), TEXT("LoadTimeStats")))
-		{
-			DirectStatsCommand(TEXT("stat group enable LinkerLoad"));
-			DirectStatsCommand(TEXT("stat group enable AsyncLoad"));
-			DirectStatsCommand(TEXT("stat dumpsum -start -ms=250"));
-		}
-		if (FParse::Param(FCommandLine::Get(), TEXT("LoadTimeFile")))
-		{
-			DirectStatsCommand(TEXT("stat group enable LinkerLoad"));
-			DirectStatsCommand(TEXT("stat group enable AsyncLoad"));
-			DirectStatsCommand(TEXT("stat startfile"));
-		}
+		DirectStatsCommand(TEXT("stat group enable LinkerLoad"));
+		DirectStatsCommand(TEXT("stat group enable AsyncLoad"));
+		DirectStatsCommand(TEXT("stat dumpsum -start -ms=250"));
+	}
+	if (FParse::Param(FCommandLine::Get(), TEXT("LoadTimeFile")))
+	{
+		DirectStatsCommand(TEXT("stat group enable LinkerLoad"));
+		DirectStatsCommand(TEXT("stat group enable AsyncLoad"));
+		DirectStatsCommand(TEXT("stat startfile"));
 	}
 
+	if( FParse::Param(FCommandLine::Get(), TEXT("MemoryProfiler")) )
+	{
+		DirectStatsCommand(TEXT("stat memoryprofiler enable"), true);
+	}
 }
 
 static FGraphEventRef LastFramesEvents[MAX_STAT_LAG];
@@ -942,8 +1014,6 @@ void FThreadStats::StopThread()
 	}
 }
 
-DECLARE_CYCLE_STAT(TEXT("WaitForStats"),STAT_WaitForStats,STATGROUP_Engine);
-
 void FThreadStats::WaitForStats()
 {
 	check(IsInGameThread());
@@ -966,4 +1036,4 @@ void FThreadStats::WaitForStats()
 	}
 }
 
-#endif
\ No newline at end of file
+#endif
diff --git a/Engine/Source/Runtime/Core/Private/Stats/StatsCommand.cpp b/Engine/Source/Runtime/Core/Private/Stats/StatsCommand.cpp
--- a/Engine/Source/Runtime/Core/Private/Stats/StatsCommand.cpp
+++ b/Engine/Source/Runtime/Core/Private/Stats/StatsCommand.cpp
@@ -8,6 +8,7 @@
 
 #include "StatsData.h"
 #include "StatsFile.h"
+#include "StatsMallocProfilerProxy.h"
 
 DECLARE_CYCLE_STAT(TEXT("Hitch Scan"),STAT_HitchScan,STATGROUP_StatSystem);
 DECLARE_CYCLE_STAT(TEXT("HUD Group"),STAT_HUDGroup,STATGROUP_StatSystem);
@@ -962,6 +963,95 @@ static void DumpCPU(int64 Frame)
 	StatsMasterEnableSubtract();
 }
 
+static void DumpMemory(int64 /*Frame*/)
+{
+	FStatsThreadState& Stats = FStatsThreadState::GetLocalState();
+	const TMap<uint64,FAllocationInfo>& Allocations = Stats.GetAllocations();
+	if( Allocations.Num() == 0  )
+	{
+		UE_LOG(LogStats, Warning, TEXT("There are no allocations, make sure memory profiler is enabled") );
+	}
+	else
+	{
+		const TMap<uint64,FAllocationInfoEx>& NonSequentialAllocations = Stats.GetNonSequentialAllocations();
+
+		struct FSizeAndCount
+		{
+			uint64 Size;
+			uint64 Count;
+			FSizeAndCount()
+				: Size(0)
+				, Count(0)
+			{}
+		};
+
+		TMap<FName,FSizeAndCount> ScopedAllocations;
+
+		uint64 TotalAllocatedMemory = 0;
+		for( const auto& It : Allocations )
+		{
+			const FAllocationInfo& Alloc = It.Value;
+			FSizeAndCount& SizeAndCount = ScopedAllocations.FindOrAdd(Alloc.Scope);
+			SizeAndCount.Size += Alloc.Size;
+			SizeAndCount.Count += 1;
+
+			TotalAllocatedMemory += Alloc.Size;
+		}
+
+		// Dump memory to the log.
+		struct FGreater
+		{
+			FORCEINLINE bool operator()(const FSizeAndCount& A, const FSizeAndCount& B) const { return B.Size < A.Size; }
+		};
+
+		ScopedAllocations.ValueSort( FGreater() );
+
+		const float MaxPctDisplayed = 0.99f;
+		int32 CurrentIndex = 0;
+		uint64 DisplayedSoFar = 0;
+		UE_LOG( LogStats, Warning, TEXT("Index, Size (Size MB), Count, Stat desc [ Group ]") );
+		for( const auto& It : ScopedAllocations )
+		{
+			const FName LongName = It.Key;
+			const FSizeAndCount& SizeAndCount = It.Value;
+
+			const FString ShortName = FStatNameAndInfo::GetShortNameFrom(LongName).ToString();
+			const FString Group = FStatNameAndInfo::GetGroupNameFrom(LongName).ToString();
+			FString Desc = FStatNameAndInfo::GetDescriptionFrom(LongName);
+			Desc.Trim();
+
+			if( Desc != ShortName )
+			{
+				if( Desc.Len() )
+				{
+					Desc += TEXT( " - " );
+				}
+				Desc += ShortName;
+			}
+
+			UE_LOG( LogStats, Warning, TEXT("%2i, %llu (%.2f MB), %llu, %s [%s] "), 
+				CurrentIndex, 
+				SizeAndCount.Size, 
+				SizeAndCount.Size/1024.0f/1024.0f,
+				SizeAndCount.Count,
+				*Desc, 
+				*Group );
+
+			CurrentIndex++;
+			DisplayedSoFar += SizeAndCount.Size;
+
+			const float CurrentPct = (float)DisplayedSoFar/(float)TotalAllocatedMemory;
+			if( CurrentPct > MaxPctDisplayed )
+			{
+				break;
+			}
+		}
+
+		UE_LOG(LogStats, Warning, TEXT("Allocated memory: %llu bytes (%.2f MB)"), TotalAllocatedMemory, TotalAllocatedMemory/1024.0f/1024.0f );
+	}
+
+	Stats.NewFrameDelegate.RemoveStatic(&DumpMemory);
+}
 
 static struct FDumpMultiple* DumpMultiple = NULL;
 
@@ -1082,7 +1172,14 @@ static void PrintStatsHelpToOutputDevice( FOutputDevice& Ar )
 	Ar.Log( TEXT("stat startfile - starts dumping a capture"));
 	Ar.Log( TEXT("stat stopfile - stops dumping a capture"));
 
+	Ar.Log( TEXT("stat startfileraw - starts dumping a raw capture"));
+	Ar.Log( TEXT("stat stopfileraw - stops dumping a raw capture"));
+
 	Ar.Log( TEXT("stat toggledebug - toggles tracking the most memory expensive stats"));
+
+	Ar.Log( TEXT("stat memoryprofiler enable - enables tracking all memory operations, run 'stat startfileraw' before"));
+	Ar.Log( TEXT("stat memoryprofiler disable - disables tracking all memory operations, run 'stat startfileraw' after"));
+	Ar.Log( TEXT("stat dumpmemory - dump basic memory statistics WIP, memoryprofiler must to be enabled"));
 }
 
 static void CommandTestFile()
@@ -1134,6 +1231,10 @@ static void StatCmd(FString InCmd)
 	{
 		StatsMasterEnableAdd();
 		Stats.NewFrameDelegate.AddStatic(&DumpFrame);
+	} 
+	else if( FParse::Command(&Cmd,TEXT("DumpMemory")) )
+	{
+		Stats.NewFrameDelegate.AddStatic(&DumpMemory);
 	}
 	else if ( FParse::Command(&Cmd,TEXT("DUMPNONFRAME")) )
 	{
@@ -1238,6 +1339,14 @@ static void StatCmd(FString InCmd)
 	{
 		FStatsThreadState::GetLocalState().ToggleFindMemoryExtensiveStats();
 	}
+	else if ( FParse::Command( &Cmd, TEXT( "memoryprofiler" ) ) )
+	{
+		IStatGroupEnableManager::Get().StatGroupEnableManagerCommand( TEXT("enable LinkerLoad") );
+		IStatGroupEnableManager::Get().StatGroupEnableManagerCommand( TEXT("enable AsyncLoad") );
+
+		const bool bEnable = FParse::Command( &Cmd, TEXT( "enable" ) ) || FParse::Command( &Cmd, TEXT( "start" ) );
+		FStatsMallocProfilerProxy::Get()->SetState( bEnable );
+	}
 	// @see FStatHierParams
 	else if ( FParse::Command( &Cmd, TEXT( "hier" ) ) )
 	{
@@ -1321,6 +1430,9 @@ bool DirectStatsCommand(const TCHAR* Cmd, bool bBlockForCompletion /*= false*/,
 		else if( FParse::Command(&TempCmd,TEXT("DUMPFRAME")) )
 		{
 		}
+		else if( FParse::Command(&TempCmd,TEXT("DumpMemory")) )
+		{
+		}
 		else if ( FParse::Command(&TempCmd,TEXT("DUMPNONFRAME")) )
 		{
 		}
@@ -1366,6 +1478,9 @@ bool DirectStatsCommand(const TCHAR* Cmd, bool bBlockForCompletion /*= false*/,
 		else if ( FParse::Command( &TempCmd, TEXT( "toggledebug" ) ) )
 		{
 		}
+		else if ( FParse::Command( &TempCmd, TEXT( "memoryprofiler" ) ) )
+		{
+		}
 		else
 #endif
 		{
diff --git a/Engine/Source/Runtime/Core/Private/Stats/StatsData.cpp b/Engine/Source/Runtime/Core/Private/Stats/StatsData.cpp
--- a/Engine/Source/Runtime/Core/Private/Stats/StatsData.cpp
+++ b/Engine/Source/Runtime/Core/Private/Stats/StatsData.cpp
@@ -459,7 +459,8 @@ void FStatPacketArray::Empty()
 {
 	for (int32 Index = 0; Index < Packets.Num(); Index++)
 	{
-		DEC_MEMORY_STAT_BY(STAT_StatMessagesMemory,Packets[Index]->StatMessages.GetAllocatedSize());
+		const uint32 PacketMemory = Packets[Index]->StatMessages.GetAllocatedSize();
+		DEC_MEMORY_STAT_BY(STAT_StatMessagesMemory, PacketMemory);
 		delete Packets[Index];
 	}
 	Packets.Empty();
@@ -642,28 +643,32 @@ void FStatsThreadState::ProcessNonFrameStats(FStatMessagesArray& Data, TSet<FNam
 				Op != EStatOperation::ChildrenEnd &&
 				Op != EStatOperation::Leaf &&
 				Op != EStatOperation::AdvanceFrameEventGameThread &&
-				Op != EStatOperation::AdvanceFrameEventRenderThread
+				Op != EStatOperation::AdvanceFrameEventRenderThread /*&&
+				Op != EStatOperation::Memory*/
 				))
 			{
 				UE_LOG(LogStats, Fatal, TEXT( "Stat %s was not cleared every frame, but was used with a scope cycle counter." ), *Item.NameAndInfo.GetRawName().ToString() );
 			}
 			else
 			{
-				FStatMessage* Result = NotClearedEveryFrame.Find(Item.NameAndInfo.GetRawName());
-				if (!Result)
+				if( Op != EStatOperation::Memory )
 				{
-					UE_LOG(LogStats, Error, TEXT( "Stat %s was cleared every frame, but we don't have metadata for it. Data loss." ), *Item.NameAndInfo.GetRawName().ToString() );
-				}
-				else
-				{
-					if (NonFrameStatsFound)
+					FStatMessage* Result = NotClearedEveryFrame.Find(Item.NameAndInfo.GetRawName());
+					if (!Result)
 					{
-						NonFrameStatsFound->Add(Item.NameAndInfo.GetRawName());
+						UE_LOG(LogStats, Error, TEXT( "Stat %s was cleared every frame, but we don't have metadata for it. Data loss." ), *Item.NameAndInfo.GetRawName().ToString() );
 					}
-					FStatsUtils::AccumulateStat(*Result, Item);
-					Item = *Result; // now just write the accumulated value back into the stream
-					check(Item.NameAndInfo.GetField<EStatOperation>() == EStatOperation::Set);
-				}
+					else
+					{
+						if (NonFrameStatsFound)
+						{
+							NonFrameStatsFound->Add(Item.NameAndInfo.GetRawName());
+						}
+						FStatsUtils::AccumulateStat(*Result, Item);
+						Item = *Result; // now just write the accumulated value back into the stream
+						check(Item.NameAndInfo.GetField<EStatOperation>() == EStatOperation::Set);
+					}
+				}			
 			}
 		}
 	}
@@ -736,7 +741,7 @@ void FStatsThreadState::AddToHistoryAndEmpty(FStatPacketArray& NewData)
 				// add the non frame stats as a new last packet
 
 				check(PacketToCopyForNonFrame);
-
+				FThreadStats* ThreadStats = FThreadStats::GetThreadStats();
 				FStatMessagesArray* NonFrameMessages = NULL;
 
 				for (TMap<FName, FStatMessage>::TConstIterator It(NotClearedEveryFrame); It; ++It)
@@ -748,11 +753,20 @@ void FStatsThreadState::AddToHistoryAndEmpty(FStatPacketArray& NewData)
 							const int32 NonFrameIndex = Frame.Packets.Add(new FStatPacket(*PacketToCopyForNonFrame));
 							NonFrameMessages = &Frame.Packets[NonFrameIndex]->StatMessages;
 						}
-						new (*NonFrameMessages) FStatMessage(It.Value());
+						
+						{
+							FStatMessageLock MessageLock(ThreadStats->MemoryMessageScope);
+							new (*NonFrameMessages) FStatMessage(It.Value());
+						}
+						
 					}
 				}
 
-				INC_MEMORY_STAT_BY(STAT_StatMessagesMemory,NonFrameMessages->GetAllocatedSize());
+				if(NonFrameMessages)
+				{
+					const uint32 PacketMemory = NonFrameMessages->GetAllocatedSize();
+					INC_MEMORY_STAT_BY(STAT_StatMessagesMemory, PacketMemory);
+				}
 
 				GoodFrames.Add(FrameNum);
 			}
@@ -798,6 +812,7 @@ void FStatsThreadState::AddToHistoryAndEmpty(FStatPacketArray& NewData)
 	for (auto It = History.CreateIterator(); It; ++It)
 	{
 		int64 ThisFrame = It.Key();
+		FStatPacketArray& StatPacketArray = It.Value();
 		if (ThisFrame <= LastFullFrameMetaAndNonFrame && ThisFrame < MinFrameToKeep)
 		{
 			check(ThisFrame <= LastFullFrameMetaAndNonFrame);
@@ -1034,6 +1049,64 @@ void FStatsThreadState::GetRawStackStats(int64 TargetFrame, FRawStatStackNode& R
 						Current = Stack.Last();
 					}
 				}
+				else if( Op == EStatOperation::Memory )
+				{
+					//const FStatMessage& ScopeStartMessage = *StartStack.Last();
+
+					// Experimental code used only to test the implementation.
+					// First memory operation is Alloc or Free
+					const uint64 EncodedPtr = Item.GetValue_Ptr();
+					const bool bIsAlloc = (EncodedPtr & (uint64)EMemoryOperation::Alloc) != 0;
+					const bool bIsFree = (EncodedPtr & (uint64)EMemoryOperation::Free) != 0;
+					const uint64 Ptr = EncodedPtr & ~(uint64)EMemoryOperation::Mask;
+					if( bIsAlloc )
+					{
+						// @see FStatsMallocProfilerProxy::TrackAlloc
+						// After alloc ptr message there is always alloc size message.
+						Index++;
+						const FStatMessage& AllocSizeMessage = Data[Index];
+						const int64 AllocSize = AllocSizeMessage.GetValue_int64();
+
+						// Check the non-sequential allocation map.
+						FAllocationInfoEx* NonSeqAllocationInfo = NonSequentialAllocationMap.Find(Ptr);
+						if( NonSeqAllocationInfo )
+						{
+							// We don't need to add this allocation again.
+							NonSequentialAllocationMap.Remove(Ptr);
+						}
+						else
+						{
+							FAllocationInfo* OldAllocationInfo = AllocationMap.Find(Ptr);
+							if( OldAllocationInfo != nullptr )
+							{
+								// Ignore for now...
+							}
+
+							// Add a new allocation.
+							AllocationMap.Add( Ptr, FAllocationInfo(AllocSize, Current->Meta.NameAndInfo.GetEncodedName()) );
+						}					
+					}
+					else if( bIsFree )
+					{
+						// Remove the allocation, check if exists.
+						const int32 NumRemoved = AllocationMap.Remove( Ptr );
+
+						// We removed non-existent allocation, 
+						// 1. the timestamp of the allocation is before we started memory profiling
+						// 2. the allocation has been done on different thread than we are currently processing
+						if( NumRemoved == 0 )
+						{
+							NonSequentialAllocationMap.Add(Ptr, FAllocationInfoEx(
+								0, 
+								Current->Meta.NameAndInfo.GetEncodedName(),
+								/*ScopeStartMessage.GetValue_int64()*/0));
+						}
+					}
+					else
+					{
+						UE_LOG(LogStats, Warning, TEXT("Pointer from a memory operation is invalid") );
+					}
+				}
 				else if (OutNonStackStats)
 				{
 					FStatsUtils::AddNonStackStats( LongName, Item, Op, ThisFrameNonStackStats );
@@ -1502,62 +1575,78 @@ void FStatsUtils::AccumulateStat(FStatMessage& Dest, FStatMessage const& Item, E
 	check(Dest.NameAndInfo.GetFlag(EStatMetaFlags::IsPackedCCAndDuration) == Item.NameAndInfo.GetFlag(EStatMetaFlags::IsPackedCCAndDuration));
 	switch (Item.NameAndInfo.GetField<EStatDataType>())
 	{
-	case EStatDataType::ST_int64:
-		switch (Op)
-		{
-		case EStatOperation::Set:
-			Dest.GetValue_int64() = Item.GetValue_int64();
-			break;
-		case EStatOperation::Clear:
-			Dest.GetValue_int64() = 0;
-			break;
-		case EStatOperation::Add:
-			Dest.GetValue_int64() += Item.GetValue_int64();
-			break;
-		case EStatOperation::Subtract:
-			if (Dest.NameAndInfo.GetFlag(EStatMetaFlags::IsPackedCCAndDuration))
+		case EStatDataType::ST_int64:
+			switch (Op)
 			{
-				// we don't subtract call counts, only times
-				Dest.GetValue_int64() = ToPackedCallCountDuration(
-					FromPackedCallCountDuration_CallCount(Dest.GetValue_int64()),
-					FromPackedCallCountDuration_Duration(Dest.GetValue_int64()) - FromPackedCallCountDuration_Duration(Item.GetValue_int64()));
+				case EStatOperation::Set:
+					Dest.GetValue_int64() = Item.GetValue_int64();
+					break;
+				case EStatOperation::Clear:
+					Dest.GetValue_int64() = 0;
+					break;
+				case EStatOperation::Add:
+					Dest.GetValue_int64() += Item.GetValue_int64();
+					break;
+				case EStatOperation::Subtract:
+					if (Dest.NameAndInfo.GetFlag(EStatMetaFlags::IsPackedCCAndDuration))
+					{
+						// we don't subtract call counts, only times
+						Dest.GetValue_int64() = ToPackedCallCountDuration(
+							FromPackedCallCountDuration_CallCount(Dest.GetValue_int64()),
+							FromPackedCallCountDuration_Duration(Dest.GetValue_int64()) - FromPackedCallCountDuration_Duration(Item.GetValue_int64()));
+					}
+					else
+					{
+						Dest.GetValue_int64() -= Item.GetValue_int64();
+					}
+					break;
+				case EStatOperation::MaxVal:
+					StatOpMaxVal_Int64( Dest.NameAndInfo, Dest.GetValue_int64(), Item.GetValue_int64() );
+					break;
+
+				// Nothing here at this moment.
+				case EStatOperation::Memory:
+					break;
+
+				default:
+					check(0);
 			}
-			else
+			break;
+
+		case EStatDataType::ST_double:
+			switch (Op)
 			{
-				Dest.GetValue_int64() -= Item.GetValue_int64();
+				case EStatOperation::Set:
+					Dest.GetValue_double() = Item.GetValue_double();
+					break;
+				case EStatOperation::Clear:
+					Dest.GetValue_double() = 0;
+					break;
+				case EStatOperation::Add:
+					Dest.GetValue_double() += Item.GetValue_double();
+					break;
+				case EStatOperation::Subtract:
+					Dest.GetValue_double() -= Item.GetValue_double();
+					break;
+				case EStatOperation::MaxVal:
+					Dest.GetValue_double() = FMath::Max<double>(Dest.GetValue_double(), Item.GetValue_double());
+					break;
+
+				// Nothing here at this moment.
+				case EStatOperation::Memory:
+					break;
+
+				default:
+					check(0);
 			}
 			break;
-		case EStatOperation::MaxVal:
-			StatOpMaxVal_Int64( Dest.NameAndInfo, Dest.GetValue_int64(), Item.GetValue_int64() );
-			break;
-		default:
-			check(0);
-		}
-		break;
-	case EStatDataType::ST_double:
-		switch (Op)
-		{
-		case EStatOperation::Set:
-			Dest.GetValue_double() = Item.GetValue_double();
-			break;
-		case EStatOperation::Clear:
-			Dest.GetValue_double() = 0;
-			break;
-		case EStatOperation::Add:
-			Dest.GetValue_double() += Item.GetValue_double();
-			break;
-		case EStatOperation::Subtract:
-			Dest.GetValue_double() -= Item.GetValue_double();
-			break;
-		case EStatOperation::MaxVal:
-			Dest.GetValue_double() = FMath::Max<double>(Dest.GetValue_double(), Item.GetValue_double());
+
+		// Nothing here at this moment.
+		case EStatDataType::ST_Ptr:
 			break;
+
 		default:
 			check(0);
-		}
-		break;
-	default:
-		check(0);
 	}
 }
 
diff --git a/Engine/Source/Runtime/Core/Private/Stats/StatsFile.cpp b/Engine/Source/Runtime/Core/Private/Stats/StatsFile.cpp
--- a/Engine/Source/Runtime/Core/Private/Stats/StatsFile.cpp
+++ b/Engine/Source/Runtime/Core/Private/Stats/StatsFile.cpp
@@ -437,7 +437,8 @@ void FStatsWriteFile::NewFrame( int64 TargetFrame )
 {
 	SCOPE_CYCLE_COUNTER( STAT_StreamFile );
 
-	// Currently raw stat files are limited to 120 frames.
+	// @TODO yrx 2014-11-25 Add stat startfile -num=number of frames to capture
+	/*
 	enum
 	{
 		MAX_NUM_RAWFRAMES = 120,
@@ -455,6 +456,7 @@ void FStatsWriteFile::NewFrame( int64 TargetFrame )
 			return;
 		}
 	}
+	}*/
 
 	WriteFrame( TargetFrame );
 	SendTask();
diff --git a/Engine/Source/Runtime/Core/Private/Stats/StatsMallocProfilerProxy.cpp b/Engine/Source/Runtime/Core/Private/Stats/StatsMallocProfilerProxy.cpp
new file mode 100644
--- /dev/null
+++ b/Engine/Source/Runtime/Core/Private/Stats/StatsMallocProfilerProxy.cpp
@@ -0,0 +1,165 @@
+// Copyright 1998-2014 Epic Games, Inc. All Rights Reserved.
+
+#include "CorePrivatePCH.h"
+
+#if STATS
+
+#include "StatsData.h"
+#include "StatsMallocProfilerProxy.h"
+
+#define DECLARE_MEMORY_PTR_STAT(CounterName,StatId,GroupId)\
+	DECLARE_STAT(CounterName,StatId,GroupId,EStatDataType::ST_Ptr, false, false, FPlatformMemory::MCR_Invalid); \
+	static DEFINE_STAT(StatId)
+
+#define DECLARE_MEMORY_SIZE_STAT(CounterName,StatId,GroupId)\
+	DECLARE_STAT(CounterName,StatId,GroupId,EStatDataType::ST_int64, false, false, FPlatformMemory::MCR_Invalid); \
+	static DEFINE_STAT(StatId)
+
+/** Fake stat group and memory stats. */
+DECLARE_STATS_GROUP(TEXT("Memory Profiler"), STATGROUP_MemoryProfiler, STATCAT_Advanced);
+
+DECLARE_MEMORY_PTR_STAT(TEXT("Memory Free Ptr"),			STAT_Memory_FreePtr,		STATGROUP_MemoryProfiler);
+DECLARE_MEMORY_PTR_STAT(TEXT("Memory Alloc Ptr"),			STAT_Memory_AllocPtr,		STATGROUP_MemoryProfiler);
+DECLARE_MEMORY_SIZE_STAT(TEXT("Memory Alloc Size"),			STAT_Memory_AllocSize,		STATGROUP_MemoryProfiler);
+
+/** Stats for memory usage by the profiler. */
+DECLARE_DWORD_COUNTER_STAT(TEXT("Profiler AllocPtr Calls"),	STAT_Memory_AllocPtr_Calls,	STATGROUP_StatSystem);
+DECLARE_DWORD_COUNTER_STAT(TEXT("Profiler FreePtr Calls"),	STAT_Memory_FreePtr_Calls,	STATGROUP_StatSystem);
+DECLARE_MEMORY_STAT( TEXT( "Profiler AllocPtr" ),			STAT_Memory_AllocPtr_Mem,	STATGROUP_StatSystem );
+DECLARE_MEMORY_STAT( TEXT( "Profiler FreePtr" ),			STAT_Memory_FreePtr_Mem,	STATGROUP_StatSystem );
+
+FStatsMallocProfilerProxy::FStatsMallocProfilerProxy( FMalloc* InMalloc ) 
+	: UsedMalloc( InMalloc )
+	, bEnabled(false)
+{
+}
+
+FStatsMallocProfilerProxy* FStatsMallocProfilerProxy::Get()
+{
+	static FStatsMallocProfilerProxy* Instance = nullptr;
+	if( Instance == nullptr )
+	{
+		Instance = new FStatsMallocProfilerProxy( GMalloc );
+	}
+	return Instance;
+}
+
+
+void FStatsMallocProfilerProxy::SetState( bool bEnable )
+{
+	bEnabled = bEnable;
+	FPlatformMisc::MemoryBarrier();
+	if( bEnable )
+	{
+		StatsMasterEnableAdd();
+	}
+	else
+	{
+		StatsMasterEnableSubtract();
+	}
+}
+
+
+void FStatsMallocProfilerProxy::InitializeStatsMetadata()
+{
+	UsedMalloc->InitializeStatsMetadata();
+
+	// Initialize the memory messages metadata.
+	// Malloc profiler proxy needs to be disabled otherwise it will hit infinite recursion in DoSetup.
+	// Needs to be changed if we want to support boot time memory profiling.
+	const FName NameAllocPtr = GET_STATFNAME(STAT_Memory_AllocPtr);
+	const FName NameFreePtr = GET_STATFNAME(STAT_Memory_FreePtr);
+	const FName NameAllocSize = GET_STATFNAME(STAT_Memory_AllocSize);
+
+	GET_STATFNAME(STAT_Memory_AllocPtr_Calls);
+	GET_STATFNAME(STAT_Memory_FreePtr_Calls);
+
+	GET_STATFNAME(STAT_Memory_AllocPtr_Mem);
+	GET_STATFNAME(STAT_Memory_FreePtr_Mem);
+
+	//StatPtr_STAT_Memory_AllocPtr;
+	//StatPtr_STAT_Memory_FreePtr;
+	//StatPtr_STAT_Memory_AllocSize;
+}
+
+
+void FStatsMallocProfilerProxy::TrackAlloc( void* Ptr, int64 Size )
+{
+	if( bEnabled )
+	{
+		if( Size > 0 )
+		{
+			FThreadStats* ThreadStats = FThreadStats::GetThreadStats();
+
+			if( ThreadStats->MemoryMessageScope == 0 )
+			{
+				ThreadStats->AddMemoryMessage( GET_STATFNAME(STAT_Memory_AllocPtr), (uint64)(UPTRINT)Ptr | (uint64)EMemoryOperation::Alloc );	// 16 bytes
+				ThreadStats->AddMemoryMessage( GET_STATFNAME(STAT_Memory_AllocSize), Size );					// 32 bytes total
+				AllocPtrCalls.Increment();
+			}
+		}	
+	}
+}
+
+void FStatsMallocProfilerProxy::TrackFree( void* Ptr )
+{
+	if( bEnabled )
+	{
+		if( Ptr != nullptr )
+		{
+			FThreadStats* ThreadStats = FThreadStats::GetThreadStats();
+
+			if( ThreadStats->MemoryMessageScope == 0 )
+			{
+				ThreadStats->AddMemoryMessage( GET_STATFNAME(STAT_Memory_FreePtr), (uint64)(UPTRINT)Ptr | (uint64)EMemoryOperation::Free );	// 16 bytes total
+				FreePtrCalls.Increment();
+			}
+		}	
+	}
+}
+
+void* FStatsMallocProfilerProxy::Malloc( SIZE_T Size, uint32 Alignment )
+{
+	void* Ptr = UsedMalloc->Malloc( Size, Alignment );
+	// We lose the Size's precision, but don't worry about it.
+	TrackAlloc( Ptr, (int64)Size );
+	return Ptr;
+}
+
+void* FStatsMallocProfilerProxy::Realloc( void* OldPtr, SIZE_T NewSize, uint32 Alignment )
+{
+	TrackFree( OldPtr );
+	void* NewPtr = UsedMalloc->Realloc( OldPtr, NewSize, Alignment );
+	TrackAlloc( NewPtr, (int64)NewSize );
+	return NewPtr;
+}
+
+void FStatsMallocProfilerProxy::Free( void* Ptr )
+{
+	TrackFree( Ptr );
+	UsedMalloc->Free( Ptr );
+}
+
+void FStatsMallocProfilerProxy::UpdateStats()
+{
+	UsedMalloc->UpdateStats();
+
+	if( bEnabled )
+	{
+		const int32 NumAllocPtrCalls = AllocPtrCalls.GetValue();
+		const int32 NumFreePtrCalls = FreePtrCalls.GetValue();
+
+		SET_DWORD_STAT( STAT_Memory_AllocPtr_Calls, NumAllocPtrCalls );
+		SET_DWORD_STAT( STAT_Memory_FreePtr_Calls, NumFreePtrCalls );
+
+		// AllocPtr, AllocSize
+		SET_MEMORY_STAT( STAT_Memory_AllocPtr_Mem, NumAllocPtrCalls*(sizeof(FStatMessage)*2) );
+		// FreePtr
+		SET_MEMORY_STAT( STAT_Memory_FreePtr_Mem, NumFreePtrCalls*(sizeof(FStatMessage)*1) );
+
+		AllocPtrCalls.Reset();
+		FreePtrCalls.Reset();
+	}
+}
+
+#endif //STATS
\ No newline at end of file
diff --git a/Engine/Source/Runtime/Core/Private/UObject/UnrealNames.cpp b/Engine/Source/Runtime/Core/Private/UObject/UnrealNames.cpp
--- a/Engine/Source/Runtime/Core/Private/UObject/UnrealNames.cpp
+++ b/Engine/Source/Runtime/Core/Private/UObject/UnrealNames.cpp
@@ -795,6 +795,10 @@ void FName::StaticInit()
 		}
 	}
 #endif
+
+	// Initialize stats metadata.
+	// We need to do here, after all hardcoded names have been initialized.
+	GMalloc->InitializeStatsMetadata();
 }
 
 bool& FName::GetIsInitialized()
diff --git a/Engine/Source/Runtime/Core/Public/HAL/MallocAnsi.h b/Engine/Source/Runtime/Core/Public/HAL/MallocAnsi.h
--- a/Engine/Source/Runtime/Core/Public/HAL/MallocAnsi.h
+++ b/Engine/Source/Runtime/Core/Public/HAL/MallocAnsi.h
@@ -39,6 +39,7 @@ public:
 	// FMalloc interface.
 	virtual void* Malloc( SIZE_T Size, uint32 Alignment ) override
 	{
+		IncrementTotalMallocCalls();
 		Alignment = FMath::Max(Size >= 16 ? (uint32)16 : (uint32)8, Alignment);
 
 #if USE_ALIGNED_MALLOC
@@ -60,6 +61,7 @@ public:
 
 	virtual void* Realloc( void* Ptr, SIZE_T NewSize, uint32 Alignment ) override
 	{
+		IncrementTotalReallocCalls();
 		void* Result;
 		Alignment = FMath::Max(NewSize >= 16 ? (uint32)16 : (uint32)8, Alignment);
 
@@ -107,6 +109,7 @@ public:
 
 	virtual void Free( void* Ptr ) override
 	{
+		IncrementTotalFreeCalls();
 #if USE_ALIGNED_MALLOC
 		_aligned_free( Ptr );
 #else
diff --git a/Engine/Source/Runtime/Core/Public/HAL/MallocBinned.h b/Engine/Source/Runtime/Core/Public/HAL/MallocBinned.h
--- a/Engine/Source/Runtime/Core/Public/HAL/MallocBinned.h
+++ b/Engine/Source/Runtime/Core/Public/HAL/MallocBinned.h
@@ -31,14 +31,14 @@
 #include "Array.h"
 
 /** Malloc binned allocator specific stats. */
-DECLARE_MEMORY_STAT_EXTERN(TEXT("Binned Os Current"),	STAT_Binned_OsCurrent,STATGROUP_MemoryAllocator, CORE_API);
-DECLARE_MEMORY_STAT_EXTERN(TEXT("Binned Os Peak"),		STAT_Binned_OsPeak,STATGROUP_MemoryAllocator, CORE_API);
+DECLARE_MEMORY_STAT_EXTERN(TEXT("Binned Os Current"),		STAT_Binned_OsCurrent,STATGROUP_MemoryAllocator, CORE_API);
+DECLARE_MEMORY_STAT_EXTERN(TEXT("Binned Os Peak"),			STAT_Binned_OsPeak,STATGROUP_MemoryAllocator, CORE_API);
 DECLARE_MEMORY_STAT_EXTERN(TEXT("Binned Waste Current"),	STAT_Binned_WasteCurrent,STATGROUP_MemoryAllocator, CORE_API);
-DECLARE_MEMORY_STAT_EXTERN(TEXT("Binned Waste Peak"),	STAT_Binned_WastePeak,STATGROUP_MemoryAllocator, CORE_API);
-DECLARE_MEMORY_STAT_EXTERN(TEXT("Binned Used Current"),	STAT_Binned_UsedCurrent,STATGROUP_MemoryAllocator, CORE_API);
+DECLARE_MEMORY_STAT_EXTERN(TEXT("Binned Waste Peak"),		STAT_Binned_WastePeak,STATGROUP_MemoryAllocator, CORE_API);
+DECLARE_MEMORY_STAT_EXTERN(TEXT("Binned Used Current"),		STAT_Binned_UsedCurrent,STATGROUP_MemoryAllocator, CORE_API);
 DECLARE_MEMORY_STAT_EXTERN(TEXT("Binned Used Peak"),		STAT_Binned_UsedPeak,STATGROUP_MemoryAllocator, CORE_API);
-DECLARE_MEMORY_STAT_EXTERN(TEXT("Binned Current Allocs"),STAT_Binned_CurrentAllocs,STATGROUP_MemoryAllocator, CORE_API);
-DECLARE_MEMORY_STAT_EXTERN(TEXT("Binned Total Allocs"),	STAT_Binned_TotalAllocs,STATGROUP_MemoryAllocator, CORE_API);
+DECLARE_MEMORY_STAT_EXTERN(TEXT("Binned Current Allocs"),	STAT_Binned_CurrentAllocs,STATGROUP_MemoryAllocator, CORE_API);
+DECLARE_MEMORY_STAT_EXTERN(TEXT("Binned Total Allocs"),		STAT_Binned_TotalAllocs,STATGROUP_MemoryAllocator, CORE_API);
 DECLARE_MEMORY_STAT_EXTERN(TEXT("Binned Slack Current"),	STAT_Binned_SlackCurrent,STATGROUP_MemoryAllocator, CORE_API);
 
 
@@ -802,6 +802,7 @@ public:
 		check(!(AddressLimit & (AddressLimit - 1)));
 		check(PageSize <= 65536); // There is internal limit on page size of 64k
 		check(AddressLimit > PageSize); // Check to catch 32 bit overflow in AddressLimit
+		InitializeStatsMetadata();
 
 		/** Shift to get the reference from the indirect tables */
 		PoolBitShift = FPlatformMath::CeilLogTwo(PageSize);
@@ -877,6 +878,8 @@ public:
 		check(MAX_POOLED_ALLOCATION_SIZE - 1 == PoolTable[POOL_COUNT - 1].BlockSize);
 	}
 	
+	virtual void InitializeStatsMetadata() override;
+
 	virtual ~FMallocBinned()
 	{}
 
diff --git a/Engine/Source/Runtime/Core/Public/HAL/MemoryBase.h b/Engine/Source/Runtime/Core/Public/HAL/MemoryBase.h
--- a/Engine/Source/Runtime/Core/Public/HAL/MemoryBase.h
+++ b/Engine/Source/Runtime/Core/Public/HAL/MemoryBase.h
@@ -75,6 +75,12 @@ class CORE_API FMalloc  :
 	public FExec
 {
 public:
+	/**
+	 *	Initializes stats metadata. We need to do this as soon as possible, but cannot be done in the constructor
+	 *	due to the FName::StaticInit
+	 */
+	virtual void InitializeStatsMetadata();
+
 	/** 
 	 * QuantizeSize returns the actual size of allocation request likely to be returned
 	 * so for the template containers that use slack, they can more wisely pick
diff --git a/Engine/Source/Runtime/Core/Public/Stats/Stats.h b/Engine/Source/Runtime/Core/Public/Stats/Stats.h
--- a/Engine/Source/Runtime/Core/Public/Stats/Stats.h
+++ b/Engine/Source/Runtime/Core/Public/Stats/Stats.h
@@ -283,6 +283,7 @@
  */
 
 #define FORCEINLINE_STATS FORCEINLINE
+//#define FORCEINLINE_STATS FORCEINLINE_DEBUGGABLE
 #define checkStats(x)
 
 #if !defined(STATS)
diff --git a/Engine/Source/Runtime/Core/Public/Stats/Stats2.h b/Engine/Source/Runtime/Core/Public/Stats/Stats2.h
--- a/Engine/Source/Runtime/Core/Public/Stats/Stats2.h
+++ b/Engine/Source/Runtime/Core/Public/Stats/Stats2.h
@@ -2,8 +2,12 @@
 
 #pragma once
 
+#include "LockFreeList.h"
+#include "LockFreeFixedSizeAllocator.h"
 #include "ChunkedArray.h"
 
+class FThreadStats;
+
 /**
 * This is thread-private information about the thread idle stats, which we always collect, even in final builds
 */
@@ -174,10 +178,16 @@ struct EStatDataType
 	enum Type
 	{
 		Invalid,
+		/** Not defined. */
 		ST_None,
+		/** int64. */
 		ST_int64,
+		/** double. */
 		ST_double,
+		/** FName. */
 		ST_FName,
+		/** Memory pointer, stored as uint64. */
+		ST_Ptr,
 
 		Num,
 		Mask = 0x7,
@@ -194,16 +204,28 @@ struct EStatOperation
 	enum Type
 	{
 		Invalid,
+		/** Indicates metadata message. */
 		SetLongName,
+		/** Special message for advancing the stats frame from the game thread. */
 		AdvanceFrameEventGameThread,    
+		/** Special message for advancing the stats frame from the render thread. */
 		AdvanceFrameEventRenderThread,  
+		/** Indicates begin of the cycle scope. */
 		CycleScopeStart,
+		/** Indicates end of the cycle scope. */
 		CycleScopeEnd,
-		SpecialMessageMarker, // This is not a regular stat operation, but just a special message marker to determine that we encountered a special data in the stat file
+		/** This is not a regular stat operation, but just a special message marker to determine that we encountered a special data in the stat file. */
+		SpecialMessageMarker,
+		/** Set operation. */
 		Set, 
+		/** Clear operation. */
 		Clear,
+		/** Add operation. */
 		Add,
+		/** Subtract operation. */
 		Subtract,
+		/** This is a memory operation. @see EMemoryOperation. */
+		Memory,
 
 		// these are special ones for processed data
 		ChildrenStart,
@@ -259,6 +281,21 @@ struct EMemoryRegion
 	static_assert(FPlatformMemory::MCR_MAX < (1 << NumBits), "Need to expand memory region field.");
 };
 
+/** Memory operation for STAT_Memory_AllocPtr. */
+enum class EMemoryOperation
+{
+	/** Invalid. */
+	Invalid,	
+	/** Alloc. */
+	Alloc,
+	/** Free. */
+	Free,
+
+	Num,
+	Mask = 0x7,
+	/** Every allocation is aligned to 8 or 16 bytes, so we have 3/4 bits free to use. */
+	NumBits = 3,
+};
 
 /**
  * A few misc final bit packing computations
@@ -541,6 +578,8 @@ private:
 	double	Float;
 	/** For ST_int64 and IsCycle or IsMemory. */
 	int64	Cycles;
+	/** For ST_Ptr. */
+	uint64	Ptr;
 	/** ST_int64 and IsPackedCCAndDuration. */
 	int32	CCAndDuration[2];
 };
@@ -593,7 +632,7 @@ struct FStatMessage
 		NameAndInfo.SetField<EStatOperation>(EStatOperation::SetLongName);
 	}
 
-	FORCEINLINE_STATS FStatMessage(FStatNameAndInfo InStatName)
+	explicit FORCEINLINE_STATS FStatMessage(FStatNameAndInfo InStatName)
 		: NameAndInfo(InStatName)
 	{
 	}
@@ -622,7 +661,7 @@ struct FStatMessage
 	/**
 	* int64 operation
 	*/
-	explicit FORCEINLINE_STATS FStatMessage(FName InStatName, EStatOperation::Type InStatOperation, int64 Value, bool bIsCycle)
+	FORCEINLINE_STATS FStatMessage(FName InStatName, EStatOperation::Type InStatOperation, int64 Value, bool bIsCycle)
 		: NameAndInfo(InStatName, true)
 	{
 		NameAndInfo.SetField<EStatOperation>(InStatOperation);
@@ -634,7 +673,7 @@ struct FStatMessage
 	/**
 	* double operation
 	*/
-	explicit FORCEINLINE_STATS FStatMessage(FName InStatName, EStatOperation::Type InStatOperation, double Value, bool)
+	FORCEINLINE_STATS FStatMessage(FName InStatName, EStatOperation::Type InStatOperation, double Value, bool)
 		: NameAndInfo(InStatName, true)
 	{
 		NameAndInfo.SetField<EStatOperation>(InStatOperation);
@@ -646,7 +685,7 @@ struct FStatMessage
 	/**
 	* name operation
 	*/
-	explicit FORCEINLINE_STATS FStatMessage(FName InStatName, EStatOperation::Type InStatOperation, FName Value, bool)
+	FORCEINLINE_STATS FStatMessage(FName InStatName, EStatOperation::Type InStatOperation, FName Value, bool)
 		: NameAndInfo(InStatName, true)
 	{
 		NameAndInfo.SetField<EStatOperation>(InStatOperation);
@@ -655,6 +694,18 @@ struct FStatMessage
 		GetValue_FName() = Value;
 	}
 
+	/**
+	 * Ptr operation
+	 */
+	FORCEINLINE_STATS FStatMessage(FName InStatName, EStatOperation::Type InStatOperation, uint64 Value, bool)
+		: NameAndInfo(InStatName, true)
+	{
+		NameAndInfo.SetField<EStatOperation>(InStatOperation);
+		checkStats(NameAndInfo.GetField<EStatDataType>() == EStatDataType::ST_Ptr);
+		checkStats(NameAndInfo.GetFlag(EStatMetaFlags::IsCycle) == false);
+		GetValue_Ptr() = Value;
+	}
+
 	/**
 	* Clear any data type
 	*/
@@ -679,6 +730,18 @@ struct FStatMessage
 		return *(int64 const*)&StatData;
 	}
 
+	FORCEINLINE_STATS uint64& GetValue_Ptr()
+	{
+		static_assert(sizeof(uint64) <= DATA_SIZE && ALIGNOF(uint64) <= DATA_ALIGN, "Bad data for stat message.");
+		checkStats(NameAndInfo.GetField<EStatDataType>() == EStatDataType::ST_Ptr);
+		return *(uint64*)&StatData;
+	}
+	FORCEINLINE_STATS uint64 GetValue_Ptr() const
+	{
+		checkStats(NameAndInfo.GetField<EStatDataType>() == EStatDataType::ST_Ptr);
+		return *(uint64 const*)&StatData;
+	}
+
 	FORCEINLINE_STATS int64 GetValue_Duration() const
 	{
 		checkStats(NameAndInfo.GetFlag(EStatMetaFlags::IsCycle) && NameAndInfo.GetField<EStatDataType>() == EStatDataType::ST_int64);
@@ -1000,15 +1063,72 @@ struct FStatPacket
 	}
 };
 
+/** Helper struct used to monitor the scope of the message. */
+struct FStatMessageLock
+{
+	FStatMessageLock(int32& InMessageScope)
+		: MessageScope(InMessageScope)
+	{
+		MessageScope++;
+	}
+
+	~FStatMessageLock()
+	{
+		MessageScope--;
+	}
+
+protected:
+	int32& MessageScope;
+};
+
+/** Preallocates FThreadStats to avoid dynamic memory allocation. */
+struct CORE_API FThreadStatsPool
+{
+	enum
+	{
+		/**
+		 *	Number of elements preallocated in the pool.
+		 *  Should be enough to handle all threads.
+		 *  Increase if you encounter a crash.
+		 */
+		NUM_ELEMENTS_IN_POOL = 128,
+	};
+
+	/** Lock free pool of FThreadStats instances. */
+	TLockFreePointerList<FThreadStats> Pool;
+
+public:
+	/** Default constructor. */
+	FThreadStatsPool();
+
+	/** Singleton accessor. */
+	static FThreadStatsPool& Get()
+	{
+		static FThreadStatsPool Singleton;
+		return Singleton;
+	}
+
+	/** Gets an instance from the pool and call the default constructor on it. */
+	FThreadStats* GetFromPool();
+
+	/** Return an instance to the pool. */
+	void ReturnToPool(FThreadStats* Instance);
+};
+
 /**
 * This is thread-private information about the stats we are acquiring. Pointers to these objects are held in TLS.
 */
 class FThreadStats
 {
+	friend class FStatsMallocProfilerProxy;
+	friend class FStatsThreadState;
+	friend class FStatsThread;
+	friend struct FThreadStatsPool;
+
 	enum
 	{
 		PRESIZE_MAX_NUM_ENTRIES = 10,
-		PRESIZE_MAX_SIZE = EStatMessagesArrayConstants::MESSAGES_CHUNK_SIZE,
+		PRESIZE_MAX_SIZE = 256*1024,
 	};
 
 	/** Used to control when we are collecting stats. User of the stats system increment and decrement this counter as they need data. **/
@@ -1024,7 +1144,8 @@ class FThreadStats
 	/** TLS slot that holds a FThreadStats. **/
 	CORE_API static uint32 TlsSlot;
 
-	friend class FStatsThread;
+	/** The data we are eventually going to send to the stats thread. **/
+	FStatPacket Packet;
 
 	/** Current game frame for this thread stat. */
 	int32 CurrentGameFrame;
@@ -1035,14 +1156,22 @@ class FThreadStats
 	/** Tracks current stack depth for cycle counters. **/
 	uint32 bWaitForExplicitFlush;
 
+	/** 
+	 * Tracks the progress of adding a new memory stat message and prevents from using the memory profiler in the scope.
+	 * Mostly related to ignoring all memory allocations in Flush() and AddStatMessage().
+	 * Memory usage of the stats messages is handled by the STAT_StatMessagesMemory.
+	 */
+	int32 MemoryMessageScope;
+
 	/** Tracks current stack depth for cycle counters. **/
 	bool bSawExplicitFlush;
 
-	/** The data we are eventually going to send to the stats thread. **/
-	FStatPacket Packet;
-
-	/** Gathers information about the current thread and sets upt eh TLS value. **/
+	/** Gathers information about the current thread and sets up the TLS value. **/
 	CORE_API FThreadStats();
+
+	/** Constructor used for the pool. */
+	CORE_API FThreadStats(bool);
+
 	/** Copys a packet for sending. !!!CAUTION!!! does not copy packets.**/
 	CORE_API FThreadStats(FThreadStats const& Other);
 
@@ -1052,7 +1181,7 @@ class FThreadStats
 		FThreadStats* Stats = (FThreadStats*)FPlatformTLS::GetTlsValue(TlsSlot);
 		if (!Stats)
 		{
-			Stats = new FThreadStats();
+			Stats = FThreadStatsPool::Get().GetFromPool();
 		}
 		return Stats;
 	}
@@ -1063,6 +1192,12 @@ class FThreadStats
 	/** Send any outstanding packets to the stats thread **/
 	CORE_API void Flush(bool bHasBrokenCallstacks = false);
 
+	FORCEINLINE_STATS void AddStatMessage( const FStatMessage& StatMessage )
+	{
+		FStatMessageLock MessageLock(MemoryMessageScope);
+		Packet.StatMessages.AddElement(StatMessage);
+	}
+
 public:
 	/** This should be called when a thread exits, this deletes FThreadStats from the heap and TLS. **/
 	static void Shutdown()
@@ -1070,8 +1205,9 @@ public:
 		FThreadStats* Stats = (FThreadStats*)FPlatformTLS::GetTlsValue(TlsSlot);
 		if (Stats)
 		{
+			// Flush stats?
 			FPlatformTLS::SetTlsValue(TlsSlot, nullptr);
-			delete Stats;
+			FThreadStatsPool::Get().ReturnToPool(Stats);
 		}
 	}
 
@@ -1084,13 +1220,13 @@ public:
 		if (InStatOperation == EStatOperation::CycleScopeStart)
 		{
 			ThreadStats->ScopeCount++;
-			new (ThreadStats->Packet.StatMessages) FStatMessage(InStatName, InStatOperation);	
+			ThreadStats->AddStatMessage(FStatMessage(InStatName, InStatOperation));
 		}
 		else if (InStatOperation == EStatOperation::CycleScopeEnd)
 		{
 			if (ThreadStats->ScopeCount > ThreadStats->bWaitForExplicitFlush)
 			{
-				new (ThreadStats->Packet.StatMessages) FStatMessage(InStatName, InStatOperation);
+				ThreadStats->AddStatMessage(FStatMessage(InStatName, InStatOperation));
 				ThreadStats->ScopeCount--;
 				if (!ThreadStats->ScopeCount)
 				{
@@ -1108,7 +1244,7 @@ public:
 		if (!InStatName.IsNone() && WillEverCollectData())
 		{
 			FThreadStats* ThreadStats = GetThreadStats();
-			new (ThreadStats->Packet.StatMessages) FStatMessage(InStatName, InStatOperation, Value, bIsCycle);
+			ThreadStats->AddStatMessage(FStatMessage(InStatName, InStatOperation, Value, bIsCycle));
 			if(!ThreadStats->ScopeCount)
 			{
 				ThreadStats->Flush();
@@ -1116,6 +1252,13 @@ public:
 		}
 	}
 
+	/** Pseudo-Memory operation. */
+	template<typename TValue>
+	FORCEINLINE_STATS void AddMemoryMessage( FName InStatName, TValue Value )
+	{
+		AddStatMessage(FStatMessage(InStatName, EStatOperation::Memory, Value, false));
+	}
+
 	/** 
 	 * Used to force a flush at the next available opportunity. This is not useful for threads other than the main and render thread. 
 	 * if DiscardCallstack is true, we also dump call stacks, making the next available opportunity at the next stat or stat close.
@@ -1264,6 +1407,7 @@ public:
 };
 
 /** Manages startup messages, usually to update the metadata. */
+// @TODO yrx 2014-11-18 Eventually add startup messages for memory profiling?
 class CORE_API FStartupMessages
 {
 public:
@@ -1349,7 +1493,7 @@ protected:
 template<class TStatData, bool TCompiledIn>
 struct FThreadSafeStaticStatInner : public FThreadSafeStaticStatBase
 {
-	FORCEINLINE TStatId GetStatId() const
+	FORCEINLINE_STATS TStatId GetStatId() const
 	{
 		static_assert(sizeof(HighPerformanceEnable) == sizeof(TStatId), "Unsafe cast requires these to be the same thing.");
 		if (!HighPerformanceEnable)
@@ -1448,6 +1592,10 @@ struct FStat_##StatName\
 
 #define STAT_GROUP_TO_FStatGroup(Group) FStatGroup_##Group
 
+/*-----------------------------------------------------------------------------
+	Local
+-----------------------------------------------------------------------------*/
+
 #define DEFINE_STAT(Stat) \
 	struct FThreadSafeStaticStat<FStat_##Stat> StatPtr_##Stat;
 
@@ -1471,13 +1619,19 @@ struct FStat_##StatName\
 #define DECLARE_DWORD_ACCUMULATOR_STAT(CounterName,StatId,GroupId) \
 	DECLARE_STAT(CounterName,StatId,GroupId,EStatDataType::ST_int64, false, false, FPlatformMemory::MCR_Invalid); \
 	static DEFINE_STAT(StatId)
+
 #define DECLARE_MEMORY_STAT(CounterName,StatId,GroupId) \
 	DECLARE_STAT(CounterName,StatId,GroupId,EStatDataType::ST_int64, false, false, FPlatformMemory::MCR_Physical); \
 	static DEFINE_STAT(StatId)
+
 #define DECLARE_MEMORY_STAT_POOL(CounterName,StatId,GroupId,Pool) \
 	DECLARE_STAT(CounterName,StatId,GroupId,EStatDataType::ST_int64, false, false, Pool); \
 	static DEFINE_STAT(StatId)
 
+/*-----------------------------------------------------------------------------
+	Extern
+-----------------------------------------------------------------------------*/
+
 #define DECLARE_CYCLE_STAT_EXTERN(CounterName,StatId,GroupId, APIX) \
 	DECLARE_STAT(CounterName,StatId,GroupId,EStatDataType::ST_int64, true, true, FPlatformMemory::MCR_Invalid); \
 	extern APIX DEFINE_STAT(StatId);
@@ -1721,6 +1875,6 @@ DECLARE_FLOAT_COUNTER_STAT_EXTERN(TEXT("StatUnit FPS"), STAT_FPS, STATGROUP_Engi
 /** Stats for the stat system */
 
 DECLARE_CYCLE_STAT_EXTERN(TEXT("DrawStats"),STAT_DrawStats,STATGROUP_StatSystem, CORE_API);
-DECLARE_MEMORY_STAT_EXTERN(TEXT("Stat messages memory"),STAT_StatMessagesMemory,STATGROUP_StatSystem, CORE_API);
+DECLARE_MEMORY_STAT_EXTERN(TEXT("Stat Messages"),STAT_StatMessagesMemory,STATGROUP_StatSystem, CORE_API);
 
 #endif
diff --git a/Engine/Source/Runtime/Core/Public/Stats/StatsData.h b/Engine/Source/Runtime/Core/Public/Stats/StatsData.h
--- a/Engine/Source/Runtime/Core/Public/Stats/StatsData.h
+++ b/Engine/Source/Runtime/Core/Public/Stats/StatsData.h
@@ -192,6 +192,7 @@ struct CORE_API FStatPacketArray
 	void Empty();
 };
 
+// @TODO yrx 2014-11-21 basic memory profiling?
 /**
 * A call stack of stat messages. Normally we don't store or transmit these; they are used for stats processing and visualization.
 */
@@ -399,6 +400,37 @@ struct IItemFiler
 	virtual bool Keep(FStatMessage const& Item) = 0;
 };
 
+/** Simple allocation info, temporary. */
+struct FAllocationInfo
+{
+	int64 Size;
+	FName Scope;
+	FAllocationInfo(int64 InSize, FName InScope)
+		: Size(InSize)
+		, Scope(InScope)
+	{}
+	FAllocationInfo()
+		: Size(0)
+	{}
+};
+
+/** Simple allocation info, temporary. */
+struct FAllocationInfoEx
+{
+	int64 Size;
+	FName Scope;
+	int64 Cycles;
+	FAllocationInfoEx(int64 InSize, FName InScope, int64 InCycles)
+		: Size(InSize)
+		, Scope(InScope)
+		, Cycles(InCycles)
+	{}
+	FAllocationInfoEx()
+		: Size(0)
+		, Cycles(0)
+	{}
+};
+
 // @TODO yrx 2014-03-21 Move metadata functionality into a separate class
 /**
 * Tracks stat state and history
@@ -477,6 +509,12 @@ private:
 	/** Cached condensed frames. This saves a lot of time since multiple listeners can use the same condensed data **/
 	mutable TMap<int64, TArray<FStatMessage>* > CondensedStackHistory;
 
+	/** Map from allocation pointer to allocation info, currently active allocations */
+	mutable TMap< uint64, FAllocationInfo > AllocationMap;
+
+	/** Map from allocation pointer to allocation info. */
+	mutable TMap< uint64, FAllocationInfoEx > NonSequentialAllocationMap;
+
 	/** After we add to the history, these frames are known good. **/
 	TSet<int64>	GoodFrames;
 
@@ -531,6 +569,17 @@ public:
 		return GoodFrames.Contains(Frame);
 	}
 
+	/** Gets allocations map. */
+	const TMap< uint64, FAllocationInfo >& GetAllocations() const 
+	{ 
+		return AllocationMap; 
+	}
+
+	const TMap< uint64, FAllocationInfoEx >& GetNonSequentialAllocations() const 
+	{ 
+		return NonSequentialAllocationMap; 
+	}
+
 	/** Used by the hitch detector. Does a very fast look a thread to decide if this is a hitch frame. **/
 	int64 GetFastThreadFrameTime(int64 TargetFrame, EThreadType::Type Thread) const;
 	int64 GetFastThreadFrameTime(int64 TargetFrame, uint32 ThreadID) const;
@@ -778,7 +827,7 @@ struct FGameThreadHudData
 /**
 * Singleton that holds the last data sent from the stats thread to the game thread for HUD stats
 */
-struct CORE_API FHUDGroupGameThreadRenderer 
+struct FHUDGroupGameThreadRenderer 
 {
 	FGameThreadHudData* Latest;
 
@@ -795,16 +844,16 @@ struct CORE_API FHUDGroupGameThreadRenderer
 
 	void NewData(FGameThreadHudData* Data);
 
-	static FHUDGroupGameThreadRenderer& Get();
+	CORE_API static FHUDGroupGameThreadRenderer& Get();
 };
 
 /**
  * Singleton that holds a list of newly registered group stats to inform the game thread of
  */
-struct CORE_API FStatGroupGameThreadNotifier
+struct FStatGroupGameThreadNotifier
 {
 public:
-	static FStatGroupGameThreadNotifier& Get();
+	CORE_API static FStatGroupGameThreadNotifier& Get();
 
 	void NewData(FStatNameAndInfo NameAndInfo)
 	{
diff --git a/Engine/Source/Runtime/Core/Public/Stats/StatsMallocProfilerProxy.h b/Engine/Source/Runtime/Core/Public/Stats/StatsMallocProfilerProxy.h
new file mode 100644
--- /dev/null
+++ b/Engine/Source/Runtime/Core/Public/Stats/StatsMallocProfilerProxy.h
@@ -0,0 +1,104 @@
+// Copyright 1998-2014 Epic Games, Inc. All Rights Reserved.
+
+#pragma once
+
+#if STATS
+
+/**
+ * Malloc proxy for gathering memory messages.
+ */
+class FStatsMallocProfilerProxy : public FMalloc
+{
+private:
+	/** Malloc we're based on. */
+	FMalloc* UsedMalloc;
+	
+	/** Whether the stats malloc profiler is enabled, disabled by default */
+	bool bEnabled;
+
+	FThreadSafeCounter AllocPtrCalls;
+	FThreadSafeCounter FreePtrCalls;
+
+public:
+	/**
+	 * Default constructor
+	 * 
+	 * @param	InMalloc - FMalloc that is going to be used for actual allocations
+	 */
+	FStatsMallocProfilerProxy( FMalloc* InMalloc);
+
+	static FStatsMallocProfilerProxy* Get();
+
+	void SetState( bool bEnable );
+
+	bool GetState() const
+	{
+		return bEnabled;
+	}
+
+	virtual void InitializeStatsMetadata() override;
+
+	virtual SIZE_T QuantizeSize( SIZE_T Size, uint32 Alignment ) override
+	{
+		return UsedMalloc->QuantizeSize(Size,Alignment); 
+	}
+
+	/**
+	 * Tracks malloc operation.
+	 *
+	 * @param	Ptr	- Allocated pointer 
+	 * @param	Size- Size of allocated pointer
+	 */
+	void TrackAlloc( void* Ptr, int64 Size );
+
+	/**
+	 * Tracks free operation
+	 *
+	 * @param	Ptr	- Freed pointer
+	 */
+	void TrackFree( void* Ptr );
+
+	virtual void* Malloc( SIZE_T Size, uint32 Alignment ) override;
+	virtual void* Realloc( void* OldPtr, SIZE_T NewSize, uint32 Alignment ) override;
+	virtual void Free( void* Ptr ) override;
+
+	virtual bool IsInternallyThreadSafe() const override
+	{ 
+		return UsedMalloc->IsInternallyThreadSafe(); 
+	}
+
+	virtual void UpdateStats() override;
+
+	virtual void GetAllocatorStats( FGenericMemoryStats& out_Stats ) override
+	{
+		UsedMalloc->GetAllocatorStats( out_Stats );
+	}
+
+	virtual void DumpAllocatorStats( class FOutputDevice& Ar ) override
+	{
+		UsedMalloc->DumpAllocatorStats( Ar );
+	}
+
+	virtual bool ValidateHeap() override
+	{
+		return UsedMalloc->ValidateHeap();
+	}
+
+	virtual bool Exec( UWorld* InWorld, const TCHAR* Cmd, FOutputDevice& Ar ) override
+	{
+		return UsedMalloc->Exec( InWorld, Cmd, Ar);
+	}
+
+	virtual bool GetAllocationSize(void *Original, SIZE_T &SizeOut) override
+	{
+		return UsedMalloc->GetAllocationSize(Original,SizeOut);
+	}
+
+	virtual const TCHAR* GetDescriptiveName() override
+	{ 
+		return UsedMalloc->GetDescriptiveName(); 
+	}
+};
+
+
+#endif //STATS
\ No newline at end of file
