Core - Misc code clean-up

diff --git a/Engine/Source/Runtime/Core/Public/HAL/MallocThreadSafeProxy.h b/Engine/Source/Runtime/Core/Public/HAL/MallocThreadSafeProxy.h
--- a/Engine/Source/Runtime/Core/Public/HAL/MallocThreadSafeProxy.h
+++ b/Engine/Source/Runtime/Core/Public/HAL/MallocThreadSafeProxy.h
@@ -25,6 +25,11 @@ public:
 	:	UsedMalloc( InMalloc )
 	{}
 
+	virtual void InitializeStatsMetadata() override
+	{
+		UsedMalloc->InitializeStatsMetadata();
+	}
+
 	/** 
 	 * QuantizeSize returns the actual size of allocation request likely to be returned
 	 * so for the template containers that use slack, they can more wisely pick
@@ -43,7 +48,7 @@ public:
 	/** 
 	 * Malloc
 	 */
-	void* Malloc( SIZE_T Size, uint32 Alignment ) override
+	virtual void* Malloc( SIZE_T Size, uint32 Alignment ) override
 	{
 		FScopeLock ScopeLock( &SynchronizationObject );
 		STAT(TotalMallocCalls++);
@@ -53,7 +58,7 @@ public:
 	/** 
 	 * Realloc
 	 */
-	void* Realloc( void* Ptr, SIZE_T NewSize, uint32 Alignment ) override
+	virtual void* Realloc( void* Ptr, SIZE_T NewSize, uint32 Alignment ) override
 	{
 		FScopeLock ScopeLock( &SynchronizationObject );
 		STAT(TotalReallocCalls++);
@@ -63,7 +68,7 @@ public:
 	/** 
 	 * Free
 	 */
-	void Free( void* Ptr ) override
+	virtual void Free( void* Ptr ) override
 	{
 		if( Ptr )
 		{
@@ -74,21 +79,21 @@ public:
 	}
 
 	/** Called once per frame, gathers and sets all memory allocator statistics into the corresponding stats. */
-	virtual void UpdateStats()
+	virtual void UpdateStats() override
 	{
 		FScopeLock ScopeLock( &SynchronizationObject );
 		UsedMalloc->UpdateStats();
 	}
 
 	/** Writes allocator stats from the last update into the specified destination. */
-	virtual void GetAllocatorStats( FGenericMemoryStats& out_Stats )
+	virtual void GetAllocatorStats( FGenericMemoryStats& out_Stats ) override
 	{
 		FScopeLock ScopeLock( &SynchronizationObject );
 		UsedMalloc->GetAllocatorStats( out_Stats );
 	}
 
 	/** Dumps allocator stats to an output device. */
-	virtual void DumpAllocatorStats( class FOutputDevice& Ar )
+	virtual void DumpAllocatorStats( class FOutputDevice& Ar ) override
 	{
 		FScopeLock ScopeLock( &SynchronizationObject );
 		UsedMalloc->DumpAllocatorStats( Ar );
@@ -103,7 +108,7 @@ public:
 		return( UsedMalloc->ValidateHeap() );
 	}
 
-	bool Exec( UWorld* InWorld, const TCHAR* Cmd, FOutputDevice& Ar ) override
+	virtual bool Exec( UWorld* InWorld, const TCHAR* Cmd, FOutputDevice& Ar ) override
 	{
 		FScopeLock ScopeLock( &SynchronizationObject );
 		return UsedMalloc->Exec( InWorld, Cmd, Ar);
@@ -116,7 +121,7 @@ public:
 	* @param SizeOut - If possible, this value is set to the size of the passed in pointer
 	* @return true if succeeded
 	*/
-	bool GetAllocationSize(void *Original, SIZE_T &SizeOut) override
+	virtual bool GetAllocationSize(void *Original, SIZE_T &SizeOut) override
 	{
 		FScopeLock ScopeLock( &SynchronizationObject );
 		return UsedMalloc->GetAllocationSize(Original,SizeOut);
