Stats - DumpSum will dump the results after 240 frames

diff --git a/Engine/Source/Runtime/Core/Private/Stats/Stats2.cpp b/Engine/Source/Runtime/Core/Private/Stats/Stats2.cpp
--- a/Engine/Source/Runtime/Core/Private/Stats/Stats2.cpp
+++ b/Engine/Source/Runtime/Core/Private/Stats/Stats2.cpp
@@ -980,7 +980,7 @@ void FThreadStats::StartThread()
 	{
 		DirectStatsCommand(TEXT("stat group enable LinkerLoad"));
 		DirectStatsCommand(TEXT("stat group enable AsyncLoad"));
-		DirectStatsCommand(TEXT("stat dumpsum -start -ms=250"));
+		DirectStatsCommand(TEXT("stat dumpsum -start -num=240"));
 	}
 	if (FParse::Param(FCommandLine::Get(), TEXT("LoadTimeFile")))
 	{
