Core - Removed obsolete FMalloc.QuantizeSize

diff --git a/Engine/Source/Runtime/Core/Private/HAL/UnrealMemory.cpp b/Engine/Source/Runtime/Core/Private/HAL/UnrealMemory.cpp
--- a/Engine/Source/Runtime/Core/Private/HAL/UnrealMemory.cpp
+++ b/Engine/Source/Runtime/Core/Private/HAL/UnrealMemory.cpp
@@ -41,33 +41,14 @@ void GCreateMalloc()
 		GMalloc = new FMallocThreadSafeProxy( GMalloc );
 	}
 
+
+
 #if	STATS
 	// Create the stats malloc profiler proxy, needs to be enabled separately.
 	//GMalloc = FStatsMallocProfilerProxy::Get();
 #endif // STATS
 }
 
-
-/** 
- * FMemory::MallocQuantizeSize returns the actual size of allocation request likely to be returned
- * so for the template containers that use slack, they can more wisely pick
- * appropriate sizes to grow and shrink to.
- *
- * @param Size			The size of a hypothetical allocation request
- * @param Alignment		The alignment of a hypothetical allocation request
- * @return				Returns the usable size that the allocation request would return. In other words you can ask for this greater amount without using any more actual memory.
- */
-SIZE_T FMemory::MallocQuantizeSize( SIZE_T Size, uint32 Alignment )
-{
-	if( !GMalloc )
-	{
-		GCreateMalloc();
-		CA_ASSUME( GMalloc != NULL );	// Don't want to assert, but suppress static analysis warnings about potentially NULL GMalloc
-	}
-	return GMalloc->QuantizeSize( Size, Alignment );
-
-}
-
 void* FMemory::Malloc( SIZE_T Count, uint32 Alignment ) 
 { 
 	if( !GMalloc )
diff --git a/Engine/Source/Runtime/Core/Public/HAL/MallocThreadSafeProxy.h b/Engine/Source/Runtime/Core/Public/HAL/MallocThreadSafeProxy.h
--- a/Engine/Source/Runtime/Core/Public/HAL/MallocThreadSafeProxy.h
+++ b/Engine/Source/Runtime/Core/Public/HAL/MallocThreadSafeProxy.h
@@ -30,21 +30,6 @@ public:
 		UsedMalloc->InitializeStatsMetadata();
 	}
 
-	/** 
-	 * QuantizeSize returns the actual size of allocation request likely to be returned
-	 * so for the template containers that use slack, they can more wisely pick
-	 * appropriate sizes to grow and shrink to.
-	 *
-	 * CAUTION: QuantizeSize is a special case and is NOT guarded by a thread lock, so must be intrinsically thread safe!
-	 *
-	 * @param Size			The size of a hypothetical allocation request
-	 * @param Alignment		The alignment of a hypothetical allocation request
-	 * @return				Returns the usable size that the allocation request would return. In other words you can ask for this greater amount without using any more actual memory.
-	 */
-	virtual SIZE_T QuantizeSize( SIZE_T Size, uint32 Alignment ) override
-	{
-		return UsedMalloc->QuantizeSize(Size,Alignment); 
-	}
 	/** 
 	 * Malloc
 	 */
diff --git a/Engine/Source/Runtime/Core/Public/HAL/MemoryBase.h b/Engine/Source/Runtime/Core/Public/HAL/MemoryBase.h
--- a/Engine/Source/Runtime/Core/Public/HAL/MemoryBase.h
+++ b/Engine/Source/Runtime/Core/Public/HAL/MemoryBase.h
@@ -81,22 +81,6 @@ public:
 	 */
 	virtual void InitializeStatsMetadata();
 
-	/** 
-	 * QuantizeSize returns the actual size of allocation request likely to be returned
-	 * so for the template containers that use slack, they can more wisely pick
-	 * appropriate sizes to grow and shrink to.
-	 *
-	 * CAUTION: QuantizeSize is a special case and is NOT guarded by a thread lock, so must be intrinsically thread safe!
-	 *
-	 * @param Size			The size of a hypothetical allocation request
-	 * @param Alignment		The alignment of a hypothetical allocation request
-	 * @return				Returns the usable size that the allocation request would return. In other words you can ask for this greater amount without using any more actual memory.
-	 */
-	virtual SIZE_T QuantizeSize( SIZE_T Size, uint32 Alignment )
-	{
-		return Size; // implementations where this is not possible need not implement it.
-	}
-
 	/** 
 	 * Malloc
 	 */
diff --git a/Engine/Source/Runtime/Core/Public/HAL/UnrealMemory.h b/Engine/Source/Runtime/Core/Public/HAL/UnrealMemory.h
--- a/Engine/Source/Runtime/Core/Public/HAL/UnrealMemory.h
+++ b/Engine/Source/Runtime/Core/Public/HAL/UnrealMemory.h
@@ -100,17 +100,6 @@ struct CORE_API FMemory
 	// C style memory allocation stubs.
 	//
 
-	/** 
-	 * FMemory::MallocQuantizeSize returns the actual size of allocation request likely to be returned
-	 * so for the template containers that use slack, they can more wisely pick
-	 * appropriate sizes to grow and shrink to.
-	 *
-	 * @param Size			The size of a hypothetical allocation request
-	 * @param Alignment		The alignment of a hypothetical allocation request
-	 * @return				Returns the usable size that the allocation request would return. In other words you can ask for this greater amount without using any more actual memory.
-	 */
-	static SIZE_T MallocQuantizeSize( SIZE_T Size, uint32 Alignment=DEFAULT_ALIGNMENT );
-
 	static void* Malloc( SIZE_T Count, uint32 Alignment=DEFAULT_ALIGNMENT );
 	static void* Realloc( void* Original, SIZE_T Count, uint32 Alignment=DEFAULT_ALIGNMENT );
 	static void Free( void* Original );
diff --git a/Engine/Source/Runtime/Core/Public/ProfilingDebugging/MallocProfiler.h b/Engine/Source/Runtime/Core/Public/ProfilingDebugging/MallocProfiler.h
--- a/Engine/Source/Runtime/Core/Public/ProfilingDebugging/MallocProfiler.h
+++ b/Engine/Source/Runtime/Core/Public/ProfilingDebugging/MallocProfiler.h
@@ -407,22 +407,6 @@ public:
 	 */
 	void BeginProfiling();
 
-	/** 
-	 * QuantizeSize returns the actual size of allocation request likely to be returned
-	 * so for the template containers that use slack, they can more wisely pick
-	 * appropriate sizes to grow and shrink to.
-	 *
-	 * CAUTION: QuantizeSize is a special case and is NOT guarded by a thread lock, so must be intrinsically thread safe!
-	 *
-	 * @param Size			The size of a hypothetical allocation request
-	 * @param Alignment		The alignment of a hypothetical allocation request
-	 * @return				Returns the usable size that the allocation request would return. In other words you can ask for this greater amount without using any more actual memory.
-	 */
-	virtual SIZE_T QuantizeSize( SIZE_T Size, uint32 Alignment ) override
-	{
-		return UsedMalloc->QuantizeSize(Size,Alignment); 
-	}
-
 	/** 
 	 * Malloc
 	 */
diff --git a/Engine/Source/Runtime/Core/Public/Stats/StatsMallocProfilerProxy.h b/Engine/Source/Runtime/Core/Public/Stats/StatsMallocProfilerProxy.h
--- a/Engine/Source/Runtime/Core/Public/Stats/StatsMallocProfilerProxy.h
+++ b/Engine/Source/Runtime/Core/Public/Stats/StatsMallocProfilerProxy.h
@@ -27,7 +27,7 @@ public:
 	 */
 	FStatsMallocProfilerProxy( FMalloc* InMalloc);
 
-	static FStatsMallocProfilerProxy* Get();
+	static CORE_API FStatsMallocProfilerProxy* Get();
 
 	void SetState( bool bEnable );
 
@@ -38,11 +38,6 @@ public:
 
 	virtual void InitializeStatsMetadata() override;
 
-	virtual SIZE_T QuantizeSize( SIZE_T Size, uint32 Alignment ) override
-	{
-		return UsedMalloc->QuantizeSize(Size,Alignment); 
-	}
-
 	/**
 	 * Tracks malloc operation.
 	 *
