Core - Fixed a stack overflow issue

diff --git a/Engine/Source/Runtime/Core/Public/HAL/MallocBinned.h b/Engine/Source/Runtime/Core/Public/HAL/MallocBinned.h
--- a/Engine/Source/Runtime/Core/Public/HAL/MallocBinned.h
+++ b/Engine/Source/Runtime/Core/Public/HAL/MallocBinned.h
@@ -802,7 +802,6 @@ public:
 		check(!(AddressLimit & (AddressLimit - 1)));
 		check(PageSize <= 65536); // There is internal limit on page size of 64k
 		check(AddressLimit > PageSize); // Check to catch 32 bit overflow in AddressLimit
-		InitializeStatsMetadata();
 
 		/** Shift to get the reference from the indirect tables */
 		PoolBitShift = FPlatformMath::CeilLogTwo(PageSize);
