Make the check for inclusion of delegate implementations from Delegate.h more explicit, rather than relying on a header guard.

diff --git a/Engine/Source/Runtime/Core/Public/Delegates/Delegate.h b/Engine/Source/Runtime/Core/Public/Delegates/Delegate.h
--- a/Engine/Source/Runtime/Core/Public/Delegates/Delegate.h
+++ b/Engine/Source/Runtime/Core/Public/Delegates/Delegate.h
@@ -1,9 +1,5 @@
 // Copyright 1998-2014 Epic Games, Inc. All Rights Reserved.
 
-// Do not remove this header guard! It ensures that DelegateCombinations.h can only be included from here.
-#ifndef __Delegate_h__
-#define __Delegate_h__
-
 #pragma once
 
 #include "SharedPointer.h"
@@ -302,6 +298,7 @@ class DynamicMulticastDelegateName : public FUNC_DECLARE_DYNAMIC_DELEGATE_BASE(T
 /*********************************************************************************************************************/
 
 // We define this as a guard to prevent DelegateSignatureImpl.inl being included outside of this file
+#define __Delegate_h__
 #define FUNC_INCLUDING_INLINE_IMPL
 
 #ifndef UE_BUILD_DOCS
@@ -315,6 +312,9 @@ class DynamicMulticastDelegateName : public FUNC_DECLARE_DYNAMIC_DELEGATE_BASE(T
 	#endif
 #endif
 
+// No longer allowed to include DelegateSignatureImpl.inl
+#undef FUNC_INCLUDING_INLINE_IMPL
+#undef __Delegate_h__
 
 /*********************************************************************************************************************/
 
@@ -322,9 +322,6 @@ class DynamicMulticastDelegateName : public FUNC_DECLARE_DYNAMIC_DELEGATE_BASE(T
 	class EventName : public BaseTypeEvent { friend class OwningType; };
 
 
-// No longer allowed to include DelegateSignatureImpl.inl
-#undef FUNC_INCLUDING_INLINE_IMPL
-
 // Undefine temporary macros
 #undef FUNC_COMBINE_ACTUAL
 #undef FUNC_COMBINE
@@ -332,6 +329,3 @@ class DynamicMulticastDelegateName : public FUNC_DECLARE_DYNAMIC_DELEGATE_BASE(T
 // Simple delegate used by various utilities such as timers
 DECLARE_DELEGATE( FSimpleDelegate );
 DECLARE_MULTICAST_DELEGATE( FSimpleMulticastDelegate );
-
-
-#endif		// __Delegate_h__
