UE-3548 - Moved the coalesced source config data to the intermediate folder.

diff --git a/Engine/Source/Runtime/Core/Private/Misc/ConfigCacheIni.cpp b/Engine/Source/Runtime/Core/Private/Misc/ConfigCacheIni.cpp
--- a/Engine/Source/Runtime/Core/Private/Misc/ConfigCacheIni.cpp
+++ b/Engine/Source/Runtime/Core/Private/Misc/ConfigCacheIni.cpp
@@ -1096,7 +1096,7 @@ void FConfigFile::SaveSourceToBackupFile()
 {
 	FString Text;
 
-	FString BetweenRunsDir = FPaths::Combine(*FPaths::GeneratedConfigDir(), TEXT("CleanSourceConfigs/"));
+	FString BetweenRunsDir = (FPaths::GameIntermediateDir() / TEXT("Config/CoalescedSourceConfigs/"));
 	FString Filename = FString::Printf( TEXT( "%s%s.ini" ), *BetweenRunsDir, *Name.ToString() );
 
 	for( TMap<FString,FConfigSection>::TIterator SectionIterator(*SourceConfigFile); SectionIterator; ++SectionIterator )
@@ -1133,7 +1133,7 @@ void FConfigFile::SaveSourceToBackupFile()
 
 void FConfigFile::ProcessSourceAndCheckAgainstBackup()
 {
-	FString BetweenRunsDir = FPaths::Combine(*FPaths::GeneratedConfigDir(), TEXT("CleanSourceConfigs/"));
+	FString BetweenRunsDir = (FPaths::GameIntermediateDir() / TEXT("Config/CoalescedSourceConfigs/"));
 	FString BackupFilename = FString::Printf( TEXT( "%s%s.ini" ), *BetweenRunsDir, *Name.ToString() );
 
 	FConfigFile BackupFile;
