Pragma enable/disable deprecation macros moved into Platform.
Deprecation unified across all platforms in *CompilerPreSetup.h headers.

diff --git a/Engine/Source/Runtime/Core/Public/Android/AndroidPlatformCompilerPreSetup.h b/Engine/Source/Runtime/Core/Public/Android/AndroidPlatformCompilerPreSetup.h
new file mode 100644
--- /dev/null
+++ b/Engine/Source/Runtime/Core/Public/Android/AndroidPlatformCompilerPreSetup.h
@@ -0,0 +1,37 @@
+// Copyright 1998-2014 Epic Games, Inc. All Rights Reserved.
+
+#pragma once
+
+
+#ifndef DISABLE_DEPRECATION
+	#pragma clang diagnostic warning "-Wdeprecated-declarations"
+
+	/**
+	 * Macro for marking up deprecated code, functions and types.
+	 *
+	 * Features that are marked as deprecated are scheduled to be removed from the code base
+	 * in a future release. If you are using a deprecated feature in your code, you should
+	 * replace it before upgrading to the next release. See the Upgrade Notes in the release
+	 * notes for the release in which the feature was marked deprecated.
+	 *
+	 * Sample usage (note the slightly different syntax for classes and structures):
+	 *
+	 *		DEPRECATED(4.6, "Message")
+	 *		void Function();
+	 *
+	 *		struct DEPRECATED(4.6, "Message") MyStruct;
+	 *		class DEPRECATED(4.6, "Message") MyClass;
+	 *
+	 * @param VERSION The release number in which the feature was marked deprecated.
+	 * @param MESSAGE A message containing upgrade notes.
+	 */
+	#define DEPRECATED(VERSION, MESSAGE) __attribute__((deprecated(MESSAGE " Please update your code to the new API before upgrading to the next release, otherwise your project will no longer compile.")))
+
+	#define PRAGMA_DISABLE_DEPRECATION_WARNINGS \
+		_Pragma ("clang diagnostic push") \
+		_Pragma ("clang diagnostic ignored \"-Wdeprecated-declarations\"")
+
+	#define PRAGMA_ENABLE_DEPRECATION_WARNINGS \
+		_Pragma("clang diagnostic pop")
+
+#endif // DISABLE_DEPRECATION
diff --git a/Engine/Source/Runtime/Core/Public/GenericPlatform/GenericPlatformCompilerPreSetup.h b/Engine/Source/Runtime/Core/Public/GenericPlatform/GenericPlatformCompilerPreSetup.h
--- a/Engine/Source/Runtime/Core/Public/GenericPlatform/GenericPlatformCompilerPreSetup.h
+++ b/Engine/Source/Runtime/Core/Public/GenericPlatform/GenericPlatformCompilerPreSetup.h
@@ -26,6 +26,14 @@
 	#define DEPRECATED(VERSION, MESSAGE)
 #endif // DEPRECATED
 
+#ifndef PRAGMA_DISABLE_DEPRECATION_WARNINGS
+	#define PRAGMA_DISABLE_DEPRECATION_WARNINGS
+#endif
+
+#ifndef PRAGMA_ENABLE_DEPRECATION_WARNINGS
+	#define PRAGMA_ENABLE_DEPRECATION_WARNINGS
+#endif
+
 #ifndef EMIT_DEPRECATED_WARNING_MESSAGE
 	#define EMIT_DEPRECATED_WARNING_MESSAGE(Msg)
 #endif // EMIT_DEPRECATED_WARNING_MESSAGE
diff --git a/Engine/Source/Runtime/Core/Public/HAL/Platform.h b/Engine/Source/Runtime/Core/Public/HAL/Platform.h
--- a/Engine/Source/Runtime/Core/Public/HAL/Platform.h
+++ b/Engine/Source/Runtime/Core/Public/HAL/Platform.h
@@ -48,8 +48,24 @@
 // Platform specific compiler pre-setup.
 #if PLATFORM_WINDOWS
 	#include "Windows/WindowsPlatformCompilerPreSetup.h"
+#elif PLATFORM_PS4
+	#include "PS4/PS4PlatformCompilerPreSetup.h"
+#elif PLATFORM_XBOXONE
+	#include "XboxOne/XboxOnePlatformCompilerPreSetup.h"
 #elif PLATFORM_MAC
 	#include "Mac/MacPlatformCompilerPreSetup.h"
+#elif PLATFORM_IOS
+	#include "IOS/IOSPlatformCompilerPreSetup.h"
+#elif PLATFORM_ANDROID
+	#include "Android/AndroidPlatformCompilerPreSetup.h"
+#elif PLATFORM_WINRT_ARM || PLATFORM_WINRT
+	#include "WinRT/WinRTPlatformCompilerPreSetup.h"
+#elif PLATFORM_HTML5
+	#include "HTML5/HTML5PlatformCompilerPreSetup.h"
+#elif PLATFORM_LINUX
+	#include "Linux/LinuxPlatformCompilerPreSetup.h"
+#else
+	#error Unknown Compiler
 #endif
 
 // Generic compiler pre-setup.
diff --git a/Engine/Source/Runtime/Core/Public/HTML5/HTML5PlatformCompilerPreSetup.h b/Engine/Source/Runtime/Core/Public/HTML5/HTML5PlatformCompilerPreSetup.h
new file mode 100644
--- /dev/null
+++ b/Engine/Source/Runtime/Core/Public/HTML5/HTML5PlatformCompilerPreSetup.h
@@ -0,0 +1,37 @@
+// Copyright 1998-2014 Epic Games, Inc. All Rights Reserved.
+
+#pragma once
+
+
+#ifndef DISABLE_DEPRECATION
+	#pragma clang diagnostic warning "-Wdeprecated-declarations"
+
+	/**
+	 * Macro for marking up deprecated code, functions and types.
+	 *
+	 * Features that are marked as deprecated are scheduled to be removed from the code base
+	 * in a future release. If you are using a deprecated feature in your code, you should
+	 * replace it before upgrading to the next release. See the Upgrade Notes in the release
+	 * notes for the release in which the feature was marked deprecated.
+	 *
+	 * Sample usage (note the slightly different syntax for classes and structures):
+	 *
+	 *		DEPRECATED(4.6, "Message")
+	 *		void Function();
+	 *
+	 *		struct DEPRECATED(4.6, "Message") MyStruct;
+	 *		class DEPRECATED(4.6, "Message") MyClass;
+	 *
+	 * @param VERSION The release number in which the feature was marked deprecated.
+	 * @param MESSAGE A message containing upgrade notes.
+	 */
+	#define DEPRECATED(VERSION, MESSAGE) __attribute__((deprecated(MESSAGE " Please update your code to the new API before upgrading to the next release, otherwise your project will no longer compile.")))
+
+	#define PRAGMA_DISABLE_DEPRECATION_WARNINGS \
+		_Pragma ("clang diagnostic push") \
+		_Pragma ("clang diagnostic ignored \"-Wdeprecated-declarations\"")
+
+	#define PRAGMA_ENABLE_DEPRECATION_WARNINGS \
+		_Pragma("clang diagnostic pop")
+
+#endif // DISABLE_DEPRECATION
diff --git a/Engine/Source/Runtime/Core/Public/IOS/IOSPlatformCompilerPreSetup.h b/Engine/Source/Runtime/Core/Public/IOS/IOSPlatformCompilerPreSetup.h
new file mode 100644
--- /dev/null
+++ b/Engine/Source/Runtime/Core/Public/IOS/IOSPlatformCompilerPreSetup.h
@@ -0,0 +1,37 @@
+// Copyright 1998-2014 Epic Games, Inc. All Rights Reserved.
+
+#pragma once
+
+
+#ifndef DISABLE_DEPRECATION
+	#pragma clang diagnostic warning "-Wdeprecated-declarations"
+
+	/**
+	 * Macro for marking up deprecated code, functions and types.
+	 *
+	 * Features that are marked as deprecated are scheduled to be removed from the code base
+	 * in a future release. If you are using a deprecated feature in your code, you should
+	 * replace it before upgrading to the next release. See the Upgrade Notes in the release
+	 * notes for the release in which the feature was marked deprecated.
+	 *
+	 * Sample usage (note the slightly different syntax for classes and structures):
+	 *
+	 *		DEPRECATED(4.6, "Message")
+	 *		void Function();
+	 *
+	 *		struct DEPRECATED(4.6, "Message") MyStruct;
+	 *		class DEPRECATED(4.6, "Message") MyClass;
+	 *
+	 * @param VERSION The release number in which the feature was marked deprecated.
+	 * @param MESSAGE A message containing upgrade notes.
+	 */
+	#define DEPRECATED(VERSION, MESSAGE) __attribute__((deprecated(MESSAGE " Please update your code to the new API before upgrading to the next release, otherwise your project will no longer compile.")))
+
+	#define PRAGMA_DISABLE_DEPRECATION_WARNINGS \
+		_Pragma ("clang diagnostic push") \
+		_Pragma ("clang diagnostic ignored \"-Wdeprecated-declarations\"")
+
+	#define PRAGMA_ENABLE_DEPRECATION_WARNINGS \
+		_Pragma("clang diagnostic pop")
+
+#endif // DISABLE_DEPRECATION
diff --git a/Engine/Source/Runtime/Core/Public/Linux/LinuxPlatformCompilerPreSetup.h b/Engine/Source/Runtime/Core/Public/Linux/LinuxPlatformCompilerPreSetup.h
new file mode 100644
--- /dev/null
+++ b/Engine/Source/Runtime/Core/Public/Linux/LinuxPlatformCompilerPreSetup.h
@@ -0,0 +1,37 @@
+// Copyright 1998-2014 Epic Games, Inc. All Rights Reserved.
+
+#pragma once
+
+
+#ifndef DISABLE_DEPRECATION
+	#pragma clang diagnostic warning "-Wdeprecated-declarations"
+
+	/**
+	 * Macro for marking up deprecated code, functions and types.
+	 *
+	 * Features that are marked as deprecated are scheduled to be removed from the code base
+	 * in a future release. If you are using a deprecated feature in your code, you should
+	 * replace it before upgrading to the next release. See the Upgrade Notes in the release
+	 * notes for the release in which the feature was marked deprecated.
+	 *
+	 * Sample usage (note the slightly different syntax for classes and structures):
+	 *
+	 *		DEPRECATED(4.6, "Message")
+	 *		void Function();
+	 *
+	 *		struct DEPRECATED(4.6, "Message") MyStruct;
+	 *		class DEPRECATED(4.6, "Message") MyClass;
+	 *
+	 * @param VERSION The release number in which the feature was marked deprecated.
+	 * @param MESSAGE A message containing upgrade notes.
+	 */
+	#define DEPRECATED(VERSION, MESSAGE) __attribute__((deprecated(MESSAGE " Please update your code to the new API before upgrading to the next release, otherwise your project will no longer compile.")))
+
+	#define PRAGMA_DISABLE_DEPRECATION_WARNINGS \
+		_Pragma ("clang diagnostic push") \
+		_Pragma ("clang diagnostic ignored \"-Wdeprecated-declarations\"")
+
+	#define PRAGMA_ENABLE_DEPRECATION_WARNINGS \
+		_Pragma("clang diagnostic pop")
+
+#endif // DISABLE_DEPRECATION
diff --git a/Engine/Source/Runtime/Core/Public/Mac/MacPlatformCompilerPreSetup.h b/Engine/Source/Runtime/Core/Public/Mac/MacPlatformCompilerPreSetup.h
--- a/Engine/Source/Runtime/Core/Public/Mac/MacPlatformCompilerPreSetup.h
+++ b/Engine/Source/Runtime/Core/Public/Mac/MacPlatformCompilerPreSetup.h
@@ -26,4 +26,12 @@
 	 * @param MESSAGE A message containing upgrade notes.
 	 */
 	#define DEPRECATED(VERSION, MESSAGE) __attribute__((deprecated(MESSAGE " Please update your code to the new API before upgrading to the next release, otherwise your project will no longer compile.")))
+
+	#define PRAGMA_DISABLE_DEPRECATION_WARNINGS \
+		_Pragma ("clang diagnostic push") \
+		_Pragma ("clang diagnostic ignored \"-Wdeprecated-declarations\"")
+
+	#define PRAGMA_ENABLE_DEPRECATION_WARNINGS \
+		_Pragma("clang diagnostic pop")
+
 #endif // DISABLE_DEPRECATION
diff --git a/Engine/Source/Runtime/Core/Public/WinRT/WinRTPlatformCompilerPreSetup.h b/Engine/Source/Runtime/Core/Public/WinRT/WinRTPlatformCompilerPreSetup.h
new file mode 100644
--- /dev/null
+++ b/Engine/Source/Runtime/Core/Public/WinRT/WinRTPlatformCompilerPreSetup.h
@@ -0,0 +1,43 @@
+// Copyright 1998-2014 Epic Games, Inc. All Rights Reserved.
+
+#pragma once
+
+
+#ifndef DISABLE_DEPRECATION
+	/**
+	 * Macro for marking up deprecated code, functions and types.
+	 *
+	 * Features that are marked as deprecated are scheduled to be removed from the code base
+	 * in a future release. If you are using a deprecated feature in your code, you should
+	 * replace it before upgrading to the next release. See the Upgrade Notes in the release
+	 * notes for the release in which the feature was marked deprecated.
+	 *
+	 * Sample usage (note the slightly different syntax for classes and structures):
+	 *
+	 *		DEPRECATED(4.6, "Message")
+	 *		void Function();
+	 *
+	 *		struct DEPRECATED(4.6, "Message") MyStruct;
+	 *		class DEPRECATED(4.6, "Message") MyClass;
+	 *
+	 * @param VERSION The release number in which the feature was marked deprecated.
+	 * @param MESSAGE A message containing upgrade notes.
+	 */
+	#define DEPRECATED(VERSION, MESSAGE) __declspec(deprecated(MESSAGE " Please update your code to the new API before upgrading to the next release, otherwise your project will no longer compile."))
+
+	#ifndef EMIT_DEPRECATED_WARNING_MESSAGE_STR
+		#define EMIT_DEPRECATED_WARNING_MESSAGE_STR1(x) #x
+		#define EMIT_DEPRECATED_WARNING_MESSAGE_STR(x) EMIT_DEPRECATED_WARNING_MESSAGE_STR1(x)
+	#endif // EMIT_DEPRECATED_WARNING_MESSAGE_STR
+
+	#define EMIT_DEPRECATED_WARNING_MESSAGE(Msg) __pragma(message(__FILE__ "(" EMIT_DEPRECATED_WARNING_MESSAGE_STR(__LINE__) "): warning C4996: " Msg))
+
+	#define PRAGMA_DISABLE_DEPRECATION_WARNINGS \
+		__pragma (warning(push)) \
+		__pragma (warning(disable:4995)) \
+		__pragma (warning(disable:4996))
+
+	#define PRAGMA_ENABLE_DEPRECATION_WARNINGS \
+		__pragma(warning(pop))
+
+#endif // DISABLE_DEPRECATION
diff --git a/Engine/Source/Runtime/Core/Public/Windows/WindowsPlatformCompilerPreSetup.h b/Engine/Source/Runtime/Core/Public/Windows/WindowsPlatformCompilerPreSetup.h
--- a/Engine/Source/Runtime/Core/Public/Windows/WindowsPlatformCompilerPreSetup.h
+++ b/Engine/Source/Runtime/Core/Public/Windows/WindowsPlatformCompilerPreSetup.h
@@ -31,4 +31,13 @@
 	#endif // EMIT_DEPRECATED_WARNING_MESSAGE_STR
 
 	#define EMIT_DEPRECATED_WARNING_MESSAGE(Msg) __pragma(message(__FILE__ "(" EMIT_DEPRECATED_WARNING_MESSAGE_STR(__LINE__) "): warning C4996: " Msg))
+
+	#define PRAGMA_DISABLE_DEPRECATION_WARNINGS \
+		__pragma (warning(push)) \
+		__pragma (warning(disable:4995)) \
+		__pragma (warning(disable:4996))
+
+	#define PRAGMA_ENABLE_DEPRECATION_WARNINGS \
+		__pragma(warning(pop))
+
 #endif // DISABLE_DEPRECATION
diff --git a/Engine/Source/Runtime/CoreUObject/Public/UObject/ObjectBase.h b/Engine/Source/Runtime/CoreUObject/Public/UObject/ObjectBase.h
--- a/Engine/Source/Runtime/CoreUObject/Public/UObject/ObjectBase.h
+++ b/Engine/Source/Runtime/CoreUObject/Public/UObject/ObjectBase.h
@@ -496,25 +496,6 @@ struct COREUOBJECT_API FReferencerInformationList
 // This macro is used to declare a thunk function in autogenerated boilerplate code
 #define DECLARE_FUNCTION(func) void func( FFrame& Stack, RESULT_DECL )
 
-#if _MSC_VER
-	#define PRAGMA_DISABLE_DEPRECATION_WARNINGS \
-		__pragma (warning(push)) \
-		__pragma (warning(disable:4995)) \
-		__pragma (warning(disable:4996))
-
-	#define PRAGMA_ENABLE_DEPRECATION_WARNINGS \
-		__pragma(warning(pop))
-#elif __clang__
-	#define PRAGMA_DISABLE_DEPRECATION_WARNINGS \
-		_Pragma ("clang diagnostic push") \
-		_Pragma ("clang diagnostic ignored \"-Wdeprecated-declarations\"")
-
-	#define PRAGMA_ENABLE_DEPRECATION_WARNINGS \
-		_Pragma("clang diagnostic pop")
-#else
-	#error Unsupported compiler
-#endif
-
 // These are used for syntax highlighting and to allow autocomplete hints
 
 namespace UC
