Deprecate IForceFeedbackSystem, add IInputInterface.
Add PS4 LightBar support.

diff --git a/Engine/Source/Runtime/Core/Private/Android/AndroidApplication.cpp b/Engine/Source/Runtime/Core/Private/Android/AndroidApplication.cpp
--- a/Engine/Source/Runtime/Core/Private/Android/AndroidApplication.cpp
+++ b/Engine/Source/Runtime/Core/Private/Android/AndroidApplication.cpp
@@ -40,6 +40,18 @@ IForceFeedbackSystem *FAndroidApplication::GetForceFeedbackSystem()
 	return InputInterface.Get();
 }
 
+IForceFeedbackSystem *FAndroidApplication::DEPRECATED_GetForceFeedbackSystem()
+{
+	// NOTE: This does not increase the reference count, so don't cache the result
+	return InputInterface.Get();
+}
+
+IInputInterface* FAndroidApplication::GetInputInterface()
+{
+	// NOTE: This does not increase the reference count, so don't cache the result
+	return InputInterface.Get();
+}
+
 void FDisplayMetrics::GetDisplayMetrics( FDisplayMetrics& OutDisplayMetrics )
 {
 	// Get screen rect
diff --git a/Engine/Source/Runtime/Core/Private/Android/AndroidInputInterface.cpp b/Engine/Source/Runtime/Core/Private/Android/AndroidInputInterface.cpp
--- a/Engine/Source/Runtime/Core/Private/Android/AndroidInputInterface.cpp
+++ b/Engine/Source/Runtime/Core/Private/Android/AndroidInputInterface.cpp
@@ -69,7 +69,7 @@ void FAndroidInputInterface::Tick(float DeltaTime)
 
 }
 
-void FAndroidInputInterface::SetChannelValue(int32 ControllerId, FForceFeedbackChannelType ChannelType, float Value)
+void FAndroidInputInterface::SetForceFeedbackChannelValue(int32 ControllerId, FForceFeedbackChannelType ChannelType, float Value)
 {
 	// For now, force the device to 0
 	// Should use Java to enumerate number of controllers and assign device ID
@@ -82,19 +82,19 @@ void FAndroidInputInterface::SetChannelValue(int32 ControllerId, FForceFeedbackC
 	// Save a copy of the value for future comparison
 	switch (ChannelType)
 	{
-		case FF_CHANNEL_LEFT_LARGE:
+		case FForceFeedbackChannelType::LEFT_LARGE:
 			NewControllerData[ControllerId].VibeValues.LeftLarge = Value;
 			break;
 
-		case FF_CHANNEL_LEFT_SMALL:
+		case FForceFeedbackChannelType::LEFT_SMALL:
 			NewControllerData[ControllerId].VibeValues.LeftSmall = Value;
 			break;
 
-		case FF_CHANNEL_RIGHT_LARGE:
+		case FForceFeedbackChannelType::RIGHT_LARGE:
 			NewControllerData[ControllerId].VibeValues.RightLarge = Value;
 			break;
 
-		case FF_CHANNEL_RIGHT_SMALL:
+		case FForceFeedbackChannelType::RIGHT_SMALL:
 			NewControllerData[ControllerId].VibeValues.RightSmall = Value;
 			break;
 
@@ -107,7 +107,7 @@ void FAndroidInputInterface::SetChannelValue(int32 ControllerId, FForceFeedbackC
 	UpdateVibeMotors(NewControllerData[ControllerId]);
 }
 
-void FAndroidInputInterface::SetChannelValues(int32 ControllerId, const FForceFeedbackValues &Values)
+void FAndroidInputInterface::SetForceFeedbackChannelValues(int32 ControllerId, const FForceFeedbackValues &Values)
 {
 	// For now, force the device to 0
 	// Should use Java to enumerate number of controllers and assign device ID
diff --git a/Engine/Source/Runtime/Core/Private/Linux/LinuxApplication.cpp b/Engine/Source/Runtime/Core/Private/Linux/LinuxApplication.cpp
--- a/Engine/Source/Runtime/Core/Private/Linux/LinuxApplication.cpp
+++ b/Engine/Source/Runtime/Core/Private/Linux/LinuxApplication.cpp
@@ -5,7 +5,7 @@
 #include "LinuxWindow.h"
 #include "LinuxCursor.h"
 #include "GenericApplicationMessageHandler.h"
-#include "IForceFeedbackSystem.h"
+#include "IInputInterface.h"
 #include "IInputDeviceModule.h"
 #include "IInputDevice.h"
 #if STEAM_CONTROLLER_SUPPORT
diff --git a/Engine/Source/Runtime/Core/Private/Windows/WindowsApplication.cpp b/Engine/Source/Runtime/Core/Private/Windows/WindowsApplication.cpp
--- a/Engine/Source/Runtime/Core/Private/Windows/WindowsApplication.cpp
+++ b/Engine/Source/Runtime/Core/Private/Windows/WindowsApplication.cpp
@@ -1728,7 +1728,7 @@ void FWindowsApplication::PollGameDeviceState( const float TimeDelta )
 	}
 }
 
-void FWindowsApplication::SetChannelValue (int32 ControllerId, FForceFeedbackChannelType ChannelType, float Value)
+void FWindowsApplication::SetForceFeedbackChannelValue(int32 ControllerId, FForceFeedbackChannelType ChannelType, float Value)
 {
 	// send vibration to externally-implemented devices
 	for( auto DeviceIt = ExternalInputDevices.CreateIterator(); DeviceIt; ++DeviceIt )
@@ -1737,7 +1737,7 @@ void FWindowsApplication::SetChannelValue (int32 ControllerId, FForceFeedbackCha
 	}
 }
 
-void FWindowsApplication::SetChannelValues (int32 ControllerId, const FForceFeedbackValues &Values)
+void FWindowsApplication::SetForceFeedbackChannelValues(int32 ControllerId, const FForceFeedbackValues &Values)
 {
 	const FForceFeedbackValues* InternalValues = &Values;
  
diff --git a/Engine/Source/Runtime/Core/Private/Windows/XInputInterface.cpp b/Engine/Source/Runtime/Core/Private/Windows/XInputInterface.cpp
--- a/Engine/Source/Runtime/Core/Private/Windows/XInputInterface.cpp
+++ b/Engine/Source/Runtime/Core/Private/Windows/XInputInterface.cpp
@@ -239,19 +239,19 @@ void XInputInterface::SetChannelValue( const int32 ControllerId, const FForceFee
 	{
 		switch( ChannelType )
 		{
-			case FF_CHANNEL_LEFT_LARGE:
+			case FForceFeedbackChannelType::LEFT_LARGE:
 				ControllerState.ForceFeedback.LeftLarge = Value;
 				break;
 
-			case FF_CHANNEL_LEFT_SMALL:
+			case FForceFeedbackChannelType::LEFT_SMALL:
 				ControllerState.ForceFeedback.LeftSmall = Value;
 				break;
 
-			case FF_CHANNEL_RIGHT_LARGE:
+			case FForceFeedbackChannelType::RIGHT_LARGE:
 				ControllerState.ForceFeedback.RightLarge = Value;
 				break;
 
-			case FF_CHANNEL_RIGHT_SMALL:
+			case FForceFeedbackChannelType::RIGHT_SMALL:
 				ControllerState.ForceFeedback.RightSmall = Value;
 				break;
 		}
diff --git a/Engine/Source/Runtime/Core/Public/Android/AndroidApplication.h b/Engine/Source/Runtime/Core/Public/Android/AndroidApplication.h
--- a/Engine/Source/Runtime/Core/Public/Android/AndroidApplication.h
+++ b/Engine/Source/Runtime/Core/Public/Android/AndroidApplication.h
@@ -25,6 +25,9 @@ namespace FAndroidEGL
 	void					BindDisplayToContext(FPlatformOpenGLContext*);
 }
 
+//disable warnings from overriding the deprecated forcefeedback.  
+//calls to the deprecated function will still generate warnings.
+PRAGMA_DISABLE_DEPRECATION_WARNINGS
 
 class FAndroidApplication : public GenericApplication
 {
@@ -48,8 +51,14 @@ public:
 	virtual void PollGameDeviceState( const float TimeDelta ) override;
 
 	virtual FPlatformRect GetWorkArea( const FPlatformRect& CurrentWindow ) const override;
+	
+	/** Function to return the current implementation of the ForceFeedback system */
+	DEPRECATED(4.7, "Please use GetInputInterface()")
+	virtual IForceFeedbackSystem* GetForceFeedbackSystem() override;
+
+	virtual IForceFeedbackSystem* DEPRECATED_GetForceFeedbackSystem() override;	
 
-	virtual IForceFeedbackSystem *GetForceFeedbackSystem() override;
+	virtual IInputInterface* GetInputInterface() override;
 
 	virtual TSharedRef< FGenericWindow > MakeWindow() override;
 	
@@ -66,3 +75,5 @@ private:
 
 	TArray< TSharedRef< FAndroidWindow > > Windows;
 };
+
+PRAGMA_ENABLE_DEPRECATION_WARNINGS
\ No newline at end of file
diff --git a/Engine/Source/Runtime/Core/Public/Android/AndroidInputInterface.h b/Engine/Source/Runtime/Core/Public/Android/AndroidInputInterface.h
--- a/Engine/Source/Runtime/Core/Public/Android/AndroidInputInterface.h
+++ b/Engine/Source/Runtime/Core/Public/Android/AndroidInputInterface.h
@@ -4,6 +4,7 @@
 #include <android/input.h>
 #include <android/keycodes.h>
 #include <android/api-level.h>
+#include "IInputInterface.h"
 #include "IForceFeedbackSystem.h"
 
 #if __ANDROID_API__ <= 18
@@ -171,8 +172,9 @@ public:
 	/**
 	* IForceFeedbackSystem implementation
 	*/
-	virtual void SetChannelValue(int32 ControllerId, FForceFeedbackChannelType ChannelType, float Value) override;
-	virtual void SetChannelValues(int32 ControllerId, const FForceFeedbackValues &values) override;
+	virtual void SetForceFeedbackChannelValue(int32 ControllerId, FForceFeedbackChannelType ChannelType, float Value) override;
+	virtual void SetForceFeedbackChannelValues(int32 ControllerId, const FForceFeedbackValues &values) override;
+	virtual void SetLightColor(int32 ControllerId, FColor Color) override {}
 
 private:
 
diff --git a/Engine/Source/Runtime/Core/Public/GenericPlatform/GenericApplication.h b/Engine/Source/Runtime/Core/Public/GenericPlatform/GenericApplication.h
--- a/Engine/Source/Runtime/Core/Public/GenericPlatform/GenericApplication.h
+++ b/Engine/Source/Runtime/Core/Public/GenericPlatform/GenericApplication.h
@@ -7,6 +7,7 @@ class FGenericWindow;
 class ICursor;
 class ITextInputMethodSystem;
 class IForceFeedbackSystem;
+class IInputInterface;
 class IAnalyticsProvider;
 
 /**
@@ -437,7 +438,12 @@ public:
 	virtual void DestroyApplication() { }
 
 	/** Function to return the current implementation of the ForceFeedback system */
-	virtual IForceFeedbackSystem *GetForceFeedbackSystem() { return NULL; }
+	DEPRECATED(4.7, "Please use GetInputInterface()  ForceFeedbackSystem not accessible.")
+	virtual IForceFeedbackSystem* GetForceFeedbackSystem() { return nullptr; }
+
+	virtual IForceFeedbackSystem* DEPRECATED_GetForceFeedbackSystem() { return nullptr; }
+
+	virtual IInputInterface* GetInputInterface() { return nullptr; }	
 
 	/** Function to return the current implementation of the Text Input Method System */
 	virtual ITextInputMethodSystem *GetTextInputMethodSystem() { return NULL; }
diff --git a/Engine/Source/Runtime/Core/Public/GenericPlatform/IForceFeedbackSystem.h b/Engine/Source/Runtime/Core/Public/GenericPlatform/IForceFeedbackSystem.h
--- a/Engine/Source/Runtime/Core/Public/GenericPlatform/IForceFeedbackSystem.h
+++ b/Engine/Source/Runtime/Core/Public/GenericPlatform/IForceFeedbackSystem.h
@@ -1,60 +1,13 @@
 // Copyright 1998-2014 Epic Games, Inc. All Rights Reserved.
 
 #pragma once
-#include "HAL/Platform.h"
-
-// General identifiers for potential force feedback channels. These will be mapped according to the
-// platform specific implementation.
-// For example, the PS4 only listens to the XXX_LARGE channels and ignores the rest, while the XBox One could
-// map the XXX_LARGE to the handle motors and XXX_SMALL to the trigger motors. And iOS can map LEFT_SMALL to
-// its single motor.
-enum FForceFeedbackChannelType
-{
-	FF_CHANNEL_LEFT_LARGE,
-	FF_CHANNEL_LEFT_SMALL,
-	FF_CHANNEL_RIGHT_LARGE,
-	FF_CHANNEL_RIGHT_SMALL
-};
-
-struct FForceFeedbackValues
-{
-	float LeftLarge;
-	float LeftSmall;
-	float RightLarge;
-	float RightSmall;
-
-	FForceFeedbackValues()
-		: LeftLarge(0.f)
-		, LeftSmall(0.f)
-		, RightLarge(0.f)
-		, RightSmall(0.f)
-	{
-	}
-};
+#include "IInputInterface.h"
 
 // Abstract interface for the force feedback system
-class IForceFeedbackSystem
+// class is deprecated and will be removed in favor of IInputInterface
+class IForceFeedbackSystem : public IInputInterface
 {
 public:
 	virtual ~IForceFeedbackSystem() {};
-
-	/**
-	 * Sets the strength/speed of the given channel for the given controller id.
-	 * NOTE: If the channel is not supported, the call will silently fail
-	 *
-	 * @param ControllerId the id of the controller whose value is to be set
-	 * @param ChannelType the type of channel whose value should be set
-	 * @param Value strength or speed of feedback, 0.0f to 1.0f. 0.0f will disable
-	 */
-	virtual void SetChannelValue (int32 ControllerId, FForceFeedbackChannelType ChannelType, float Value) = 0;
-
-	/**
-	 * Sets the strength/speed of all the channels for the given controller id.
-	 * NOTE: Unsupported channels are silently ignored
-	 *
-	 * @param ControllerId the id of the controller whose value is to be set
-	 * @param FForceFeedbackChannelValues strength or speed of feedback for all channels
-	 */
-	virtual void SetChannelValues (int32 ControllerId, const FForceFeedbackValues &Values) = 0;
 };
 
diff --git a/Engine/Source/Runtime/Core/Public/Windows/WindowsApplication.h b/Engine/Source/Runtime/Core/Public/Windows/WindowsApplication.h
--- a/Engine/Source/Runtime/Core/Public/Windows/WindowsApplication.h
+++ b/Engine/Source/Runtime/Core/Public/Windows/WindowsApplication.h
@@ -9,6 +9,7 @@
 	#include "OleIdl.h"
 	#include <ShObjIdl.h>
 #include "HideWindowsPlatformTypes.h"
+#include "IInputInterface.h"
 #include "IForceFeedbackSystem.h"
 #include "WindowsTextInputMethodSystem.h"
 
@@ -192,6 +193,10 @@ public:
 	POINTL CursorPosition;
 };
 
+//disable warnings from overriding the deprecated forcefeedback.  
+//calls to the deprecated function will still generate warnings.
+PRAGMA_DISABLE_DEPRECATION_WARNINGS
+
 /**
  * Windows-specific application implementation.
  */
@@ -251,15 +256,29 @@ public:
 	virtual void GetInitialDisplayMetrics( FDisplayMetrics& OutDisplayMetrics ) const override;
 
 	virtual EWindowTitleAlignment::Type GetWindowTitleAlignment() const override;
-
-	virtual IForceFeedbackSystem *GetForceFeedbackSystem() override 
+	
+	/** Function to return the current implementation of the ForceFeedback system */
+	DEPRECATED(4.7, "Please use GetInputInterface()")
+	virtual IForceFeedbackSystem* GetForceFeedbackSystem() override
 	{
 		return this; 
 	}
 
-	virtual void SetChannelValue (int32 ControllerId, FForceFeedbackChannelType ChannelType, float Value) override;
+	virtual IForceFeedbackSystem* DEPRECATED_GetForceFeedbackSystem() override
+	{
+		return this;
+	}
+
+	virtual IInputInterface* GetInputInterface() override
+	{
+		return this;
+	}
+
+	virtual void SetForceFeedbackChannelValue (int32 ControllerId, FForceFeedbackChannelType ChannelType, float Value) override;
 
-	virtual void SetChannelValues (int32 ControllerId, const FForceFeedbackValues &Values) override;
+	virtual void SetForceFeedbackChannelValues(int32 ControllerId, const FForceFeedbackValues &Values) override;
+
+	virtual void SetLightColor(int32 ControllerId, FColor Color) override {}
 
 	virtual ITextInputMethodSystem *GetTextInputMethodSystem() override
 	{
@@ -360,4 +379,6 @@ private:
 	STICKYKEYS							StartupStickyKeys;
 	TOGGLEKEYS							StartupToggleKeys;
 	FILTERKEYS							StartupFilterKeys;
-};
\ No newline at end of file
+};
+
+PRAGMA_ENABLE_DEPRECATION_WARNINGS
\ No newline at end of file
diff --git a/Engine/Source/Runtime/Engine/Classes/GameFramework/PlayerController.h b/Engine/Source/Runtime/Engine/Classes/GameFramework/PlayerController.h
--- a/Engine/Source/Runtime/Engine/Classes/GameFramework/PlayerController.h
+++ b/Engine/Source/Runtime/Engine/Classes/GameFramework/PlayerController.h
@@ -11,7 +11,7 @@
 #include "GameFramework/OnlineReplStructs.h"
 #include "GameFramework/Controller.h"
 #include "Engine/LatentActionManager.h"
-#include "GenericPlatform/IForceFeedbackSystem.h"
+#include "GenericPlatform/IInputInterface.h"
 #include "SlateCore.h"
 
 #include "PlayerController.generated.h"
diff --git a/Engine/Source/Runtime/Engine/Private/GameplayStatics.cpp b/Engine/Source/Runtime/Engine/Private/GameplayStatics.cpp
--- a/Engine/Source/Runtime/Engine/Private/GameplayStatics.cpp
+++ b/Engine/Source/Runtime/Engine/Private/GameplayStatics.cpp
@@ -5,7 +5,7 @@
 #include "SoundDefinitions.h"
 #include "PlatformFeatures.h"
 #include "LatentActions.h"
-#include "IForceFeedbackSystem.h"
+#include "IInputInterface.h"
 #include "SlateBasics.h"
 #include "Particles/ParticleSystemComponent.h"
 #include "GameFramework/Character.h"
diff --git a/Engine/Source/Runtime/Engine/Private/PlayerController.cpp b/Engine/Source/Runtime/Engine/Private/PlayerController.cpp
--- a/Engine/Source/Runtime/Engine/Private/PlayerController.cpp
+++ b/Engine/Source/Runtime/Engine/Private/PlayerController.cpp
@@ -13,7 +13,7 @@
 #include "SoundDefinitions.h"
 #include "OnlineSubsystemUtils.h"
 #include "IHeadMountedDisplay.h"
-#include "IForceFeedbackSystem.h"
+#include "IInputInterface.h"
 #include "SlateBasics.h"
 #include "GameFramework/TouchInterface.h"
 #include "DisplayDebugHelpers.h"
@@ -25,7 +25,6 @@
 #include "ContentStreaming.h"
 #include "GameFramework/PawnMovementComponent.h"
 #include "Camera/CameraActor.h"
-#include "GenericPlatform/IForceFeedbackSystem.h"
 #include "Engine/InputDelegateBinding.h"
 #include "SVirtualJoystick.h"
 #include "GameFramework/LocalMessage.h"
@@ -1337,10 +1336,10 @@ void APlayerController::EndPlay(const EEndPlayReason::Type EndPlayReason)
 		}
 
 		// Stop any force feedback effects that may be active
-		IForceFeedbackSystem* ForceFeedbackSystem = FSlateApplication::Get().GetForceFeedbackSystem();
-		if (ForceFeedbackSystem)
+		IInputInterface* InputInterface = FSlateApplication::Get().GetInputInterface();
+		if (InputInterface)
 		{
-			ForceFeedbackSystem->SetChannelValues(LocalPlayer->GetControllerId(), FForceFeedbackValues());
+			InputInterface->SetForceFeedbackChannelValues(LocalPlayer->GetControllerId(), FForceFeedbackValues());
 		}
 	}
 
@@ -3502,10 +3501,10 @@ void APlayerController::ProcessForceFeedback(const float DeltaTime, const bool b
 		}
 	}
 
-	IForceFeedbackSystem* ForceFeedbackSystem = FSlateApplication::Get().GetForceFeedbackSystem();
-	if (ForceFeedbackSystem)
+	IInputInterface* InputInterface = FSlateApplication::Get().GetInputInterface();
+	if (InputInterface)
 	{
-		ForceFeedbackSystem->SetChannelValues(CastChecked<ULocalPlayer>(Player)->GetControllerId(), (bForceFeedbackEnabled ? ForceFeedbackValues : FForceFeedbackValues()));
+		InputInterface->SetForceFeedbackChannelValues(CastChecked<ULocalPlayer>(Player)->GetControllerId(), (bForceFeedbackEnabled ? ForceFeedbackValues : FForceFeedbackValues()));
 	}
 }
 
diff --git a/Engine/Source/Runtime/Engine/Private/UserInterface/ForceFeedbackEffect.cpp b/Engine/Source/Runtime/Engine/Private/UserInterface/ForceFeedbackEffect.cpp
--- a/Engine/Source/Runtime/Engine/Private/UserInterface/ForceFeedbackEffect.cpp
+++ b/Engine/Source/Runtime/Engine/Private/UserInterface/ForceFeedbackEffect.cpp
@@ -1,7 +1,7 @@
 // Copyright 1998-2014 Epic Games, Inc. All Rights Reserved.
 
 #include "EnginePrivate.h"
-#include "IForceFeedbackSystem.h"
+#include "IInputInterface.h"
 
 UForceFeedbackEffect::UForceFeedbackEffect(const FObjectInitializer& ObjectInitializer)
 	: Super(ObjectInitializer)
diff --git a/Engine/Source/Runtime/Slate/Public/Framework/Application/SlateApplication.h b/Engine/Source/Runtime/Slate/Public/Framework/Application/SlateApplication.h
--- a/Engine/Source/Runtime/Slate/Public/Framework/Application/SlateApplication.h
+++ b/Engine/Source/Runtime/Slate/Public/Framework/Application/SlateApplication.h
@@ -192,7 +192,10 @@ public:
 	float GetSoundDuration(const FSlateSound& Sound) const;
 
 	/** @return The force feedback interface for this application */
-	IForceFeedbackSystem *GetForceFeedbackSystem() const { return PlatformApplication->GetForceFeedbackSystem(); }
+	DEPRECATED(4.7, "Please use GetInputInterface().")
+	IForceFeedbackSystem* GetForceFeedbackSystem() const { return PlatformApplication->DEPRECATED_GetForceFeedbackSystem(); }
+
+	IInputInterface* GetInputInterface() const { return PlatformApplication->GetInputInterface(); }
 
 	/** @return Whether or not the current platform supports system help */
 	bool SupportsSystemHelp() const { return PlatformApplication->SupportsSystemHelp(); }
