Add missing file.

diff --git a/Engine/Source/Runtime/Core/Public/GenericPlatform/IInputInterface.h b/Engine/Source/Runtime/Core/Public/GenericPlatform/IInputInterface.h
new file mode 100644
--- /dev/null
+++ b/Engine/Source/Runtime/Core/Public/GenericPlatform/IInputInterface.h
@@ -0,0 +1,69 @@
+// Copyright 1998-2014 Epic Games, Inc. All Rights Reserved.
+
+#pragma once
+#include "HAL/Platform.h"
+
+// General identifiers for potential force feedback channels. These will be mapped according to the
+// platform specific implementation.
+// For example, the PS4 only listens to the XXX_LARGE channels and ignores the rest, while the XBox One could
+// map the XXX_LARGE to the handle motors and XXX_SMALL to the trigger motors. And iOS can map LEFT_SMALL to
+// its single motor.
+enum class FForceFeedbackChannelType
+{
+	LEFT_LARGE,
+	LEFT_SMALL,
+	RIGHT_LARGE,
+	RIGHT_SMALL
+};
+
+struct FForceFeedbackValues
+{
+	float LeftLarge;
+	float LeftSmall;
+	float RightLarge;
+	float RightSmall;
+
+	FForceFeedbackValues()
+		: LeftLarge(0.f)
+		, LeftSmall(0.f)
+		, RightLarge(0.f)
+		, RightSmall(0.f)
+	{
+	}
+};
+
+// Abstract interface for the input interface.
+class IInputInterface
+{
+public:
+	virtual ~IInputInterface() {};
+
+	/**
+	* Sets the strength/speed of the given channel for the given controller id.
+	* NOTE: If the channel is not supported, the call will silently fail
+	*
+	* @param ControllerId the id of the controller whose value is to be set
+	* @param ChannelType the type of channel whose value should be set
+	* @param Value strength or speed of feedback, 0.0f to 1.0f. 0.0f will disable
+	*/
+	DEPRECATED(4.7, "Please use SetForceFeedbackChannel()")
+	void SetChannelValue(int32 ControllerId, FForceFeedbackChannelType ChannelType, float Value) { SetForceFeedbackChannelValue(ControllerId, ChannelType, Value); }
+	virtual void SetForceFeedbackChannelValue(int32 ControllerId, FForceFeedbackChannelType ChannelType, float Value) = 0;
+
+	/**
+	* Sets the strength/speed of all the channels for the given controller id.
+	* NOTE: Unsupported channels are silently ignored
+	*
+	* @param ControllerId the id of the controller whose value is to be set
+	* @param FForceFeedbackChannelValues strength or speed of feedback for all channels
+	*/
+	DEPRECATED(4.7, "Please use SetForceFeedbackChannelValues()")
+	void SetChannelValues(int32 ControllerId, const FForceFeedbackValues &Values) { SetForceFeedbackChannelValues(ControllerId, Values); }
+	virtual void SetForceFeedbackChannelValues(int32 ControllerId, const FForceFeedbackValues &Values) = 0;
+
+	/*
+	 * Sets the controller for the given controller.  Ignored if controller does not support a color.
+	 */
+	virtual void SetLightColor(int32 ControllerId, FColor Color) = 0;
+};
+
