#UE4

- Adding ClosestPointOnSegment2D into FMath.

#RememberBUNS

diff --git a/Engine/Source/Runtime/Core/Private/Math/UnrealMath.cpp b/Engine/Source/Runtime/Core/Private/Math/UnrealMath.cpp
--- a/Engine/Source/Runtime/Core/Private/Math/UnrealMath.cpp
+++ b/Engine/Source/Runtime/Core/Private/Math/UnrealMath.cpp
@@ -1221,6 +1221,29 @@ FVector FMath::ClosestPointOnSegment(const FVector &Point, const FVector &StartP
 	return StartPoint + Segment * (Dot1 / Dot2);
 }
 
+FVector2D FMath::ClosestPointOnSegment2D(const FVector2D &Point, const FVector2D &StartPoint, const FVector2D &EndPoint)
+{
+	const FVector2D Segment = EndPoint - StartPoint;
+	const FVector2D VectToPoint = Point - StartPoint;
+
+	// See if closest point is before StartPoint
+	const float Dot1 = VectToPoint | Segment;
+	if (Dot1 <= 0)
+	{
+		return StartPoint;
+	}
+
+	// See if closest point is beyond EndPoint
+	const float Dot2 = Segment | Segment;
+	if (Dot2 <= Dot1)
+	{
+		return EndPoint;
+	}
+
+	// Closest Point is within segment
+	return StartPoint + Segment * (Dot1 / Dot2);
+}
+
 float FMath::PointDistToSegment(const FVector &Point, const FVector &StartPoint, const FVector &EndPoint) 
 {
 	const FVector ClosestPoint = ClosestPointOnSegment(Point, StartPoint, EndPoint);
diff --git a/Engine/Source/Runtime/Core/Public/Math/UnrealMathUtility.h b/Engine/Source/Runtime/Core/Public/Math/UnrealMathUtility.h
--- a/Engine/Source/Runtime/Core/Public/Math/UnrealMathUtility.h
+++ b/Engine/Source/Runtime/Core/Public/Math/UnrealMathUtility.h
@@ -861,6 +861,20 @@ struct FMath : public FPlatformMath
 	 */
 	static CORE_API FVector ClosestPointOnSegment(const FVector &Point, const FVector &StartPoint, const FVector &EndPoint);
 
+	/**
+	* FVector2D version of ClosestPointOnSegment.
+	* Returns closest point on a segment to a given 2D point.
+	* The idea is to project point on line formed by segment.
+	* Then we see if the closest point on the line is outside of segment or inside.
+	*
+	* @param	Point			point for which we find the closest point on the segment
+	* @param	StartPoint		StartPoint of segment
+	* @param	EndPoint		EndPoint of segment
+	*
+	* @return	point on the segment defined by (StartPoint, EndPoint) that is closest to Point.
+	*/
+	static CORE_API FVector2D ClosestPointOnSegment2D(const FVector2D &Point, const FVector2D &StartPoint, const FVector2D &EndPoint);
+
 	/**
 	 * Returns distance from a point to the closest point on a segment.
 	 *
