Added -UserDir= commandline argument to run multiple game instances with different config files easily

diff --git a/Engine/Source/Runtime/Core/Private/Misc/Paths.cpp b/Engine/Source/Runtime/Core/Private/Misc/Paths.cpp
--- a/Engine/Source/Runtime/Core/Private/Misc/Paths.cpp
+++ b/Engine/Source/Runtime/Core/Private/Misc/Paths.cpp
@@ -91,6 +91,12 @@ FString FPaths::GameUserDir()
 	}
 	else
 	{
+		FString UserDir;
+		if (FParse::Value(FCommandLine::Get(), TEXT("UserDir="), UserDir))
+		{
+			return FPaths::Combine(*FPaths::GameDir(), *UserDir) + TEXT("/");
+		}
+
 		return FPaths::GameDir();
 	}
 }
