[UE-6046] JNI errors when using ART runtime. Fixed various errors with JNI calls and references.. Moved build OS version to GameActivity to work around ART limitation. Fix FAndroidApplication::ClassLoader being a local ref instead of a global ref. Added exception check to FindJavaClass to find class lookup errors early. Fixed various local vs. global refs in FJavaClassObject.

diff --git a/Engine/Build/Android/Java/src/com/epicgames/ue4/GameActivity.java b/Engine/Build/Android/Java/src/com/epicgames/ue4/GameActivity.java
--- a/Engine/Build/Android/Java/src/com/epicgames/ue4/GameActivity.java
+++ b/Engine/Build/Android/Java/src/com/epicgames/ue4/GameActivity.java
@@ -130,6 +130,15 @@ public class GameActivity extends NativeActivity
 	{
 		return _activity;
 	}
+	
+	/**
+	Get the SDK level of the OS we are running in.
+	We do this instead of accessing the SDK_INT
+	with JNI from C++ as the new ART runtime seems to have
+	problems dynamically finding/loading static inner classes.
+	*/
+	public static final int ANDROID_BUILD_VERSION = android.os.Build.VERSION.SDK_INT;
+	
 	private StoreHelper IapStoreHelper;
 
 	@Override
diff --git a/Engine/Source/Runtime/Core/Private/Android/AndroidApplication.cpp b/Engine/Source/Runtime/Core/Private/Android/AndroidApplication.cpp
--- a/Engine/Source/Runtime/Core/Private/Android/AndroidApplication.cpp
+++ b/Engine/Source/Runtime/Core/Private/Android/AndroidApplication.cpp
@@ -5,6 +5,8 @@
 #include "AndroidInputInterface.h"
 #include "AndroidWindow.h"
 
+DEFINE_LOG_CATEGORY_STATIC(LogAndroidApplication, Log, All);
+
 FAndroidApplication* FAndroidApplication::CreateAndroidApplication()
 {
 	return new FAndroidApplication();
@@ -145,7 +147,8 @@ void FAndroidApplication::InitializeJavaEnv( JavaVM* VM, jint Version, jobject G
 		jclass classClass = Env->FindClass("java/lang/Class");
 		jclass classLoaderClass = Env->FindClass("java/lang/ClassLoader");
 		jmethodID getClassLoaderMethod = Env->GetMethodID(classClass, "getClassLoader", "()Ljava/lang/ClassLoader;");
-		ClassLoader = Env->CallObjectMethod(MainClass, getClassLoaderMethod);
+		jobject classLoader = Env->CallObjectMethod(MainClass, getClassLoaderMethod);
+		ClassLoader = Env->NewGlobalRef(classLoader);
 		FindClassMethod = Env->GetMethodID(classLoaderClass, "findClass", "(Ljava/lang/String;)Ljava/lang/Class;");
 	}
 	GlobalObjectRef = GlobalThis;
@@ -204,6 +207,7 @@ jclass FAndroidApplication::FindJavaClass( const char* name )
 	}
 	jstring ClassNameObj = Env->NewStringUTF(name);
 	jclass FoundClass = static_cast<jclass>(Env->CallObjectMethod(ClassLoader, FindClassMethod, ClassNameObj));
+	CheckJavaException();
 	Env->DeleteLocalRef(ClassNameObj);
 	return FoundClass;
 }
@@ -212,3 +216,20 @@ void FAndroidApplication::DetachJavaEnv()
 {
 	CurrentJavaVM->DetachCurrentThread();
 }
+
+bool FAndroidApplication::CheckJavaException()
+{
+	JNIEnv* Env = GetJavaEnv();
+	if (!Env)
+	{
+		return true;
+	}
+	if (Env->ExceptionCheck())
+	{
+		Env->ExceptionDescribe();
+		Env->ExceptionClear();
+		verify(false && "Java JNI call failed with an exception.");
+		return true;
+	}
+	return false;
+}
diff --git a/Engine/Source/Runtime/Core/Private/Android/AndroidJava.cpp b/Engine/Source/Runtime/Core/Private/Android/AndroidJava.cpp
--- a/Engine/Source/Runtime/Core/Private/Android/AndroidJava.cpp
+++ b/Engine/Source/Runtime/Core/Private/Android/AndroidJava.cpp
@@ -14,19 +14,18 @@ FJavaClassObject::FJavaClassObject(FName ClassName, const char* CtorSig, ...)
 	check(Constructor);
 	va_list Params;
 	va_start(Params, CtorSig);
-	Object = JEnv->NewObjectV(Class, Constructor, Params);
+	jobject object = JEnv->NewObjectV(Class, Constructor, Params);
 	va_end(Params);
 	VerifyException();
+	check(object);
 
-	check(Object);
 	// Promote local references to global
-	JEnv->NewGlobalRef(Class);
-	JEnv->NewGlobalRef(Object);
+	Object = JEnv->NewGlobalRef(object);
+	JEnv->DeleteLocalRef(object);
 }
 
 FJavaClassObject::~FJavaClassObject()
 {
-	JEnv->DeleteGlobalRef(Class);
 	JEnv->DeleteGlobalRef(Object);
 }
 
@@ -93,11 +92,11 @@ jobject FJavaClassObject::CallMethod<jobject>(FJavaClassMethod Method, ...)
 {
 	va_list Params;
 	va_start(Params, Method);
-	jobject RetVal = JEnv->CallObjectMethodV(Object, Method.Method, Params);
+	jobject val = JEnv->CallObjectMethodV(Object, Method.Method, Params);
 	va_end(Params);
-	JEnv->NewGlobalRef(RetVal);
-	JEnv->DeleteLocalRef(RetVal);
 	VerifyException();
+	jobject RetVal = JEnv->NewGlobalRef(val);
+	JEnv->DeleteLocalRef(val);
 	return RetVal;
 }
 
@@ -130,9 +129,9 @@ FString FJavaClassObject::CallMethod<FString>(FJavaClassMethod Method, ...)
 jstring FJavaClassObject::GetJString(const FString& String)
 {
 	auto StringCastObj = StringCast<ANSICHAR>(*String);
-	jstring result = FAndroidApplication::GetJavaEnv()->NewStringUTF(StringCastObj.Get());
-	FAndroidApplication::GetJavaEnv()->NewGlobalRef(result);
-	FAndroidApplication::GetJavaEnv()->DeleteLocalRef(result);
+	jstring local = FAndroidApplication::GetJavaEnv()->NewStringUTF(StringCastObj.Get());
+	jstring result = (jstring)FAndroidApplication::GetJavaEnv()->NewGlobalRef(local);
+	FAndroidApplication::GetJavaEnv()->DeleteLocalRef(local);
 	return result;
 }
 
diff --git a/Engine/Source/Runtime/Core/Private/Android/AndroidMisc.cpp b/Engine/Source/Runtime/Core/Private/Android/AndroidMisc.cpp
--- a/Engine/Source/Runtime/Core/Private/Android/AndroidMisc.cpp
+++ b/Engine/Source/Runtime/Core/Private/Android/AndroidMisc.cpp
@@ -642,16 +642,15 @@ int32 FAndroidMisc::GetAndroidBuildVersion()
 		JNIEnv* JEnv = FAndroidApplication::GetJavaEnv();
 		if (nullptr != JEnv)
 		{
-			jclass Class = FAndroidApplication::FindJavaClass("android/os/Build$VERSION");
+			jclass Class = FAndroidApplication::FindJavaClass("com/epicgames/ue4/GameActivity");
 			if (nullptr != Class)
 			{
-				JEnv->NewGlobalRef(Class);
-				jfieldID Field = JEnv->GetStaticFieldID(Class, "SDK_INT", "I");
+				jfieldID Field = JEnv->GetStaticFieldID(Class, "ANDROID_BUILD_VERSION", "I");
 				if (nullptr != Field)
 				{
 					AndroidBuildVersion = JEnv->GetStaticIntField(Class, Field);
 				}
-				JEnv->DeleteGlobalRef(Class);
+				JEnv->DeleteLocalRef(Class);
 			}
 		}
 	}
@@ -666,17 +665,15 @@ bool FAndroidMisc::IsDebuggerPresent()
 	JNIEnv* JEnv = FAndroidApplication::GetJavaEnv();
 	if (nullptr != JEnv)
 	{
-		jclass Class = FAndroidApplication::FindJavaClass("android.os.Debug");
+		jclass Class = FAndroidApplication::FindJavaClass("android/os/Debug");
 		if (nullptr != Class)
 		{
-			JEnv->NewGlobalRef(Class);
 			// This segfaults for some reason. So this is all disabled for now.
 			jmethodID Method = JEnv->GetStaticMethodID(Class, "isDebuggerConnected", "()Z");
 			if (nullptr != Method)
 			{
 				Result = JEnv->CallStaticBooleanMethod(Class, Method);
 			}
-			JEnv->DeleteGlobalRef(Class);
 		}
 	}
 #endif
diff --git a/Engine/Source/Runtime/Core/Public/Android/AndroidApplication.h b/Engine/Source/Runtime/Core/Public/Android/AndroidApplication.h
--- a/Engine/Source/Runtime/Core/Public/Android/AndroidApplication.h
+++ b/Engine/Source/Runtime/Core/Public/Android/AndroidApplication.h
@@ -40,7 +40,7 @@ public:
 	static JNIEnv* GetJavaEnv(bool bRequireGlobalThis = true);
 	static jclass FindJavaClass(const char* name);
 	static void DetachJavaEnv();
-
+	static bool CheckJavaException();
 
 public:	
 	
