Blueprint methods for FDateTime to FText conversion now assume the supplied FDateTime is of the "Unknown" timezone. Added the timezone as an optional parameter to FText::AsDate, defaulting to the user's local time zone.

diff --git a/Engine/Source/Runtime/Core/Private/Internationalization/ICUCulture.cpp b/Engine/Source/Runtime/Core/Private/Internationalization/ICUCulture.cpp
--- a/Engine/Source/Runtime/Core/Private/Internationalization/ICUCulture.cpp
+++ b/Engine/Source/Runtime/Core/Private/Internationalization/ICUCulture.cpp
@@ -43,24 +43,27 @@ namespace
 	TSharedRef<const icu::DateFormat> CreateDateFormat( const icu::Locale& ICULocale )
 	{
 		UErrorCode ICUStatus = U_ZERO_ERROR;
-		TSharedPtr<const icu::DateFormat> Ptr = MakeShareable( icu::DateFormat::createDateInstance( icu::DateFormat::EStyle::kDefault, ICULocale ) );
+		TSharedPtr<icu::DateFormat> Ptr = MakeShareable( icu::DateFormat::createDateInstance( icu::DateFormat::EStyle::kDefault, ICULocale ) );
 		checkf(Ptr.IsValid(), TEXT("Creating a date format object failed using locale %s. Perhaps this locale has no data."), StringCast<TCHAR>(ICULocale.getName()).Get());
+		Ptr->adoptTimeZone( icu::TimeZone::createDefault() );
 		return Ptr.ToSharedRef();
 	}
 
 	TSharedRef<const icu::DateFormat> CreateTimeFormat( const icu::Locale& ICULocale )
 	{
 		UErrorCode ICUStatus = U_ZERO_ERROR;
-		TSharedPtr<const icu::DateFormat> Ptr = MakeShareable( icu::DateFormat::createTimeInstance( icu::DateFormat::EStyle::kDefault, ICULocale ) );
+		TSharedPtr<icu::DateFormat> Ptr = MakeShareable( icu::DateFormat::createTimeInstance( icu::DateFormat::EStyle::kDefault, ICULocale ) );
 		checkf(Ptr.IsValid(), TEXT("Creating a time format object failed using locale %s. Perhaps this locale has no data."), StringCast<TCHAR>(ICULocale.getName()).Get());
+		Ptr->adoptTimeZone( icu::TimeZone::createDefault() );
 		return Ptr.ToSharedRef();
 	}
 
 	TSharedRef<const icu::DateFormat> CreateDateTimeFormat( const icu::Locale& ICULocale )
 	{
 		UErrorCode ICUStatus = U_ZERO_ERROR;
-		TSharedPtr<const icu::DateFormat> Ptr = MakeShareable( icu::DateFormat::createDateTimeInstance( icu::DateFormat::EStyle::kDefault, icu::DateFormat::EStyle::kDefault, ICULocale ) );
+		TSharedPtr<icu::DateFormat> Ptr = MakeShareable( icu::DateFormat::createDateTimeInstance( icu::DateFormat::EStyle::kDefault, icu::DateFormat::EStyle::kDefault, ICULocale ) );
 		checkf(Ptr.IsValid(), TEXT("Creating a date-time format object failed using locale %s. Perhaps this locale has no data."), StringCast<TCHAR>(ICULocale.getName()).Get());
+		Ptr->adoptTimeZone( icu::TimeZone::createDefault() );
 		return Ptr.ToSharedRef();
 	}
 }
@@ -306,11 +309,34 @@ TSharedRef<const icu::DecimalFormat> FCulture::FICUCultureImplementation::GetPer
 	}
 }
 
-TSharedRef<const icu::DateFormat> FCulture::FICUCultureImplementation::GetDateFormatter(const EDateTimeStyle::Type DateStyle) const
+TSharedRef<const icu::DateFormat> FCulture::FICUCultureImplementation::GetDateFormatter(const EDateTimeStyle::Type DateStyle, const FString& TimeZone) const
 {
-	const bool bIsDefault = (DateStyle == EDateTimeStyle::Default);
+	icu::UnicodeString InputTimeZoneID;
+	ICUUtilities::ConvertString(TimeZone, InputTimeZoneID, false);
 
 	const TSharedRef<const icu::DateFormat> DefaultFormatter( ICUDateFormat );
+
+	bool bIsDefaultTimeZone = TimeZone.IsEmpty();
+	if( !bIsDefaultTimeZone )
+	{
+		UErrorCode ICUStatus = U_ZERO_ERROR;
+
+		icu::UnicodeString CanonicalInputTimeZoneID;
+		icu::TimeZone::getCanonicalID(InputTimeZoneID, CanonicalInputTimeZoneID, ICUStatus);
+
+		icu::UnicodeString DefaultTimeZoneID;
+		DefaultFormatter->getTimeZone().getID(DefaultTimeZoneID);
+
+		icu::UnicodeString CanonicalDefaultTimeZoneID;
+		icu::TimeZone::getCanonicalID(DefaultTimeZoneID, CanonicalDefaultTimeZoneID, ICUStatus);
+
+		bIsDefaultTimeZone = (CanonicalInputTimeZoneID == CanonicalDefaultTimeZoneID ? true : false);
+	}
+
+	const bool bIsDefault = 
+		DateStyle == EDateTimeStyle::Default &&
+		bIsDefaultTimeZone;
+
 	if(bIsDefault)
 	{
 		return DefaultFormatter;
@@ -318,8 +344,7 @@ TSharedRef<const icu::DateFormat> FCulture::FICUCultureImplementation::GetDateFo
 	else
 	{
 		const TSharedRef<icu::DateFormat> Formatter( icu::DateFormat::createDateInstance( UEToICU(DateStyle), ICULocale ) );
-		const icu::TimeZone& TimeZone = *(icu::TimeZone::getGMT());
-		Formatter->setTimeZone( TimeZone );
+		Formatter->adoptTimeZone( bIsDefaultTimeZone ? icu::TimeZone::createDefault() :icu::TimeZone::createTimeZone(InputTimeZoneID) );
 		return Formatter;
 	}
 }
diff --git a/Engine/Source/Runtime/Core/Private/Internationalization/ICUCulture.h b/Engine/Source/Runtime/Core/Private/Internationalization/ICUCulture.h
--- a/Engine/Source/Runtime/Core/Private/Internationalization/ICUCulture.h
+++ b/Engine/Source/Runtime/Core/Private/Internationalization/ICUCulture.h
@@ -140,7 +140,7 @@ class FCulture::FICUCultureImplementation
 	TSharedRef<const icu::DecimalFormat> GetDecimalFormatter(const FNumberFormattingOptions* const Options = NULL) const;
 	TSharedRef<const icu::DecimalFormat> GetCurrencyFormatter(const FString& CurrencyCode = FString(), const FNumberFormattingOptions* const Options = NULL) const;
 	TSharedRef<const icu::DecimalFormat> GetPercentFormatter(const FNumberFormattingOptions* const Options = NULL) const;
-	TSharedRef<const icu::DateFormat> GetDateFormatter(const EDateTimeStyle::Type DateStyle) const;
+	TSharedRef<const icu::DateFormat> GetDateFormatter(const EDateTimeStyle::Type DateStyle, const FString& TimeZone) const;
 	TSharedRef<const icu::DateFormat> GetTimeFormatter(const EDateTimeStyle::Type TimeStyle, const FString& TimeZone) const;
 	TSharedRef<const icu::DateFormat> GetDateTimeFormatter(const EDateTimeStyle::Type DateStyle, const EDateTimeStyle::Type TimeStyle, const FString& TimeZone) const;
 
diff --git a/Engine/Source/Runtime/Core/Private/Internationalization/ICUText.cpp b/Engine/Source/Runtime/Core/Private/Internationalization/ICUText.cpp
--- a/Engine/Source/Runtime/Core/Private/Internationalization/ICUText.cpp
+++ b/Engine/Source/Runtime/Core/Private/Internationalization/ICUText.cpp
@@ -21,7 +21,7 @@ bool FText::IsWhitespace( const TCHAR Char )
 	return u_isWhitespace(ICUChar) != 0;
 }
 
-FText FText::AsDate(const FDateTime& DateTime, const EDateTimeStyle::Type DateStyle, const FCulturePtr& TargetCulture)
+FText FText::AsDate(const FDateTime& DateTime, const EDateTimeStyle::Type DateStyle, const FString& TimeZone, const FCulturePtr& TargetCulture)
 {
 	FInternationalization& I18N = FInternationalization::Get();
 	checkf(I18N.IsInitialized() == true, TEXT("FInternationalization is not initialized. An FText formatting method was likely used in static object initialization - this is not supported."));
@@ -31,7 +31,7 @@ FText FText::AsDate(const FDateTime& DateTime, const EDateTimeStyle::Type DateSt
 	UDate ICUDate = static_cast<double>(UNIXTimestamp) * U_MILLIS_PER_SECOND;
 
 	UErrorCode ICUStatus = U_ZERO_ERROR;
-	const TSharedRef<const icu::DateFormat> ICUDateFormat( Culture->Implementation->GetDateFormatter(DateStyle) );
+	const TSharedRef<const icu::DateFormat> ICUDateFormat( Culture->Implementation->GetDateFormatter(DateStyle, TimeZone) );
 	icu::UnicodeString FormattedString;
 	ICUDateFormat->format(ICUDate, FormattedString);
 
@@ -39,7 +39,7 @@ FText FText::AsDate(const FDateTime& DateTime, const EDateTimeStyle::Type DateSt
 	ICUUtilities::ConvertString(FormattedString, NativeString);
 
 	FText ResultText = FText::CreateChronologicalText(NativeString);
-	ResultText.History = MakeShareable(new FTextHistory_AsDate(DateTime, DateStyle, TargetCulture));
+	ResultText.History = MakeShareable(new FTextHistory_AsDate(DateTime, DateStyle, TimeZone, TargetCulture));
 
 	return ResultText;
 }
diff --git a/Engine/Source/Runtime/Core/Private/Internationalization/LegacyText.cpp b/Engine/Source/Runtime/Core/Private/Internationalization/LegacyText.cpp
--- a/Engine/Source/Runtime/Core/Private/Internationalization/LegacyText.cpp
+++ b/Engine/Source/Runtime/Core/Private/Internationalization/LegacyText.cpp
@@ -10,7 +10,7 @@ bool FText::IsWhitespace( const TCHAR Char )
 	return FChar::IsWhitespace(Char);
 }
 
-FText FText::AsDate( const FDateTime& DateTime, const EDateTimeStyle::Type DateStyle, const FCulturePtr& TargetCulture )
+FText FText::AsDate( const FDateTime& DateTime, const EDateTimeStyle::Type DateStyle, const FString& TimeZone, const FCulturePtr& TargetCulture )
 {
 	checkf(FInternationalization::Get().IsInitialized() == true, TEXT("FInternationalization is not initialized. An FText formatting method was likely used in static object initialization - this is not supported."));
 	return FText::FromString( DateTime.ToString( TEXT("%Y.%m.%d") ) );
diff --git a/Engine/Source/Runtime/Core/Private/Internationalization/TextHistory.cpp b/Engine/Source/Runtime/Core/Private/Internationalization/TextHistory.cpp
--- a/Engine/Source/Runtime/Core/Private/Internationalization/TextHistory.cpp
+++ b/Engine/Source/Runtime/Core/Private/Internationalization/TextHistory.cpp
@@ -465,9 +465,10 @@ void FTextHistory_AsCurrency::Serialize( FArchive& Ar )
 ///////////////////////////////////////
 // FTextHistory_AsDate
 
-FTextHistory_AsDate::FTextHistory_AsDate(const FDateTime& InSourceDateTime, const EDateTimeStyle::Type InDateStyle, const FCulturePtr InTargetCulture)
+FTextHistory_AsDate::FTextHistory_AsDate(const FDateTime& InSourceDateTime, const EDateTimeStyle::Type InDateStyle, const FString& InTimeZone, const FCulturePtr InTargetCulture)
 	: SourceDateTime(InSourceDateTime)
 	, DateStyle(InDateStyle)
+	, TimeZone(InTimeZone)
 	, TargetCulture(InTargetCulture)
 {
 }
@@ -486,6 +487,11 @@ void FTextHistory_AsDate::Serialize(FArchive& Ar)
 	Ar << DateStyleInt8;
 	DateStyle = (EDateTimeStyle::Type)DateStyleInt8;
 
+	if( Ar.UE4Ver() >= VER_UE4_FTEXT_HISTORY_DATE_TIMEZONE )
+	{
+		Ar << TimeZone;
+	}
+
 	if(Ar.IsSaving())
 	{
 		FString CultureName = TargetCulture.IsValid()? TargetCulture->GetName() : FString();
@@ -507,7 +513,7 @@ FText FTextHistory_AsDate::ToText(bool bInAsSource) const
 {
 	TSharedPtr< FCulture, ESPMode::ThreadSafe > CurrentCulture = bInAsSource? FInternationalization::Get().GetInvariantCulture() : TargetCulture;
 
-	return FText::AsDate(SourceDateTime, DateStyle, CurrentCulture);
+	return FText::AsDate(SourceDateTime, DateStyle, TimeZone, CurrentCulture);
 }
 ///////////////////////////////////////
 // FTextHistory_AsTime
diff --git a/Engine/Source/Runtime/Core/Private/Internationalization/TextHistory.h b/Engine/Source/Runtime/Core/Private/Internationalization/TextHistory.h
--- a/Engine/Source/Runtime/Core/Private/Internationalization/TextHistory.h
+++ b/Engine/Source/Runtime/Core/Private/Internationalization/TextHistory.h
@@ -198,7 +198,7 @@ class CORE_API FTextHistory_AsDate : public FTextHistory
 {
 public:
 	FTextHistory_AsDate() {};
-	FTextHistory_AsDate(const FDateTime& InSourceDateTime, const EDateTimeStyle::Type InDateStyle, const FCulturePtr InTargetCulture);
+	FTextHistory_AsDate(const FDateTime& InSourceDateTime, const EDateTimeStyle::Type InDateStyle, const FString& InTimeZone, const FCulturePtr InTargetCulture);
 
 	// Begin FTextHistory interface
 	virtual FText ToText(bool bInAsSource) const override;
@@ -210,6 +210,8 @@ private:
 	FDateTime SourceDateTime;
 	/** Style to format the date using */
 	EDateTimeStyle::Type DateStyle;
+	/** Timezone to put the time in */
+	FString TimeZone;
 	/** Culture to format the date in */
 	FCulturePtr TargetCulture;
 };
diff --git a/Engine/Source/Runtime/Core/Public/Internationalization/Text.h b/Engine/Source/Runtime/Core/Public/Internationalization/Text.h
--- a/Engine/Source/Runtime/Core/Public/Internationalization/Text.h
+++ b/Engine/Source/Runtime/Core/Public/Internationalization/Text.h
@@ -175,7 +175,7 @@ public:
 	/**
 	 * Generate an FText that represents the passed number as a date and/or time in the current culture
 	 */
-	static FText AsDate(const FDateTime& DateTime, const EDateTimeStyle::Type DateStyle = EDateTimeStyle::Default, const FCulturePtr& TargetCulture = NULL);
+	static FText AsDate(const FDateTime& DateTime, const EDateTimeStyle::Type DateStyle = EDateTimeStyle::Default, const FString& TimeZone = TEXT(""), const FCulturePtr& TargetCulture = NULL);
 	static FText AsDateTime(const FDateTime& DateTime, const EDateTimeStyle::Type DateStyle = EDateTimeStyle::Default, const EDateTimeStyle::Type TimeStyle = EDateTimeStyle::Default, const FString& TimeZone = TEXT(""), const FCulturePtr& TargetCulture = NULL);
 	static FText AsTime(const FDateTime& DateTime, const EDateTimeStyle::Type TimeStyle = EDateTimeStyle::Default, const FString& TimeZone = TEXT(""), const FCulturePtr& TargetCulture = NULL);
 	static FText AsTimespan(const FTimespan& Timespan, const FCulturePtr& TargetCulture = NULL);
diff --git a/Engine/Source/Runtime/Core/Public/UObject/ObjectVersion.h b/Engine/Source/Runtime/Core/Public/UObject/ObjectVersion.h
--- a/Engine/Source/Runtime/Core/Public/UObject/ObjectVersion.h
+++ b/Engine/Source/Runtime/Core/Public/UObject/ObjectVersion.h
@@ -657,6 +657,8 @@ enum EUnrealEngineObjectUE4Version
 	VER_UE4_ADD_PIVOT_TO_WIDGET_COMPONENT,
 	// Added finer control over when AI Pawns are automatically possessed. Also renamed Pawn.AutoPossess to Pawn.AutoPossessPlayer indicate this was a setting for players and not AI.
 	VER_UE4_PAWN_AUTO_POSSESS_AI,
+	// Added serialization of timezone to FTextHistory for AsDate operations.
+	VER_UE4_FTEXT_HISTORY_DATE_TIMEZONE,
 	// -----<new versions can be added before this line>-------------------------------------------------
 	// - this needs to be the last line (see note below)
 	VER_UE4_AUTOMATIC_VERSION_PLUS_ONE,
diff --git a/Engine/Source/Runtime/Engine/Private/KismetTextLibrary.cpp b/Engine/Source/Runtime/Engine/Private/KismetTextLibrary.cpp
--- a/Engine/Source/Runtime/Engine/Private/KismetTextLibrary.cpp
+++ b/Engine/Source/Runtime/Engine/Private/KismetTextLibrary.cpp
@@ -199,19 +199,19 @@ FText UKismetTextLibrary::AsPercent_Float(float Value, TEnumAsByte<ERoundingMode
 
 FText UKismetTextLibrary::AsDate_DateTime(const FDateTime& InDateTime)
 {
-	return FText::AsDate(InDateTime);
+	return FText::AsDate(InDateTime, EDateTimeStyle::Default, TEXT("Unknown"));
 }
 
 
 FText UKismetTextLibrary::AsDateTime_DateTime(const FDateTime& InDateTime)
 {
-	return FText::AsDateTime(InDateTime);
+	return FText::AsDateTime(InDateTime, EDateTimeStyle::Default, EDateTimeStyle::Default, TEXT("Unknown"));
 }
 
 
 FText UKismetTextLibrary::AsTime_DateTime(const FDateTime& InDateTime)
 {
-	return FText::AsTime(InDateTime);
+	return FText::AsTime(InDateTime, EDateTimeStyle::Default, TEXT("Unknown"));
 }
 
 
