- Removed handling of / commandline switches in FCommandLine::Parse(). It was not compatible with Mac/Linux and absolute paths on the commandline.

diff --git a/Engine/Source/Runtime/Core/Private/Misc/CoreMisc.cpp b/Engine/Source/Runtime/Core/Private/Misc/CoreMisc.cpp
--- a/Engine/Source/Runtime/Core/Private/Misc/CoreMisc.cpp
+++ b/Engine/Source/Runtime/Core/Private/Misc/CoreMisc.cpp
@@ -642,7 +642,7 @@ void FCommandLine::Parse(const TCHAR* InCmdLine, TArray<FString>& Tokens, TArray
 	FString NextToken;
 	while (FParse::Token(InCmdLine, NextToken, false))
 	{
-		if ((**NextToken == TCHAR('-')) || (**NextToken == TCHAR('/')))
+		if ((**NextToken == TCHAR('-')))
 		{
 			new(Switches) FString(NextToken.Mid(1));
 		}
