Fix old and incorrect comment

diff --git a/Engine/Source/Runtime/Core/Public/Misc/AutomationTest.h b/Engine/Source/Runtime/Core/Public/Misc/AutomationTest.h
--- a/Engine/Source/Runtime/Core/Public/Misc/AutomationTest.h
+++ b/Engine/Source/Runtime/Core/Public/Misc/AutomationTest.h
@@ -491,7 +491,7 @@ public:
 	void SetDeveloperDirectoryIncluded(const bool bInDeveloperDirectoryIncluded);
 
 	/**
-	 * Sets whether we want to convert old A2 AnimBlueprint content to new A3 AnimBlueprint
+	 * Sets whether the automation tests should include visual commandlet.
 	 */
 	void SetVisualCommandletFilter(const bool bInVisualCommandletFilterOn);
 
