log2 not available on all platforms
#ue4
#android

diff --git a/Engine/Source/Runtime/Core/Public/Math/UnrealMathFPU.h b/Engine/Source/Runtime/Core/Public/Math/UnrealMathFPU.h
--- a/Engine/Source/Runtime/Core/Public/Math/UnrealMathFPU.h
+++ b/Engine/Source/Runtime/Core/Public/Math/UnrealMathFPU.h
@@ -984,7 +984,7 @@ FORCEINLINE VectorRegister VectorLog(const VectorRegister& X)
 //TODO: Vectorize
 FORCEINLINE VectorRegister VectorLog2(const VectorRegister& X)
 {
-	return MakeVectorRegister((float)log2(VectorGetComponent(X, 0)), (float)log2(VectorGetComponent(X, 1)), (float)log2(VectorGetComponent(X, 2)), (float)log2(VectorGetComponent(X, 3)));
+	return MakeVectorRegister((float)FMath::Log2(VectorGetComponent(X, 0)), (float)FMath::Log2(VectorGetComponent(X, 1)), (float)FMath::Log2(VectorGetComponent(X, 2)), (float)FMath::Log2(VectorGetComponent(X, 3)));
 }
 
 //TODO: Vectorize
