Added size constraints for windows on Mac

diff --git a/Engine/Source/Runtime/Core/Private/Mac/MacWindow.cpp b/Engine/Source/Runtime/Core/Private/Mac/MacWindow.cpp
--- a/Engine/Source/Runtime/Core/Private/Mac/MacWindow.cpp
+++ b/Engine/Source/Runtime/Core/Private/Mac/MacWindow.cpp
@@ -140,6 +140,15 @@ void FMacWindow::Initialize( FMacApplication* const Application, const TSharedRe
 
 			[WindowHandle setOpaque: NO];
 
+			if (Definition->SizeLimits.GetMinWidth().IsSet() && Definition->SizeLimits.GetMinHeight().IsSet())
+			{
+				[WindowHandle setMinSize:NSMakeSize(Definition->SizeLimits.GetMinWidth().GetValue(), Definition->SizeLimits.GetMinHeight().GetValue())];
+			}
+			if (Definition->SizeLimits.GetMaxWidth().IsSet() && Definition->SizeLimits.GetMaxHeight().IsSet())
+			{
+				[WindowHandle setMaxSize:NSMakeSize(Definition->SizeLimits.GetMaxWidth().GetValue(), Definition->SizeLimits.GetMaxHeight().GetValue())];
+			}
+
 			ReshapeWindow( X, Y, SizeX, SizeY );
 
 			if( Definition->IsRegularWindow )
diff --git a/Engine/Source/Runtime/Core/Public/GenericPlatform/GenericWindowDefinition.h b/Engine/Source/Runtime/Core/Public/GenericPlatform/GenericWindowDefinition.h
--- a/Engine/Source/Runtime/Core/Public/GenericPlatform/GenericWindowDefinition.h
+++ b/Engine/Source/Runtime/Core/Public/GenericPlatform/GenericWindowDefinition.h
@@ -3,6 +3,7 @@
 #pragma once
 #include "Containers/UnrealString.h"
 #include "HAL/Platform.h"
+#include "GenericApplicationMessageHandler.h"
 
 struct CORE_API FGenericWindowDefinition
 {
@@ -32,4 +33,6 @@ struct CORE_API FGenericWindowDefinition
 	FString Title;
 	float Opacity;
 	int32 CornerRadius;
+
+	FWindowSizeLimits SizeLimits;
 };
\ No newline at end of file
diff --git a/Engine/Source/Runtime/Slate/Private/Framework/Application/SlateApplication.cpp b/Engine/Source/Runtime/Slate/Private/Framework/Application/SlateApplication.cpp
--- a/Engine/Source/Runtime/Slate/Private/Framework/Application/SlateApplication.cpp
+++ b/Engine/Source/Runtime/Slate/Private/Framework/Application/SlateApplication.cpp
@@ -1438,6 +1438,8 @@ TSharedRef< FGenericWindow > FSlateApplication::MakeWindow( TSharedRef<SWindow>
 	Definition->Opacity = InSlateWindow->GetOpacity();
 	Definition->CornerRadius = InSlateWindow->GetCornerRadius();
 
+	Definition->SizeLimits = InSlateWindow->GetSizeLimits();
+
 	TSharedRef< FGenericWindow > NewWindow = PlatformApplication->MakeWindow();
 
 	InSlateWindow->SetNativeWindow( NewWindow );
