Changed Mac window size constrains setting code to allow per-axis control

diff --git a/Engine/Source/Runtime/Core/Private/Mac/MacWindow.cpp b/Engine/Source/Runtime/Core/Private/Mac/MacWindow.cpp
--- a/Engine/Source/Runtime/Core/Private/Mac/MacWindow.cpp
+++ b/Engine/Source/Runtime/Core/Private/Mac/MacWindow.cpp
@@ -140,14 +140,8 @@ void FMacWindow::Initialize( FMacApplication* const Application, const TSharedRe
 
 			[WindowHandle setOpaque: NO];
 
-			if (Definition->SizeLimits.GetMinWidth().IsSet() && Definition->SizeLimits.GetMinHeight().IsSet())
-			{
-				[WindowHandle setMinSize:NSMakeSize(Definition->SizeLimits.GetMinWidth().GetValue(), Definition->SizeLimits.GetMinHeight().GetValue())];
-			}
-			if (Definition->SizeLimits.GetMaxWidth().IsSet() && Definition->SizeLimits.GetMaxHeight().IsSet())
-			{
-				[WindowHandle setMaxSize:NSMakeSize(Definition->SizeLimits.GetMaxWidth().GetValue(), Definition->SizeLimits.GetMaxHeight().GetValue())];
-			}
+			[WindowHandle setMinSize:NSMakeSize(Definition->SizeLimits.GetMinWidth().Get(10.0f), Definition->SizeLimits.GetMinHeight().Get(10.0f))];
+			[WindowHandle setMaxSize:NSMakeSize(Definition->SizeLimits.GetMaxWidth().Get(10000.0f), Definition->SizeLimits.GetMaxHeight().Get(10000.0f))];
 
 			ReshapeWindow( X, Y, SizeX, SizeY );
 
