Fix config files being unable to save urls properly

diff --git a/Engine/Source/Runtime/Core/Private/Misc/ConfigCacheIni.cpp b/Engine/Source/Runtime/Core/Private/Misc/ConfigCacheIni.cpp
--- a/Engine/Source/Runtime/Core/Private/Misc/ConfigCacheIni.cpp
+++ b/Engine/Source/Runtime/Core/Private/Misc/ConfigCacheIni.cpp
@@ -551,8 +551,8 @@ void FConfigFile::Read( const FString& Filename )
 bool FConfigFile::ShouldExportQuotedString(const FString& PropertyValue) const
 {
 	// The value should be exported as quoted string if it begins with a space (which is stripped on import) or
-	// when it contains '//' (interpreted as a comment when importing).
-	return **PropertyValue == TEXT(' ') || FCString::Strstr(*PropertyValue, TEXT("//")) != NULL;
+	// when it contains ' //' (interpreted as a comment when importing) or starts with '//'
+	return **PropertyValue == TEXT(' ') || FCString::Strstr(*PropertyValue, TEXT(" //")) != NULL || PropertyValue.StartsWith(TEXT("//"));
 }
 
 
