Improve FIntPoint's hash, FMemCrc_Deprecated was showing up on profiles of the landscape editing code...

diff --git a/Engine/Source/Runtime/Core/Public/Math/IntPoint.h b/Engine/Source/Runtime/Core/Public/Math/IntPoint.h
--- a/Engine/Source/Runtime/Core/Public/Math/IntPoint.h
+++ b/Engine/Source/Runtime/Core/Public/Math/IntPoint.h
@@ -454,9 +454,9 @@ FORCEINLINE int32 FIntPoint::GetMin( ) const
 	return FMath::Min(X,Y);
 }
 
-FORCEINLINE uint32 GetTypeHash( const FIntPoint& InPoint )
+FORCEINLINE uint32 GetTypeHash(const FIntPoint& InPoint)
 {
-	return FCrc::MemCrc_DEPRECATED(&InPoint, sizeof(FIntPoint));
+	return HashCombine(GetTypeHash(InPoint.X), GetTypeHash(InPoint.Y));
 }
 
 
