Fix HTML5-win32 build.

diff --git a/Engine/Source/Runtime/Core/Public/HTML5/HTML5PlatformCompilerPreSetup.h b/Engine/Source/Runtime/Core/Public/HTML5/HTML5PlatformCompilerPreSetup.h
--- a/Engine/Source/Runtime/Core/Public/HTML5/HTML5PlatformCompilerPreSetup.h
+++ b/Engine/Source/Runtime/Core/Public/HTML5/HTML5PlatformCompilerPreSetup.h
@@ -4,6 +4,8 @@
 
 
 #ifndef DISABLE_DEPRECATION
+#if PLATFORM_HTML5_BROWSER
+
 	#pragma clang diagnostic warning "-Wdeprecated-declarations"
 
 	/**
@@ -34,4 +36,5 @@
 	#define PRAGMA_ENABLE_DEPRECATION_WARNINGS \
 		_Pragma("clang diagnostic pop")
 
+#endif 
 #endif // DISABLE_DEPRECATION
