Gamepad for HTML5.

diff --git a/Engine/Source/Runtime/Core/Private/HTML5/HTML5InputInterface.cpp b/Engine/Source/Runtime/Core/Private/HTML5/HTML5InputInterface.cpp
--- a/Engine/Source/Runtime/Core/Private/HTML5/HTML5InputInterface.cpp
+++ b/Engine/Source/Runtime/Core/Private/HTML5/HTML5InputInterface.cpp
@@ -8,6 +8,52 @@
 DECLARE_LOG_CATEGORY_EXTERN(LogHTML5, Log, All);
 DEFINE_LOG_CATEGORY(LogHTML5);
 
+
+#if PLATFORM_HTML5_BROWSER
+#include "emscripten.h"
+#include "html5.h"
+
+static EControllerButtons::Type AxisMapping[4] =
+{
+	EControllerButtons::LeftAnalogX,
+	EControllerButtons::LeftAnalogY, 
+	EControllerButtons::RightAnalogX,
+	EControllerButtons::RightAnalogY
+};
+
+// Axis Mapping, reversed or not. 
+static int Reversed[4] =
+{
+	1,
+	-1, 
+	1,
+	-1
+};
+
+// all are digital except Left and Right Trigger Analog. 
+static EControllerButtons::Type  ButtonMapping[16] = 
+{
+	EControllerButtons::FaceButtonBottom,
+	EControllerButtons::FaceButtonRight,
+	EControllerButtons::FaceButtonLeft,
+	EControllerButtons::FaceButtonTop,
+	EControllerButtons::LeftShoulder,
+	EControllerButtons::RightShoulder,
+	EControllerButtons::LeftTriggerThreshold, 
+	EControllerButtons::RightTriggerThreshold,
+	EControllerButtons::SpecialLeft, 
+	EControllerButtons::SpecialRight,
+	EControllerButtons::LeftStickDown,
+	EControllerButtons::RightStickDown,
+	EControllerButtons::DPadUp,
+	EControllerButtons::DPadDown,
+	EControllerButtons::DPadLeft,
+	EControllerButtons::DPadRight
+};
+
+#endif 
+
+
 TSharedRef< FHTML5InputInterface > FHTML5InputInterface::Create(  const TSharedRef< FGenericApplicationMessageHandler >& InMessageHandler, const TSharedPtr< ICursor >& InCursor )
 {
 	return MakeShareable( new FHTML5InputInterface( InMessageHandler, InCursor ) );
@@ -206,4 +252,76 @@ void FHTML5InputInterface::Tick(float DeltaTime, const SDL_Event& Event,TSharedR
 
 void FHTML5InputInterface::SendControllerEvents()
 {
+#if PLATFORM_HTML5_BROWSER
+	// game pads can only be polled.
+	static int PrevNumGamepads = 0;
+	static bool GamepadSupported = true;
+
+	const double CurrentTime = FPlatformTime::Seconds();
+
+	if (GamepadSupported)
+	{
+		int NumGamepads = emscripten_get_num_gamepads();
+		if (NumGamepads != PrevNumGamepads)
+		{
+			if (NumGamepads == EMSCRIPTEN_RESULT_NOT_SUPPORTED)
+			{
+				GamepadSupported = false;
+				return;
+			}
+		}
+
+		for(int CurrentGamePad = 0; CurrentGamePad < NumGamepads && CurrentGamePad < 5; ++CurrentGamePad) // max 5 game pads.
+		{
+			EmscriptenGamepadEvent GamePadEvent;
+			int Failed = emscripten_get_gamepad_status(CurrentGamePad, &GamePadEvent);
+			if (!Failed)
+			{
+				check(CurrentGamePad == GamePadEvent.index);
+				for(int CurrentAxis = 0; CurrentAxis < GamePadEvent.numAxes; ++CurrentAxis)
+				{
+					if (GamePadEvent.axis[CurrentAxis] != PrevGamePadState[CurrentGamePad].axis[CurrentAxis])
+					{
+
+						MessageHandler->OnControllerAnalog(AxisMapping[CurrentAxis],CurrentGamePad, Reversed[CurrentAxis]*GamePadEvent.axis[CurrentAxis]);
+					}
+				}
+				// edge trigger. 
+				for(int CurrentButton = 0; CurrentButton < GamePadEvent.numButtons; ++CurrentButton)
+				{
+					// trigger for digital buttons. 
+					if ( GamePadEvent.digitalButton[CurrentButton]   != PrevGamePadState[CurrentGamePad].digitalButton[CurrentButton] )
+					{
+						bool Triggered = GamePadEvent.digitalButton[CurrentButton] != 0;
+						if ( Triggered )
+						{
+							MessageHandler->OnControllerButtonPressed(ButtonMapping[CurrentButton],CurrentGamePad, false); 
+							LastPressedTime[CurrentGamePad][CurrentButton] = CurrentTime;
+						}
+						else 
+						{
+							MessageHandler->OnControllerButtonReleased(ButtonMapping[CurrentButton],CurrentGamePad, false); 
+						}
+					}
+				}
+				// repeat trigger. 
+				const float RepeatDelta = 0.2f; 
+				for(int CurrentButton = 0; CurrentButton < GamePadEvent.numButtons; ++CurrentButton)
+				{
+					if ( GamePadEvent.digitalButton[CurrentButton]  )
+					{
+						if (CurrentTime - LastPressedTime[CurrentGamePad][CurrentButton] > RepeatDelta) 
+						{
+							MessageHandler->OnControllerButtonPressed(ButtonMapping[CurrentButton],CurrentGamePad, true); 
+							LastPressedTime[CurrentGamePad][CurrentButton] = CurrentTime; 
+						}
+					}
+				}
+				PrevGamePadState[CurrentGamePad] = GamePadEvent;
+			}
+		}
+	}
+
+#endif
+
 }
diff --git a/Engine/Source/Runtime/Core/Public/HTML5/HTML5InputInterface.h b/Engine/Source/Runtime/Core/Public/HTML5/HTML5InputInterface.h
--- a/Engine/Source/Runtime/Core/Public/HTML5/HTML5InputInterface.h
+++ b/Engine/Source/Runtime/Core/Public/HTML5/HTML5InputInterface.h
@@ -3,6 +3,10 @@
 #pragma once
 #include <SDL.h>
 
+#if PLATFORM_HTML5_BROWSER
+#include <html5.h>
+#endif 
+
 
 /**
  * Interface class for HTML5 input devices                 
@@ -38,4 +42,12 @@ private:
 	const TSharedPtr< ICursor > Cursor;
 
 	TBitArray<FDefaultBitArrayAllocator> KeyStates;
+
+#if PLATFORM_HTML5_BROWSER
+
+	EmscriptenGamepadEvent PrevGamePadState[5]; 
+	double LastPressedTime[5][15];
+
+#endif 
+
 };
\ No newline at end of file
