Improve Android startup load times
#android
#ue4

diff --git a/Engine/Build/Android/Java/src/com/epicgames/ue4/GameActivity.java b/Engine/Build/Android/Java/src/com/epicgames/ue4/GameActivity.java
--- a/Engine/Build/Android/Java/src/com/epicgames/ue4/GameActivity.java
+++ b/Engine/Build/Android/Java/src/com/epicgames/ue4/GameActivity.java
@@ -744,6 +744,11 @@ public class GameActivity extends NativeActivity
 		return JavaBuildSettings.PackageType.AMAZON == JavaBuildSettings.PACKAGING;
 	}
 
+	public static String GetOBBFilename()
+	{
+		return JavaBuildSettings.OBBFILENAME;
+	}
+
 	public void AndroidThunkJava_Minimize()
 	{
 		Intent startMain = new Intent(Intent.ACTION_MAIN);
@@ -780,7 +785,7 @@ public class GameActivity extends NativeActivity
 
 		for ( String fontdir : fontdirs )
         {
-			Log.debug(fontdir);
+//			Log.debug(fontdir);
             File dir = new File( fontdir );
 
 			if(dir.exists())
diff --git a/Engine/Source/Programs/UnrealBuildTool/Android/UEDeployAndroid.cs b/Engine/Source/Programs/UnrealBuildTool/Android/UEDeployAndroid.cs
--- a/Engine/Source/Programs/UnrealBuildTool/Android/UEDeployAndroid.cs
+++ b/Engine/Source/Programs/UnrealBuildTool/Android/UEDeployAndroid.cs
@@ -261,31 +261,47 @@ namespace UnrealBuildTool.Android
             return Path.Combine(GetUE4JavaFilePath(EngineDirectory), "JavaBuildSettings.java");
         }
 
-        public void WriteJavaBuildSettingsFile(string FileName, bool OBBinAPK)
+        public void WriteJavaBuildSettingsFile(string FileName, bool OBBinAPK, string CookFlavor)
         {
              // (!UEBuildConfiguration.bOBBinAPK ? "PackageType.AMAZON" : /*bPackageForGoogle ? "PackageType.GOOGLE" :*/ "PackageType.DEVELOPMENT") + ";\n");
             string Setting = OBBinAPK ? "AMAZON" : "DEVELOPMENT";
-            if (!File.Exists(FileName) || ShouldWriteJavaBuildSettingsFile(FileName, Setting))
+            string OBBFilename = "Android" + CookFlavor + ".obb" + (OBBinAPK ? ".png" : "");
+            if (!File.Exists(FileName) || ShouldWriteJavaBuildSettingsFile(FileName, Setting, OBBFilename))
             {
 				Log.TraceInformation("\n===={0}====WRITING JAVABUILDSETTINGS.JAVA========================================================", DateTime.Now.ToString());
 				StringBuilder BuildSettings = new StringBuilder("package com.epicgames.ue4;\npublic class JavaBuildSettings\n{\n");
                 BuildSettings.Append("\tpublic enum PackageType {AMAZON, GOOGLE, DEVELOPMENT};\n");
                 BuildSettings.Append("\tpublic static final PackageType PACKAGING = PackageType." + Setting + ";\n");
+                BuildSettings.Append("\tpublic static final String OBBFILENAME = \"" + OBBFilename + "\";\n");                    
                 BuildSettings.Append("}\n");
                 File.WriteAllText(FileName, BuildSettings.ToString());
             }
         }
 
-        public bool ShouldWriteJavaBuildSettingsFile(string FileName, string setting)
+        public bool ShouldWriteJavaBuildSettingsFile(string FileName, string setting, string OBBFilename)
         {
             var fileContent = File.ReadAllLines(FileName);
-            if (fileContent.Length < 5)
-                return true;
-            var packageLine = fileContent[4]; // We know this to be true... because we write it below...
-            int location = packageLine.IndexOf("PACKAGING") + 12 + 12; // + ("PACKAGING = ") + ("PackageType.")
-            if (location == -1)
-                return true;
-            return String.Compare(setting, packageLine.Substring(location, Math.Min(packageLine.Length - location, setting.Length))) != 0;
+
+            // compare packaging type and OBB filenames
+            int location;
+            bool packageTypeMatch = false;
+            bool obbFilenameMatch = false;
+            foreach (var fileLine in fileContent)
+            {
+                if ((location = fileLine.IndexOf("PACKAGING = ")) != -1)
+                {
+                    location += 12 + 12; // + ("PACKAGING = ") + ("PackageType.")
+                    packageTypeMatch = String.Compare(setting, fileLine.Substring(location, Math.Min(fileLine.Length - location - 1, setting.Length))) == 0;
+                }
+                else if ((location = fileLine.IndexOf("OBBFILENAME = ")) != -1)
+                {
+                    location += 15; // + ("OBBFILENAME = \"")
+                    obbFilenameMatch = String.Compare(OBBFilename, fileLine.Substring(location, Math.Min(fileLine.Length - location - 2, OBBFilename.Length))) == 0;
+                }
+            }
+
+            // write required if both don't match
+            return !(packageTypeMatch && obbFilenameMatch);
         }
 
 		private static string GetNDKArch(string UE4Arch)
@@ -585,7 +601,7 @@ namespace UnrealBuildTool.Android
 
             // See if we need to create a 'default' Java Build settings file if one doesn't exist (if it does exist we have to assume it has been setup correctly)
             string UE4JavaBuildSettingsFileName = GetUE4JavaBuildSettingsFileName(EngineDirectory);
-            WriteJavaBuildSettingsFile(UE4JavaBuildSettingsFileName, UEBuildConfiguration.bOBBinAPK);
+            WriteJavaBuildSettingsFile(UE4JavaBuildSettingsFileName, UEBuildConfiguration.bOBBinAPK, CookFlavor);
 
 			// check to see if any "meta information" is newer than last time we build
 			string CurrentBuildSettings = GetAllBuildSettings(UE4BuildPath, bForDistribution, bMakeSeparateApks, UEBuildConfiguration.bOBBinAPK);
diff --git a/Engine/Source/Runtime/Core/Private/Android/AndroidFile.cpp b/Engine/Source/Runtime/Core/Private/Android/AndroidFile.cpp
--- a/Engine/Source/Runtime/Core/Private/Android/AndroidFile.cpp
+++ b/Engine/Source/Runtime/Core/Private/Android/AndroidFile.cpp
@@ -43,6 +43,9 @@ FString GFontPathBase;
 // Is the OBB in an APK file or not
 bool GOBBinAPK;
 
+// Filename of the OBB (used for OBBinAPK case to avoid close search)
+FString GOBBFilename;
+
 extern jobject AndroidJNI_GetJavaAssetManager();
 extern AAssetManager * AndroidThunkCpp_GetAssetManager();
 
@@ -659,8 +662,11 @@ public:
 		{
 			// Since we control the APK we mount any OBBs we find
 			// inside of the APK.
-			FMountOBBVisitor MountOBB(*this);
-			IterateDirectory(TEXT(""), MountOBB);
+			//FMountOBBVisitor MountOBB(*this);
+			//IterateDirectory(TEXT(""), MountOBB, false, true);
+
+			// Getting the OBB to mount faster on some devices (IterateDirectory above can be slow!)
+			MountOBB(*GOBBFilename);
 		}
 		else
 		{
@@ -706,15 +712,13 @@ public:
 		PathToAndroidPaths(LocalPath, AssetPath, Filename, AllowLocal);
 
 		bool result = false;
-		if (IsLocal(LocalPath))
+		struct stat FileInfo;
+		if (!LocalPath.IsEmpty() && (stat(TCHAR_TO_UTF8(*LocalPath), &FileInfo) == 0))
 		{
 			// For local files we need to check if it's a plain
 			// file, as opposed to directories.
-			struct stat FileInfo;
-			if (stat(TCHAR_TO_UTF8(*LocalPath), &FileInfo) == 0)
-			{
-				result = S_ISREG(FileInfo.st_mode);
-			}
+			result = S_ISREG(FileInfo.st_mode);
+			return result;
 		}
 		else
 		{
@@ -743,16 +747,14 @@ public:
 
 		struct stat FileInfo;
 		FileInfo.st_size = -1;
-		if (IsLocal(LocalPath))
+		if (!LocalPath.IsEmpty() && (stat(TCHAR_TO_UTF8(*LocalPath), &FileInfo) == 0))
 		{
-			if (stat(TCHAR_TO_UTF8(*LocalPath), &FileInfo) == 0)
+			// make sure to return -1 for directories
+			if (S_ISDIR(FileInfo.st_mode))
 			{
-				// make sure to return -1 for directories
-				if (S_ISDIR(FileInfo.st_mode))
-				{
-					FileInfo.st_size = -1;
-				}
+				FileInfo.st_size = -1;
 			}
+			return FileInfo.st_size;
 		}
 		else if (IsResource(AssetPath))
 		{
@@ -1075,10 +1077,10 @@ public:
 
 	virtual bool DirectoryExists(const TCHAR* Directory) override
 	{
-		return DirectoryExists(Directory, false);
+		return DirectoryExists(Directory, false, false);
 	}
 	
-	bool DirectoryExists(const TCHAR* Directory, bool AllowLocal)
+	bool DirectoryExists(const TCHAR* Directory, bool AllowLocal, bool AllowAsset)
 	{
 #if LOG_ANDROID_FILE
 		FPlatformMisc::LowLevelOutputDebugStringf(TEXT("FAndroidPlatformFile::DirectoryExists('%s')"), Directory);
@@ -1103,8 +1105,9 @@ public:
 			FPlatformMisc::LowLevelOutputDebugStringf(TEXT("FAndroidPlatformFile::DirectoryExists('%s') => Found as resource: '%s'"), Directory, *(AssetPath + "/"));
 #endif
 		}
-		else
+		else if (AllowAsset)
 		{
+			// This can be very slow on some devices so only do it if we enabled
 			AAssetDir * dir = AAssetManager_openDir(AssetMgr, TCHAR_TO_UTF8(*AssetPath));
 			Found = AAssetDir_getNextFileName(dir) != nullptr;
 			AAssetDir_close(dir);
@@ -1166,10 +1169,10 @@ public:
 
 	virtual bool IterateDirectory(const TCHAR* Directory, FDirectoryVisitor& Visitor) override
 	{
-		return IterateDirectory(Directory, Visitor, false);
+		return IterateDirectory(Directory, Visitor, false, false);
 	}
 
-	bool IterateDirectory(const TCHAR* Directory, FDirectoryVisitor& Visitor, bool AllowLocal)
+	bool IterateDirectory(const TCHAR* Directory, FDirectoryVisitor& Visitor, bool AllowLocal, bool AllowAsset)
 	{
 #if LOG_ANDROID_FILE
 		FPlatformMisc::LowLevelOutputDebugStringf(TEXT("FAndroidPlatformFile::IterateDirectory('%s')"), Directory);
@@ -1191,7 +1194,7 @@ public:
 					{
 						FString DirPath = FString(Directory) / UTF8_TO_TCHAR(Entry->d_name);
 #if LOG_ANDROID_FILE
-						FPlatformMisc::LowLevelOutputDebugStringf(TEXT("FAndroidPlatformFile::IterateDirectory('%s').. Visit: '%s'"), Directory, *DirPath);
+						FPlatformMisc::LowLevelOutputDebugStringf(TEXT("FAndroidPlatformFile::IterateDirectory('%s').. LOCAL Visit: '%s'"), Directory, *DirPath);
 #endif
 						Result = Visitor.Visit(*DirPath, Entry->d_type == DT_DIR);
 					}
@@ -1207,7 +1210,7 @@ public:
 			while (Result && ResourceDir.Next())
 			{
 #if LOG_ANDROID_FILE
-				FPlatformMisc::LowLevelOutputDebugStringf(TEXT("FAndroidPlatformFile::IterateDirectory('%s').. Visit: '%s'"), Directory, *ResourceDir.Current.Key());
+				FPlatformMisc::LowLevelOutputDebugStringf(TEXT("FAndroidPlatformFile::IterateDirectory('%s').. RESOURCE Visit: '%s'"), Directory, *ResourceDir.Current.Key());
 #endif
 				Result = Visitor.Visit(*ResourceDir.Current.Key(), false);
 			}
@@ -1219,12 +1222,12 @@ public:
 			while (Result && ResourceDir.Next())
 			{
 #if LOG_ANDROID_FILE
-				FPlatformMisc::LowLevelOutputDebugStringf(TEXT("FAndroidPlatformFile::IterateDirectory('%s').. Visit: '%s'"), Directory, *ResourceDir.Current.Key());
+				FPlatformMisc::LowLevelOutputDebugStringf(TEXT("FAndroidPlatformFile::IterateDirectory('%s').. RESOURCE/ Visit: '%s'"), Directory, *ResourceDir.Current.Key());
 #endif
 				Result = Visitor.Visit(*ResourceDir.Current.Key(), false);
 			}
 		}
-		else
+		else if (AllowAsset)
 		{
 			AAssetDir * dir = AAssetManager_openDir(AssetMgr, TCHAR_TO_UTF8(*AssetPath));
 			if (nullptr != dir)
@@ -1241,7 +1244,7 @@ public:
 						AAssetDir_close(dir);
 					}
 #if LOG_ANDROID_FILE
-					FPlatformMisc::LowLevelOutputDebugStringf(TEXT("FAndroidPlatformFile::IterateDirectory('%s').. Visit: '%s'"), Directory, UTF8_TO_TCHAR(fileName));
+					FPlatformMisc::LowLevelOutputDebugStringf(TEXT("FAndroidPlatformFile::IterateDirectory('%s').. ASSET Visit: '%s'"), Directory, UTF8_TO_TCHAR(fileName));
 #endif
 					Result = Visitor.Visit(UTF8_TO_TCHAR(fileName), isDirectory);
 				}
diff --git a/Engine/Source/Runtime/Launch/Private/Android/AndroidJNI.cpp b/Engine/Source/Runtime/Launch/Private/Android/AndroidJNI.cpp
--- a/Engine/Source/Runtime/Launch/Private/Android/AndroidJNI.cpp
+++ b/Engine/Source/Runtime/Launch/Private/Android/AndroidJNI.cpp
@@ -20,6 +20,7 @@ static IVirtualKeyboardEntry *VirtualKeyboardWidget = NULL;
 extern FString GFilePathBase;
 extern FString GFontPathBase;
 extern bool GOBBinAPK;
+extern FString GOBBFilename;
 
 //////////////////////////////////////////////////////////////////////////
 
@@ -481,6 +482,13 @@ JNIEXPORT jint JNI_OnLoad(JavaVM* InJavaVM, void* InReserved)
 	jmethodID isOBBInAPKMethod = env->GetStaticMethodID(JDef_GameActivity::ClassID, "isOBBInAPK", "()Z");
 	GOBBinAPK = (bool)env->CallStaticBooleanMethod(JDef_GameActivity::ClassID, isOBBInAPKMethod, nullptr);
 
+	// Get the OBB filename
+	jmethodID GetOBBFilenameMethod = env->GetStaticMethodID(JDef_GameActivity::ClassID, "GetOBBFilename", "()Ljava/lang/String;");
+	jstring obbString = (jstring)env->CallStaticObjectMethod(JDef_GameActivity::ClassID, GetOBBFilenameMethod, nullptr);
+	const char * nativeObbString = env->GetStringUTFChars(obbString, 0);
+	GOBBFilename = FString(nativeObbString);
+	env->ReleaseStringUTFChars(obbString, nativeObbString);
+
 	// Get the system font directory
 	jstring fontPath = (jstring)env->CallStaticObjectMethod(JDef_GameActivity::ClassID, JDef_GameActivity::AndroidThunkJava_GetFontDirectory);
 	const char * nativeFontPathString = env->GetStringUTFChars(fontPath, 0);
