Treat non-foreign projects correctly when generating project files for new projects.

diff --git a/Engine/Source/Developer/DesktopPlatform/Private/DesktopPlatformBase.cpp b/Engine/Source/Developer/DesktopPlatform/Private/DesktopPlatformBase.cpp
--- a/Engine/Source/Developer/DesktopPlatform/Private/DesktopPlatformBase.cpp
+++ b/Engine/Source/Developer/DesktopPlatform/Private/DesktopPlatformBase.cpp
@@ -427,8 +427,8 @@ bool FDesktopPlatformBase::GenerateProjectFiles(const FString& RootDir, const FS
 	FString Arguments = TEXT("-projectfiles");
 #endif
 
-	// Build the arguments to pass to UBT
-	if ( !ProjectFileName.IsEmpty() )
+	// Build the arguments to pass to UBT. If it's a non-foreign project, just build full project files.
+	if ( !ProjectFileName.IsEmpty() && GetCachedProjectDictionary(RootDir).IsForeignProject(ProjectFileName) )
 	{
 		// Figure out whether it's a foreign project
 		const FUProjectDictionary &ProjectDictionary = GetCachedProjectDictionary(RootDir);
diff --git a/Engine/Source/Runtime/Core/Private/Misc/UProjectInfo.cpp b/Engine/Source/Runtime/Core/Private/Misc/UProjectInfo.cpp
--- a/Engine/Source/Runtime/Core/Private/Misc/UProjectInfo.cpp
+++ b/Engine/Source/Runtime/Core/Private/Misc/UProjectInfo.cpp
@@ -9,69 +9,58 @@ DEFINE_LOG_CATEGORY(LogUProjectInfo);
 
 FUProjectDictionary::FUProjectDictionary(const FString& InRootDir)
 {
-	FString RootDir = FPaths::ConvertRelativePathToFull(InRootDir);
+	RootDir = InRootDir;
+	Refresh();
+}
 
-	TArray<FString> DirectoriesToSearch;
-	
-	// Find all the .uprojectdirs files contained in the root folder and add their entries to the search array
-	TArray<FString> UProjectDirsFiles;
-	FString SearchPath = RootDir / FString(TEXT("*.uprojectdirs"));
-	IFileManager::Get().FindFiles(UProjectDirsFiles, *SearchPath, true, false);
+void FUProjectDictionary::Refresh()
+{
+	ProjectRootDirs.Reset();
+	ShortProjectNameDictionary.Reset();
 
-	for (int32 FileIndex = 0; FileIndex < UProjectDirsFiles.Num(); ++FileIndex)
-	{
-		FString ProjDirsFile = RootDir / UProjectDirsFiles[FileIndex];
+	// Find all the .uprojectdirs files contained in the root folder
+	TArray<FString> ProjectDirsFiles;
+	IFileManager::Get().FindFiles(ProjectDirsFiles, *(RootDir / FString(TEXT("*.uprojectdirs"))), true, false);
 
-		TArray<FString> FileStrings;
-		if (FFileHelper::LoadANSITextFileToStrings(*ProjDirsFile, &IFileManager::Get(), FileStrings))
+	// Add all the project root directories to ProjectRootDirs
+	for(const FString& ProjectDirsFile: ProjectDirsFiles)
+	{
+		TArray<FString> Lines;
+		if(FFileHelper::LoadANSITextFileToStrings(*(RootDir / ProjectDirsFile), &IFileManager::Get(), Lines))
 		{
-			for (int32 EntryIdx = 0; EntryIdx < FileStrings.Num(); EntryIdx++)
+			for(const FString& Line: Lines)
 			{
-				FString ProjDirEntry = FileStrings[EntryIdx];
-				ProjDirEntry.Trim();
-				if (ProjDirEntry.IsEmpty() == false)
+				FString Entry = FString(Line).Trim();
+				if(!Entry.IsEmpty() && !Entry.StartsWith(";"))
 				{
-					if (ProjDirEntry.StartsWith(";"))
-					{
-						// Commented out line... skip it
-						continue;
-					}
-					else
-					{
-						DirectoriesToSearch.Add(RootDir / ProjDirEntry);
-					}
+					FString DirectoryName = FPaths::ConvertRelativePathToFull(RootDir, Entry);
+					FPaths::NormalizeDirectoryName(DirectoryName);
+					ProjectRootDirs.AddUnique(DirectoryName);
 				}
 			}
 		}
 	}
 
-	for (int32 DirIdx = 0; DirIdx < DirectoriesToSearch.Num(); DirIdx++)
+	// Search for all the projects under each root directory
+	for(const FString& ProjectRootDir: ProjectRootDirs)
 	{
-		FString DirToSearch = DirectoriesToSearch[DirIdx];
-
-		// To specific the <UE4> folder itself, './' is used... strip it out if present
-		DirToSearch = DirToSearch.Replace(TEXT("/./"), TEXT("/"));
+		// Enumerate the subdirectories
+		TArray<FString> ProjectDirs;
+		IFileManager::Get().FindFiles(ProjectDirs, *(ProjectRootDir / FString(TEXT("*"))), false, true);
 
-		// Get each sub-directory of this folder and check it for uproject files
-		TArray<FString> SubDirectories;
-		SearchPath = DirToSearch / FString(TEXT("*"));
-		IFileManager::Get().FindFiles(SubDirectories, *SearchPath, false, true);
-		
-		for (int32 SubIdx = 0; SubIdx < SubDirectories.Num(); SubIdx++)
+		// Check each one for project files
+		for(const FString& ProjectDir: ProjectDirs)
 		{
-			FString SubDir = DirToSearch / SubDirectories[SubIdx];
+			// Enumerate all the project files
+			TArray<FString> ProjectFiles;
+			IFileManager::Get().FindFiles(ProjectFiles, *(ProjectRootDir / ProjectDir / TEXT("*.uproject")), true, false);
 
-			TArray<FString> SubDirFiles;
-			SearchPath = SubDir / TEXT("*.uproject");
-			IFileManager::Get().FindFiles(SubDirFiles, *SearchPath, true, false);
-
-			for (int32 UProjIdx = 0; UProjIdx < SubDirFiles.Num(); UProjIdx++)
+			// Add all the projects to the dictionary
+			for(const FString& ProjectFile: ProjectFiles)
 			{
-				FString ProjectFileName = SubDir / SubDirFiles[UProjIdx];
-				FPaths::NormalizeFilename(ProjectFileName);
-
-				FString ShortName = FPaths::GetBaseFilename(ProjectFileName);
-				ShortProjectNameDictionary.Add(ShortName, ProjectFileName);
+				FString ShortName = FPaths::GetBaseFilename(ProjectFile);
+				FString FullProjectFile = ProjectRootDir / ProjectDir / ProjectFile;
+				ShortProjectNameDictionary.Add(ShortName, FullProjectFile);
 			}
 		}
 	}
@@ -81,6 +70,7 @@ bool FUProjectDictionary::IsForeignProject(const FString& InProjectFileName) con
 {
 	FString ProjectFileName = FPaths::ConvertRelativePathToFull(InProjectFileName);
 
+	// Check if it's already in the project dictionary
 	for(TMap<FString, FString>::TConstIterator Iter(ShortProjectNameDictionary); Iter; ++Iter)
 	{
 		if(Iter.Value() == ProjectFileName)
@@ -89,6 +79,14 @@ bool FUProjectDictionary::IsForeignProject(const FString& InProjectFileName) con
 		}
 	}
 
+	// If not, it may be a new project. Check if its parent directory is a project root dir.
+	FString ProjectRootDir = FPaths::GetPath(FPaths::GetPath(ProjectFileName));
+	if(ProjectRootDirs.Contains(ProjectRootDir))
+	{
+		return false;
+	}
+
+	// Otherwise it's a foreign project
 	return true;
 }
 
diff --git a/Engine/Source/Runtime/Core/Public/Misc/UProjectInfo.h b/Engine/Source/Runtime/Core/Public/Misc/UProjectInfo.h
--- a/Engine/Source/Runtime/Core/Public/Misc/UProjectInfo.h
+++ b/Engine/Source/Runtime/Core/Public/Misc/UProjectInfo.h
@@ -2,17 +2,18 @@
 
 #pragma once
 
-
 /**
  * Dictionary of all the non-foreign projects for an engine installation, found by parsing .uprojectdirs files for source directories.
  */
 class CORE_API FUProjectDictionary
 {
 public:
-
 	/** Scans the engine root directory for all the known projects. */
 	FUProjectDictionary(const FString& InRootDir);
 	
+	/** Refreshes the list of known projects */
+	void Refresh();
+
 	/** Determines whether a project is a foreign project or not. */
 	bool IsForeignProject(const FString& ProjectFileName) const;
 
@@ -26,6 +27,11 @@ public:
 	static FUProjectDictionary& GetDefault();
 
 private:
+	/** The root directory for this engine installation */
+	FString RootDir;
+
+	/** List of project root directories */
+	TArray<FString> ProjectRootDirs;
 
 	/** Map of short game names to full project paths. */
 	TMap<FString, FString> ShortProjectNameDictionary;
