#ue4 - Add FSphere::Intersects() to test for intersection between two spheres, with an optional tolerance.

diff --git a/Engine/Source/Runtime/Core/Public/Math/Sphere.h b/Engine/Source/Runtime/Core/Public/Math/Sphere.h
--- a/Engine/Source/Runtime/Core/Public/Math/Sphere.h
+++ b/Engine/Source/Runtime/Core/Public/Math/Sphere.h
@@ -97,6 +97,19 @@ public:
 		return (Center - Other.Center).SizeSquared() <= FMath::Square(Other.W - Tolerance - W);
 	}
 
+	
+	/**
+	 * Test whether this sphere intersects another.
+	 * 
+	 * @param  Other The other sphere.
+	 * @param  Tolerance Error tolerance.
+	 * @return true if spheres intersect, false otherwise.
+	 */
+	FORCEINLINE bool Intersects( const FSphere& Other, float Tolerance = KINDA_SMALL_NUMBER ) const
+	{
+		return (Center - Other.Center).SizeSquared() <= FMath::Square(FMath::Max(0.f, Other.W + W + Tolerance));
+	}
+
 	/**
 	 * Get result of Transforming sphere by Matrix.
 	 *
