#ue4 - Add FBoxSphereBounds::SpheresIntersect() and FBoxSphereBounds::BoxesIntersect() to test for intersection/overlap between two bounds.

diff --git a/Engine/Source/Runtime/Core/Public/Math/BoxSphereBounds.h b/Engine/Source/Runtime/Core/Public/Math/BoxSphereBounds.h
--- a/Engine/Source/Runtime/Core/Public/Math/BoxSphereBounds.h
+++ b/Engine/Source/Runtime/Core/Public/Math/BoxSphereBounds.h
@@ -135,12 +135,22 @@ public:
 		return ::ComputeSquaredDistanceFromBoxToPoint(Mins, Maxs, Point);
 	}
 
+	FORCEINLINE static bool SpheresIntersect(const FBoxSphereBounds& A, const FBoxSphereBounds& B)
+	{
+		return (A.Origin - B.Origin).SizeSquared() <= FMath::Square(A.SphereRadius + B.SphereRadius);
+	}
+
+	FORCEINLINE static bool BoxesIntersect(const FBoxSphereBounds& A, const FBoxSphereBounds& B)
+	{
+		return A.GetBox().Intersect(B.GetBox());
+	}
+
 	/**
 	 * Gets the bounding box.
 	 *
 	 * @return The bounding box.
 	 */
-	FBox GetBox( ) const
+	FORCEINLINE FBox GetBox( ) const
 	{
 		return FBox(Origin - BoxExtent,Origin + BoxExtent);
 	}
@@ -166,7 +176,7 @@ public:
 	 *
 	 * @return The bounding sphere.
 	 */
-	FSphere GetSphere( ) const
+	FORCEINLINE FSphere GetSphere( ) const
 	{
 		return FSphere(Origin,SphereRadius);
 	}
