#ue4 - Optimize TArray::RemoveAt() and TArray::RemoveAtSwap() to only call CalculateSlack if they are allowed to shrink the array.

diff --git a/Engine/Source/Runtime/Core/Public/Containers/Array.h b/Engine/Source/Runtime/Core/Public/Containers/Array.h
--- a/Engine/Source/Runtime/Core/Public/Containers/Array.h
+++ b/Engine/Source/Runtime/Core/Public/Containers/Array.h
@@ -1442,11 +1442,14 @@ public:
 		}
 		ArrayNum -= Count;
 
-		const int32 NewArrayMax = AllocatorInstance.CalculateSlack(ArrayNum, ArrayMax, sizeof(ElementType));
-		if (NewArrayMax != ArrayMax && bAllowShrinking)
+		if (bAllowShrinking)
 		{
-			ArrayMax = NewArrayMax;
-			AllocatorInstance.ResizeAllocation(ArrayNum, ArrayMax, sizeof(ElementType));
+			const int32 NewArrayMax = AllocatorInstance.CalculateSlack(ArrayNum, ArrayMax, sizeof(ElementType));
+			if (NewArrayMax != ArrayMax)
+			{
+				ArrayMax = NewArrayMax;
+				AllocatorInstance.ResizeAllocation(ArrayNum, ArrayMax, sizeof(ElementType));
+			}
 		}
 	}
 
@@ -1483,11 +1486,14 @@ public:
 		}
 		ArrayNum -= Count;
 
-		const int32 NewArrayMax = AllocatorInstance.CalculateSlack(ArrayNum, ArrayMax, sizeof(ElementType));
-		if (NewArrayMax != ArrayMax && bAllowShrinking)
+		if (bAllowShrinking)
 		{
-			ArrayMax = NewArrayMax;
-			AllocatorInstance.ResizeAllocation(ArrayNum, ArrayMax, sizeof(ElementType));
+			const int32 NewArrayMax = AllocatorInstance.CalculateSlack(ArrayNum, ArrayMax, sizeof(ElementType));
+			if (NewArrayMax != ArrayMax)
+			{
+				ArrayMax = NewArrayMax;
+				AllocatorInstance.ResizeAllocation(ArrayNum, ArrayMax, sizeof(ElementType));
+			}
 		}
 	}
 
