#ue4 - Add error tolerance to FBoxSphereBounds::SpheresIntersect(), and document method.

diff --git a/Engine/Source/Runtime/Core/Public/Math/BoxSphereBounds.h b/Engine/Source/Runtime/Core/Public/Math/BoxSphereBounds.h
--- a/Engine/Source/Runtime/Core/Public/Math/BoxSphereBounds.h
+++ b/Engine/Source/Runtime/Core/Public/Math/BoxSphereBounds.h
@@ -135,11 +135,26 @@ public:
 		return ::ComputeSquaredDistanceFromBoxToPoint(Mins, Maxs, Point);
 	}
 
-	FORCEINLINE static bool SpheresIntersect(const FBoxSphereBounds& A, const FBoxSphereBounds& B)
+	/**
+	 * Test whether the spheres from two BoxSphereBounds intersect/overlap.
+	 * 
+	 * @param  A First BoxSphereBounds to test.
+	 * @param  B Second BoxSphereBounds to test.
+	 * @param  Tolerance Error tolerance added to test distance.
+	 * @return true if spheres intersect, false otherwise.
+	 */
+	FORCEINLINE static bool SpheresIntersect(const FBoxSphereBounds& A, const FBoxSphereBounds& B, float Tolerance = KINDA_SMALL_NUMBER)
 	{
-		return (A.Origin - B.Origin).SizeSquared() <= FMath::Square(A.SphereRadius + B.SphereRadius);
+		return (A.Origin - B.Origin).SizeSquared() <= FMath::Square(FMath::Max(0.f, A.SphereRadius + B.SphereRadius + Tolerance));
 	}
 
+	/**
+	 * Test whether the boxes from two BoxSphereBounds intersect/overlap.
+	 * 
+	 * @param  A First BoxSphereBounds to test.
+	 * @param  B Second BoxSphereBounds to test.
+	 * @return true if boxes intersect, false otherwise.
+	 */
 	FORCEINLINE static bool BoxesIntersect(const FBoxSphereBounds& A, const FBoxSphereBounds& B)
 	{
 		return A.GetBox().Intersect(B.GetBox());
