Linux: band-aid, non-performant support for case-insensitive OpenRead().

- No caching whatsoever, performance impact possibly large on large projects.
- Also reorganized TestPAL to add new tests more easily.

diff --git a/Engine/Source/Programs/TestPAL/Private/Main.cpp b/Engine/Source/Programs/TestPAL/Private/Main.cpp
--- a/Engine/Source/Programs/TestPAL/Private/Main.cpp
+++ b/Engine/Source/Programs/TestPAL/Private/Main.cpp
@@ -7,32 +7,15 @@ DEFINE_LOG_CATEGORY(LogTestPAL);
 
 IMPLEMENT_APPLICATION(TestPAL, "TestPAL");
 
-#define CHILD_INVOCATION_ARG			"-child"
-
-namespace
+namespace TestPAL
 {
-	/**
-	 * Checks if this instance is nested (i.e. run under itself)
-	 *
-	 * @param ArgC Number of commandline arguments.
-	 * @param ArgV Commandline arguments
-	 * @return true if nested
-	 */
-	bool IsNestedInstance(int32 ArgC, ANSICHAR* ArgV[])
-	{
-		for (int32 IdxArg = 0; IdxArg < ArgC; ++IdxArg)
-		{
-			if (!FCStringAnsi::Strcmp(ArgV[IdxArg], CHILD_INVOCATION_ARG))
-			{
-				return true;
-			}
-		}
-
-		return false;
-	}
-}
+	FString CommandLine;
+};
 
-int32 RunAsChild(const TCHAR* CommandLine)
+/**
+ * FProcHandle test (child instance)
+ */
+int32 ProcRunAsChild(const TCHAR* CommandLine)
 {
 	FPlatformMisc::SetCrashHandler(NULL);
 	FPlatformMisc::SetGracefulTerminationHandler();
@@ -42,7 +25,7 @@ int32 RunAsChild(const TCHAR* CommandLine)
 	// set a random delay pretending to do some useful work up to a minute. 
 	double RandomWorkTime = FMath::FRandRange(0.0f, 60.0f);
 
-	UE_LOG(LogTestPAL, Display, TEXT("Running as child (pid %d), will be doing work for %f seconds."), FPlatformProcess::GetCurrentProcessId(), RandomWorkTime);
+	UE_LOG(LogTestPAL, Display, TEXT("Running proc test as child (pid %d), will be doing work for %f seconds."), FPlatformProcess::GetCurrentProcessId(), RandomWorkTime);
 
 	double StartTime = FPlatformTime::Seconds();
 
@@ -63,13 +46,16 @@ int32 RunAsChild(const TCHAR* CommandLine)
 	return 0;
 }
 
-int32 RunAsParent(const TCHAR* CommandLine)
+/**
+ * FProcHandle test (parent instance)
+ */
+int32 ProcRunAsParent(const TCHAR* CommandLine)
 {
 	FPlatformMisc::SetCrashHandler(NULL);
 	FPlatformMisc::SetGracefulTerminationHandler();
 
 	GEngineLoop.PreInit(CommandLine);
-	UE_LOG(LogTestPAL, Display, TEXT("Running as parent."));
+	UE_LOG(LogTestPAL, Display, TEXT("Running proc test as parent."));
 
 	// Run slave instance continuously
 	int NumChildrenToSpawn = 255, MaxAtOnce = 5;
@@ -84,12 +70,94 @@ int32 RunAsParent(const TCHAR* CommandLine)
 	return 0;
 }
 
+/**
+ * Tests a single file.
+ */
+void TestCaseInsensitiveFile(const FString & Filename, const FString & WrongFilename)
+{
+	IPlatformFile& PlatformFile = FPlatformFileManager::Get().GetPlatformFile();
+
+	IFileHandle * CreationHandle = PlatformFile.OpenWrite(*Filename);
+	checkf(CreationHandle, TEXT("Could not create a test file for '%s'"), *Filename);
+	delete CreationHandle;
+	
+	IFileHandle* CheckGoodHandle = PlatformFile.OpenRead(*Filename);
+	checkf(CheckGoodHandle, TEXT("Could not open a test file for '%s' (zero probe)"), *Filename);
+	delete CheckGoodHandle;
+	
+	IFileHandle* CheckWrongCaseRelHandle = PlatformFile.OpenRead(*WrongFilename);
+	checkf(CheckWrongCaseRelHandle, TEXT("Could not open a test file for '%s'"), *WrongFilename);
+	delete CheckWrongCaseRelHandle;
+	
+	PlatformFile.DeleteFile(*Filename);
+}
+
+/**
+ * Case-(in)sensitivity test/
+ */
+int32 CaseTest(const TCHAR* CommandLine)
+{
+	FPlatformMisc::SetCrashHandler(NULL);
+	FPlatformMisc::SetGracefulTerminationHandler();
+	
+	GEngineLoop.PreInit(CommandLine);
+	UE_LOG(LogTestPAL, Display, TEXT("Running case sensitivity test."));
+	
+	TestCaseInsensitiveFile(TEXT("Test.Test"), TEXT("teSt.teSt"));
+	
+	FString File(TEXT("Test^%!CaseInsens"));
+	FString AbsFile = FPaths::ConvertRelativePathToFull(File);
+	FString AbsFileUpper = AbsFile.ToUpper();
+
+	TestCaseInsensitiveFile(AbsFile, AbsFileUpper);
+	
+	FEngineLoop::AppPreExit();
+	FEngineLoop::AppExit();
+	return 0;
+}
+
+/**
+ * Selects and runs one of test cases.
+ *
+ * @param ArgC Number of commandline arguments.
+ * @param ArgV Commandline arguments.
+ * @return Test case's return code.
+ */
+int32 MultiplexedMain(int32 ArgC, char* ArgV[])
+{
+	for (int32 IdxArg = 0; IdxArg < ArgC; ++IdxArg)
+	{
+		if (!FCStringAnsi::Strcmp(ArgV[IdxArg], ARG_PROC_TEST_CHILD))
+		{
+			return ProcRunAsChild(*TestPAL::CommandLine);
+		}
+		else if (!FCStringAnsi::Strcmp(ArgV[IdxArg], ARG_PROC_TEST))
+		{
+			return ProcRunAsParent(*TestPAL::CommandLine);
+		}
+		else if (!FCStringAnsi::Strcmp(ArgV[IdxArg], ARG_CASE_SENSITIVITY_TEST))
+		{
+			return CaseTest(*TestPAL::CommandLine);
+		}
+	}
+
+	FPlatformMisc::SetCrashHandler(NULL);
+	FPlatformMisc::SetGracefulTerminationHandler();
+	
+	GEngineLoop.PreInit(*TestPAL::CommandLine);
+	UE_LOG(LogTestPAL, Warning, TEXT("Unable to find any known test name, no test started."));
+	FEngineLoop::AppPreExit();
+	FEngineLoop::AppExit();
+	return 1;
+}
+
 int32 main(int32 ArgC, char* ArgV[])
 {
-	FString CommandLine;
+	TestPAL::CommandLine = TEXT("");
+
 	for (int32 Option = 1; Option < ArgC; Option++)
 	{
-		CommandLine += TEXT(" ");
+		TestPAL::CommandLine += TEXT(" ");
 		FString Argument(ANSI_TO_TCHAR(ArgV[Option]));
 		if (Argument.Contains(TEXT(" ")))
 		{
@@ -105,9 +173,8 @@ int32 main(int32 ArgC, char* ArgV[])
 				Argument = FString::Printf(TEXT("\"%s\""), *Argument);
 			}
 		}
-		CommandLine += Argument;
+		TestPAL::CommandLine += Argument;
 	}
 
-	// fork and run under self
-	return IsNestedInstance(ArgC, ArgV) ? RunAsChild(*CommandLine) : RunAsParent(*CommandLine);
+	return MultiplexedMain(ArgC, ArgV);
 }
diff --git a/Engine/Source/Programs/TestPAL/Private/Parent.cpp b/Engine/Source/Programs/TestPAL/Private/Parent.cpp
--- a/Engine/Source/Programs/TestPAL/Private/Parent.cpp
+++ b/Engine/Source/Programs/TestPAL/Private/Parent.cpp
@@ -17,7 +17,7 @@ FProcHandle FParent::Launch()
 
 	uint32 WorkerId = 0;
 	FString WorkerName = FPlatformProcess::ExecutableName(false);
-	FProcHandle WorkerHandle = FPlatformProcess::CreateProc(*WorkerName, TEXT("-child"), true, false, false, &WorkerId, PriorityModifier, NULL, NULL);
+	FProcHandle WorkerHandle = FPlatformProcess::CreateProc(*WorkerName, TEXT("proc-child"), true, false, false, &WorkerId, PriorityModifier, NULL, NULL);
 	if (!WorkerHandle.IsValid())
 	{
 		// If this doesn't error, the app will hang waiting for jobs that can never be completed
diff --git a/Engine/Source/Programs/TestPAL/Private/PrivatePCH.h b/Engine/Source/Programs/TestPAL/Private/PrivatePCH.h
--- a/Engine/Source/Programs/TestPAL/Private/PrivatePCH.h
+++ b/Engine/Source/Programs/TestPAL/Private/PrivatePCH.h
@@ -5,5 +5,10 @@
 // other modules' headers
 #include "Core.h"
 
+#define ARG_PROC_TEST						"proc"
+#define ARG_PROC_TEST_CHILD					"proc-child"
+#define ARG_CASE_SENSITIVITY_TEST			"case"
+
+
 // this module headers
 #include "Parent.h"
diff --git a/Engine/Source/Runtime/Core/Private/Linux/LinuxPlatformFile.cpp b/Engine/Source/Runtime/Core/Private/Linux/LinuxPlatformFile.cpp
--- a/Engine/Source/Runtime/Core/Private/Linux/LinuxPlatformFile.cpp
+++ b/Engine/Source/Runtime/Core/Private/Linux/LinuxPlatformFile.cpp
@@ -231,19 +231,208 @@ FString FLinuxPlatformFile::GetFilenameOnDisk(const TCHAR* Filename)
 	return Filename;
 }
 
+/**
+ * A class to handle case insensitive file opening. This is a band-aid, non-performant approach,
+ * without any caching.
+ */
+class FLinuxFileMapper
+{
+	/** Max path component */
+	int MaxPathComponents;
+
+public:
+
+	FLinuxFileMapper()
+		:	MaxPathComponents(-1)
+	{
+	}
+
+	FString GetPathComponent(const FString & Filename, int NumPathComponent)
+	{
+		// skip over empty part
+		int StartPosition = (Filename[0] == TEXT('/')) ? 1 : 0;
+		
+		for (int ComponentIdx = 0; ComponentIdx < NumPathComponent; ++ComponentIdx)
+		{
+			int FoundAtIndex = Filename.Find(TEXT("/"), ESearchCase::CaseSensitive,
+											 ESearchDir::FromStart, StartPosition);
+			
+			if (FoundAtIndex == INDEX_NONE)
+			{
+				checkf(false, TEXT("Asked to get %d-th path component, but filename '%s' doesn't have that many!"), 
+					   NumPathComponent, *Filename);
+				break;
+			}
+			
+			StartPosition = FoundAtIndex + 1;	// skip the '/' itself
+		}
+
+		// now return the 
+		int NextSlash = Filename.Find(TEXT("/"), ESearchCase::CaseSensitive,
+									  ESearchDir::FromStart, StartPosition);
+		if (NextSlash == INDEX_NONE)
+		{
+			// just return the rest of the string
+			return Filename.RightChop(StartPosition);
+		}
+		else if (NextSlash == StartPosition)
+		{
+			return TEXT("");	// encountered an invalid path like /foo/bar//baz
+		}
+		
+		return Filename.Mid(StartPosition, NextSlash - StartPosition);
+	}
+
+	int32 CountPathComponents(const FString & Filename)
+	{
+		if (Filename.Len() == 0)
+		{
+			return 0;
+		}
+
+		// if the first character is not a separator, it's part of a distinct component
+		int NumComponents = (Filename[0] != TEXT('/')) ? 1 : 0;
+		for (const auto & Char : Filename)
+		{
+			if (Char == TEXT('/'))
+			{
+				++NumComponents;
+			}
+		}
+
+		// cannot be 0 components if path is non-empty
+		return FMath::Max(NumComponents, 1);
+	}
+
+	/**
+	 * Tries to recursively find (using case-insensitive comparison) and open the file.
+	 * The first file found will be opened.
+	 * 
+	 * @param Filename Original file path as requested (absolute)
+	 * @param PathComponentToLookFor Part of path we are currently trying to find.
+	 * @param ConstructedPath The real (absolute) path that we have found so far
+	 * 
+	 * @return a handle opened with open()
+	 */
+	int32 TryOpenRecursively(const FString & Filename, int PathComponentToLookFor, FString & ConstructedPath)
+	{
+		// get the directory without the last path component
+		FString BaseDir = ConstructedPath;
+
+		// get the path component to compare
+		FString PathComponent = GetPathComponent(Filename, PathComponentToLookFor);
+		FString PathComponentLower = PathComponent.ToLower();
+
+		int32 Handle = -1;
+
+		// see if we can open this (we should)
+		DIR* DirHandle = opendir(TCHAR_TO_UTF8(*BaseDir));
+		if (DirHandle)
+		{
+			struct dirent *Entry;
+			while ((Entry = readdir(DirHandle)) != nullptr)
+			{
+				FString DirEntry = UTF8_TO_TCHAR(Entry->d_name);
+				if (DirEntry.ToLower() == PathComponentLower)
+				{
+					if (PathComponentToLookFor < MaxPathComponents - 1)
+					{
+						// make sure this is a directory
+						bool bIsDirectory = Entry->d_type == DT_DIR;
+						if(Entry->d_type == DT_UNKNOWN)
+						{
+							struct stat StatInfo;
+							if(stat(TCHAR_TO_UTF8(*(BaseDir / Entry->d_name)), &StatInfo) == 0)
+							{
+								bIsDirectory = S_ISDIR(StatInfo.st_mode);
+							}
+						}
+
+						if (bIsDirectory)
+						{
+							// recurse with the new filename
+							FString NewConstructedPath = ConstructedPath;
+							NewConstructedPath /= DirEntry;
+
+							Handle = TryOpenRecursively(Filename, PathComponentToLookFor + 1, NewConstructedPath);
+							if (Handle != -1)
+							{
+								ConstructedPath = NewConstructedPath;
+								break;
+							}
+						}
+					}
+					else
+					{
+						// last level, try opening directly
+						FString ConstructedFilename = ConstructedPath;
+						ConstructedFilename /= DirEntry;
+
+						Handle = open(TCHAR_TO_UTF8(*ConstructedFilename), O_RDONLY);
+						if (Handle != -1)
+						{
+							ConstructedPath = ConstructedFilename;
+							break;
+						}
+					}
+				}
+			}
+			closedir(DirHandle);
+		}
+
+		return Handle;
+	}
+	
+	/**
+	 * Opens a file for reading, disregarding the case.
+	 * 
+	 * @param Filename absolute filename
+	 */
+	int32 OpenCaseInsensitiveRead(const FString & Filename)
+	{
+		// try opening right away
+		int32 Handle = open(TCHAR_TO_UTF8(*Filename), O_RDONLY);
+		if (Handle == -1)
+		{
+			// log non-standard errors only
+			if (ENOENT != errno)
+			{
+				int ErrNo = errno;
+				UE_LOG(LogLinuxPlatformFile, Warning, TEXT( "open('%s', ORDONLY) failed: errno=%d (%s)" ), *Filename, ErrNo, ANSI_TO_TCHAR(strerror(ErrNo)));
+			}
+			else
+			{
+				// perform a case-insensitive search
+				// make sure we get the absolute filename
+				checkf(Filename[0] == TEXT('/'), TEXT("Filename '%s' given to OpenCaseInsensitiveRead is not absolute!"), *Filename);
+				
+				MaxPathComponents = CountPathComponents(Filename);
+				if (MaxPathComponents > 0)
+				{
+					FString FoundFilename(TEXT("/"));	// start with root
+					Handle = TryOpenRecursively(Filename, 0, FoundFilename);
+
+					if (Handle != -1)
+					{
+						UE_LOG(LogLinuxPlatformFile, Log, TEXT("Mapped '%s' to '%s'"), *Filename, *FoundFilename);
+					}
+				}
+			}
+		}
+		return Handle;
+	}
+};
+
+
 IFileHandle* FLinuxPlatformFile::OpenRead(const TCHAR* Filename)
 {
-	int32 Handle = open(TCHAR_TO_UTF8(*NormalizeFilename(Filename)), O_RDONLY);
+	FLinuxFileMapper CaseInsensMapper;
+	int32 Handle = CaseInsensMapper.OpenCaseInsensitiveRead(TCHAR_TO_UTF8(*NormalizeFilename(Filename)));
 	if (Handle != -1)
 	{
 		return new FFileHandleLinux(Handle);
 	}
-	// log non-standard errors only
-	if (ENOENT != errno)
-	{
-		UE_LOG(LogLinuxPlatformFile, Warning, TEXT( "open('%s', ORDONLY) failed: errno=%d (%s)" ), *NormalizeFilename(Filename), errno, ANSI_TO_TCHAR(strerror(errno)));
-	}
-	return NULL;
+	return nullptr;
 }
 
 IFileHandle* FLinuxPlatformFile::OpenWrite(const TCHAR* Filename, bool bAppend, bool bAllowRead)
