Fixed a crash when using the Bing IME

https://answers.unrealengine.com/questions/145623/can-not-input-zh-cn-words-with-slate-inputting-wid.html

The IME was sending both a result string and an empty composition string in the same message. The latter was causing the compositing to be ended twice, resulting in an assert.

diff --git a/Engine/Source/Runtime/Core/Private/Windows/WindowsTextInputMethodSystem.cpp b/Engine/Source/Runtime/Core/Private/Windows/WindowsTextInputMethodSystem.cpp
--- a/Engine/Source/Runtime/Core/Private/Windows/WindowsTextInputMethodSystem.cpp
+++ b/Engine/Source/Runtime/Core/Private/Windows/WindowsTextInputMethodSystem.cpp
@@ -824,7 +824,8 @@ int32 FWindowsTextInputMethodSystem::ProcessMessage(HWND hwnd, uint32 msg, WPARA
 
 					// Not all IMEs send a cancel request when you press escape, but instead just set the string to empty
 					// We need to cancel out the composition string here to avoid weirdness when you start typing again
-					if(CompositionString.Len() == 0)
+					// Don't do this if we have a result string, as we'll have already called EndIMMComposition to finish the composition
+					if(CompositionString.Len() == 0 && !bHasResultStringFlag)
 					{
 						CancelIMMComposition();
 					}
