Stopped FName converting an ANSICHAR source string to a WIDECHAR string, only to convert it back to ANSICHAR again later

FName::Init could take either an ANSICHAR* or a WIDECHAR*, however the ANSICHAR* variant would convert the string to a WIDECHAR[], and would later convert it back to an ANSICHAR[] because it contains only ANSI characters.

Recent refactoring in FName has made it possible to optimize the ANSICHAR* path, by avoiding the double conversion and just using the ANSICHAR* directly when adding to the FName string table.

diff --git a/Engine/Source/Runtime/Core/Private/UObject/UnrealNames.cpp b/Engine/Source/Runtime/Core/Private/UObject/UnrealNames.cpp
--- a/Engine/Source/Runtime/Core/Private/UObject/UnrealNames.cpp
+++ b/Engine/Source/Runtime/Core/Private/UObject/UnrealNames.cpp
@@ -441,9 +441,20 @@ int32 FName::Compare( const FName& Other ) const
 	}
 }
 
-void FName::Init(const WIDECHAR* InName, int32 InNumber, EFindName FindType, bool bSplitName, int32 HardcodeIndex )
+void FName::Init(const WIDECHAR* InName, int32 InNumber, EFindName FindType, bool bSplitName, int32 HardcodeIndex)
 {
-	check(FCString::Strlen(InName)<=NAME_SIZE);
+	InitInternal<WIDECHAR>(InName, InNumber, FindType, bSplitName, HardcodeIndex);
+}
+
+void FName::Init(const ANSICHAR* InName, int32 InNumber, EFindName FindType, bool bSplitName, int32 HardcodeIndex)
+{
+	InitInternal<ANSICHAR>(InName, InNumber, FindType, bSplitName, HardcodeIndex);
+}
+
+template <typename TCharType>
+void FName::InitInternal(const TCharType* InName, int32 InNumber, const EFindName FindType, const bool bSplitName, const int32 HardcodeIndex)
+{
+	check(TCString<TCharType>::Strlen(InName)<=NAME_SIZE);
 
 	// initialize the name subsystem if necessary
 	if (!GetIsInitialized())
@@ -451,13 +462,13 @@ void FName::Init(const WIDECHAR* InName, int32 InNumber, EFindName FindType, boo
 		StaticInit();
 	}
 
-	WIDECHAR TempBuffer[NAME_SIZE];
+	TCharType TempBuffer[NAME_SIZE];
 	int32 TempNumber;
-	// if we were passed in a number, we can't split again, other wise, a_1_2_3_4 would change everytime
+	// if we were passed in a number, we can't split again, other wise, a_1_2_3_4 would change every time
 	// it was loaded in
 	if (InNumber == NAME_NO_NUMBER_INTERNAL && bSplitName == true )
 	{
-		if (SplitNameWithCheck(InName, TempBuffer, ARRAY_COUNT(TempBuffer), TempNumber))
+		if (SplitNameWithCheckImpl<TCharType>(InName, TempBuffer, ARRAY_COUNT(TempBuffer), TempNumber))
 		{
 			InName = TempBuffer;
 			InNumber = NAME_EXTERNAL_TO_INTERNAL(TempNumber);
@@ -485,18 +496,14 @@ void FName::Init(const WIDECHAR* InName, int32 InNumber, EFindName FindType, boo
 	int32 OutComparisonIndex = HardcodeIndex;
 	int32 OutDisplayIndex = HardcodeIndex;
 
-	const bool bIsPureAnsi = FCStringWide::IsPureAnsi( InName );
+	const bool bIsPureAnsi = TCString<TCharType>::IsPureAnsi(InName);
 	if(bIsPureAnsi)
 	{
-		// Convert to an ansi display name
-		ANSICHAR AnsiDisplayName[NAME_SIZE];
-		FCStringAnsi::Strncpy(AnsiDisplayName, StringCast<ANSICHAR>(InName).Get(), ARRAY_COUNT(AnsiDisplayName));
-
-		bWasFoundOrAdded = InitInternal_FindOrAdd<ANSICHAR>(AnsiDisplayName, FindType, HardcodeIndex, OutComparisonIndex, OutDisplayIndex);
+		bWasFoundOrAdded = InitInternal_FindOrAdd<ANSICHAR>(StringCast<ANSICHAR>(InName).Get(), FindType, HardcodeIndex, OutComparisonIndex, OutDisplayIndex);
 	}
 	else
 	{
-		bWasFoundOrAdded = InitInternal_FindOrAdd<WIDECHAR>(InName, FindType, HardcodeIndex, OutComparisonIndex, OutDisplayIndex);
+		bWasFoundOrAdded = InitInternal_FindOrAdd<WIDECHAR>(StringCast<WIDECHAR>(InName).Get(), FindType, HardcodeIndex, OutComparisonIndex, OutDisplayIndex);
 	}
 
 	if(bWasFoundOrAdded)
@@ -824,17 +831,23 @@ void FName::DisplayHash( FOutputDevice& Ar )
 }
 
 bool FName::SplitNameWithCheck(const WIDECHAR* OldName, WIDECHAR* NewName, int32 NewNameLen, int32& NewNumber)
+{
+	return SplitNameWithCheckImpl<WIDECHAR>(OldName, NewName, NewNameLen, NewNumber);
+}
+
+template <typename TCharType>
+bool FName::SplitNameWithCheckImpl(const TCharType* OldName, TCharType* NewName, int32 NewNameLen, int32& NewNumber)
 {
 	bool bSucceeded = false;
-	const int32 OldNameLength = FCStringWide::Strlen(OldName);
+	const int32 OldNameLength = TCString<TCharType>::Strlen(OldName);
 
 	if(OldNameLength > 0)
 	{
 		// get string length
-		const WIDECHAR* LastChar = OldName + (OldNameLength - 1);
+		const TCharType* LastChar = OldName + (OldNameLength - 1);
 		
 		// if the last char is a number, then we will try to split
-		const WIDECHAR* Ch = LastChar;
+		const TCharType* Ch = LastChar;
 		if (*Ch >= '0' && *Ch <= '9')
 		{
 			// go backwards, looking an underscore or the start of the string
@@ -854,12 +867,12 @@ bool FName::SplitNameWithCheck(const WIDECHAR* OldName, WIDECHAR* NewName, int32
 				if (Ch[1] != '0' || LastChar - Ch == 1)
 				{
 					// attempt to convert what's following it to a number
-					uint64 TempConvert = FCStringWide::Atoi64(Ch + 1);
+					uint64 TempConvert = TCString<TCharType>::Atoi64(Ch + 1);
 					if (TempConvert <= MAX_int32)
 					{
 						NewNumber = (int32)TempConvert;
 						// copy the name portion into the buffer
-						FCStringWide::Strncpy(NewName, OldName, FMath::Min<int32>(Ch - OldName + 1, NewNameLen));
+						TCString<TCharType>::Strncpy(NewName, OldName, FMath::Min<int32>(Ch - OldName + 1, NewNameLen));
 
 						// mark successful
 						bSucceeded = true;
diff --git a/Engine/Source/Runtime/Core/Public/UObject/NameTypes.h b/Engine/Source/Runtime/Core/Public/UObject/NameTypes.h
--- a/Engine/Source/Runtime/Core/Public/UObject/NameTypes.h
+++ b/Engine/Source/Runtime/Core/Public/UObject/NameTypes.h
@@ -939,7 +939,7 @@ private:
 	friend FNameEntry* AllocateNameEntry( const void* Name, NAME_INDEX Index, FNameEntry* HashNext, bool bIsPureAnsi );
 
 	/**
-	 * Shared initialization code (between two constructors)
+	 * Initialization from a wide string
 	 * 
 	 * @param InName String name of the name/number pair
 	 * @param InNumber Number part of the name/number pair
@@ -950,17 +950,18 @@ private:
 	void Init(const WIDECHAR* InName, int32 InNumber, EFindName FindType, bool bSplitName=true, int32 HardcodeIndex = -1);
 
 	/**
-	 * Non-optimized initialization code for ansi names.
+	 * Initialization from an ANSI string
 	 * 
 	 * @param InName		String name of the name/number pair
 	 * @param InNumber		Number part of the name/number pair
 	 * @param FindType		Operation to perform on names
 	 * @param bSplitName	If true, this function will attempt to split a number off of the string portion (turning Rocket_17 to Rocket and number 17)
+	 * @param HardcodeIndex If >= 0, this represents a hardcoded FName and so automatically gets this index
 	 */
-	void Init(const ANSICHAR* InName, int32 InNumber, EFindName FindType, bool bSplitName=true, int32 HardcodeIndex = -1)
-	{
-		Init(StringCast<WIDECHAR>(InName).Get(), InNumber, FindType, bSplitName, HardcodeIndex);
-	}
+	void Init(const ANSICHAR* InName, int32 InNumber, EFindName FindType, bool bSplitName=true, int32 HardcodeIndex = -1);
+
+	template <typename TCharType>
+	void InitInternal(const TCharType* InName, int32 InNumber, const EFindName FindType, const bool bSplitName, const int32 HardcodeIndex);
 
 	template <typename TCharType>
 	static bool InitInternal_FindOrAdd(const TCharType* InName, const EFindName FindType, const int32 HardcodeIndex, int32& OutComparisonIndex, int32& OutDisplayIndex);
@@ -968,6 +969,9 @@ private:
 	template <typename TCharType>
 	static bool InitInternal_FindOrAddNameEntry(const TCharType* InName, const EFindName FindType, const ENameCase ComparisonMode, int32& OutIndex);
 
+	template <typename TCharType>
+	static bool SplitNameWithCheckImpl(const TCharType* OldName, TCharType* NewName, int32 NewNameLen, int32& NewNumber);
+
 	FORCEINLINE NAME_INDEX GetComparisonIndexFast() const
 	{
 		return ComparisonIndex;
