Prevented FName crashing if constructed from a null string

This now initializes to NAME_None instead.

diff --git a/Engine/Source/Runtime/Core/Private/UObject/UnrealNames.cpp b/Engine/Source/Runtime/Core/Private/UObject/UnrealNames.cpp
--- a/Engine/Source/Runtime/Core/Private/UObject/UnrealNames.cpp
+++ b/Engine/Source/Runtime/Core/Private/UObject/UnrealNames.cpp
@@ -356,12 +356,26 @@ int32							FName::NumWideNames;
  */
 FName::FName( const WIDECHAR* Name, EFindName FindType, bool )
 {
-	Init(Name, NAME_NO_NUMBER_INTERNAL, FindType);
+	if (Name)
+	{
+		Init(Name, NAME_NO_NUMBER_INTERNAL, FindType);
+	}
+	else
+	{
+		*this = FName(NAME_None);
+	}
 }
 
 FName::FName( const ANSICHAR* Name, EFindName FindType, bool )
 {
-	Init(Name, NAME_NO_NUMBER_INTERNAL, FindType);
+	if (Name)
+	{
+		Init(Name, NAME_NO_NUMBER_INTERNAL, FindType);
+	}
+	else
+	{
+		*this = FName(NAME_None);
+	}
 }
 
 /**
@@ -927,6 +941,7 @@ void FName::AutoTest()
 	const FName AutoTest1Find("autoTEST_1", EFindName::FNAME_Find);
 	const FName AutoTest_2(TEXT("AutoTest_2"));
 	const FName AutoTestB_2(TEXT("AutoTestB_2"));
+	const FName NullName(static_cast<ANSICHAR*>(nullptr));
 
 	check(AutoTest_1 != AutoTest_2);
 	check(AutoTest_1 == autoTest_1);
@@ -943,6 +958,7 @@ void FName::AutoTest()
 	check(*AutoTestB_2.GetPlainNameString() != *AutoTest_2.GetPlainNameString());
 	check(AutoTestB_2.GetNumber() == AutoTest_2.GetNumber());
 	check(autoTest_1.GetNumber() != AutoTest_2.GetNumber());
+	check(NullName.IsNone());
 }
 
 /*-----------------------------------------------------------------------------
