PR #677: fix GetKeyFromCodes(0, 'A') == EKeys::A on Mac (Contributed by [REDACTED_PERSON])

https://github.com/EpicGames/UnrealEngine/pull/677

Also applied the same fix to HTML5 and Linux

diff --git a/Engine/Source/Runtime/Core/Private/HTML5/HTML5PlatformMisc.cpp b/Engine/Source/Runtime/Core/Private/HTML5/HTML5PlatformMisc.cpp
--- a/Engine/Source/Runtime/Core/Private/HTML5/HTML5PlatformMisc.cpp
+++ b/Engine/Source/Runtime/Core/Private/HTML5/HTML5PlatformMisc.cpp
@@ -34,7 +34,7 @@ GenericApplication* FHTML5Misc::CreateApplication()
 
 uint32 FHTML5Misc::GetCharKeyMap(uint16* KeyCodes, FString* KeyNames, uint32 MaxMappings)
 {
-	return FGenericPlatformMisc::GetStandardPrintableKeyMap(KeyCodes, KeyNames, MaxMappings, false, true);
+	return FGenericPlatformMisc::GetStandardPrintableKeyMap(KeyCodes, KeyNames, MaxMappings, true, false);
 }
 
 uint32 FHTML5Misc::GetKeyMap(uint16* KeyCodes, FString* KeyNames, uint32 MaxMappings)
diff --git a/Engine/Source/Runtime/Core/Private/Linux/LinuxPlatformMisc.cpp b/Engine/Source/Runtime/Core/Private/Linux/LinuxPlatformMisc.cpp
--- a/Engine/Source/Runtime/Core/Private/Linux/LinuxPlatformMisc.cpp
+++ b/Engine/Source/Runtime/Core/Private/Linux/LinuxPlatformMisc.cpp
@@ -220,7 +220,7 @@ void FLinuxPlatformMisc::PumpMessages( bool bFromMainLoop )
 
 uint32 FLinuxPlatformMisc::GetCharKeyMap(uint16* KeyCodes, FString* KeyNames, uint32 MaxMappings)
 {
-	return FGenericPlatformMisc::GetStandardPrintableKeyMap(KeyCodes, KeyNames, MaxMappings, false, true);
+	return FGenericPlatformMisc::GetStandardPrintableKeyMap(KeyCodes, KeyNames, MaxMappings, true, false);
 }
 
 void FLinuxPlatformMisc::LowLevelOutputDebugString(const TCHAR *Message)
diff --git a/Engine/Source/Runtime/Core/Private/Mac/MacPlatformMisc.cpp b/Engine/Source/Runtime/Core/Private/Mac/MacPlatformMisc.cpp
--- a/Engine/Source/Runtime/Core/Private/Mac/MacPlatformMisc.cpp
+++ b/Engine/Source/Runtime/Core/Private/Mac/MacPlatformMisc.cpp
@@ -484,7 +484,7 @@ void FMacPlatformMisc::PumpMessages( bool bFromMainLoop )
 
 uint32 FMacPlatformMisc::GetCharKeyMap(uint16* KeyCodes, FString* KeyNames, uint32 MaxMappings)
 {
-	return FGenericPlatformMisc::GetStandardPrintableKeyMap(KeyCodes, KeyNames, MaxMappings, false, true);
+	return FGenericPlatformMisc::GetStandardPrintableKeyMap(KeyCodes, KeyNames, MaxMappings, true, false);
 }
 
 uint32 FMacPlatformMisc::GetKeyMap( uint16* KeyCodes, FString* KeyNames, uint32 MaxMappings )
