Delayed the creation of FCanvas::WordWrapper until it is actually needed

Canvas doesn't actually use this that often.

diff --git a/Engine/Source/Runtime/Core/Private/Internationalization/ICUBreakIterator.cpp b/Engine/Source/Runtime/Core/Private/Internationalization/ICUBreakIterator.cpp
--- a/Engine/Source/Runtime/Core/Private/Internationalization/ICUBreakIterator.cpp
+++ b/Engine/Source/Runtime/Core/Private/Internationalization/ICUBreakIterator.cpp
@@ -29,6 +29,8 @@ FICUBreakIteratorManager& FICUBreakIteratorManager::Get()
 
 TWeakPtr<icu::BreakIterator> FICUBreakIteratorManager::CreateCharacterBoundaryIterator()
 {
+	check(IsInGameThread());
+
 	UErrorCode ICUStatus = U_ZERO_ERROR;
 	TSharedRef<icu::BreakIterator> Iterator = MakeShareable(icu::BreakIterator::createCharacterInstance(icu::Locale::getDefault(), ICUStatus));
 	AllocatedIterators.Add(Iterator);
@@ -37,6 +39,8 @@ TWeakPtr<icu::BreakIterator> FICUBreakIteratorManager::CreateCharacterBoundaryIt
 
 TWeakPtr<icu::BreakIterator> FICUBreakIteratorManager::CreateWordBreakIterator()
 {
+	check(IsInGameThread());
+
 	UErrorCode ICUStatus = U_ZERO_ERROR;
 	TSharedRef<icu::BreakIterator> Iterator = MakeShareable(icu::BreakIterator::createWordInstance(icu::Locale::getDefault(), ICUStatus));
 	AllocatedIterators.Add(Iterator);
@@ -45,6 +49,8 @@ TWeakPtr<icu::BreakIterator> FICUBreakIteratorManager::CreateWordBreakIterator()
 
 TWeakPtr<icu::BreakIterator> FICUBreakIteratorManager::CreateLineBreakIterator()
 {
+	check(IsInGameThread());
+
 	UErrorCode ICUStatus = U_ZERO_ERROR;
 	TSharedRef<icu::BreakIterator> Iterator = MakeShareable(icu::BreakIterator::createLineInstance(icu::Locale::getDefault(), ICUStatus));
 	AllocatedIterators.Add(Iterator);
@@ -53,6 +59,8 @@ TWeakPtr<icu::BreakIterator> FICUBreakIteratorManager::CreateLineBreakIterator()
 
 void FICUBreakIteratorManager::DestroyIterator(TWeakPtr<icu::BreakIterator>& InIterator)
 {
+	check(IsInGameThread());
+
 	TSharedPtr<icu::BreakIterator> Iterator = InIterator.Pin();
 	if(Iterator.IsValid())
 	{
diff --git a/Engine/Source/Runtime/Engine/Private/SubtitleManager.cpp b/Engine/Source/Runtime/Engine/Private/SubtitleManager.cpp
--- a/Engine/Source/Runtime/Engine/Private/SubtitleManager.cpp
+++ b/Engine/Source/Runtime/Engine/Private/SubtitleManager.cpp
@@ -215,7 +215,7 @@ void FSubtitleManager::SplitLinesToSafeZone( FCanvas* Canvas, FIntRect& Subtitle
 		// Word wrap into lines
 		TArray<FWrappedStringElement> Lines;
 		FTextSizingParameters RenderParms( 0.0f, 0.0f, SubtitleRegion.Width(), 0.0f, GEngine->GetSubtitleFont() );
-		UCanvas::WrapString( Canvas->WordWrapper, RenderParms, 0, *Concatenated, Lines );
+		Canvas->WrapString( RenderParms, 0, *Concatenated, Lines );
 
 		// Set up the times
 		Subtitle.Subtitles.Empty();
diff --git a/Engine/Source/Runtime/Engine/Private/UserInterface/Canvas.cpp b/Engine/Source/Runtime/Engine/Private/UserInterface/Canvas.cpp
--- a/Engine/Source/Runtime/Engine/Private/UserInterface/Canvas.cpp
+++ b/Engine/Source/Runtime/Engine/Private/UserInterface/Canvas.cpp
@@ -1083,6 +1083,14 @@ int32 FCanvas::DrawShadowedText( float StartX,float StartY,const FText& Text,con
 	return TextItem.DrawnSize.Y;	
 }
 
+void FCanvas::WrapString( FTextSizingParameters& Parameters, const float InCurX, const TCHAR* const pText, TArray<FWrappedStringElement>& out_Lines, FCanvasWordWrapper::FWrappedLineData* const OutWrappedLineData)
+{
+	if (!WordWrapper.IsValid())
+	{
+		WordWrapper = MakeShareable(new FCanvasWordWrapper());
+	}
+	UCanvas::WrapString( *WordWrapper, Parameters, InCurX, pText, out_Lines, OutWrappedLineData);
+}
 
 ENGINE_API void StringSize(const UFont* Font,int32& XL,int32& YL,const TCHAR* Text)
 {
@@ -1214,7 +1222,7 @@ void UCanvas::WrapString( FCanvasWordWrapper& Wrapper, FTextSizingParameters& Pa
 
 void UCanvas::WrapString( FTextSizingParameters& Parameters, const float InCurX, const TCHAR* const pText, TArray<FWrappedStringElement>& out_Lines, FCanvasWordWrapper::FWrappedLineData* const OutWrappedLineData)
 {
-	UCanvas::WrapString( Canvas->WordWrapper, Parameters, InCurX, pText, out_Lines, OutWrappedLineData);
+	Canvas->WrapString(Parameters, InCurX, pText, out_Lines, OutWrappedLineData);
 }
 
 /*-----------------------------------------------------------------------------
diff --git a/Engine/Source/Runtime/Engine/Public/CanvasTypes.h b/Engine/Source/Runtime/Engine/Public/CanvasTypes.h
--- a/Engine/Source/Runtime/Engine/Public/CanvasTypes.h
+++ b/Engine/Source/Runtime/Engine/Public/CanvasTypes.h
@@ -539,7 +539,7 @@ public:
 
 public:
 	/** Private class for handling word wrapping behavior. */
-	FCanvasWordWrapper WordWrapper;
+	TSharedPtr<FCanvasWordWrapper> WordWrapper;
 
 private:
 	/** Stack of SortKeys. All rendering is done using the top most sort key */
@@ -667,6 +667,8 @@ public:
 	
 	ENGINE_API int32 DrawShadowedText( float StartX, float StartY, const FText& Text, const UFont* Font, const FLinearColor& Color, const FLinearColor& ShadowColor = FLinearColor::Black );
 
+	ENGINE_API void WrapString( FTextSizingParameters& Parameters, const float InCurX, const TCHAR* const pText, TArray<FWrappedStringElement>& out_Lines, FCanvasWordWrapper::FWrappedLineData* const OutWrappedLineData = nullptr);
+
 	ENGINE_API void DrawNGon(const FVector2D& Center, const FColor& Color, int32 NumSides, float Radius);
 
 	/** 
