Made FICUBreakIteratorManager thread-safe

This helps Canvas, which can be used on the render-thread.

diff --git a/Engine/Source/Runtime/Core/Private/Internationalization/ICUBreakIterator.cpp b/Engine/Source/Runtime/Core/Private/Internationalization/ICUBreakIterator.cpp
--- a/Engine/Source/Runtime/Core/Private/Internationalization/ICUBreakIterator.cpp
+++ b/Engine/Source/Runtime/Core/Private/Internationalization/ICUBreakIterator.cpp
@@ -29,41 +29,43 @@ FICUBreakIteratorManager& FICUBreakIteratorManager::Get()
 
 TWeakPtr<icu::BreakIterator> FICUBreakIteratorManager::CreateCharacterBoundaryIterator()
 {
-	check(IsInGameThread());
-
 	UErrorCode ICUStatus = U_ZERO_ERROR;
 	TSharedRef<icu::BreakIterator> Iterator = MakeShareable(icu::BreakIterator::createCharacterInstance(icu::Locale::getDefault(), ICUStatus));
-	AllocatedIterators.Add(Iterator);
+	{
+		FScopeLock ScopeLock(&AllocatedIteratorsCS);
+		AllocatedIterators.Add(Iterator);
+	}
 	return Iterator;
 }
 
 TWeakPtr<icu::BreakIterator> FICUBreakIteratorManager::CreateWordBreakIterator()
 {
-	check(IsInGameThread());
-
 	UErrorCode ICUStatus = U_ZERO_ERROR;
 	TSharedRef<icu::BreakIterator> Iterator = MakeShareable(icu::BreakIterator::createWordInstance(icu::Locale::getDefault(), ICUStatus));
-	AllocatedIterators.Add(Iterator);
+	{
+		FScopeLock ScopeLock(&AllocatedIteratorsCS);
+		AllocatedIterators.Add(Iterator);
+	}
 	return Iterator;
 }
 
 TWeakPtr<icu::BreakIterator> FICUBreakIteratorManager::CreateLineBreakIterator()
 {
-	check(IsInGameThread());
-
 	UErrorCode ICUStatus = U_ZERO_ERROR;
 	TSharedRef<icu::BreakIterator> Iterator = MakeShareable(icu::BreakIterator::createLineInstance(icu::Locale::getDefault(), ICUStatus));
-	AllocatedIterators.Add(Iterator);
+	{
+		FScopeLock ScopeLock(&AllocatedIteratorsCS);
+		AllocatedIterators.Add(Iterator);
+	}
 	return Iterator;
 }
 
 void FICUBreakIteratorManager::DestroyIterator(TWeakPtr<icu::BreakIterator>& InIterator)
 {
-	check(IsInGameThread());
-
 	TSharedPtr<icu::BreakIterator> Iterator = InIterator.Pin();
 	if(Iterator.IsValid())
 	{
+		FScopeLock ScopeLock(&AllocatedIteratorsCS);
 		AllocatedIterators.Remove(Iterator);
 	}
 	InIterator.Reset();
diff --git a/Engine/Source/Runtime/Core/Private/Internationalization/ICUBreakIterator.h b/Engine/Source/Runtime/Core/Private/Internationalization/ICUBreakIterator.h
--- a/Engine/Source/Runtime/Core/Private/Internationalization/ICUBreakIterator.h
+++ b/Engine/Source/Runtime/Core/Private/Internationalization/ICUBreakIterator.h
@@ -25,6 +25,7 @@ public:
 
 private:
 	static FICUBreakIteratorManager* Singleton;
+	FCriticalSection AllocatedIteratorsCS;
 	TSet<TSharedPtr<icu::BreakIterator>> AllocatedIterators;
 };
 
