added ConsoleHelp.html output when executing "help" from the console
  (allows to browse/search CVars in a single interactive HTML document)

diff --git a/Engine/Source/Runtime/Core/Private/HAL/ConsoleManager.cpp b/Engine/Source/Runtime/Core/Private/HAL/ConsoleManager.cpp
--- a/Engine/Source/Runtime/Core/Private/HAL/ConsoleManager.cpp
+++ b/Engine/Source/Runtime/Core/Private/HAL/ConsoleManager.cpp
@@ -1508,32 +1508,6 @@ void CreateConsoleVariables()
 	}
 }
 
-#if !(UE_BUILD_SHIPPING || UE_BUILD_TEST)
-static void DumpHelp(FOutputDevice& Ar)
-{
-	Ar.Logf(TEXT("Console Help:"));
-	Ar.Logf(TEXT("============="));
-	Ar.Logf(TEXT(" "));
-	Ar.Logf(TEXT("A console variable is a engine wide key value pair. The key is a string usually starting with the subsystem prefix followed"));
-	Ar.Logf(TEXT("by '.' e.g. r.BloomQuality. The value can be of different tpe (e.g. float, int, string). A console command has no state associated with"));
-	Ar.Logf(TEXT("and gets executed immediately."));
-	Ar.Logf(TEXT(" "));
-	Ar.Logf(TEXT("Console variables can be put into ini files (e.g. ConsoleVariables.ini or BaseEngine.ini) with this syntax:"));
-	Ar.Logf(TEXT("<Console variable> = <value>"));
-	Ar.Logf(TEXT(" "));
-	Ar.Logf(TEXT("DumpConsoleCommands         Lists all console variables and commands that are registered (Some are not registered)"));
-	Ar.Logf(TEXT("<Console variable>          Get the console variable state"));
-	Ar.Logf(TEXT("<Console variable> ?        Get the console variable help text"));
-	Ar.Logf(TEXT("<Console variable> <value>  Set the console variable value"));
-	Ar.Logf(TEXT("<Console command> [Params]  Execute the console command with optional parameters"));
-}
-static FAutoConsoleCommandWithOutputDevice GConsoleCommandHelp(
-	TEXT("help"),
-	TEXT("Outputs some helptext to the console and the log"),
-	FConsoleCommandWithOutputDeviceDelegate::CreateStatic(DumpHelp)
-	);
-#endif // !(UE_BUILD_SHIPPING || UE_BUILD_TEST)
-
 // Naming conventions:
 //
 // Console variable should start with (suggestion):
diff --git a/Engine/Source/Runtime/Core/Private/Misc/Parse.cpp b/Engine/Source/Runtime/Core/Private/Misc/Parse.cpp
--- a/Engine/Source/Runtime/Core/Private/Misc/Parse.cpp
+++ b/Engine/Source/Runtime/Core/Private/Misc/Parse.cpp
@@ -1,6 +1,7 @@
 // Copyright 1998-2015 Epic Games, Inc. All Rights Reserved.
 
 #include "CorePrivatePCH.h"
+#include "LazyPrintf.h"
 
 #if !(UE_BUILD_SHIPPING || UE_BUILD_TEST)
 /**
@@ -28,12 +29,8 @@ public:
 		}
 	}
 
-private:
-
-	TSet<FString>		KnownNames;
 	const FString&		Pattern;
-
-	friend void ConsoleCommandLibrary_DumpLibrary(class UWorld* InWorld, FExec& SubSystem, const FString& Pattern, FOutputDevice& Ar);
+	TSet<FString>		KnownNames;
 };
 
 // 0 if gathering of names is deactivated
@@ -69,7 +66,6 @@ public:
 	}
 };
 
-/** Needed for the console command "DumpConsoleCommands" */
 void ConsoleCommandLibrary_DumpLibrary(UWorld* InWorld, FExec& SubSystem, const FString& Pattern, FOutputDevice& Ar)
 {
 	ConsoleCommandLibrary LocalConsoleCommandLibrary(Pattern);
@@ -78,7 +74,6 @@ void ConsoleCommandLibrary_DumpLibrary(UWorld* InWorld, FExec& SubSystem, const
 
 	bool bExecuted = SubSystem.Exec( InWorld, *Pattern, Null);
 
-
 	{
 		IConsoleManager::Get().ForEachConsoleObject(
 			FConsoleObjectVisitor::CreateStatic< TSet<FString>& >(
@@ -86,7 +81,6 @@ void ConsoleCommandLibrary_DumpLibrary(UWorld* InWorld, FExec& SubSystem, const
 			LocalConsoleCommandLibrary.KnownNames ) );
 	}
 
-
 	LocalConsoleCommandLibrary.KnownNames.Sort( TLess<FString>() );
 
 	for(TSet<FString>::TConstIterator It(LocalConsoleCommandLibrary.KnownNames); It; ++It)
@@ -104,6 +98,86 @@ void ConsoleCommandLibrary_DumpLibrary(UWorld* InWorld, FExec& SubSystem, const
 		Ar.Logf(TEXT("However Exec() returned true which means we either executed a command or the command parsing returned true where it shouldn't."));
 	}
 }
+
+void ConsoleCommandLibrary_DumpLibraryHTML(UWorld* InWorld, FExec& SubSystem, const FString& OutPath)
+{
+	ConsoleCommandLibrary LocalConsoleCommandLibrary(FString(TEXT("*")));
+
+	FOutputDeviceNull Null;
+
+	bool bExecuted = SubSystem.Exec( InWorld, *LocalConsoleCommandLibrary.Pattern, Null);
+
+	{
+		IConsoleManager::Get().ForEachConsoleObject(
+			FConsoleObjectVisitor::CreateStatic< TSet<FString>& >(
+			&FConsoleVariableDumpVisitor::OnConsoleVariable,
+			LocalConsoleCommandLibrary.KnownNames ) );
+	}
+
+	LocalConsoleCommandLibrary.KnownNames.Sort( TLess<FString>() );
+
+	FString TemplateFilename = FPaths::Combine(FPlatformProcess::BaseDir(), TEXT("../../Documentation/Extras"), TEXT("ConsoleHelpTemplate.html"));
+	FString TemplateFile;
+	if(FFileHelper::LoadFileToString(TemplateFile, *TemplateFilename, FFileHelper::EHashOptions::EnableVerify | FFileHelper::EHashOptions::ErrorMissingHash) )
+	{
+		// todo: do we need to create the directory?
+		FArchive* File = IFileManager::Get().CreateDebugFileWriter(*OutPath);
+
+		if(File)
+		{
+			FLazyPrintf LazyPrintf(*TemplateFile);
+
+			// title
+			LazyPrintf.PushParam(TEXT("UE4 Console Variables and Commands"));
+			// headline
+			LazyPrintf.PushParam(TEXT("Unreal Engine 4 Console Variables and Commands"));
+			// generated by
+			LazyPrintf.PushParam(TEXT("Unreal Engine 4 'DumpConsoleCommands'"));
+			// version
+			LazyPrintf.PushParam(TEXT("0.93"));
+			// date
+			LazyPrintf.PushParam(*FDateTime::Now().ToString());
+
+			FString AllData;
+
+			for(TSet<FString>::TConstIterator It(LocalConsoleCommandLibrary.KnownNames); It; ++It)
+			{
+				const FString& Name = *It;
+
+				auto CVar = IConsoleManager::Get().FindConsoleVariable(*Name);
+
+				if(CVar)
+				{
+					FString Help = CVar->GetHelp();
+
+					Help = Help.ReplaceCharWithEscapedChar();
+
+					//{name: "r.SetRes", help:"To change the screen/window resolution."},
+					FString DataLine = FString::Printf(TEXT("{name: \"%s\", help:\"%s\"},\r\n"), *Name, *Help);
+
+					AllData += DataLine;
+				}
+			}
+
+			LazyPrintf.PushParam(*AllData);
+
+			auto AnsiHelp = StringCast<ANSICHAR>(*LazyPrintf.GetResultString());
+			File->Serialize((ANSICHAR*)AnsiHelp.Get(), AnsiHelp.Length());
+
+			delete File;
+			File = 0;
+		}
+	}
+
+/*
+	// the pattern (e.g. Motion*) should not really trigger the execution
+	if(bExecuted)
+	{
+		Ar.Logf(TEXT("ERROR: The function was supposed to only find matching commands but not have any side effect."));
+		Ar.Logf(TEXT("However Exec() returned true which means we either executed a command or the command parsing returned true where it shouldn't."));
+	}
+*/
+}
 #endif // !(UE_BUILD_SHIPPING || UE_BUILD_TEST)
 
 //
diff --git a/Engine/Source/Runtime/Core/Public/Containers/LazyPrintf.h b/Engine/Source/Runtime/Core/Public/Containers/LazyPrintf.h
new file mode 100644
--- /dev/null
+++ b/Engine/Source/Runtime/Core/Public/Containers/LazyPrintf.h
@@ -0,0 +1,70 @@
+// Copyright 1998-2015 Epic Games, Inc. All Rights Reserved.
+
+#pragma once
+
+
+// to avoid limits with the Printf parameter count, for more readability and type safety
+class FLazyPrintf
+{
+public:
+	// constructor
+	FLazyPrintf(const TCHAR* InputWithPercentS)
+		: CurrentInputPos(InputWithPercentS)
+	{
+		// to avoid reallocations
+		CurrentState.Empty(50 * 1024);
+	}
+
+	FString GetResultString()
+	{
+		// internal error more %s than %s in MaterialTemplate.usf
+		check(!ProcessUntilPercentS());
+
+		// copy all remaining input data
+		CurrentState += CurrentInputPos;
+
+		return CurrentState;
+	}
+
+	// %s
+	void PushParam(const TCHAR* Data)
+	{
+		if(ProcessUntilPercentS())
+		{
+			CurrentState += Data;
+		}
+		else
+		{
+			// internal error, more ReplacePercentS() calls than %s in MaterialTemplate.usf
+			check(0);
+		}
+	}
+
+private:
+
+	// @param Pattern e.g. TEXT("%s")
+	bool ProcessUntilPercentS()
+	{
+		const TCHAR* Found = FCString::Strstr(CurrentInputPos, TEXT("%s"));
+
+		if(Found == 0)
+		{
+			return false;
+		}
+
+		// copy from input until %s
+		while(CurrentInputPos < Found)
+		{
+			// can cause reallocations we could avoid
+			CurrentState += *CurrentInputPos++;
+		}
+
+		// jump over %s
+		CurrentInputPos += 2;
+
+		return true;
+	}
+
+	const TCHAR* CurrentInputPos;
+	FString CurrentState;
+};
diff --git a/Engine/Source/Runtime/Core/Public/Misc/StringUtility.h b/Engine/Source/Runtime/Core/Public/Misc/StringUtility.h
--- a/Engine/Source/Runtime/Core/Public/Misc/StringUtility.h
+++ b/Engine/Source/Runtime/Core/Public/Misc/StringUtility.h
@@ -10,4 +10,6 @@
 #if !(UE_BUILD_SHIPPING || UE_BUILD_TEST)
 /** Needed for the console command "DumpConsoleCommands" */
 CORE_API void ConsoleCommandLibrary_DumpLibrary(class UWorld* InWorld, FExec& SubSystem, const FString& Pattern, FOutputDevice& Ar);
+/** Needed for the console command "Help" */
+CORE_API void ConsoleCommandLibrary_DumpLibraryHTML(class UWorld* InWorld, FExec& SubSystem, const FString& OutPath);
 #endif // !(UE_BUILD_SHIPPING || UE_BUILD_TEST)
diff --git a/Engine/Source/Runtime/Engine/Private/Materials/HLSLMaterialTranslator.h b/Engine/Source/Runtime/Engine/Private/Materials/HLSLMaterialTranslator.h
--- a/Engine/Source/Runtime/Engine/Private/Materials/HLSLMaterialTranslator.h
+++ b/Engine/Source/Runtime/Engine/Private/Materials/HLSLMaterialTranslator.h
@@ -22,6 +22,7 @@
 #include "MaterialUniformExpressions.h"
 #include "ParameterCollection.h"
 #include "Materials/MaterialParameterCollection.h"
+#include "LazyPrintf.h"
 
 /** @return the number of components in a vector type. */
 uint32 GetNumComponents(EMaterialValueType Type)
@@ -100,72 +101,6 @@ struct FShaderCodeChunk
 	{}
 };
 
-// to avoid limits with the Printf parameter count, for more readability and type safety
-class FLazyPrintf
-{
-public:
-	// constructor
-	FLazyPrintf(const TCHAR* InputWithPercentS)
-		: CurrentInputPos(InputWithPercentS)
-	{
-		// to avoid reallocations
-		CurrentState.Empty(50 * 1024);
-	}
-
-	FString GetResultString()
-	{
-		// internal error more %s than %s in MaterialTemplate.usf
-		check(!ProcessUntilPercentS());
-
-		// copy all remaining input data
-		CurrentState += CurrentInputPos;
-
-		return CurrentState;
-	}
-
-	// %s
-	void PushParam(const TCHAR* Data)
-	{
-		if(ProcessUntilPercentS())
-		{
-			CurrentState += Data;
-		}
-		else
-		{
-			// internal error, more ReplacePercentS() calls than %s in MaterialTemplate.usf
-			check(0);
-		}
-	}
-
-private:
-
-	// @param Pattern e.g. TEXT("%s")
-	bool ProcessUntilPercentS()
-	{
-		const TCHAR* Found = FCString::Strstr(CurrentInputPos, TEXT("%s"));
-
-		if(Found == 0)
-		{
-			return false;
-		}
-
-		// copy from input until %s
-		while(CurrentInputPos < Found)
-		{
-			// can cause reallocations we could avoid
-			CurrentState += *CurrentInputPos++;
-		}
-
-		// jump over %s
-		CurrentInputPos += 2;
-
-		return true;
-	}
-
-	const TCHAR* CurrentInputPos;
-	FString CurrentState;
-};
-
 class FHLSLMaterialTranslator : public FMaterialCompiler
 {
 protected:
diff --git a/Engine/Source/Runtime/Engine/Private/UnrealEngine.cpp b/Engine/Source/Runtime/Engine/Private/UnrealEngine.cpp
--- a/Engine/Source/Runtime/Engine/Private/UnrealEngine.cpp
+++ b/Engine/Source/Runtime/Engine/Private/UnrealEngine.cpp
@@ -2671,6 +2671,43 @@ bool UEngine::HandleHotReloadCommand( const TCHAR* Cmd, FOutputDevice& Ar )
 }
 #endif // WITH_HOT_RELOAD
 
+
+#if !(UE_BUILD_SHIPPING || UE_BUILD_TEST)
+static void DumpHelp(UWorld* InWorld)
+{
+	UE_LOG(LogEngine, Display, TEXT("Console Help:"));
+	UE_LOG(LogEngine, Display, TEXT("============="));
+	UE_LOG(LogEngine, Display, TEXT(" "));
+	UE_LOG(LogEngine, Display, TEXT("A console variable is a engine wide key value pair. The key is a string usually starting with the subsystem prefix followed"));
+	UE_LOG(LogEngine, Display, TEXT("by '.' e.g. r.BloomQuality. The value can be of different tpe (e.g. float, int, string). A console command has no state associated with"));
+	UE_LOG(LogEngine, Display, TEXT("and gets executed immediately."));
+	UE_LOG(LogEngine, Display, TEXT(" "));
+	UE_LOG(LogEngine, Display, TEXT("Console variables can be put into ini files (e.g. ConsoleVariables.ini or BaseEngine.ini) with this syntax:"));
+	UE_LOG(LogEngine, Display, TEXT("<Console variable> = <value>"));
+	UE_LOG(LogEngine, Display, TEXT(" "));
+	UE_LOG(LogEngine, Display, TEXT("DumpConsoleCommands         Lists all console variables and commands that are registered (Some are not registered)"));
+	UE_LOG(LogEngine, Display, TEXT("<Console variable>          Get the console variable state"));
+	UE_LOG(LogEngine, Display, TEXT("<Console variable> ?        Get the console variable help text"));
+	UE_LOG(LogEngine, Display, TEXT("<Console variable> <value>  Set the console variable value"));
+	UE_LOG(LogEngine, Display, TEXT("<Console command> [Params]  Execute the console command with optional parameters"));
+
+	UE_LOG(LogEngine, Display, TEXT(" "));
+
+	FString FilePath = FPaths::GameSavedDir() + TEXT("ConsoleHelp.html");
+
+	UE_LOG(LogEngine, Display, TEXT("To browse console variables open this: '%s'"), *FilePath);
+	UE_LOG(LogEngine, Display, TEXT(" "));
+
+	ConsoleCommandLibrary_DumpLibraryHTML(InWorld, *GEngine, FilePath);
+}
+static FAutoConsoleCommandWithWorld GConsoleCommandHelp(
+	TEXT("help"),
+	TEXT("Outputs some helptext to the console and the log"),
+	FConsoleCommandWithWorldDelegate::CreateStatic(DumpHelp)
+	);
+#endif // !(UE_BUILD_SHIPPING || UE_BUILD_TEST)
+
+
 bool UEngine::HandleDumpConsoleCommandsCommand( const TCHAR* Cmd, FOutputDevice& Ar, UWorld* InWorld )
 {
 	Ar.Logf(TEXT("DumpConsoleCommands: %s*"), Cmd);
