UE-5577 UE logo is distorted on splash screen on iPhone 4s
was picking the wrong image for iPhone4/4s
#ios

diff --git a/Engine/Source/Runtime/Core/Private/IOS/IOSAppDelegate.cpp b/Engine/Source/Runtime/Core/Private/IOS/IOSAppDelegate.cpp
--- a/Engine/Source/Runtime/Core/Private/IOS/IOSAppDelegate.cpp
+++ b/Engine/Source/Runtime/Core/Private/IOS/IOSAppDelegate.cpp
@@ -431,11 +431,15 @@ void InstallSignalHandlers()
 	}
 	else
 	{
-		if (MainFrame.size.height == 320 && !self.bDeviceInPortraitMode)
+		if (MainFrame.size.height == 320 && MainFrame.size.width != 480 && !self.bDeviceInPortraitMode)
 		{
 			[ImageString appendString:@"-568h"];
 			orient = UIImageOrientationLeft;
 		}
+		else if (MainFrame.size.height == 320 && MainFrame.size.width == 480 && !self.bDeviceInPortraitMode)
+		{
+			orient = UIImageOrientationLeft;
+		}
 		else if (MainFrame.size.height == 568)
 		{
 			[ImageString appendString:@"-568h"];
