fixed check() when cvars are getting set in a certain order
https://udn.unrealengine.com/questions/229788/46-cvar-issue.html
https://udn.unrealengine.com/questions/229813/crash-when-initializing-console-variable-from-ini.html

diff --git a/Engine/Source/Runtime/Core/Private/Misc/ConfigCacheIni.cpp b/Engine/Source/Runtime/Core/Private/Misc/ConfigCacheIni.cpp
--- a/Engine/Source/Runtime/Core/Private/Misc/ConfigCacheIni.cpp
+++ b/Engine/Source/Runtime/Core/Private/Misc/ConfigCacheIni.cpp
@@ -2883,7 +2883,8 @@ void FConfigCacheIni::LoadConsoleVariablesFromINI()
 				else
 				{
 					// Create a dummy that is used when someone registers the variable later on.
-					IConsoleManager::Get().RegisterConsoleVariable(*KeyString, *ValueString, TEXT("IAmNoRealVariable"), (uint32)ECVF_Unregistered | (uint32)ECVF_CreatedFromIni);
+					IConsoleManager::Get().RegisterConsoleVariable(*KeyString, *ValueString, TEXT("IAmNoRealVariable"),
+						(uint32)ECVF_Unregistered | (uint32)ECVF_CreatedFromIni | (uint32)ECVF_SetByConsoleVariablesIni);
 				}
 			}
 			bFoundSection = true;
