UE-5663 Mobile orientation stuck in portrait mode
Implemented updated interfaces for auto-rotate and supported orientations
Fixed crash when only one orientation is supported
#ios

diff --git a/Engine/Source/Runtime/Core/Private/IOS/IOSView.cpp b/Engine/Source/Runtime/Core/Private/IOS/IOSView.cpp
--- a/Engine/Source/Runtime/Core/Private/IOS/IOSView.cpp
+++ b/Engine/Source/Runtime/Core/Private/IOS/IOSView.cpp
@@ -486,23 +486,24 @@ id<MTLDevice> GMetalDevice = nil;
 }
 
 /**
- * Tell the OS that our view controller can auto-rotate between the two landscape modes
+ * Tell the OS what the default supported orientations are
  */
-- (BOOL)shouldAutorotateToInterfaceOrientation:(UIInterfaceOrientation)interfaceOrientation
+- (NSUInteger)supportedInterfaceOrientations
 {
-	//BOOL bIsValidOrientation;
-
-	//// will the app let us rotate?
-	//BOOL bCanRotate = [IPhoneAppDelegate GetDelegate].bCanAutoRotate;
-	//
-	//// is it one of the valid orientations to rotate to?
-	//bIsValidOrientation = bCanRotate && 
-	//	(((interfaceOrientation == UIInterfaceOrientationLandscapeLeft) && ![IOSAppDelegate GetDelegate].bDeviceInPortraitMode) || 
-	//	((interfaceOrientation == UIInterfaceOrientationLandscapeRight) && ![IOSAppDelegate GetDelegate].bDeviceInPortraitMode) || 
-	//	((interfaceOrientation == UIInterfaceOrientationPortrait) && [IOSAppDelegate GetDelegate].bDeviceInPortraitMode) ||
-	//	((interfaceOrientation == UIInterfaceOrientationPortraitUpsideDown) && [IOSAppDelegate GetDelegate].bDeviceInPortraitMode));
-	//	
-	//return bIsValidOrientation;
+	return UIInterfaceOrientationMaskAll;
+}
+
+/**
+ * Tell the OS that our view controller can auto-rotate between supported orientations
+ */
+- (BOOL)shouldAutorotate
+{
+	NSUInteger orient = [[UIApplication sharedApplication] supportedInterfaceOrientationsForWindow: [IOSAppDelegate GetDelegate].Window];
+	if (orient == UIInterfaceOrientationMaskPortrait || orient == UIInterfaceOrientationMaskPortraitUpsideDown
+		|| orient == UIInterfaceOrientationMaskLandscapeLeft || orient == UIInterfaceOrientationMaskLandscapeRight)
+	{
+		return NO;
+	}
 	return YES;
 }
 
