UE-6361 Fix typo in FMatrix::InverseFast message
https://github.com/EpicGames/UnrealEngine/pull/661 (from Pierdek)

diff --git a/Engine/Source/Runtime/Core/Public/Math/UnrealMatrix.h b/Engine/Source/Runtime/Core/Public/Math/UnrealMatrix.h
--- a/Engine/Source/Runtime/Core/Public/Math/UnrealMatrix.h
+++ b/Engine/Source/Runtime/Core/Public/Math/UnrealMatrix.h
@@ -264,8 +264,8 @@ inline FMatrix FMatrix::InverseFast() const
 		GetScaledAxis( EAxis::Y ).IsNearlyZero(SMALL_NUMBER) && 
 		GetScaledAxis( EAxis::Z ).IsNearlyZero(SMALL_NUMBER) ) 
 	{
-		UE_LOG(LogUnrealMath, Error, TEXT("FMatrix::Inverse(), trying to invert a NIL matrix, this results in NaNs! Use Inverse() instead."));
-		ensureMsgf(false, TEXT("FMatrix::Inverse(), trying to invert a NIL matrix, this results in NaNs! Use Inverse() instead."));
+		UE_LOG(LogUnrealMath, Error, TEXT("FMatrix::InverseFast(), trying to invert a NIL matrix, this results in NaNs! Use Inverse() instead."));
+		ensureMsgf(false, TEXT("FMatrix::InverseFast(), trying to invert a NIL matrix, this results in NaNs! Use Inverse() instead."));
 	}
 #endif
 	FMatrix Result;
