UE-6116  Add RandBool() to UnrealMathUtility.h (from SaxonRah)
https://github.com/EpicGames/UnrealEngine/pull/647

diff --git a/Engine/Source/Runtime/Core/Public/Math/UnrealMathUtility.h b/Engine/Source/Runtime/Core/Public/Math/UnrealMathUtility.h
--- a/Engine/Source/Runtime/Core/Public/Math/UnrealMathUtility.h
+++ b/Engine/Source/Runtime/Core/Public/Math/UnrealMathUtility.h
@@ -112,6 +112,12 @@ struct FMath : public FPlatformMath
 		return InMin + (InMax - InMin) * FRand();
 	}
 
+	/** Util to generate a random boolean. */
+	static FORCEINLINE bool RandBool()
+	{
+		return (RandRange(0,1) == 1) ? true : false;
+	}
+
 	/** Return a uniformly distributed random unit length vector = point on the unit sphere surface. */
 	static FVector VRand();
 	
diff --git a/Engine/Source/Runtime/Engine/Private/KismetMathLibrary.cpp b/Engine/Source/Runtime/Engine/Private/KismetMathLibrary.cpp
--- a/Engine/Source/Runtime/Engine/Private/KismetMathLibrary.cpp
+++ b/Engine/Source/Runtime/Engine/Private/KismetMathLibrary.cpp
@@ -33,7 +33,7 @@ UKismetMathLibrary::UKismetMathLibrary(const FObjectInitializer& ObjectInitializ
 
 bool UKismetMathLibrary::RandomBool()
 {
-	return (FMath::RandRange(0,1) == 1) ? true : false;
+	return FMath::RandBool();
 }
 
 bool UKismetMathLibrary::Not_PreBool(bool A)
