Stopped the default constructor of FText allocating memory

This gets called when default constructing widget arguments, and it now defaults to constructing from the already allocated string from the static "empty" text.

diff --git a/Engine/Source/Runtime/Core/Private/Internationalization/Text.cpp b/Engine/Source/Runtime/Core/Private/Internationalization/Text.cpp
--- a/Engine/Source/Runtime/Core/Private/Internationalization/Text.cpp
+++ b/Engine/Source/Runtime/Core/Private/Internationalization/Text.cpp
@@ -221,6 +221,8 @@ FArchive& operator<<(FArchive& Ar, FNumberFormattingOptions& Value)
 bool FText::bEnableErrorCheckingResults = ENABLE_TEXT_ERROR_CHECKING_RESULTS;
 bool FText::bSuppressWarnings = false;
 
+const FText FText::StaticEmptyText = FText(FText::EInitToEmptyString::Value);
+
 //Some error text formats
 const FText FText::UnusedArgumentsError = LOCTEXT("Error_UnusedArguments", "ERR: The following arguments were unused ({0}).");
 const FText FText::CommentStartError = LOCTEXT("Error_CommentDoesntStartWithQMark", "ERR: The Comment for Arg Block {0} does not start with a '?'.");
@@ -233,6 +235,12 @@ const FText FText::UnEscapedCloseBraceOutsideOfArgumentBlock = LOCTEXT("Error_Un
 const FText FText::SerializationFailureError = LOCTEXT("Error_SerializationFailure", "ERR: Transient text cannot be serialized \"{0}\".");
 
 FText::FText()
+	: DisplayString( StaticEmptyText.DisplayString )
+	, Flags(0)
+{
+}
+
+FText::FText( EInitToEmptyString )
 	: DisplayString( new FString() )
 	, Flags(0)
 {
diff --git a/Engine/Source/Runtime/Core/Public/Internationalization/Text.h b/Engine/Source/Runtime/Core/Public/Internationalization/Text.h
--- a/Engine/Source/Runtime/Core/Public/Internationalization/Text.h
+++ b/Engine/Source/Runtime/Core/Public/Internationalization/Text.h
@@ -123,7 +123,6 @@ public:
 
 	static const FText& GetEmpty()
 	{
-		static const FText StaticEmptyText = FText();
 		return StaticEmptyText;
 	}
 
@@ -298,6 +297,10 @@ public:
 
 private:
 
+	/** Special constructor used to create StaticEmptyText without also allocating a history object */
+	enum class EInitToEmptyString : uint8 { Value };
+	explicit FText( EInitToEmptyString );
+
 	explicit FText( FString InSourceString );
 
 	explicit FText( FString InSourceString, FString InNamespace, FString InKey, int32 InFlags=0 );
@@ -355,6 +358,8 @@ private:
 	static bool bSuppressWarnings;
 
 public:
+	static const FText StaticEmptyText;
+
 	//Some error text formats
 	static const FText UnusedArgumentsError;
 	static const FText CommentStartError;
