Implemented FMacApplication::IsCursorDirectlyOverSlateWindow()

diff --git a/Engine/Source/Runtime/Core/Private/Mac/MacApplication.cpp b/Engine/Source/Runtime/Core/Private/Mac/MacApplication.cpp
--- a/Engine/Source/Runtime/Core/Private/Mac/MacApplication.cpp
+++ b/Engine/Source/Runtime/Core/Private/Mac/MacApplication.cpp
@@ -316,6 +316,14 @@ static TSharedPtr< FMacWindow > FindWindowByNSWindow( const TArray< TSharedRef<
 	return TSharedPtr< FMacWindow >( NULL );
 }
 
+bool FMacApplication::IsCursorDirectlyOverSlateWindow() const
+{
+	SCOPED_AUTORELEASE_POOL;
+	const NSInteger WindowNumber = [NSWindow windowNumberAtPoint:[NSEvent mouseLocation] belowWindowWithWindowNumber:0];
+	NSWindow* const Window = [NSApp windowWithWindowNumber:WindowNumber];
+	return Window && [Window isKindOfClass:[FCocoaWindow class]];
+}
+
 void FMacApplication::ProcessEvent(FMacEvent const* const Event)
 {
 	// Must have an event
diff --git a/Engine/Source/Runtime/Core/Public/Mac/MacApplication.h b/Engine/Source/Runtime/Core/Public/Mac/MacApplication.h
--- a/Engine/Source/Runtime/Core/Public/Mac/MacApplication.h
+++ b/Engine/Source/Runtime/Core/Public/Mac/MacApplication.h
@@ -48,6 +48,8 @@ public:
 
 	virtual FModifierKeysState GetModifierKeys() const override;
 
+	virtual bool IsCursorDirectlyOverSlateWindow() const override;
+
 	virtual FPlatformRect GetWorkArea( const FPlatformRect& CurrentWindow ) const override;
 
 	virtual EWindowTitleAlignment::Type GetWindowTitleAlignment() const override
