Fixed dragging on Mac

diff --git a/Engine/Source/Runtime/Core/Private/Mac/MacApplication.cpp b/Engine/Source/Runtime/Core/Private/Mac/MacApplication.cpp
--- a/Engine/Source/Runtime/Core/Private/Mac/MacApplication.cpp
+++ b/Engine/Source/Runtime/Core/Private/Mac/MacApplication.cpp
@@ -321,7 +321,7 @@ bool FMacApplication::IsCursorDirectlyOverSlateWindow() const
 	SCOPED_AUTORELEASE_POOL;
 	const NSInteger WindowNumber = [NSWindow windowNumberAtPoint:[NSEvent mouseLocation] belowWindowWithWindowNumber:0];
 	NSWindow* const Window = [NSApp windowWithWindowNumber:WindowNumber];
-	return Window && [Window isKindOfClass:[FCocoaWindow class]];
+	return Window && ([Window isKindOfClass:[FCocoaWindow class]] || [Window isKindOfClass:[FMouseCaptureWindow class]]);
 }
 
 void FMacApplication::ProcessEvent(FMacEvent const* const Event)
