Implemented IsCursorDirectlyOverSlateWindow() for Linux.

diff --git a/Engine/Source/Runtime/Core/Private/Linux/LinuxApplication.cpp b/Engine/Source/Runtime/Core/Private/Linux/LinuxApplication.cpp
--- a/Engine/Source/Runtime/Core/Private/Linux/LinuxApplication.cpp
+++ b/Engine/Source/Runtime/Core/Private/Linux/LinuxApplication.cpp
@@ -70,6 +70,7 @@ FLinuxApplication::FLinuxApplication() : GenericApplication( MakeShareable( new
 	, bIsMouseCursorLocked(false)
 	, bIsMouseCaptureEnabled(false)
 	, bHasLoadedInputPlugins(false)
+	, bInsideOwnWindow(false)
 {
 	bUsingHighPrecisionMouseInput = false;
 	bAllowedToDeferMessageProcessing = true;
@@ -156,6 +157,10 @@ void FLinuxApplication::PumpMessages( const float TimeDelta )
 	FPlatformMisc::PumpMessages( true );
 }
 
+bool FLinuxApplication::IsCursorDirectlyOverSlateWindow() const
+{
+	return bInsideOwnWindow;
+}
 
 void FLinuxApplication::AddPendingEvent( SDL_Event SDLEvent )
 {
@@ -708,15 +713,24 @@ void FLinuxApplication::ProcessDeferredMessage( SDL_Event Event )
 						if (CurrentEventWindow.IsValid())
 						{
 							MessageHandler->OnCursorSet();
+
+							bInsideOwnWindow = true;
+							UE_LOG(LogLinuxWindow, Verbose, TEXT("Entered one of application windows"));
 						}
 					}
 					break;
 
 				case SDL_WINDOWEVENT_LEAVE:
 					{
-						if (CurrentEventWindow.IsValid() && GetCapture() != NULL)
+						if (CurrentEventWindow.IsValid())
 						{
-							UpdateMouseCaptureWindow((SDL_HWindow)GetCapture());
+							if (GetCapture() != nullptr)
+							{
+								UpdateMouseCaptureWindow((SDL_HWindow)GetCapture());
+							}
+
+							bInsideOwnWindow = false;
+							UE_LOG(LogLinuxWindow, Verbose, TEXT("Left an application window we were hovering above."));
 						}
 					}
 					break;
diff --git a/Engine/Source/Runtime/Core/Public/Linux/LinuxApplication.h b/Engine/Source/Runtime/Core/Public/Linux/LinuxApplication.h
--- a/Engine/Source/Runtime/Core/Public/Linux/LinuxApplication.h
+++ b/Engine/Source/Runtime/Core/Public/Linux/LinuxApplication.h
@@ -54,13 +54,17 @@ public:
 	void AddPendingEvent( SDL_Event event );
 
 	void OnMouseCursorLock( bool bLockEnabled );
-	
+
 	void RemoveEventWindow(SDL_HWindow Window);
 
 	EWindowZone::Type WindowHitTest( const TSharedPtr< FLinuxWindow > &window, int x, int y );
+
 	TSharedPtr< FLinuxWindow > FindWindowBySDLWindow( SDL_Window *win );
 
+	virtual bool IsCursorDirectlyOverSlateWindow() const override;
+
 private:
+
 	FLinuxApplication();
 
 	TCHAR ConvertChar( SDL_Keysym Keysym );
@@ -132,6 +136,8 @@ private:
 	/** Whether input plugins have been loaded */
 	bool bHasLoadedInputPlugins;
 
+	/** Whether we entered one of our own windows */
+	bool bInsideOwnWindow;
 };
 
 
