Fix splash screen persisting sometimes (and causing crashes).

diff --git a/Engine/Source/Runtime/Core/Private/Linux/LinuxPlatformSplash.cpp b/Engine/Source/Runtime/Core/Private/Linux/LinuxPlatformSplash.cpp
--- a/Engine/Source/Runtime/Core/Private/Linux/LinuxPlatformSplash.cpp
+++ b/Engine/Source/Runtime/Core/Private/Linux/LinuxPlatformSplash.cpp
@@ -68,7 +68,13 @@ static FString GIconPath;
 
 static int32 SplashWidth = 0, SplashHeight = 0;
 static unsigned char *ScratchSpace = nullptr;
-static volatile int32 ThreadState = 0;
+/** 
+ * Used for communication with a splash thread. 
+ * Negative values -> thread is (or must be) stopped. 
+ * >= 0 -> thread is running
+ * > 0 -> there is an update that necessitates (re)rendering the string.
+ */
+static volatile int32 ThreadState = -1;
 static int32 SplashBPP = 0;
 
 //////////////////////////////////
@@ -793,8 +799,8 @@ static int StartSplashScreenThread(void *ptr)
 	SDL_GL_DeleteContext(Context);
 	FMemory::Free(ScratchSpace);
 
-	// set the thread state to 0 to let the caller know we're done (FIXME: can be done without busy-loops)
-	ThreadState = 0;
+	// set the thread state to -1 to let the caller know we're done (FIXME: can be done without busy-loops)
+	ThreadState = -1;
 
 	return 0;
 }
@@ -946,13 +952,11 @@ void FLinuxPlatformSplash::Hide()
 {
 #if WITH_EDITOR
 	// signal thread it's time to quit
-	GSplashThread = nullptr;
-
-	if (ThreadState > 0)	// if there's a thread at all...
+	if (ThreadState >= 0)	// if there's a thread at all...
 	{
 		ThreadState = -99;
 		// wait for the thread to be done before tearing it tearing it down (it will set the ThreadState to 0)
-		while (ThreadState != 0)
+		while (ThreadState != -1)
 		{
 			// busy loop!
 			FPlatformProcess::Sleep(0.01f);
@@ -961,6 +965,8 @@ void FLinuxPlatformSplash::Hide()
 		// tear down resources that thread used
 		LinuxSplash_TearDownSplashResources();
 	}
+
+	GSplashThread = nullptr;
 #endif
 }
 
@@ -989,7 +995,7 @@ void FLinuxPlatformSplash::SetSplashText( const SplashTextType::Type InType, con
 			}
 		}
 
-		if (bWasUpdated)
+		if (bWasUpdated && GSplashThread && ThreadState >= 0)
 		{
 			ThreadState++;
 		}
