Fixed crash caused by FLegacyTextHelper::Format mutating the static empty FText instance

diff --git a/Engine/Source/Runtime/Core/Private/Internationalization/LegacyText.cpp b/Engine/Source/Runtime/Core/Private/Internationalization/LegacyText.cpp
--- a/Engine/Source/Runtime/Core/Private/Internationalization/LegacyText.cpp
+++ b/Engine/Source/Runtime/Core/Private/Internationalization/LegacyText.cpp
@@ -251,8 +251,7 @@ public:
 
 		const FString& PatternString = Pattern.ToString();
 
-		FText Result;
-		FString& ResultString = Result.DisplayString.Get();
+		FString ResultString;
 
 		EEscapeState::Type EscapeState = EEscapeState::None;
 		EBlockState::Type BlockState = EBlockState::None;
@@ -402,11 +401,11 @@ public:
 			}
 		}
 
+		FText Result = FText(MoveTemp(ResultString));
 		if (!GIsEditor)
 		{
 			Result.Flags = Result.Flags | ETextFlag::Transient;
 		}
-
 		return Result;
 	}
 };
