Fixed WindowsPlatform compiler warning about initializer syntax when compiling using Clang

diff --git a/Engine/Source/Runtime/Core/Private/Windows/WindowsPlatformMisc.cpp b/Engine/Source/Runtime/Core/Private/Windows/WindowsPlatformMisc.cpp
--- a/Engine/Source/Runtime/Core/Private/Windows/WindowsPlatformMisc.cpp
+++ b/Engine/Source/Runtime/Core/Private/Windows/WindowsPlatformMisc.cpp
@@ -368,7 +368,8 @@ static void InitSHAHashes()
 static void SetProcessMemoryLimit( SIZE_T ProcessMemoryLimitMB )
 {
 	HANDLE JobObject = ::CreateJobObject(nullptr, TEXT("UE4-JobObject"));
-	JOBOBJECT_EXTENDED_LIMIT_INFORMATION JobLimitInfo = {0};
+	JOBOBJECT_EXTENDED_LIMIT_INFORMATION JobLimitInfo;
+	FMemory::MemZero( JobLimitInfo );
 	JobLimitInfo.ProcessMemoryLimit = 1024*1024*ProcessMemoryLimitMB;
 	JobLimitInfo.BasicLimitInformation.LimitFlags = JOB_OBJECT_LIMIT_PROCESS_MEMORY;
 	const BOOL bSetJob = ::SetInformationJobObject(JobObject,JobObjectExtendedLimitInformation,&JobLimitInfo,sizeof(JobLimitInfo));
