#stats - Add STATSGROUP_Math to Core.

- Add stat counters for Rotator->Quat and Quat->Rotator conversion.

diff --git a/Engine/Source/Runtime/Core/Private/Math/UnrealMath.cpp b/Engine/Source/Runtime/Core/Private/Math/UnrealMath.cpp
--- a/Engine/Source/Runtime/Core/Private/Math/UnrealMath.cpp
+++ b/Engine/Source/Runtime/Core/Private/Math/UnrealMath.cpp
@@ -7,6 +7,10 @@
 #include "CorePrivatePCH.h"
 
 DEFINE_LOG_CATEGORY(LogUnrealMath);
+
+DEFINE_STAT(STAT_MathConvertQuatToRotator);
+DEFINE_STAT(STAT_MathConvertRotatorToQuat);
+
 /*-----------------------------------------------------------------------------
 	Globals
 -----------------------------------------------------------------------------*/
@@ -260,6 +264,8 @@ CORE_API FVector FRotator::Vector() const
 
 FQuat FRotator::Quaternion() const
 {
+	SCOPE_CYCLE_COUNTER(STAT_MathConvertRotatorToQuat);
+
 #if USE_MATRIX_ROTATOR 
 	FQuat RotationMatrix = FQuat( FRotationMatrix( *this ) );
 #endif
@@ -378,6 +384,69 @@ void FMatrix::DebugPrint() const
 	UE_LOG(LogUnrealMath, Log, TEXT("%s"), *ToString());
 }
 
+//////////////////////////////////////////////////////////////////////////
+// FQuat
+
+FRotator FQuat::Rotator() const
+{
+	SCOPE_CYCLE_COUNTER(STAT_MathConvertQuatToRotator);
+
+#if USE_MATRIX_ROTATOR 
+	// if you think this function is problem, you can undo previous matrix rotator by returning RotatorFromMatrix
+	FRotator RotatorFromMatrix = FQuatRotationTranslationMatrix(*this, FVector::ZeroVector).Rotator();
+	checkSlow(IsNormalized());
+#endif
+
+	static float RAD_TO_DEG = 180.f/PI;
+
+	FRotator RotatorFromQuat;
+	float SingularityTest = Z*X-W*Y;
+	float Pitch = FMath::Asin(2*(SingularityTest));
+
+	RotatorFromQuat.Yaw = FMath::Atan2(2.f*(W*Z+X*Y), (1-2.f*(FMath::Square(Y) + FMath::Square(Z))))*RAD_TO_DEG;
+
+	// reference 
+	// http://en.wikipedia.org/wiki/Conversion_between_quaternions_and_Euler_angles
+	// http://www.euclideanspace.com/maths/geometry/rotations/conversions/quaternionToEuler/
+
+	// this value was found from experience, the above websites recommend different values
+	// but that isn't the case for us, so I went through different testing, and finally found the case 
+	// where both of world lives happily. 
+#define SINGULARITY_THRESHOLD	0.4999995f
+
+	if (SingularityTest < -SINGULARITY_THRESHOLD)
+	{
+		RotatorFromQuat.Roll = -RotatorFromQuat.Yaw - 2.f* FMath::Atan2(X, W) * RAD_TO_DEG;
+		RotatorFromQuat.Pitch = 270.f;
+	}
+	else if (SingularityTest > SINGULARITY_THRESHOLD)
+	{
+		RotatorFromQuat.Roll = RotatorFromQuat.Yaw - 2.f* FMath::Atan2(X, W) * RAD_TO_DEG;
+		RotatorFromQuat.Pitch = 90.f;
+	}
+	else
+	{
+		RotatorFromQuat.Roll = FMath::Atan2(-2.f*(W*X+Y*Z), (1-2.f*(FMath::Square(X) + FMath::Square(Y))))*RAD_TO_DEG;
+		RotatorFromQuat.Pitch = FRotator::ClampAxis(Pitch*RAD_TO_DEG); //clamp it so within 360 - this is for if below
+	}
+
+	RotatorFromQuat.Normalize();
+
+#if USE_MATRIX_ROTATOR
+	RotatorFromMatrix = RotatorFromMatrix.Clamp();
+
+	// this Euler is degree, so less 1 is negligible
+	if (!DebugRotatorEquals(RotatorFromQuat, RotatorFromMatrix, 0.1f))
+	{
+		FPlatformMisc::LowLevelOutputDebugStringf(TEXT("WRONG: (Singularity: %.9f) RotationMatrix (%s), RotationQuat(%s)"), SingularityTest, *RotatorFromMatrix.ToString(), *RotatorFromQuat.ToString());
+	}
+
+	return RotatorFromMatrix;
+#else
+	return RotatorFromQuat;
+#endif
+}
+
 FQuat FQuat::MakeFromEuler(const FVector& Euler)
 {
 	return FRotator::MakeFromEuler(Euler).Quaternion();
diff --git a/Engine/Source/Runtime/Core/Public/Math/Quat.h b/Engine/Source/Runtime/Core/Public/Math/Quat.h
--- a/Engine/Source/Runtime/Core/Public/Math/Quat.h
+++ b/Engine/Source/Runtime/Core/Public/Math/Quat.h
@@ -325,7 +325,7 @@ public:
 	FORCEINLINE FVector GetAxisZ( ) const;
 
 	/** @return rotator representation of this quaternion */
-	FRotator Rotator( ) const;
+	CORE_API FRotator Rotator( ) const;
 
 	/** @return Vector of the axis of the quaternion */
 	FORCEINLINE FVector GetRotationAxis( ) const;
@@ -575,64 +575,6 @@ FORCEINLINE bool DebugRotatorEquals(const FRotator& R1, const FRotator& R2, floa
 #endif
 
 
-inline FRotator FQuat::Rotator() const
-{
-#if USE_MATRIX_ROTATOR 
-	// if you think this function is problem, you can undo previous matrix rotator by returning RotatorFromMatrix
-	FRotator RotatorFromMatrix = FQuatRotationTranslationMatrix( *this, FVector::ZeroVector ).Rotator();
-	checkSlow (IsNormalized());
-#endif
-
-	static float RAD_TO_DEG = 180.f/PI;
-
-	FRotator RotatorFromQuat;
-	float SingularityTest = Z*X-W*Y;
-	float Pitch = FMath::Asin(2*(SingularityTest));
-
-	RotatorFromQuat.Yaw = FMath::Atan2(2.f*(W*Z+X*Y), (1-2.f*(FMath::Square(Y) + FMath::Square(Z))))*RAD_TO_DEG;
-
-	// reference 
-	// http://en.wikipedia.org/wiki/Conversion_between_quaternions_and_Euler_angles
-	// http://www.euclideanspace.com/maths/geometry/rotations/conversions/quaternionToEuler/
-
-// this value was found from experience, the above websites recommend different values
-// but that isn't the case for us, so I went through different testing, and finally found the case 
-// where both of world lives happily. 
-#define SINGULARITY_THRESHOLD	0.4999995f
-
-	if ( SingularityTest < -SINGULARITY_THRESHOLD )
-	{
-		RotatorFromQuat.Roll = -RotatorFromQuat.Yaw - 2.f* FMath::Atan2(X, W) * RAD_TO_DEG;
-		RotatorFromQuat.Pitch = 270.f;
-	}
-	else if ( SingularityTest > SINGULARITY_THRESHOLD )
-	{
-		RotatorFromQuat.Roll = RotatorFromQuat.Yaw - 2.f* FMath::Atan2(X, W) * RAD_TO_DEG;
-		RotatorFromQuat.Pitch = 90.f;
-	}
-	else
-	{
-		RotatorFromQuat.Roll = FMath::Atan2(-2.f*(W*X+Y*Z), (1-2.f*(FMath::Square(X) + FMath::Square(Y))))*RAD_TO_DEG;
-		RotatorFromQuat.Pitch = FRotator::ClampAxis(Pitch*RAD_TO_DEG); //clamp it so within 360 - this is for if below
-	}
-
-	RotatorFromQuat.Normalize();
-
-#if USE_MATRIX_ROTATOR
-	RotatorFromMatrix = RotatorFromMatrix.Clamp();
-
-	// this Euler is degree, so less 1 is negligible
-	if (!DebugRotatorEquals(RotatorFromQuat, RotatorFromMatrix, 0.1f))
-	{
-		FPlatformMisc::LowLevelOutputDebugStringf(TEXT("WRONG: (Singularity: %.9f) RotationMatrix (%s), RotationQuat(%s)"), SingularityTest, *RotatorFromMatrix.ToString(), *RotatorFromQuat.ToString());
-	}
-
-	return RotatorFromMatrix;
-#else
-	return RotatorFromQuat;
-#endif
-}
-
 
 /* FQuat inline functions
  *****************************************************************************/
diff --git a/Engine/Source/Runtime/Core/Public/Misc/CoreStats.h b/Engine/Source/Runtime/Core/Public/Misc/CoreStats.h
--- a/Engine/Source/Runtime/Core/Public/Misc/CoreStats.h
+++ b/Engine/Source/Runtime/Core/Public/Misc/CoreStats.h
@@ -81,3 +81,9 @@ DECLARE_FLOAT_ACCUMULATOR_STAT_EXTERN(TEXT("Main thread block time"),STAT_AsyncI
 DECLARE_FLOAT_ACCUMULATOR_STAT_EXTERN(TEXT("Async package precache wait time"),STAT_AsyncIO_AsyncPackagePrecacheWaitTime,STATGROUP_AsyncIO, CORE_API);
 DECLARE_FLOAT_COUNTER_STAT_EXTERN(TEXT("Bandwidth (MByte/ sec)"),STAT_AsyncIO_Bandwidth,STATGROUP_AsyncIO, CORE_API);
 
+/**
+ * Math stats
+ */ 
+
+DECLARE_CYCLE_STAT_EXTERN(TEXT("Convert Rotator to Quat"), STAT_MathConvertRotatorToQuat, STATGROUP_Math, CORE_API);
+DECLARE_CYCLE_STAT_EXTERN(TEXT("Convert Quat to Rotator"), STAT_MathConvertQuatToRotator, STATGROUP_Math, CORE_API);
diff --git a/Engine/Source/Runtime/Core/Public/Stats/Stats2.h b/Engine/Source/Runtime/Core/Public/Stats/Stats2.h
--- a/Engine/Source/Runtime/Core/Public/Stats/Stats2.h
+++ b/Engine/Source/Runtime/Core/Public/Stats/Stats2.h
@@ -1826,6 +1826,7 @@ DECLARE_STATS_GROUP(TEXT("Init Views"),STATGROUP_InitViews, STATCAT_Advanced);
 DECLARE_STATS_GROUP(TEXT("Landscape"),STATGROUP_Landscape, STATCAT_Advanced);
 DECLARE_STATS_GROUP(TEXT("Light Rendering"),STATGROUP_LightRendering, STATCAT_Advanced);
 DECLARE_STATS_GROUP(TEXT("Load Time"), STATGROUP_LoadTime, STATCAT_Advanced);
+DECLARE_STATS_GROUP(TEXT("Math"), STATGROUP_Math, STATCAT_Advanced);
 DECLARE_STATS_GROUP(TEXT("Memory Allocator"),STATGROUP_MemoryAllocator, STATCAT_Advanced);
 DECLARE_STATS_GROUP(TEXT("Memory Platform"),STATGROUP_MemoryPlatform, STATCAT_Advanced);
 DECLARE_STATS_GROUP(TEXT("Memory StaticMesh"),STATGROUP_MemoryStaticMesh, STATCAT_Advanced);
@@ -1871,7 +1872,6 @@ DECLARE_STATS_GROUP(TEXT("User"),STATGROUP_User, STATCAT_Advanced);
 DECLARE_CYCLE_STAT_EXTERN(TEXT("FrameTime"),STAT_FrameTime,STATGROUP_Engine, CORE_API);
 DECLARE_FLOAT_COUNTER_STAT_EXTERN(TEXT("StatUnit FPS"), STAT_FPS, STATGROUP_Engine, CORE_API);
 
-
 /** Stats for the stat system */
 
 DECLARE_CYCLE_STAT_EXTERN(TEXT("DrawStats"),STAT_DrawStats,STATGROUP_StatSystem, CORE_API);
