Windows custom warning emitter fix.

Additional parenthesis was preventing red highlight in VS output.

diff --git a/Engine/Source/Runtime/Core/Public/Windows/WindowsPlatformCompilerPreSetup.h b/Engine/Source/Runtime/Core/Public/Windows/WindowsPlatformCompilerPreSetup.h
--- a/Engine/Source/Runtime/Core/Public/Windows/WindowsPlatformCompilerPreSetup.h
+++ b/Engine/Source/Runtime/Core/Public/Windows/WindowsPlatformCompilerPreSetup.h
@@ -49,7 +49,7 @@
 
 #if _MSC_VER
 #define EMIT_CUSTOM_WARNING(Warning) \
-	__pragma(message(WARNING_LOCATION "(: warning C4996: " Warning))
+	__pragma(message(WARNING_LOCATION ": warning C4996: " Warning))
 #elif __clang__
 #define EMIT_CUSTOM_WARNING(Warning) \
 	_Pragma(PREPROCESSOR_TO_STRING(message(WARNING_LOCATION Warning)))
