FVector and FRotator equality checks now properly handle a zero tolerance

diff --git a/Engine/Source/Runtime/Core/Public/Math/Rotator.h b/Engine/Source/Runtime/Core/Public/Math/Rotator.h
--- a/Engine/Source/Runtime/Core/Public/Math/Rotator.h
+++ b/Engine/Source/Runtime/Core/Public/Math/Rotator.h
@@ -499,9 +499,9 @@ FORCEINLINE bool FRotator::IsZero() const
 
 FORCEINLINE bool FRotator::Equals(const FRotator &R, float Tolerance) const
 {
-	return (FMath::Abs(NormalizeAxis(Pitch - R.Pitch)) < Tolerance) 
-		&& (FMath::Abs(NormalizeAxis(Yaw - R.Yaw)) < Tolerance) 
-		&& (FMath::Abs(NormalizeAxis(Roll - R.Roll)) < Tolerance);
+	return (FMath::Abs(NormalizeAxis(Pitch - R.Pitch)) <= Tolerance) 
+		&& (FMath::Abs(NormalizeAxis(Yaw - R.Yaw)) <= Tolerance) 
+		&& (FMath::Abs(NormalizeAxis(Roll - R.Roll)) <= Tolerance);
 }
 
 
diff --git a/Engine/Source/Runtime/Core/Public/Math/Vector.h b/Engine/Source/Runtime/Core/Public/Math/Vector.h
--- a/Engine/Source/Runtime/Core/Public/Math/Vector.h
+++ b/Engine/Source/Runtime/Core/Public/Math/Vector.h
@@ -1279,7 +1279,7 @@ FORCEINLINE bool FVector::operator!=( const FVector& V ) const
 
 FORCEINLINE bool FVector::Equals(const FVector& V, float Tolerance) const
 {
-	return FMath::Abs(X-V.X) < Tolerance && FMath::Abs(Y-V.Y) < Tolerance && FMath::Abs(Z-V.Z) < Tolerance;
+	return FMath::Abs(X-V.X) <= Tolerance && FMath::Abs(Y-V.Y) <= Tolerance && FMath::Abs(Z-V.Z) <= Tolerance;
 }
 
 FORCEINLINE bool FVector::AllComponentsEqual(float Tolerance) const
