Removed mouse capture code from Mac version

diff --git a/Engine/Source/Runtime/Core/Private/Mac/CocoaWindow.cpp b/Engine/Source/Runtime/Core/Private/Mac/CocoaWindow.cpp
--- a/Engine/Source/Runtime/Core/Private/Mac/CocoaWindow.cpp
+++ b/Engine/Source/Runtime/Core/Private/Mac/CocoaWindow.cpp
@@ -652,65 +652,3 @@ NSString* NSPerformDragOperation = @"NSPerformDragOperation";
 }
 
 @end
-
-/**
- * Custom window class used for mouse capture
- */
-@implementation FMouseCaptureWindow
-
-- (id)initWithTargetWindow: (FCocoaWindow*)Window
-{
-	self = [super initWithContentRect: [[Window screen] frame] styleMask: NSBorderlessWindowMask backing: NSBackingStoreBuffered defer: NO];
-	[self setBackgroundColor: [NSColor clearColor]];
-	[self setOpaque: NO];
-	[self setLevel: NSMainMenuWindowLevel + 1];
-	[self setIgnoresMouseEvents: NO];
-	[self setAcceptsMouseMovedEvents: YES];
-	[self setHidesOnDeactivate: YES];
-	
-	TargetWindow = Window;
-	
-	return self;
-}
-
-- (FCocoaWindow*)targetWindow
-{
-	return TargetWindow;
-}
-
-- (void)setTargetWindow: (FCocoaWindow*)Window
-{
-	TargetWindow = Window;
-}
-
-- (void)mouseDown:(NSEvent*)Event
-{
-	FMacEvent::SendToGameRunLoop(Event, EMacEventSendMethod::Async);
-}
-
-- (void)rightMouseDown:(NSEvent*)Event
-{
-	FMacEvent::SendToGameRunLoop(Event, EMacEventSendMethod::Async);
-}
-
-- (void)otherMouseDown:(NSEvent*)Event
-{
-	FMacEvent::SendToGameRunLoop(Event, EMacEventSendMethod::Async);
-}
-
-- (void)mouseUp:(NSEvent*)Event
-{
-	FMacEvent::SendToGameRunLoop(Event, EMacEventSendMethod::Async);
-}
-
-- (void)rightMouseUp:(NSEvent*)Event
-{
-	FMacEvent::SendToGameRunLoop(Event, EMacEventSendMethod::Async);
-}
-
-- (void)otherMouseUp:(NSEvent*)Event
-{
-	FMacEvent::SendToGameRunLoop(Event, EMacEventSendMethod::Async);
-}
-
-@end
diff --git a/Engine/Source/Runtime/Core/Private/Mac/MacApplication.cpp b/Engine/Source/Runtime/Core/Private/Mac/MacApplication.cpp
--- a/Engine/Source/Runtime/Core/Private/Mac/MacApplication.cpp
+++ b/Engine/Source/Runtime/Core/Private/Mac/MacApplication.cpp
@@ -102,8 +102,6 @@ FMacApplication::FMacApplication()
 	, bUsingTrackpad( false )
 	, HIDInput( HIDInputInterface::Create( MessageHandler ) )
 	, DraggedWindow( NULL )
-	, MouseCaptureWindow( NULL )
-	, bIsMouseCaptureEnabled( false )
 	, bIsMouseCursorLocked( false )
 	, bSystemModalMode( false )
 	, bIsProcessingNSEvent( false )
@@ -270,12 +268,7 @@ FMacApplication::~FMacApplication()
 	}
 
 	CGDisplayRemoveReconfigurationCallback(FMacApplication::OnDisplayReconfiguration, this);
-	if (MouseCaptureWindow)
-	{
-		[MouseCaptureWindow close];
-		MouseCaptureWindow = NULL;
-	}
-	
+
 	TextInputMethodSystem->Terminate();
 
 	MacApplication = NULL;
@@ -321,7 +314,7 @@ bool FMacApplication::IsCursorDirectlyOverSlateWindow() const
 	SCOPED_AUTORELEASE_POOL;
 	const NSInteger WindowNumber = [NSWindow windowNumberAtPoint:[NSEvent mouseLocation] belowWindowWithWindowNumber:0];
 	NSWindow* const Window = [NSApp windowWithWindowNumber:WindowNumber];
-	return Window && ([Window isKindOfClass:[FCocoaWindow class]] || [Window isKindOfClass:[FMouseCaptureWindow class]]);
+	return Window && [Window isKindOfClass:[FCocoaWindow class]];
 }
 
 void FMacApplication::ProcessEvent(FMacEvent const* const Event)
@@ -902,10 +895,6 @@ FCocoaWindow* FMacApplication::FindEventWindow( NSEvent* Event )
 		{
 			EventWindow = DraggedWindow;
 		}
-		else if( MouseCaptureWindow && MouseCaptureWindow == [Event GetWindow] && [MouseCaptureWindow targetWindow] )
-		{
-			EventWindow = [MouseCaptureWindow targetWindow];
-		}
 		else
 		{
 			NSPoint CursorPos = [Event locationInWindow];
@@ -971,84 +960,6 @@ void FMacApplication::PollGameDeviceState( const float TimeDelta )
 	HIDInput->SendControllerEvents();
 }
 
-void FMacApplication::SetCapture( const TSharedPtr< FGenericWindow >& InWindow )
-{
-	bIsMouseCaptureEnabled = InWindow.IsValid();
-	UpdateMouseCaptureWindow( bIsMouseCaptureEnabled ? ((FMacWindow*)InWindow.Get())->GetWindowHandle() : NULL );
-}
-
-void* FMacApplication::GetCapture( void ) const
-{
-	return ( bIsMouseCaptureEnabled && MouseCaptureWindow ) ? [MouseCaptureWindow targetWindow] : NULL;
-}
-
-void FMacApplication::UseMouseCaptureWindow( bool bUseMouseCaptureWindow )
-{
-	bIsMouseCaptureEnabled = bUseMouseCaptureWindow;
-	if(bUseMouseCaptureWindow)
-	{
-		// Bring the mouse capture window to front
-		if(MouseCaptureWindow)
-		{
-			[MouseCaptureWindow orderFront: nil];
-		}
-	}
-	else
-	{
-		if(MouseCaptureWindow)
-		{
-			[MouseCaptureWindow orderOut: nil];
-		}
-	}
-}
-
-void FMacApplication::UpdateMouseCaptureWindow( FCocoaWindow* TargetWindow )
-{
-	SCOPED_AUTORELEASE_POOL;
-
-	// To prevent mouse events to get passed to the Dock or top menu bar, we create a transparent, full screen window above everything.
-	// This assures that only the editor (and with that, its active window) gets the mouse events while mouse capture is active.
-
-	const bool bEnable = bIsMouseCaptureEnabled || bIsMouseCursorLocked;
-
-	if( bEnable )
-	{
-		if( !TargetWindow )
-		{
-			if( [MouseCaptureWindow targetWindow] )
-			{
-				TargetWindow = [MouseCaptureWindow targetWindow];
-			}
-			else if( [[NSApp keyWindow] isKindOfClass: [FCocoaWindow class]] )
-			{
-				TargetWindow = (FCocoaWindow*)[NSApp keyWindow];
-			}
-		}
-
-		if( !MouseCaptureWindow )
-		{
-			MouseCaptureWindow = [[FMouseCaptureWindow alloc] initWithTargetWindow: TargetWindow];
-			check(MouseCaptureWindow);
-		}
-		else
-		{
-			[MouseCaptureWindow setTargetWindow: TargetWindow];
-		}
-
-		// Bring the mouse capture window to front
-		[MouseCaptureWindow orderFront: nil];
-	}
-	else
-	{
-		// Hide the mouse capture window
-		if( MouseCaptureWindow )
-		{
-			[MouseCaptureWindow orderOut: nil];
-			[MouseCaptureWindow setTargetWindow: NULL];
-		}
-	}
-}
-
 void FMacApplication::SetHighPrecisionMouseMode( const bool Enable, const TSharedPtr< FGenericWindow >& InWindow )
 {
 	bUsingHighPrecisionMouseInput = Enable;
@@ -1335,7 +1246,6 @@ void FMacApplication::OnWindowClose( FCocoaWindow* Window )
 void FMacApplication::OnMouseCursorLock( bool bLockEnabled )
 {
 	bIsMouseCursorLocked = bLockEnabled;
-	UpdateMouseCaptureWindow( NULL );
 }
 
 bool FMacApplication::IsPrintableKey(uint32 Character)
diff --git a/Engine/Source/Runtime/Core/Private/Mac/MacPlatformMisc.cpp b/Engine/Source/Runtime/Core/Private/Mac/MacPlatformMisc.cpp
--- a/Engine/Source/Runtime/Core/Private/Mac/MacPlatformMisc.cpp
+++ b/Engine/Source/Runtime/Core/Private/Mac/MacPlatformMisc.cpp
@@ -635,11 +635,6 @@ EAppReturnType::Type FMacPlatformMisc::MessageBoxExt(EAppMsgType::Type MsgType,
 {
 	SCOPED_AUTORELEASE_POOL;
 
-	if (MacApplication)
-	{
-		MacApplication->UseMouseCaptureWindow(false);
-	}
-
 	EAppReturnType::Type RetValue = MainThreadReturn(^{
 		EAppReturnType::Type RetValue = EAppReturnType::Cancel;
 		NSInteger Result;
@@ -803,11 +798,6 @@ EAppReturnType::Type FMacPlatformMisc::MessageBoxExt(EAppMsgType::Type MsgType,
 		return RetValue;
 	});
 
-	if (MacApplication)
-	{
-		MacApplication->UseMouseCaptureWindow(true);
-	}
-
 	return RetValue;
 }
 
diff --git a/Engine/Source/Runtime/Core/Public/Mac/CocoaWindow.h b/Engine/Source/Runtime/Core/Public/Mac/CocoaWindow.h
--- a/Engine/Source/Runtime/Core/Public/Mac/CocoaWindow.h
+++ b/Engine/Source/Runtime/Core/Public/Mac/CocoaWindow.h
@@ -2,8 +2,6 @@
 
 #pragma once
 
-#ifdef __OBJC__
-
 /**
  * Custom window class used for input handling
  */
@@ -71,32 +69,8 @@
 
 @end
 
-/**
- * Custom window class used for mouse capture
- */
-@interface FMouseCaptureWindow : NSWindow <NSWindowDelegate>
-{
-	FCocoaWindow*	TargetWindow;
-}
-
-- (id)initWithTargetWindow: (FCocoaWindow*)Window;
-- (FCocoaWindow*)targetWindow;
-- (void)setTargetWindow: (FCocoaWindow*)Window;
-
-@end
-
 extern NSString* NSWindowRedrawContents;
 extern NSString* NSDraggingExited;
 extern NSString* NSDraggingUpdated;
 extern NSString* NSPrepareForDragOperation;
 extern NSString* NSPerformDragOperation;
-
-#else // __OBJC__
-
-class FCocoaWindow;
-class FMouseCaptureWindow;
-class NSWindow;
-class NSEvent;
-class NSScreen;
-
-#endif // __OBJC__
diff --git a/Engine/Source/Runtime/Core/Public/Mac/MacApplication.h b/Engine/Source/Runtime/Core/Public/Mac/MacApplication.h
--- a/Engine/Source/Runtime/Core/Public/Mac/MacApplication.h
+++ b/Engine/Source/Runtime/Core/Public/Mac/MacApplication.h
@@ -36,10 +36,6 @@ public:
 
 	virtual void InitializeWindow( const TSharedRef< FGenericWindow >& Window, const TSharedRef< FGenericWindowDefinition >& InDefinition, const TSharedPtr< FGenericWindow >& InParent, const bool bShowImmediately ) override;
 
-	virtual void SetCapture( const TSharedPtr< FGenericWindow >& InWindow ) override;
-
-	virtual void* GetCapture( void ) const override;
-
 	virtual void SetHighPrecisionMouseMode( const bool Enable, const TSharedPtr< FGenericWindow >& InWindow ) override;
 
 	virtual bool IsUsingHighPrecisionMouseMode() const override { return bUsingHighPrecisionMouseInput; }
@@ -74,8 +70,6 @@ public:
 
 	uint32 GetModifierKeysFlags() { return ModifierKeysFlags; }
 
-	void UseMouseCaptureWindow(bool bUseMouseCaptureWindow);
-
 	bool IsProcessingNSEvent() const { return bIsProcessingNSEvent; }
 
 	bool IsWorkspaceSessionActive() const { return bIsWorkspaceSessionActive; }
@@ -135,8 +129,6 @@ private:
 
 	NSScreen* FindScreenByPoint( int32 X, int32 Y ) const;
 
-	void UpdateMouseCaptureWindow( FCocoaWindow* TargetWindow );
-
 	void HandleModifierChange(NSUInteger NewModifierFlags, NSUInteger FlagsShift, NSUInteger UE4Shift, EMacModifierKeys TranslatedCode);
 
 #if WITH_EDITOR
@@ -170,8 +162,6 @@ private:
 
 	FCocoaWindow* DraggedWindow;
 
-	FMouseCaptureWindow* MouseCaptureWindow;
-	bool bIsMouseCaptureEnabled;
 	bool bIsMouseCursorLocked;
 
 	bool bSystemModalMode;
