Core: Removed unused file header comments; fixed spacing, line breaks, spelling

diff --git a/Engine/Source/Runtime/Core/Public/Math/BigInt.h b/Engine/Source/Runtime/Core/Public/Math/BigInt.h
--- a/Engine/Source/Runtime/Core/Public/Math/BigInt.h
+++ b/Engine/Source/Runtime/Core/Public/Math/BigInt.h
@@ -1,6 +1,8 @@
 // Copyright 1998-2015 Epic Games, Inc. All Rights Reserved.
+
 #pragma once
 
+
 /**
  * n-bit integer. @todo: optimize
  * Data is stored as an array of 32-bit words from the least to the most significant
@@ -62,9 +64,7 @@ class TBigInt
 
 public:
 
-	/**
-	 * Sets this integer to 0.
-	 */
+	/** Sets this integer to 0. */
 	FORCEINLINE void Zero()
 	{
 		FMemory::Memset(Bits, 0, sizeof(Bits));
@@ -72,6 +72,8 @@ public:
 
 	/**
 	 * Initializes this big int with a 64 bit integer value.
+	 *
+	 * @param Value The value to set.
 	 */
 	FORCEINLINE void Set(int64 Value)
 	{
@@ -80,9 +82,7 @@ public:
 		Bits[1] = (Value >> BitsPerWord) & 0xffffffff;
 	}
 
-	/**
-	 * Default constructor. Initializes the number to zero.
-	 */
+	/** Default constructor. Initializes the number to zero. */
 	TBigInt()
 	{
 		Zero();
@@ -90,6 +90,8 @@ public:
 
 	/**
 	 * Constructor. Initializes this big int with a 64 bit integer value.
+	 *
+	 * @parma Other The value to set.
 	 */
 	TBigInt(int64 Other)
 	{
diff --git a/Engine/Source/Runtime/Core/Public/Math/Box.h b/Engine/Source/Runtime/Core/Public/Math/Box.h
--- a/Engine/Source/Runtime/Core/Public/Math/Box.h
+++ b/Engine/Source/Runtime/Core/Public/Math/Box.h
@@ -1,9 +1,5 @@
 // Copyright 1998-2015 Epic Games, Inc. All Rights Reserved.
 
-/*=============================================================================
-	Box.h: Declares the FBox class.
-=============================================================================*/
-
 #pragma once
 
 
@@ -28,14 +24,10 @@ public:
 
 public:
 
-	/**
-	 * Default constructor (no initialization).
-	 */
-	FBox( ) { }
+	/** Default constructor (no initialization). */
+	FBox() { }
 
-	/**
-	 * Creates and initializes a new box with zero extent and marks it as invalid.
-	 */
+	/** Creates and initializes a new box with zero extent and marks it as invalid. */
 	FBox( int32 )
 	{
 		Init();
@@ -210,7 +202,7 @@ public:
 	 * @return The center point.
 	 * @see GetCenterAndExtents, GetExtent, GetSize, GetVolume
 	 */
-	FVector GetCenter( ) const
+	FVector GetCenter() const
 	{
 		return FVector((Min + Max) * 0.5f);
 	}
@@ -242,7 +234,7 @@ public:
 	 * @return The box extents.
 	 * @see GetCenter, GetCenterAndExtents, GetSize, GetVolume
 	 */
-	FVector GetExtent( ) const
+	FVector GetExtent() const
 	{
 		return 0.5f * (Max - Min);
 	}
@@ -275,7 +267,7 @@ public:
 	 * @return The box size.
 	 * @see GetCenter, GetCenterAndExtents, GetExtent, GetVolume
 	 */
-	FVector GetSize( ) const
+	FVector GetSize() const
 	{
 		return (Max - Min);
 	}
@@ -286,7 +278,7 @@ public:
 	 * @return The box volume.
 	 * @see GetCenter, GetCenterAndExtents, GetExtent, GetSize
 	 */
-	float GetVolume( ) const
+	float GetVolume() const
 	{
 		return ((Max.X - Min.X) * (Max.Y - Min.Y) * (Max.Z - Min.Z));
 	}
@@ -294,7 +286,7 @@ public:
 	/**
 	 * Set the initial values of the bounding box to Zero.
 	 */
-	void Init( )
+	void Init()
 	{
 		Min = Max = FVector::ZeroVector;
 		IsValid = 0;
@@ -402,7 +394,7 @@ public:
 	 *
 	 * @return A string describing the box.
 	 */
-	FString ToString( ) const;
+	FString ToString() const;
 
 public:
 
@@ -564,7 +556,7 @@ FORCEINLINE bool FBox::IntersectXY( const FBox& Other ) const
 }
 
 
-FORCEINLINE FString FBox::ToString( ) const
+FORCEINLINE FString FBox::ToString() const
 {
 	return FString::Printf(TEXT("IsValid=%s, Min=(%s), Max=(%s)"), IsValid ? TEXT("true") : TEXT("false"), *Min.ToString(), *Max.ToString());
 }
diff --git a/Engine/Source/Runtime/Core/Public/Math/Box2D.h b/Engine/Source/Runtime/Core/Public/Math/Box2D.h
--- a/Engine/Source/Runtime/Core/Public/Math/Box2D.h
+++ b/Engine/Source/Runtime/Core/Public/Math/Box2D.h
@@ -1,9 +1,5 @@
 // Copyright 1998-2015 Epic Games, Inc. All Rights Reserved.
 
-/*=============================================================================
-	Box2D.h: Declares the FBox2D class.
-=============================================================================*/
-
 #pragma once
 
 
@@ -24,10 +20,8 @@ public:
 
 public:
 
-	/**
-	 * Default constructor (no initialization).
-	 */
-	FBox2D( ) { }
+	/** Default constructor (no initialization). */
+	FBox2D() { }
 
 	/**
 	 * Creates and initializes a new box.
@@ -198,7 +192,7 @@ public:
 	 * @return Box area.
 	 * @see GetCenter, GetCenterAndExtents, GetExtent, GetSize
 	 */
-	float GetArea( ) const
+	float GetArea() const
 	{
 		return (Max.X - Min.X) * (Max.Y - Min.Y);
 	}
@@ -209,7 +203,7 @@ public:
 	 * @return Th center point.
 	 * @see GetArea, GetCenterAndExtents, GetExtent, GetSize
 	 */
-	FVector2D GetCenter( ) const
+	FVector2D GetCenter() const
 	{
 		return FVector2D((Min + Max) * 0.5f);
 	}
@@ -230,7 +224,7 @@ public:
 	/**
 	 * Calculates the closest point on or inside the box to a given point in space.
 	 *
-	 * @param Point - The point in space.
+	 * @param Point The point in space.
 	 *
 	 * @return The closest point on or inside the box.
 	 */
@@ -242,7 +236,7 @@ public:
 	 * @return Box extents.
 	 * @see GetArea, GetCenter, GetCenterAndExtents, GetSize
 	 */
-	FVector2D GetExtent( ) const
+	FVector2D GetExtent() const
 	{
 		return 0.5f * (Max - Min);
 	}
@@ -254,7 +248,7 @@ public:
 	 * @return Box size.
 	 * @see GetArea, GetCenter, GetCenterAndExtents, GetExtent
 	 */
-	FVector2D GetSize( ) const
+	FVector2D GetSize() const
 	{
 		return (Max - Min);
 	}
@@ -262,7 +256,7 @@ public:
 	/**
 	 * Set the initial values of the bounding box to Zero.
 	 */
-	void Init( )
+	void Init()
 	{
 		Min = Max = FVector2D::ZeroVector;
 		bIsValid = false;
@@ -314,15 +308,15 @@ public:
 	 *
 	 * @return A string describing the box.
 	 */
-	FString ToString( ) const;
+	FString ToString() const;
 
 public:
 
 	/**
 	 * Serializes the bounding box.
 	 *
-	 * @param Ar - The archive to serialize into.
-	 * @param Box - The box to serialize.
+	 * @param Ar The archive to serialize into.
+	 * @param Box The box to serialize.
 	 *
 	 * @return Reference to the Archive after serialization.
 	 */
@@ -404,6 +398,7 @@ FORCEINLINE FVector2D FBox2D::GetClosestPointTo( const FVector2D& Point ) const
 	return ClosestPoint;
 }
 
+
 FORCEINLINE bool FBox2D::Intersect( const FBox2D & Other ) const
 {
 	if ((Min.X > Other.Max.X) || (Other.Min.X > Max.X))
@@ -419,7 +414,8 @@ FORCEINLINE bool FBox2D::Intersect( const FBox2D & Other ) const
 	return true;
 }
 
-FORCEINLINE FString FBox2D::ToString( ) const
+
+FORCEINLINE FString FBox2D::ToString() const
 {
 	return FString::Printf(TEXT("bIsValid=%s, Min=(%s), Max=(%s)"), bIsValid ? TEXT("true") : TEXT("false"), *Min.ToString(), *Max.ToString());
 }
diff --git a/Engine/Source/Runtime/Core/Public/Math/BoxSphereBounds.h b/Engine/Source/Runtime/Core/Public/Math/BoxSphereBounds.h
--- a/Engine/Source/Runtime/Core/Public/Math/BoxSphereBounds.h
+++ b/Engine/Source/Runtime/Core/Public/Math/BoxSphereBounds.h
@@ -1,9 +1,5 @@
 // Copyright 1998-2015 Epic Games, Inc. All Rights Reserved.
 
-/*=============================================================================
-	BoxSphereBounds.h: Declares the FBoxSphereBounds structure.
-=============================================================================*/
-
 #pragma once
 
 
@@ -23,10 +19,8 @@ struct FBoxSphereBounds
 
 public:
 
-	/**
-	 * Default constructor.
-	 */
-	FBoxSphereBounds( ) { }
+	/** Default constructor. */
+	FBoxSphereBounds() { }
 
 	/**
 	 * Creates and initializes a new instance.
@@ -165,7 +159,7 @@ public:
 	 *
 	 * @return The bounding box.
 	 */
-	FORCEINLINE FBox GetBox( ) const
+	FORCEINLINE FBox GetBox() const
 	{
 		return FBox(Origin - BoxExtent,Origin + BoxExtent);
 	}
@@ -191,7 +185,7 @@ public:
 	 *
 	 * @return The bounding sphere.
 	 */
-	FORCEINLINE FSphere GetSphere( ) const
+	FORCEINLINE FSphere GetSphere() const
 	{
 		return FSphere(Origin,SphereRadius);
 	}
@@ -228,7 +222,7 @@ public:
 	 *
 	 * @return Text describing the bounding box.
 	 */
-	FString ToString( ) const;
+	FString ToString() const;
 
 	/**
 	 * Constructs a bounding volume containing both A and B.
@@ -254,14 +248,14 @@ public:
 	}
 
 #if ENABLE_NAN_DIAGNOSTIC
-	FORCEINLINE void DiagnosticCheckNaN( ) const
+	FORCEINLINE void DiagnosticCheckNaN() const
 	{
 		checkf(!Origin.ContainsNaN(), TEXT("Origin contains NaN: %s"), *Origin.ToString());
 		checkf(!BoxExtent.ContainsNaN(), TEXT("BoxExtent contains NaN: %s"), *BoxExtent.ToString());
 		checkf(!FMath::IsNaN(SphereRadius) && FMath::IsFinite(SphereRadius), TEXT("SphereRadius contains NaN: %f"), SphereRadius);
 	}
 #else
-	FORCEINLINE void DiagnosticCheckNaN( ) const {}
+	FORCEINLINE void DiagnosticCheckNaN() const {}
 #endif
 
 public:
@@ -326,7 +320,7 @@ FORCEINLINE FBoxSphereBounds FBoxSphereBounds::operator+( const FBoxSphereBounds
 }
 
 
-FORCEINLINE FString FBoxSphereBounds::ToString( ) const
+FORCEINLINE FString FBoxSphereBounds::ToString() const
 {
 	return FString::Printf(TEXT("Origin=%s, BoxExtent=(%s), SphereRadius=(%f)"), *Origin.ToString(), *BoxExtent.ToString(), SphereRadius);
 }
\ No newline at end of file
diff --git a/Engine/Source/Runtime/Core/Public/Math/Color.h b/Engine/Source/Runtime/Core/Public/Math/Color.h
--- a/Engine/Source/Runtime/Core/Public/Math/Color.h
+++ b/Engine/Source/Runtime/Core/Public/Math/Color.h
@@ -1,18 +1,17 @@
 // Copyright 1998-2015 Epic Games, Inc. All Rights Reserved.
 
-/*=============================================================================
-	Color.h: Unreal color definitions.
-=============================================================================*/
-
 #pragma once
+
 #include "Math/UnrealMathUtility.h"
 #include "Misc/Parse.h"
 #include "Serialization/ArchiveBase.h"
 
+
 class FColor;
 class FVector;
 class FFloat16Color;
 
+
 /**
  * A linear, 32-bit/component floating point RGBA color.
  */
diff --git a/Engine/Source/Runtime/Core/Public/Math/ColorList.h b/Engine/Source/Runtime/Core/Public/Math/ColorList.h
--- a/Engine/Source/Runtime/Core/Public/Math/ColorList.h
+++ b/Engine/Source/Runtime/Core/Public/Math/ColorList.h
@@ -1,12 +1,10 @@
 // Copyright 1998-2015 Epic Games, Inc. All Rights Reserved.
 
-/*=============================================================================
-	ColorList.h: List of common colors definitions.
-=============================================================================*/
-
 #pragma once
 
-class CORE_API FColorList : public FExec
+
+class CORE_API FColorList
+	: public FExec
 {
 public:
 	typedef TMap< FString, const FColor* > TColorsMap;
@@ -153,5 +151,5 @@ protected:
 	TColorsLookup ColorsLookup;
 };
 
-extern CORE_API FColorList GColorList;
 
+extern CORE_API FColorList GColorList;
diff --git a/Engine/Source/Runtime/Core/Public/Math/ConvexHull2d.h b/Engine/Source/Runtime/Core/Public/Math/ConvexHull2d.h
--- a/Engine/Source/Runtime/Core/Public/Math/ConvexHull2d.h
+++ b/Engine/Source/Runtime/Core/Public/Math/ConvexHull2d.h
@@ -1,11 +1,8 @@
 // Copyright 1998-2015 Epic Games, Inc. All Rights Reserved.
 
-/*=============================================================================
-	ConvexHull2d.h: Declares utilites for creating convex hulls in 2d.
-=============================================================================*/
-
 #pragma once
 
+
 namespace ConvexHull2D
 {
 	/** Returns <0 if C is left of A-B */
diff --git a/Engine/Source/Runtime/Core/Public/Math/Edge.h b/Engine/Source/Runtime/Core/Public/Math/Edge.h
--- a/Engine/Source/Runtime/Core/Public/Math/Edge.h
+++ b/Engine/Source/Runtime/Core/Public/Math/Edge.h
@@ -1,9 +1,5 @@
 // Copyright 1998-2015 Epic Games, Inc. All Rights Reserved.
 
-/*=============================================================================
-	Edge.h: Declares the FEdge structure.
-=============================================================================*/
-
 #pragma once
 
 
@@ -20,9 +16,7 @@ struct FEdge
 
 public:
 
-	/**
-	 * Default constructor (no initialization).
-	 */
+	/** Default constructor (no initialization). */
 	FEdge () { }
 
 	/**
diff --git a/Engine/Source/Runtime/Core/Public/Math/Float16Color.h b/Engine/Source/Runtime/Core/Public/Math/Float16Color.h
--- a/Engine/Source/Runtime/Core/Public/Math/Float16Color.h
+++ b/Engine/Source/Runtime/Core/Public/Math/Float16Color.h
@@ -2,12 +2,14 @@
 
 #pragma once
 
+
 /**
  *	RGBA Color made up of FFloat16
  */
 class FFloat16Color
 {
 public:
+
 	FFloat16 R;
 	FFloat16 G;
 	FFloat16 B;
@@ -29,15 +31,14 @@ public:
 	 * Checks whether two colors are identical.
 	 *
 	 * @param Src The other color.
-	 *
 	 * @return true if the two colors are identical, otherwise false.
 	 */
 	bool operator==(const FFloat16Color& Src);
 };
 
-FORCEINLINE FFloat16Color::FFloat16Color()
-{
-}
+
+FORCEINLINE FFloat16Color::FFloat16Color() { }
+
 
 FORCEINLINE FFloat16Color::FFloat16Color(const FFloat16Color& Src)
 {
@@ -47,13 +48,14 @@ FORCEINLINE FFloat16Color::FFloat16Color(const FFloat16Color& Src)
 	A = Src.A;
 }
 
+
 FORCEINLINE FFloat16Color::FFloat16Color(const FLinearColor& Src) :
 	R(Src.R),
 	G(Src.G),
 	B(Src.B),
 	A(Src.A)
-{
-}
+{ }
+
 
 FORCEINLINE FFloat16Color& FFloat16Color::operator=(const FFloat16Color& Src)
 {
@@ -64,6 +66,7 @@ FORCEINLINE FFloat16Color& FFloat16Color::operator=(const FFloat16Color& Src)
 	return *this;
 }
 
+
 FORCEINLINE bool FFloat16Color::operator==(const FFloat16Color& Src)
 {
 	return (
@@ -72,4 +75,4 @@ FORCEINLINE bool FFloat16Color::operator==(const FFloat16Color& Src)
 		(B == Src.B) &&
 		(A == Src.A)
 		);
-}
\ No newline at end of file
+}
diff --git a/Engine/Source/Runtime/Core/Public/Math/FloatPacker.h b/Engine/Source/Runtime/Core/Public/Math/FloatPacker.h
--- a/Engine/Source/Runtime/Core/Public/Math/FloatPacker.h
+++ b/Engine/Source/Runtime/Core/Public/Math/FloatPacker.h
@@ -1,13 +1,11 @@
 // Copyright 1998-2015 Epic Games, Inc. All Rights Reserved.
 
-/*=============================================================================
-	FloatPacker.h:
-=============================================================================*/
-
 #pragma once
 
+
 DECLARE_LOG_CATEGORY_EXTERN(LogFloatPacker, Log, All);
 
+
 /**
  *
  */
@@ -37,6 +35,7 @@ public:
 	}
 };
 
+
 /**
  *
  */
@@ -201,4 +200,3 @@ public:
 	}
 #endif
 };
-
diff --git a/Engine/Source/Runtime/Core/Public/Math/IntPoint.h b/Engine/Source/Runtime/Core/Public/Math/IntPoint.h
--- a/Engine/Source/Runtime/Core/Public/Math/IntPoint.h
+++ b/Engine/Source/Runtime/Core/Public/Math/IntPoint.h
@@ -1,13 +1,11 @@
 // Copyright 1998-2015 Epic Games, Inc. All Rights Reserved.
 
-/*=============================================================================
-	IntPoint.h: Declares the FIntPoint structure.
-=============================================================================*/
-
 #pragma once
+
 #include "Containers/UnrealString.h"
 #include "HAL/Platform.h"
 
+
 /**
  * Structure for integer points in 2-d space.
  *
@@ -31,10 +29,8 @@ public:
 
 public:
 
-	/**
-	 * Default constructor (no initialization).
-	 */
-	FIntPoint( );
+	/** Default constructor (no initialization). */
+	FIntPoint();
 
 	/**
 	 * Creates and initializes a new instance with the specified coordinates.
@@ -203,7 +199,7 @@ public:
 	 * @return The maximum value in the point.
 	 * @see GetMin, Size, SizeSquared
 	 */
-	int32 GetMax( ) const;
+	int32 GetMax() const;
 
 	/**
 	 * Gets the minimum value in the point.
@@ -211,7 +207,7 @@ public:
 	 * @return The minimum value in the point.
 	 * @see GetMax, Size, SizeSquared
 	 */
-	int32 GetMin( ) const;
+	int32 GetMin() const;
 
 	/**
 	 * Gets the distance of this point from (0,0).
@@ -219,7 +215,7 @@ public:
 	 * @return The distance of this point from (0,0).
 	 * @see GetMax, GetMin, SizeSquared
 	 */
-	int32 Size( ) const;
+	int32 Size() const;
 
 	/**
 	 * Gets the squared distance of this point from (0,0).
@@ -227,14 +223,14 @@ public:
 	 * @return The squared distance of this point from (0,0).
 	 * @see GetMax, GetMin, Size
 	 */
-	int32 SizeSquared( ) const;
+	int32 SizeSquared() const;
 	
 	/**
 	 * Get a textual representation of this point.
 	 *
 	 * @return A string describing the point.
 	 */
-	FString ToString( ) const;
+	FString ToString() const;
 
 public:
 
@@ -262,7 +258,7 @@ public:
 	 * Gets number of components point has.
 	 * @return number of components point has.
 	 */
-	static int32 Num( );
+	static int32 Num();
 
 public:
 
@@ -283,7 +279,7 @@ public:
 /* FInterval inline functions
  *****************************************************************************/
 
-FORCEINLINE FIntPoint::FIntPoint( ) { }
+FORCEINLINE FIntPoint::FIntPoint() { }
 
 
 FORCEINLINE FIntPoint::FIntPoint( int32 InX, int32 InY )
@@ -310,7 +306,7 @@ FORCEINLINE int32& FIntPoint::operator()( int32 PointIndex )
 }
 
 
-FORCEINLINE int32 FIntPoint::Num( )
+FORCEINLINE int32 FIntPoint::Num()
 {
 	return 2;
 }
@@ -393,38 +389,45 @@ FORCEINLINE FIntPoint FIntPoint::operator/( int32 Divisor ) const
 	return FIntPoint(*this) /= Divisor;
 }
 
+
 FORCEINLINE int32& FIntPoint::operator[](int32 Index)
 {
 	check(Index >= 0 && Index < 2);
 	return ((Index == 0) ? X : Y);
 }
 
+
 FORCEINLINE int32 FIntPoint::operator[](int32 Index) const
 {
 	check(Index >= 0 && Index < 2);
 	return ((Index == 0) ? X : Y);
 }
 
+
 FORCEINLINE FIntPoint FIntPoint::ComponentMin(const FIntPoint& Other) const
 {
 	return FIntPoint(FMath::Min(X, Other.X), FMath::Min(Y, Other.Y));
 }
 
+
 FORCEINLINE FIntPoint FIntPoint::ComponentMax(const FIntPoint& Other) const
 {
 	return FIntPoint(FMath::Max(X, Other.X), FMath::Max(Y, Other.Y));
 }
 
+
 FORCEINLINE FIntPoint FIntPoint::DivideAndRoundUp( FIntPoint lhs, int32 Divisor )
 {
 	return FIntPoint(FMath::DivideAndRoundUp(lhs.X, Divisor), FMath::DivideAndRoundUp(lhs.Y, Divisor));
 }	
 
+
 FORCEINLINE FIntPoint FIntPoint::DivideAndRoundDown( FIntPoint lhs, int32 Divisor )
 {
 	return FIntPoint(FMath::DivideAndRoundDown(lhs.X, Divisor), FMath::DivideAndRoundDown(lhs.Y, Divisor));
 }	
 
+
 FORCEINLINE FIntPoint FIntPoint::operator+( const FIntPoint& Other ) const
 {
 	return FIntPoint(*this) += Other;
@@ -443,13 +446,13 @@ FORCEINLINE FIntPoint FIntPoint::operator/( const FIntPoint& Other ) const
 }
 
 
-FORCEINLINE int32 FIntPoint::GetMax( ) const
+FORCEINLINE int32 FIntPoint::GetMax() const
 {
 	return FMath::Max(X, Y);
 }
 
 
-FORCEINLINE int32 FIntPoint::GetMin( ) const
+FORCEINLINE int32 FIntPoint::GetMin() const
 {
 	return FMath::Min(X,Y);
 }
@@ -460,7 +463,7 @@ FORCEINLINE uint32 GetTypeHash(const FIntPoint& InPoint)
 }
 
 
-FORCEINLINE int32 FIntPoint::Size( ) const
+FORCEINLINE int32 FIntPoint::Size() const
 {
 	return int32(FMath::Sqrt( float(X*X + Y*Y)));
 }
@@ -470,7 +473,7 @@ FORCEINLINE int32 FIntPoint::SizeSquared() const
 	return X*X + Y*Y;
 }
 
-FORCEINLINE FString FIntPoint::ToString( ) const
+FORCEINLINE FString FIntPoint::ToString() const
 {
 	return FString::Printf(TEXT("X=%d Y=%d"), X, Y);
-}
\ No newline at end of file
+}
diff --git a/Engine/Source/Runtime/Core/Public/Math/IntRect.h b/Engine/Source/Runtime/Core/Public/Math/IntRect.h
--- a/Engine/Source/Runtime/Core/Public/Math/IntRect.h
+++ b/Engine/Source/Runtime/Core/Public/Math/IntRect.h
@@ -1,9 +1,5 @@
 // Copyright 1998-2015 Epic Games, Inc. All Rights Reserved.
 
-/*=============================================================================
-	IntRect.h: Declares the FIntRect structure.
-=============================================================================*/
-
 #pragma once
 
 
@@ -23,7 +19,7 @@ struct FIntRect
 public:
 
 	/** Constructor */
-	FIntRect( );
+	FIntRect();
 
 	/**
 	 * Constructor
@@ -164,7 +160,7 @@ public:
 	 *
 	 * @return The area of this rectangle.
 	 */
-	int32 Area( ) const;
+	int32 Area() const;
 
 	/**
 	 * Creates a rectangle from the bottom part of this rectangle.
@@ -176,6 +172,7 @@ public:
 
 	/**
 	 * Clip a rectangle using the bounds of another rectangle.
+	 *
 	 * @param Other The other rectangle to clip against.
 	 */
 	void Clip( const FIntRect& Other );
@@ -201,7 +198,7 @@ public:
 	 *
 	 * @return The Height of the rectangle.
 	 */
-	int32 Height( ) const;
+	int32 Height() const;
 
 	/**
 	 * Inflates or deflates the rectangle.
@@ -246,21 +243,21 @@ public:
 	 *
 	 * @return The distance from one corner of the rectangle to the other.
 	 */
-	FIntPoint Size( ) const;
+	FIntPoint Size() const;
 
 	/**
 	 * Get a textual representation of this rectangle.
 	 *
 	 * @return A string describing the rectangle.
 	 */
-	FString ToString( ) const;
+	FString ToString() const;
 
 	/**
 	 * Gets the width of the rectangle.
 	 *
 	 * @return The width of the rectangle.
 	 */
-	int32 Width( ) const;
+	int32 Width() const;
 
 	/**
 	 * Returns true if the rectangle is 0 x 0.
@@ -285,16 +282,15 @@ public:
 	 *
 	 * @return Number of points in the Rectangle.
 	 */
-	static int32 Num( );
+	static int32 Num();
 
 public:
 
 	/**
 	 * Serializes the Rectangle.
 	 *
-	 * @param Ar - The archive to serialize into.
-	 * @param Rect - The rectangle to serialize.
-	 *
+	 * @param Ar The archive to serialize into.
+	 * @param Rect The rectangle to serialize.
 	 * @return Reference to the Archive after serialization.
 	 */
 	friend FArchive& operator<<( FArchive& Ar, FIntRect& Rect )
@@ -316,7 +312,7 @@ FORCEINLINE FIntRect FIntRect::Scale( float Fraction ) const
 /* FIntRect inline functions
  *****************************************************************************/
 
-FORCEINLINE FIntRect::FIntRect( )
+FORCEINLINE FIntRect::FIntRect()
 	: Min(ForceInit)
 	, Max(ForceInit)
 { }
@@ -427,7 +423,7 @@ FORCEINLINE FIntRect FIntRect::operator-( const FIntRect& Other ) const
 }
 
 
-FORCEINLINE int32 FIntRect::Area( ) const
+FORCEINLINE int32 FIntRect::Area() const
 {
 	return (Max.X - Min.X) * (Max.Y - Min.Y);
 }
@@ -474,7 +470,7 @@ FORCEINLINE void FIntRect::GetCenterAndExtents( FIntPoint& OutCenter, FIntPoint&
 }
 
 
-FORCEINLINE int32 FIntRect::Height( ) const
+FORCEINLINE int32 FIntRect::Height() const
 {
 	return (Max.Y - Min.Y);
 }
@@ -488,6 +484,7 @@ FORCEINLINE void FIntRect::InflateRect( int32 Amount )
 	Max.Y += Amount;
 }
 
+
 FORCEINLINE void FIntRect::Include( FIntPoint Point )
 {
 	Min.X = FMath::Min(Min.X, Point.X);
@@ -514,19 +511,19 @@ FORCEINLINE FIntRect FIntRect::Right( int32 InWidth ) const
 }
 
 
-FORCEINLINE FIntPoint FIntRect::Size( ) const
+FORCEINLINE FIntPoint FIntRect::Size() const
 {
 	return FIntPoint( Max.X-Min.X, Max.Y-Min.Y );
 }
 
 
-FORCEINLINE FString FIntRect::ToString( ) const
+FORCEINLINE FString FIntRect::ToString() const
 {
 	return FString::Printf(TEXT("Min=(%s) Max=(%s)"), *Min.ToString(), *Max.ToString());
 }
 
 
-FORCEINLINE int32 FIntRect::Width( ) const
+FORCEINLINE int32 FIntRect::Width() const
 {
 	return Max.X-Min.X;
 }
diff --git a/Engine/Source/Runtime/Core/Public/Math/IntVector.h b/Engine/Source/Runtime/Core/Public/Math/IntVector.h
--- a/Engine/Source/Runtime/Core/Public/Math/IntVector.h
+++ b/Engine/Source/Runtime/Core/Public/Math/IntVector.h
@@ -1,9 +1,5 @@
 // Copyright 1998-2015 Epic Games, Inc. All Rights Reserved.
 
-/*=============================================================================
-	IntVector.h: Declares the FIntVector structure.
-=============================================================================*/
-
 #pragma once
 
 
@@ -34,7 +30,7 @@ public:
 	/**
 	 * Default constructor (no initialization).
 	 */
-	FIntVector( );
+	FIntVector();
 
 	/**
 	 * Creates and initializes a new instance with the specified coordinates.
@@ -172,28 +168,28 @@ public:
 	 *
 	 * @return The maximum value in the point.
 	 */
-	float GetMax( ) const;
+	float GetMax() const;
 
 	/**
 	 * Gets the minimum value in the point.
 	 *
 	 * @return The minimum value in the point.
 	 */
-	float GetMin( ) const;
+	float GetMin() const;
 
 	/**
 	 * Gets the distance of this point from (0,0).
 	 *
 	 * @return The distance of this point from (0,0).
 	 */
-	int32 Size( ) const;
+	int32 Size() const;
 
 	/**
 	 * Get a textual representation of this vector.
 	 *
 	 * @return A string describing the vector.
 	 */
-	FString ToString( ) const;
+	FString ToString() const;
 
 public:
 
@@ -211,7 +207,7 @@ public:
 	 *
 	 * @return Number of components point has.
 	 */
-	static int32 Num( );
+	static int32 Num();
 
 public:
 
@@ -232,7 +228,7 @@ public:
 /* FIntVector inline functions
  *****************************************************************************/
 
-FORCEINLINE FIntVector::FIntVector( )
+FORCEINLINE FIntVector::FIntVector()
 { }
 
 
@@ -360,31 +356,31 @@ FORCEINLINE FIntVector FIntVector::DivideAndRoundUp( FIntVector lhs, int32 Divis
 }
 
 
-FORCEINLINE float FIntVector::GetMax( ) const
+FORCEINLINE float FIntVector::GetMax() const
 {
 	return FMath::Max(FMath::Max(X, Y), Z);
 }
 
 
-FORCEINLINE float FIntVector::GetMin( ) const
+FORCEINLINE float FIntVector::GetMin() const
 {
 	return FMath::Min(FMath::Min(X, Y), Z);
 }
 
 
-FORCEINLINE int32 FIntVector::Num( )
+FORCEINLINE int32 FIntVector::Num()
 {
 	return 2;
 }
 
 
-FORCEINLINE int32 FIntVector::Size( ) const
+FORCEINLINE int32 FIntVector::Size() const
 {
 	return int32(FMath::Sqrt(float(X * X + Y * Y + Z * Z)));
 }
 
 
-FORCEINLINE FString FIntVector::ToString( ) const
+FORCEINLINE FString FIntVector::ToString() const
 {
 	return FString::Printf(TEXT("X=%3.3f Y=%3.3f Z=%3.3f"), X, Y, Z);
 }
diff --git a/Engine/Source/Runtime/Core/Public/Math/InterpCurve.h b/Engine/Source/Runtime/Core/Public/Math/InterpCurve.h
--- a/Engine/Source/Runtime/Core/Public/Math/InterpCurve.h
+++ b/Engine/Source/Runtime/Core/Public/Math/InterpCurve.h
@@ -1,9 +1,5 @@
 // Copyright 1998-2015 Epic Games, Inc. All Rights Reserved.
 
-/*=============================================================================
-	InterpCurve.h: Declares the FInterpCurve class.
-=============================================================================*/
-
 #pragma once
 
 
@@ -23,10 +19,8 @@ public:
 
 public:
 
-	/**
-	 * Default constructor (no initialization).
-	 */
-	FInterpCurve( ) { }
+	/** Default constructor (no initialization). */
+	FInterpCurve() { }
 
 public:
 
@@ -35,7 +29,6 @@ public:
 	 *
 	 * @param InVal
 	 * @param OutVal
-	 *
 	 * @return The index of the new key.
 	 */
 	int32 AddPoint( const float InVal, const T &OutVal );
@@ -51,9 +44,7 @@ public:
 	 */
 	int32 MovePoint( int32 PointIndex, float NewInVal );
 
-	/**
-	 * Clears all keypoints from InterpCurve.
-	 */
+	/** Clears all keypoints from InterpCurve. */
 	void Reset();
 
 	/** 
diff --git a/Engine/Source/Runtime/Core/Public/Math/InterpCurvePoint.h b/Engine/Source/Runtime/Core/Public/Math/InterpCurvePoint.h
--- a/Engine/Source/Runtime/Core/Public/Math/InterpCurvePoint.h
+++ b/Engine/Source/Runtime/Core/Public/Math/InterpCurvePoint.h
@@ -1,11 +1,8 @@
 // Copyright 1998-2015 Epic Games, Inc. All Rights Reserved.
 
-/*=============================================================================
-	InterpCurvePoint.h: Declares the FInterpCurvePoint class and related types.
-=============================================================================*/
-
 #pragma once
 
+
 enum EInterpCurveMode
 {
 	/** A straight line between two keypoint values. */
@@ -65,7 +62,7 @@ public:
 	/**
 	 * Default constructor (no initialization).
 	 */
-	FInterpCurvePoint( ) { };
+	FInterpCurvePoint() { };
 
 	/** 
 	 * Constructor 
@@ -90,7 +87,7 @@ public:
 public:
 
 	/** @return true if the key value is using a curve interp mode, otherwise false */
-	FORCEINLINE bool IsCurveKey( ) const;
+	FORCEINLINE bool IsCurveKey() const;
 
 public:
 
@@ -138,7 +135,7 @@ FORCEINLINE FInterpCurvePoint<T>::FInterpCurvePoint( const float In, const T &Ou
 
 
 template< class T > 
-FORCEINLINE bool FInterpCurvePoint<T>::IsCurveKey( ) const
+FORCEINLINE bool FInterpCurvePoint<T>::IsCurveKey() const
 {
 	return ((InterpMode == CIM_CurveAuto) || (InterpMode == CIM_CurveAutoClamped) || (InterpMode == CIM_CurveUser) || (InterpMode == CIM_CurveBreak));
 }
diff --git a/Engine/Source/Runtime/Core/Public/Math/Interval.h b/Engine/Source/Runtime/Core/Public/Math/Interval.h
--- a/Engine/Source/Runtime/Core/Public/Math/Interval.h
+++ b/Engine/Source/Runtime/Core/Public/Math/Interval.h
@@ -1,9 +1,5 @@
 // Copyright 1998-2015 Epic Games, Inc. All Rights Reserved.
 
-/*=============================================================================
-	Interval.h: Declares the FInterval structure.
-=============================================================================*/
-
 #pragma once
 
 
@@ -28,7 +24,7 @@ public:
 	 *
 	 * The interval is initialized to [0, 0].
 	 */
-	FInterval( )
+	FInterval()
 		: Min(0.0f)
 		, Max(0.0f)
 		, bIsEmpty(true)
@@ -106,7 +102,6 @@ public:
 	 *
 	 * @param A The first interval.
 	 * @param B The second interval.
-	 *
 	 * @return The intersection.
 	 */
 	friend FInterval Intersect( const FInterval& A, const FInterval& B )
diff --git a/Engine/Source/Runtime/Core/Public/Math/InverseRotationMatrix.h b/Engine/Source/Runtime/Core/Public/Math/InverseRotationMatrix.h
--- a/Engine/Source/Runtime/Core/Public/Math/InverseRotationMatrix.h
+++ b/Engine/Source/Runtime/Core/Public/Math/InverseRotationMatrix.h
@@ -2,34 +2,37 @@
 
 #pragma once
 
+
 /** Inverse Rotation matrix */
-class FInverseRotationMatrix : public FMatrix
+class FInverseRotationMatrix
+	: public FMatrix
 {
 public:
-	/** Constructor
+	/**
+	 * Constructor.
 	 *
 	 * @param Rot rotation
 	 */
 	FInverseRotationMatrix(const FRotator& Rot);
 };
 
-FORCEINLINE FInverseRotationMatrix::FInverseRotationMatrix(const FRotator& Rot) :
-FMatrix(
-	FMatrix( // Yaw
-	FPlane(+FMath::Cos(Rot.Yaw * PI / 180.f),	-FMath::Sin(Rot.Yaw * PI / 180.f),	0.0f,									0.0f),
-	FPlane(+FMath::Sin(Rot.Yaw * PI / 180.f),	+FMath::Cos(Rot.Yaw * PI / 180.f),	0.0f,									0.0f),
-	FPlane(0.0f,							0.0f,							1.0f,									0.0f),
-	FPlane(0.0f,							0.0f,							0.0f,									1.0f)) *
-	FMatrix( // Pitch
-	FPlane(+FMath::Cos(Rot.Pitch * PI / 180.f),	0.0f,							-FMath::Sin(Rot.Pitch * PI / 180.f),		0.0f),
-	FPlane(0.0f,							1.0f,							0.0f,									0.0f),
-	FPlane(+FMath::Sin(Rot.Pitch * PI / 180.f),	0.0f,							+FMath::Cos(Rot.Pitch * PI / 180.f),		0.0f),
-	FPlane(0.0f,							0.0f,							0.0f,									1.0f)) *
-	FMatrix( // Roll
-	FPlane(1.0f,							0.0f,							0.0f,									0.0f),
-	FPlane(0.0f,							+FMath::Cos(Rot.Roll * PI / 180.f),	+FMath::Sin(Rot.Roll * PI / 180.f),			0.0f),
-	FPlane(0.0f,							-FMath::Sin(Rot.Roll * PI / 180.f),	+FMath::Cos(Rot.Roll * PI / 180.f),			0.0f),
-	FPlane(0.0f,							0.0f,							0.0f,									1.0f))
+
+FORCEINLINE FInverseRotationMatrix::FInverseRotationMatrix(const FRotator& Rot)
+	: FMatrix(
+		FMatrix( // Yaw
+		FPlane(+FMath::Cos(Rot.Yaw * PI / 180.f), -FMath::Sin(Rot.Yaw * PI / 180.f), 0.0f, 0.0f),
+		FPlane(+FMath::Sin(Rot.Yaw * PI / 180.f), +FMath::Cos(Rot.Yaw * PI / 180.f), 0.0f, 0.0f),
+		FPlane(0.0f, 0.0f, 1.0f, 0.0f),
+		FPlane(0.0f, 0.0f, 0.0f, 1.0f)) *
+		FMatrix( // Pitch
+		FPlane(+FMath::Cos(Rot.Pitch * PI / 180.f), 0.0f, -FMath::Sin(Rot.Pitch * PI / 180.f), 0.0f),
+		FPlane(0.0f, 1.0f, 0.0f, 0.0f),
+		FPlane(+FMath::Sin(Rot.Pitch * PI / 180.f), 0.0f, +FMath::Cos(Rot.Pitch * PI / 180.f), 0.0f),
+		FPlane(0.0f, 0.0f, 0.0f, 1.0f)) *
+		FMatrix( // Roll
+		FPlane(1.0f, 0.0f, 0.0f, 0.0f),
+		FPlane(0.0f, +FMath::Cos(Rot.Roll * PI / 180.f), +FMath::Sin(Rot.Roll * PI / 180.f), 0.0f),
+		FPlane(0.0f, -FMath::Sin(Rot.Roll * PI / 180.f), +FMath::Cos(Rot.Roll * PI / 180.f), 0.0f),
+		FPlane(0.0f, 0.0f, 0.0f, 1.0f))
 	)
-{
-}
+{ }
diff --git a/Engine/Source/Runtime/Core/Public/Math/Matrix.h b/Engine/Source/Runtime/Core/Public/Math/Matrix.h
--- a/Engine/Source/Runtime/Core/Public/Math/Matrix.h
+++ b/Engine/Source/Runtime/Core/Public/Math/Matrix.h
@@ -2,6 +2,7 @@
 
 #pragma once
 
+
 /**
  * 4x4 matrix of floating point values.
  * Matrix-matrix multiplication happens with a pre-multiple of the transpose --
@@ -60,57 +61,51 @@ public:
 	 * Gets the result of multiplying a Matrix to this.
 	 *
 	 * @param Other The matrix to multiply this by.
-	 *
 	 * @return The result of multiplication.
 	 */
-	FORCEINLINE FMatrix		operator* (const FMatrix& Other) const;
+	FORCEINLINE FMatrix operator* (const FMatrix& Other) const;
 
 	/**
-	 * Multiply this by a matix.
+	 * Multiply this by a matrix.
 	 *
 	 * @param Other the matrix to multiply by this.
-	 *
 	 * @return reference to this after multiply.
 	 */
-	FORCEINLINE void		operator*=(const FMatrix& Other);
+	FORCEINLINE void operator*=(const FMatrix& Other);
 
 	/**
 	 * Gets the result of adding a matrix to this.
 	 *
 	 * @param Other The Matrix to add.
-	 *
 	 * @return The result of addition.
 	 */
-	FORCEINLINE FMatrix		operator+ (const FMatrix& Other) const;
+	FORCEINLINE FMatrix operator+ (const FMatrix& Other) const;
 
 	/**
 	 * Adds to this matrix.
 	 *
 	 * @param Other The matrix to add to this.
-	 *
 	 * @return Reference to this after addition.
 	 */
-	FORCEINLINE void		operator+=(const FMatrix& Other);
+	FORCEINLINE void operator+=(const FMatrix& Other);
 
 	/** 
 	  * This isn't applying SCALE, just multiplying float to all members - i.e. weighting
 	  */
-	FORCEINLINE FMatrix		operator* (float Other) const;
+	FORCEINLINE FMatrix operator* (float Other) const;
 
 	/**
 	 * Multiply this matrix by a weighting factor.
 	 *
 	 * @param other The weight.
-	 *
 	 * @return a reference to this after weighting.
 	 */
-	FORCEINLINE void		operator*=(float Other);
+	FORCEINLINE void operator*=(float Other);
 
  	/**
 	 * Checks whether two matrix are identical.
 	 *
 	 * @param Other The other matrix.
-	 *
 	 * @return true if two matrix are identical, otherwise false.
 	 */
 	inline bool operator==(const FMatrix& Other) const;
@@ -120,7 +115,6 @@ public:
 	 *
 	 * @param Other The other Matrix.
 	 * @param Tolerance Error Tolerance.
-	 *
 	 * @return true if two Matrix are equal, within specified tolerance, otherwise false.
 	 */
 	inline bool Equals(const FMatrix& Other, float Tolerance=KINDA_SMALL_NUMBER) const;
@@ -129,12 +123,10 @@ public:
 	 * Checks whether another Matrix is not equal to this, within specified tolerance.
 	 *
 	 * @param Other The other Matrix.
-	 *
 	 * @return true if two Matrix are not equal, within specified tolerance, otherwise false.
 	 */
 	inline bool operator!=(const FMatrix& Other) const;
 
-	
 	// Homogeneous transform.
 	FORCEINLINE FVector4 TransformFVector4(const FVector4& V) const;
 
@@ -156,7 +148,6 @@ public:
 	 */
 	FORCEINLINE FVector InverseTransformVector(const FVector &V) const;
 
-
 	// Transpose.
 
 	FORCEINLINE FMatrix GetTransposed() const;
@@ -170,6 +161,7 @@ public:
 
 	/** Fast path, doesn't check for nil matrices in final release builds */
 	inline FMatrix InverseFast() const;
+
 	/** Fast path, and handles nil matrices. */
 	inline FMatrix Inverse() const;
 
@@ -230,7 +222,7 @@ public:
 	 * get unit length axis of this matrix
 	 *
 	 * @param i index into the axis of the matrix
-	 * @ return vector of the axis
+	 * @return vector of the axis
 	 */
 	inline FVector GetUnitAxis(EAxis::Type Axis) const;
 
@@ -278,6 +270,7 @@ public:
 
 	/** 
 	 * Transform a rotation matrix into a quaternion.
+	 *
 	 * @warning rotation part will need to be unit length for this to be right!
 	 */
 	CORE_API FQuat ToQuat() const;
@@ -302,8 +295,7 @@ public:
 	FORCEINLINE bool GetFrustumBottomPlane(FPlane& OutPlane) const;
 
 	/**
-	 * Utility for mirroring this transform across a certain plane,
-	 * and flipping one of the axis as well.
+	 * Utility for mirroring this transform across a certain plane, and flipping one of the axis as well.
 	 */
 	inline void Mirror(EAxis::Type MirrorAxis, EAxis::Type FlipAxis);
 
@@ -322,14 +314,13 @@ public:
 	 *
 	 * @param Ar Reference to the serialization archive.
 	 * @param M Reference to the matrix being serialized.
-	 *
 	 * @return Reference to the Archive after serialization.
 	 */
 	friend CORE_API FArchive& operator<<(FArchive& Ar,FMatrix& M);
 
 	/**
-	* Convert this Atom to the 3x4 transpose of the transformation matrix.
-	*/
+	 * Convert this Atom to the 3x4 transpose of the transformation matrix.
+	 */
 	void To3x4MatrixTranspose( float* Out ) const
 	{
 		const float* RESTRICT Src = &(M[0][0]);
@@ -352,6 +343,7 @@ public:
 	}	
 };
 
+
 /**
  * A storage class for compile-time fixed size matrices.
  */
@@ -374,10 +366,10 @@ public:
 	TMatrix(const FMatrix& InMatrix);
 };
 
+
 template<uint32 NumRows,uint32 NumColumns>
-FORCEINLINE TMatrix<NumRows,NumColumns>::TMatrix()
-{
-}
+FORCEINLINE TMatrix<NumRows,NumColumns>::TMatrix() { }
+
 
 template<uint32 NumRows,uint32 NumColumns>
 FORCEINLINE TMatrix<NumRows,NumColumns>::TMatrix(const FMatrix& InMatrix)
@@ -398,6 +390,7 @@ struct FBasisVectorMatrix: FMatrix
 	FBasisVectorMatrix(const FVector& XAxis,const FVector& YAxis,const FVector& ZAxis,const FVector& Origin);
 };
 
+
 FORCEINLINE FBasisVectorMatrix::FBasisVectorMatrix(const FVector& XAxis,const FVector& YAxis,const FVector& ZAxis,const FVector& Origin)
 {
 	for(uint32 RowIndex = 0;RowIndex < 3;RowIndex++)
@@ -423,6 +416,7 @@ struct FLookAtMatrix : FMatrix
 	FLookAtMatrix(const FVector& EyePosition, const FVector& LookAtPosition, const FVector& UpVector);
 };
 
+
 FORCEINLINE FLookAtMatrix::FLookAtMatrix(const FVector& EyePosition, const FVector& LookAtPosition, const FVector& UpVector)
 {
 	const FVector ZAxis = (LookAtPosition - EyePosition).GetSafeNormal();
@@ -445,6 +439,7 @@ FORCEINLINE FLookAtMatrix::FLookAtMatrix(const FVector& EyePosition, const FVect
 
 template <> struct TIsPODType<FMatrix> { enum { Value = true }; };
 
+
 // very high quality 4x4 matrix inverse
 static inline void Inverse4x4( double* dst, const float* src )
 {
@@ -480,4 +475,4 @@ static inline void Inverse4x4( double* dst, const float* src )
 	{
 		dst[i] = inv[i] * det;
 	}
-}
\ No newline at end of file
+}
diff --git a/Engine/Source/Runtime/Core/Public/Math/MirrorMatrix.h b/Engine/Source/Runtime/Core/Public/Math/MirrorMatrix.h
--- a/Engine/Source/Runtime/Core/Public/Math/MirrorMatrix.h
+++ b/Engine/Source/Runtime/Core/Public/Math/MirrorMatrix.h
@@ -2,21 +2,24 @@
 
 #pragma once
 
+
 /**
  * Mirrors a point about an abitrary plane 
  */
-class FMirrorMatrix : public FMatrix
+class FMirrorMatrix
+	: public FMatrix
 {
 public:
+
 	/** 
-	 * Constructor.
-	 * Updated for the fact that our FPlane uses Ax+By+Cz=D.
+	 * Constructor. Updated for the fact that our FPlane uses Ax+By+Cz=D.
 	 * 
-	 * @param	Plane - source plane for mirroring (assumed normalized)
+	 * @param Plane source plane for mirroring (assumed normalized)
 	 */
 	FMirrorMatrix( const FPlane& Plane );
 };
 
+
 FORCEINLINE FMirrorMatrix::FMirrorMatrix( const FPlane& Plane ) :
 FMatrix(
 	FPlane( -2.f*Plane.X*Plane.X + 1.f,	-2.f*Plane.Y*Plane.X,		-2.f*Plane.Z*Plane.X,		0.f ),
diff --git a/Engine/Source/Runtime/Core/Public/Math/NumericLimits.h b/Engine/Source/Runtime/Core/Public/Math/NumericLimits.h
--- a/Engine/Source/Runtime/Core/Public/Math/NumericLimits.h
+++ b/Engine/Source/Runtime/Core/Public/Math/NumericLimits.h
@@ -1,9 +1,5 @@
 // Copyright 1998-2015 Epic Games, Inc. All Rights Reserved.
 
-/*=============================================================================
-	NumericLimits.h: Unreal numeric limits
-=============================================================================*/
-
 #pragma once
 
 
diff --git a/Engine/Source/Runtime/Core/Public/Math/OrientedBox.h b/Engine/Source/Runtime/Core/Public/Math/OrientedBox.h
--- a/Engine/Source/Runtime/Core/Public/Math/OrientedBox.h
+++ b/Engine/Source/Runtime/Core/Public/Math/OrientedBox.h
@@ -1,9 +1,5 @@
 // Copyright 1998-2015 Epic Games, Inc. All Rights Reserved.
 
-/*=============================================================================
-	OrientedBox.h: Declares the FOrientedBox structure.
-=============================================================================*/
-
 #pragma once
 
 
@@ -40,7 +36,7 @@ public:
 	 *
 	 * Constructs a unit-sized, origin-centered box with axes aligned to the coordinate system.
 	 */
-	FOrientedBox ()
+	FOrientedBox()
 		: Center(0.0f)
 		, AxisX(1.0f, 0.0f, 0.0f)
 		, AxisY(0.0f, 1.0f, 0.0f)
@@ -57,14 +53,14 @@ public:
 	 *
 	 * @param Verts The array to fill in with the vertices.
 	 */
-	FORCEINLINE void CalcVertices (FVector* Verts) const;
+	FORCEINLINE void CalcVertices(FVector* Verts) const;
 
 	/**
 	 * Finds the projection interval of the box when projected onto Axis.
 	 *
 	 * @param Axis The unit vector defining the axis to project the box onto.
 	 */
-	FORCEINLINE FInterval Project (const FVector& Axis) const;
+	FORCEINLINE FInterval Project(const FVector& Axis) const;
 };
 
 
diff --git a/Engine/Source/Runtime/Core/Public/Math/Plane.h b/Engine/Source/Runtime/Core/Public/Math/Plane.h
--- a/Engine/Source/Runtime/Core/Public/Math/Plane.h
+++ b/Engine/Source/Runtime/Core/Public/Math/Plane.h
@@ -1,9 +1,5 @@
 // Copyright 1998-2015 Epic Games, Inc. All Rights Reserved.
 
-/*=============================================================================
-	Plane.h: Declares the FPlane class.
-=============================================================================*/
-
 #pragma once
 
 
@@ -22,10 +18,8 @@ public:
 
 public:
 
-	/**
-	 * Default constructor (no initialization).
-	 */
-	FORCEINLINE FPlane( );
+	/** Default constructor (no initialization). */
+	FORCEINLINE FPlane();
 
 	/**
 	 * Copy Constructor.
diff --git a/Engine/Source/Runtime/Core/Public/Math/Quat.h b/Engine/Source/Runtime/Core/Public/Math/Quat.h
--- a/Engine/Source/Runtime/Core/Public/Math/Quat.h
+++ b/Engine/Source/Runtime/Core/Public/Math/Quat.h
@@ -1,9 +1,5 @@
 // Copyright 1998-2015 Epic Games, Inc. All Rights Reserved.
 
-/*=============================================================================
-	Quat.h: Declares the FQuat class.
-=============================================================================*/
-
 #pragma once
 
 
@@ -44,7 +40,7 @@ public:
 	/**
 	 * Default constructor (no initialization).
 	 */
-	FORCEINLINE FQuat( ) { }
+	FORCEINLINE FQuat() { }
 
 	/**
 	 * Creates and initializes a new quaternion, with the W component either 0 or 1.
@@ -252,7 +248,7 @@ public:
 	static CORE_API FQuat MakeFromEuler( const FVector& Euler );
 
 	/** Convert a Quaternion into floating-point Euler angles (in degrees). */
-	CORE_API FVector Euler( ) const;
+	CORE_API FVector Euler() const;
 
 	/**
 	 * Normalize this quaternion if it is large enough.
@@ -262,21 +258,21 @@ public:
 	FORCEINLINE void Normalize( float Tolerance=SMALL_NUMBER );
 
 	// Return true if this quaternion is normalized
-	bool IsNormalized( ) const;
+	bool IsNormalized() const;
 
 	/**
 	 * Get the length of this quaternion.
 	 *
 	 * @return The length of this quaternion.
 	 */
-	FORCEINLINE float Size( ) const;
+	FORCEINLINE float Size() const;
 
 	/**
 	 * Get the length squared of this quaternion.
 	 *
 	 * @return The length of this quaternion.
 	 */
-	FORCEINLINE float SizeSquared( ) const;
+	FORCEINLINE float SizeSquared() const;
 
 	/** 
 	 * get the axis and angle of rotation of this quaternion
@@ -298,16 +294,16 @@ public:
 	/**
 	 * @return quaternion with W=0 and V=theta*v.
 	 */
-	CORE_API FQuat Log( ) const;
+	CORE_API FQuat Log() const;
 
 	/**
 	 * @note Exp should really only be used after Log.
 	 * Assumes a quaternion with W=0 and V=theta*v (where |v| = 1).
 	 * Exp(q) = (sin(theta)*v, cos(theta))
 	 */
-	CORE_API FQuat Exp( ) const;
+	CORE_API FQuat Exp() const;
 
-	FORCEINLINE FQuat Inverse( ) const;
+	FORCEINLINE FQuat Inverse() const;
 
 	/**
 	 * Enforce that the delta between this Quaternion and another one represents
@@ -316,19 +312,19 @@ public:
 	void EnforceShortestArcWith( const FQuat& OtherQuat );
 	
 	/** Get X Rotation Axis. */
-	FORCEINLINE FVector GetAxisX( ) const;
+	FORCEINLINE FVector GetAxisX() const;
 
 	/** Get Y Rotation Axis. */
-	FORCEINLINE FVector GetAxisY( ) const;
+	FORCEINLINE FVector GetAxisY() const;
 
 	/** Get Z Rotation Axis. */
-	FORCEINLINE FVector GetAxisZ( ) const;
+	FORCEINLINE FVector GetAxisZ() const;
 
 	/** @return rotator representation of this quaternion */
-	CORE_API FRotator Rotator( ) const;
+	CORE_API FRotator Rotator() const;
 
 	/** @return Vector of the axis of the quaternion */
-	FORCEINLINE FVector GetRotationAxis( ) const;
+	FORCEINLINE FVector GetRotationAxis() const;
 
 	/**
 	 * Serializes the vector compressed for e.g. network transmission.
@@ -342,19 +338,19 @@ public:
 	 *
 	 * @return true if there are any NaNs in this Quaternion, otherwise false.
 	 */
-	bool ContainsNaN( ) const;
+	bool ContainsNaN() const;
 
 	/**
 	 * Get a textual representation of the vector.
 	 *
 	 * @return Text describing the vector.
 	 */
-	FString ToString( ) const;
+	FString ToString() const;
 
 public:
 
 #if ENABLE_NAN_DIAGNOSTIC
-	FORCEINLINE void DiagnosticCheckNaN( ) const
+	FORCEINLINE void DiagnosticCheckNaN() const
 	{
 		checkf(!ContainsNaN(), TEXT("FQuat contains NaN: %s"), *ToString());
 	}
@@ -602,7 +598,7 @@ FORCEINLINE FQuat::FQuat( const FQuat& Q )
 { }
 
 
-FORCEINLINE FString FQuat::ToString( ) const
+FORCEINLINE FString FQuat::ToString() const
 {
 	return FString::Printf(TEXT("X=%3.3f Y=%3.3f Z=%3.3f W=%3.3f"), X, Y, Z, W);
 }
@@ -772,19 +768,19 @@ FORCEINLINE void FQuat::Normalize(float Tolerance)
 }
 
 
-FORCEINLINE bool FQuat::IsNormalized( ) const
+FORCEINLINE bool FQuat::IsNormalized() const
 {
 	return (FMath::Abs(1.f - SizeSquared()) < THRESH_QUAT_NORMALIZED);
 }
 
 
-FORCEINLINE float FQuat::Size( ) const
+FORCEINLINE float FQuat::Size() const
 {
 	return FMath::Sqrt(X * X + Y * Y + Z * Z + W * W);
 }
 
 
-FORCEINLINE float FQuat::SizeSquared( ) const
+FORCEINLINE float FQuat::SizeSquared() const
 {
 	return (X * X + Y * Y + Z * Z + W * W);
 }
@@ -797,7 +793,7 @@ FORCEINLINE void FQuat::ToAxisAndAngle( FVector& Axis, float& Angle ) const
 }
 
 
-FORCEINLINE FVector FQuat::GetRotationAxis( ) const
+FORCEINLINE FVector FQuat::GetRotationAxis() const
 {
 	// Ensure we never try to sqrt a neg number
 	const float S = FMath::Sqrt(FMath::Max(1.f - (W * W), 0.f));
@@ -824,7 +820,7 @@ FORCEINLINE FVector FQuat::RotateVector( FVector V ) const
 }
 
 
-FORCEINLINE FQuat FQuat::Inverse( ) const
+FORCEINLINE FQuat FQuat::Inverse() const
 {
 	checkSlow(IsNormalized());
 
@@ -844,19 +840,19 @@ FORCEINLINE void FQuat::EnforceShortestArcWith( const FQuat& OtherQuat )
 }
 
 	
-FORCEINLINE FVector FQuat::GetAxisX( ) const
+FORCEINLINE FVector FQuat::GetAxisX() const
 {
 	return RotateVector(FVector(1.f, 0.f, 0.f));
 }
 
 
-FORCEINLINE FVector FQuat::GetAxisY( ) const
+FORCEINLINE FVector FQuat::GetAxisY() const
 {
 	return RotateVector(FVector(0.f, 1.f, 0.f));
 }
 
 
-FORCEINLINE FVector FQuat::GetAxisZ( ) const
+FORCEINLINE FVector FQuat::GetAxisZ() const
 {
 	return RotateVector(FVector(0.f, 0.f, 1.f));
 }
@@ -900,7 +896,7 @@ FORCEINLINE FQuat FQuat::FastBilerp( const FQuat& P00, const FQuat& P10, const F
 }
 
 
-FORCEINLINE bool FQuat::ContainsNaN( ) const
+FORCEINLINE bool FQuat::ContainsNaN() const
 {
 	return (FMath::IsNaN(X) || !FMath::IsFinite(X) ||
 			FMath::IsNaN(Y) || !FMath::IsFinite(Y) ||
diff --git a/Engine/Source/Runtime/Core/Public/Math/RandomStream.h b/Engine/Source/Runtime/Core/Public/Math/RandomStream.h
--- a/Engine/Source/Runtime/Core/Public/Math/RandomStream.h
+++ b/Engine/Source/Runtime/Core/Public/Math/RandomStream.h
@@ -1,9 +1,5 @@
 // Copyright 1998-2015 Epic Games, Inc. All Rights Reserved.
 
-/*=============================================================================
-	RandomStream.h: Declares the FRandomStream class.
-=============================================================================*/
-
 #pragma once
 
 
@@ -28,7 +24,7 @@ public:
 	 *
 	 * The seed should be set prior to use.
 	 */
-	FRandomStream( )
+	FRandomStream()
 		: InitialSeed(0)
 		, Seed(0)
 	{ }
@@ -59,7 +55,7 @@ public:
 	/**
 	 * Resets this random stream to the initial seed value.
 	 */
-	void Reset( ) const
+	void Reset() const
 	{
 		Seed = InitialSeed;
 	}
@@ -72,7 +68,7 @@ public:
 	/**
 	 * Generates a new random seed.
 	 */
-	void GenerateNewSeed( )
+	void GenerateNewSeed()
 	{
 		Initialize(FMath::Rand());
 	}
@@ -82,7 +78,7 @@ public:
 	 *
 	 * @return Random number.
 	 */
-	float GetFraction( ) const
+	float GetFraction() const
 	{
 		MutateSeed();
 
@@ -99,7 +95,7 @@ public:
 	 *
 	 * @return Random number.
 	 */
-	uint32 GetUnsignedInt( ) const
+	uint32 GetUnsignedInt() const
 	{
 		MutateSeed();
 
@@ -111,7 +107,7 @@ public:
 	 *
 	 * @return Random unit vector.
 	 */
-	FVector GetUnitVector( ) const
+	FVector GetUnitVector() const
 	{
 		FVector Result;
 		float L;
@@ -134,7 +130,7 @@ public:
 	 *
 	 * @return Current seed.
 	 */
-	int32 GetCurrentSeed( ) const
+	int32 GetCurrentSeed() const
 	{
 		return Seed;
 	}
@@ -144,7 +140,7 @@ public:
 	 *
 	 * @return Random number.
 	 */
-	FORCEINLINE float FRand( ) const
+	FORCEINLINE float FRand() const
 	{
 		return GetFraction();
 	}
@@ -187,7 +183,7 @@ public:
 	 *
 	 * @return Random unit vector.
 	 */
-	FORCEINLINE FVector VRand( ) const
+	FORCEINLINE FVector VRand() const
 	{
 		return GetUnitVector();
 	}
@@ -291,7 +287,7 @@ protected:
 	/**
 	 * Mutates the current seed into the next seed.
 	 */
-	void MutateSeed( ) const
+	void MutateSeed() const
 	{
 		Seed = (Seed * 196314165) + 907633515; 
 	}
diff --git a/Engine/Source/Runtime/Core/Public/Math/Range.h b/Engine/Source/Runtime/Core/Public/Math/Range.h
--- a/Engine/Source/Runtime/Core/Public/Math/Range.h
+++ b/Engine/Source/Runtime/Core/Public/Math/Range.h
@@ -1,9 +1,5 @@
 // Copyright 1998-2015 Epic Games, Inc. All Rights Reserved.
 
-/*=============================================================================
-	Range.h: Declares the TRange class template.
-=============================================================================*/
-
 #pragma once
 
 
@@ -50,10 +46,8 @@ public:
 
 	typedef TRangeBound<ElementType> BoundsType;
 
-	/**
-	 * Default constructor (no initialization).
-	 */
-	TRange( ) { }
+	/** Default constructor (no initialization). */
+	TRange() { }
 
 	/**
 	 * Creates a range with a single element.
@@ -149,7 +143,8 @@ public:
 			return ((UpperBound.IsInclusive() && Other.LowerBound.IsExclusive()) ||
 					(UpperBound.IsExclusive() && Other.LowerBound.IsInclusive()));
 		}
-		else if (Other.UpperBound.GetValue() == LowerBound.GetValue())
+		
+		if (Other.UpperBound.GetValue() == LowerBound.GetValue())
 		{
 			return ((Other.UpperBound.IsInclusive() && LowerBound.IsExclusive()) ||
 					(Other.UpperBound.IsExclusive() && LowerBound.IsInclusive()));
@@ -220,7 +215,7 @@ public:
 	 *
 	 * @return Lower bound.
 	 */
-	BoundsType GetLowerBound( ) const
+	BoundsType GetLowerBound() const
 	{
 		return LowerBound;
 	}
@@ -233,7 +228,7 @@ public:
 	 * @return Bound value.
 	 * @see GetUpperBoundValue, HasLowerBound
 	 */
-	const ElementType& GetLowerBoundValue( ) const
+	const ElementType& GetLowerBoundValue() const
 	{
 		return LowerBound.GetValue();
 	}
@@ -243,7 +238,7 @@ public:
 	 *
 	 * @return Upper bound.
 	 */
-	BoundsType GetUpperBound( ) const
+	BoundsType GetUpperBound() const
 	{
 		return UpperBound;
 	}
@@ -256,7 +251,7 @@ public:
 	 * @return Bound value.
 	 * @see GetLowerBoundValue, HasUpperBound
 	 */
-	const ElementType& GetUpperBoundValue( ) const
+	const ElementType& GetUpperBoundValue() const
 	{
 		return UpperBound.GetValue();
 	}
@@ -266,7 +261,7 @@ public:
 	 *
 	 * @return true if the range has a lower bound, false otherwise.
 	 */
-	bool HasLowerBound( ) const
+	bool HasLowerBound() const
 	{
 		return LowerBound.IsClosed();
 	}
@@ -276,7 +271,7 @@ public:
 	 *
 	 * @return true if the range has an upper bound, false otherwise.
 	 */
-	bool HasUpperBound( ) const
+	bool HasUpperBound() const
 	{
 		return UpperBound.IsClosed();
 	}
@@ -289,7 +284,7 @@ public:
 	 *
 	 * @return true if the range is degenerate, false otherwise.
 	 */
-	bool IsDegenerate( ) const
+	bool IsDegenerate() const
 	{
 		return (LowerBound.IsInclusive() && (LowerBound == UpperBound));
 	}
@@ -304,7 +299,7 @@ public:
 	 *
 	 * @return true if the range is empty, false otherwise.
 	 */
-	bool IsEmpty( ) const
+	bool IsEmpty() const
 	{
 		if (LowerBound.IsClosed() && UpperBound.IsClosed())
 		{
@@ -369,7 +364,7 @@ public:
 	 *
 	 * @return Range size.
 	 */
-	template<typename DifferenceType> DifferenceType Size( ) const
+	template<typename DifferenceType> DifferenceType Size() const
 	{
 		check(LowerBound.IsClosed() && UpperBound.IsClosed());
 
@@ -443,7 +438,7 @@ public:
 	 *
 	 * @return Empty range.
 	 */
-	static FORCEINLINE TRange Empty( )
+	static FORCEINLINE TRange Empty()
 	{
 		return TRange(BoundsType::Exclusive(ElementType()), BoundsType::Exclusive(ElementType()));
 	}
@@ -587,7 +582,7 @@ public:
 	 *
 	 * @return A new range.
 	 */
-	static FORCEINLINE TRange All( )
+	static FORCEINLINE TRange All()
 	{
 		return TRange(BoundsType::Open(), BoundsType::Open());
 	}
@@ -674,37 +669,23 @@ private:
 /* Default ranges for built-in types
  *****************************************************************************/
 
-/**
- * Defines a discrete range of dates.
- */
+/** Defines a discrete range of dates. */
 typedef TRange<FDateTime> FDateRange;
 
-/**
- * Defines a discrete range of dates.
- */
+/** Defines a discrete range of dates. */
 typedef TRange<double> FDoubleRange;
 
-/**
- * Defines a discrete range of dates.
- */
+/** Defines a discrete range of dates. */
 typedef TRange<float> FFloatRange;
 
-/**
- * Defines a discrete range of 8-bit signed integers.
- */
+/** Defines a discrete range of 8-bit signed integers. */
 typedef TRange<int8> FInt8Range;
 
-/**
- * Defines a discrete range of 16-bit signed integers.
- */
+/** Defines a discrete range of 16-bit signed integers. */
 typedef TRange<int16> FInt16Range;
 
-/**
- * Defines a discrete range of 32-bit signed integers.
- */
+/** Defines a discrete range of 32-bit signed integers. */
 typedef TRange<int32> FInt32Range;
 
-/**
- * Defines a discrete range of 64-bit signed integers.
- */
+/** Defines a discrete range of 64-bit signed integers. */
 typedef TRange<int64> FInt64Range;
diff --git a/Engine/Source/Runtime/Core/Public/Math/RangeBound.h b/Engine/Source/Runtime/Core/Public/Math/RangeBound.h
--- a/Engine/Source/Runtime/Core/Public/Math/RangeBound.h
+++ b/Engine/Source/Runtime/Core/Public/Math/RangeBound.h
@@ -1,9 +1,5 @@
 // Copyright 1998-2015 Epic Games, Inc. All Rights Reserved.
 
-/*=============================================================================
-	RangeBound.h: Declares the TRangeBound template.
-=============================================================================*/
-
 #pragma once
 
 
@@ -38,7 +34,7 @@ public:
 	 *
 	 * @see Exclusive, Inclusive, Open
 	 */
-	TRangeBound( )
+	TRangeBound()
 		: Type(ERangeBoundTypes::Open)
 		, Value()
 	{ }
@@ -88,7 +84,7 @@ public:
 	 * @return Bound value.
 	 * @see IsOpen
 	 */
-	const ElementType& GetValue( ) const
+	const ElementType& GetValue() const
 	{
 		check(Type != ERangeBoundTypes::Open);
 
@@ -100,7 +96,7 @@ public:
 	 *
 	 * @return true if the bound is closed, false otherwise.
 	 */
-	bool IsClosed( ) const
+	bool IsClosed() const
 	{
 		return (Type != ERangeBoundTypes::Open);
 	}
@@ -110,7 +106,7 @@ public:
 	 *
 	 * @return true if the bound is exclusive, false otherwise.
 	 */
-	bool IsExclusive( ) const
+	bool IsExclusive() const
 	{
 		return (Type == ERangeBoundTypes::Exclusive);
 	}
@@ -120,7 +116,7 @@ public:
 	 *
 	 * @return true if the bound is inclusive, false otherwise.
 	 */
-	bool IsInclusive( ) const
+	bool IsInclusive() const
 	{
 		return (Type == ERangeBoundTypes::Inclusive);
 	}
@@ -130,7 +126,7 @@ public:
 	 *
 	 * @return true if the bound is open, false otherwise.
 	 */
-	bool IsOpen( ) const
+	bool IsOpen() const
 	{
 		return (Type == ERangeBoundTypes::Open);
 	}
@@ -199,7 +195,7 @@ public:
 	 *
 	 * @return An open bound.
 	 */
-	static FORCEINLINE TRangeBound Open( )
+	static FORCEINLINE TRangeBound Open()
 	{
 		TRangeBound Result;
 
@@ -317,37 +313,23 @@ private:
 /* Default range bounds for built-in types
  *****************************************************************************/
 
-/**
- * Defines a discrete range bound for dates.
- */
+/** Defines a discrete range bound for dates. */
 typedef TRangeBound<FDateTime> FDateRangeBound;
 
-/**
- * Defines a discrete range bound for dates.
- */
+/** Defines a discrete range bound for dates. */
 typedef TRangeBound<double> FDoubleRangeBound;
 
-/**
- * Defines a discrete range bound for dates.
- */
+/** Defines a discrete range bound for dates. */
 typedef TRangeBound<float> FFloatRangeBound;
 
-/**
- * Defines a discrete range bound for 8-bit signed integers.
- */
+/** Defines a discrete range bound for 8-bit signed integers. */
 typedef TRangeBound<int8> FInt8RangeBound;
 
-/**
- * Defines a discrete range bound for 16-bit signed integers.
- */
+/** Defines a discrete range bound for 16-bit signed integers. */
 typedef TRangeBound<int16> FInt16RangeBound;
 
-/**
- * Defines a discrete range bound for 32-bit signed integers.
- */
+/** Defines a discrete range bound for 32-bit signed integers. */
 typedef TRangeBound<int32> FInt32RangeBound;
 
-/**
- * Defines a discrete range bound for 64-bit signed integers.
- */
+/** Defines a discrete range bound for 64-bit signed integers. */
 typedef TRangeBound<int64> FInt64RangeBound;
diff --git a/Engine/Source/Runtime/Core/Public/Math/RangeSet.h b/Engine/Source/Runtime/Core/Public/Math/RangeSet.h
--- a/Engine/Source/Runtime/Core/Public/Math/RangeSet.h
+++ b/Engine/Source/Runtime/Core/Public/Math/RangeSet.h
@@ -1,9 +1,5 @@
 // Copyright 1998-2015 Epic Games, Inc. All Rights Reserved.
 
-/*=============================================================================
-	RangeSet.h: Declares the TRangeSet class template.
-=============================================================================*/
-
 #pragma once
 
 
@@ -18,15 +14,11 @@ template<typename ElementType> class TRangeSet
 
 public:
 
-	/**
-	 * Default constructor.
-	 */
-	TRangeSet( ) { }
+	/** Default constructor. */
+	TRangeSet() { }
 
-	/**
-	 * Destructor.
-	 */
-	~TRangeSet( ) { }
+	/** Destructor. */
+	~TRangeSet() { }
 
 public:
 	
@@ -105,7 +97,7 @@ public:
 	/**
 	 * Removes all ranges from the set.
 	 */
-	void Empty( )
+	void Empty()
 	{
 		Ranges.Empty();
 	}
@@ -155,7 +147,7 @@ public:
 	 *
 	 * @return Array of ranges.
 	 */
-	const TArray<RangeType>& GetRanges( ) const
+	const TArray<RangeType>& GetRanges() const
 	{
 		return Ranges;
 	}
@@ -165,7 +157,7 @@ public:
 	 *
 	 * @return true if the range set is empty, false otherwise.
 	 */
-	bool IsEmpty( ) const
+	bool IsEmpty() const
 	{
 		return (Ranges.Num() == 0);
 	}
diff --git a/Engine/Source/Runtime/Core/Public/Math/RotationAboutPointMatrix.h b/Engine/Source/Runtime/Core/Public/Math/RotationAboutPointMatrix.h
--- a/Engine/Source/Runtime/Core/Public/Math/RotationAboutPointMatrix.h
+++ b/Engine/Source/Runtime/Core/Public/Math/RotationAboutPointMatrix.h
@@ -2,11 +2,15 @@
 
 #pragma once
 
+
 /** Rotates about an Origin point. */
-class FRotationAboutPointMatrix : public FRotationTranslationMatrix
+class FRotationAboutPointMatrix
+	: public FRotationTranslationMatrix
 {
 public:
-	/** Constructor
+
+	/**
+	 * Constructor.
 	 *
 	 * @param Rot rotation
 	 * @param Origin about which to rotate.
@@ -23,8 +27,9 @@ public:
 	static CORE_API FMatrix Make(const FQuat& Rot, const FVector& Origin);
 };
 
+
 FORCEINLINE FRotationAboutPointMatrix::FRotationAboutPointMatrix(const FRotator& Rot, const FVector& Origin)
-	:FRotationTranslationMatrix(Rot, Origin)
+	: FRotationTranslationMatrix(Rot, Origin)
 {
 	// FRotationTranslationMatrix generates R * T.
 	// We need -T * R * T, so prepend that translation:
@@ -35,4 +40,4 @@ FORCEINLINE FRotationAboutPointMatrix::FRotationAboutPointMatrix(const FRotator&
 	M[3][0]	-= XAxis | Origin;
 	M[3][1]	-= YAxis | Origin;
 	M[3][2]	-= ZAxis | Origin;
-}
\ No newline at end of file
+}
diff --git a/Engine/Source/Runtime/Core/Public/Math/RotationMatrix.h b/Engine/Source/Runtime/Core/Public/Math/RotationMatrix.h
--- a/Engine/Source/Runtime/Core/Public/Math/RotationMatrix.h
+++ b/Engine/Source/Runtime/Core/Public/Math/RotationMatrix.h
@@ -2,15 +2,21 @@
 
 #pragma once
 
+
 /** Rotation matrix no translation */
-class FRotationMatrix : public FRotationTranslationMatrix
+class FRotationMatrix
+	: public FRotationTranslationMatrix
 {
 public:
-	/** Constructor
+
+	/**
+	 * Constructor.
 	 *
 	 * @param Rot rotation
 	 */
-	FRotationMatrix(const FRotator& Rot) : FRotationTranslationMatrix(Rot, FVector::ZeroVector)	{}
+	FRotationMatrix(const FRotator& Rot)
+		: FRotationTranslationMatrix(Rot, FVector::ZeroVector)
+	{ }
 
 	/** Matrix factory. Return an FMatrix so we don't have type conversion issues in expressions. */
 	static FMatrix Make(FRotator const& Rot)
@@ -47,4 +53,4 @@ public:
 
 	/** Builds a matrix with given Z and Y axes. Z will remain fixed, Y may be changed minimally to enforce orthogonality. X will be computed. Inputs need not be normalized. */
 	static CORE_API FMatrix MakeFromZY(FVector const& ZAxis, FVector const& YAxis);
-};
\ No newline at end of file
+};
diff --git a/Engine/Source/Runtime/Core/Public/Math/RotationTranslationMatrix.h b/Engine/Source/Runtime/Core/Public/Math/RotationTranslationMatrix.h
--- a/Engine/Source/Runtime/Core/Public/Math/RotationTranslationMatrix.h
+++ b/Engine/Source/Runtime/Core/Public/Math/RotationTranslationMatrix.h
@@ -2,11 +2,15 @@
 
 #pragma once
 
+
 /** Combined rotation and translation matrix */
-class FRotationTranslationMatrix : public FMatrix
+class FRotationTranslationMatrix
+	: public FMatrix
 {
 public:
-	/** Constructor
+
+	/**
+	 * Constructor.
 	 *
 	 * @param Rot rotation
 	 * @param Origin translation to apply
@@ -18,9 +22,9 @@ public:
 	{
 		return FRotationTranslationMatrix(Rot, Origin);
 	}
-
 };
 
+
 FORCEINLINE FRotationTranslationMatrix::FRotationTranslationMatrix(const FRotator& Rot, const FVector& Origin)
 {
 	const float	SR	= FMath::Sin(Rot.Roll * PI / 180.f);
@@ -49,4 +53,4 @@ FORCEINLINE FRotationTranslationMatrix::FRotationTranslationMatrix(const FRotato
 	M[3][1]	= Origin.Y;
 	M[3][2]	= Origin.Z;
 	M[3][3]	= 1.f;
-}
\ No newline at end of file
+}
diff --git a/Engine/Source/Runtime/Core/Public/Math/Rotator.h b/Engine/Source/Runtime/Core/Public/Math/Rotator.h
--- a/Engine/Source/Runtime/Core/Public/Math/Rotator.h
+++ b/Engine/Source/Runtime/Core/Public/Math/Rotator.h
@@ -1,9 +1,5 @@
 // Copyright 1998-2015 Epic Games, Inc. All Rights Reserved.
 
-/*=============================================================================
-	Rotator.h: Declares the FRotator class.
-=============================================================================*/
-
 #pragma once
 
 
@@ -35,7 +31,7 @@ public:
 	/**
 	 * Default constructor (no initialization).
 	 */
-	FORCEINLINE FRotator( ) { };
+	FORCEINLINE FRotator() { };
 
 	/**
 	 * Constructor
@@ -110,7 +106,7 @@ public:
 	 *
 	 * @return A negated copy of the rotator.
 	 */
-	FORCEINLINE FRotator operator-( ) const;
+	FORCEINLINE FRotator operator-() const;
 
 	// Binary comparison operators.
 
@@ -165,7 +161,7 @@ public:
 	 *
 	 * @return true if this has exactly zero rotation, otherwise false.
 	 */
-	bool IsZero( ) const;
+	bool IsZero() const;
 
 	/**
 	 * Checks whether two rotators are equal, within specified tolerance.
@@ -199,21 +195,21 @@ public:
 	 *
 	 * @return Rotation as a direction vector.
 	 */
-	CORE_API FVector Vector( ) const;
+	CORE_API FVector Vector() const;
 
 	/**
 	 * Get Rotation as a quaternion.
 	 *
 	 * @return Rotation as a quaternion.
 	 */
-	CORE_API FQuat Quaternion( ) const;
+	CORE_API FQuat Quaternion() const;
 
 	/**
 	 * Convert a Rotator into floating-point Euler angles (in degrees). Rotator now stored in degrees.
 	 *
 	 * @return Rotation as a Euler angle vector.
 	 */
-	CORE_API FVector Euler( ) const;
+	CORE_API FVector Euler() const;
 
 	/**
 	 * Rotate a vector rotated by this rotator.
@@ -236,26 +232,26 @@ public:
 	 *
 	 * @return Clamped version of rotator.
 	 */
-	FRotator Clamp( ) const;
+	FRotator Clamp() const;
 
 	/** 
 	 * Create a copy of this rotator and normalize, removes all winding and creates the "shortest route" rotation. 
 	 *
 	 * @return Normalized copy of this rotator
 	 */
-	FRotator GetNormalized( ) const;
+	FRotator GetNormalized() const;
 
 	/** 
 	 * Create a copy of this rotator and denormalize, clamping each axis to 0 - 360. 
 	 *
 	 * @return Denormalized copy of this rotator
 	 */
-	FRotator GetDenormalized( ) const;
+	FRotator GetDenormalized() const;
 
 	/**
 	 * In-place normalize, removes all winding and creates the "shortest route" rotation.
 	 */
-	void Normalize( );
+	void Normalize();
 
 	/** 
 	 * Decompose this Rotator into a Winding part (multiples of 360) and a Remainder part. 
@@ -271,10 +267,10 @@ public:
 	 *
 	 * @return Text describing the vector.
 	 */
-	FString ToString( ) const;
+	FString ToString() const;
 
 	/** Get a short textural representation of this vector, for compact readable logging. */
-	FString ToCompactString( ) const;
+	FString ToCompactString() const;
 
 	/**
 	 * Initialize this Rotator based on an FString. The String is expected to contain P=, Y=, R=.
@@ -290,7 +286,7 @@ public:
 	 *
 	 * @return true if there are any NaNs in this Rotator, otherwise false.
 	 */
-	bool ContainsNaN( ) const;
+	bool ContainsNaN() const;
 
 	/**
 	 * Serializes the rotator compressed for e.g. network transmission.
@@ -525,7 +521,7 @@ FORCEINLINE FRotator FRotator::GridSnap( const FRotator &RotGrid ) const
 }
 
 
-FORCEINLINE FRotator FRotator::Clamp( ) const
+FORCEINLINE FRotator FRotator::Clamp() const
 {
 	return FRotator(ClampAxis(Pitch), ClampAxis(Yaw), ClampAxis(Roll));
 }
@@ -589,7 +585,7 @@ FORCEINLINE float FRotator::DecompressAxisFromShort( uint16 Angle )
 }
 
 
-FORCEINLINE FRotator FRotator::GetNormalized( ) const
+FORCEINLINE FRotator FRotator::GetNormalized() const
 {
 	FRotator Rot = *this;
 	Rot.Normalize();
@@ -597,7 +593,7 @@ FORCEINLINE FRotator FRotator::GetNormalized( ) const
 }
 
 
-FORCEINLINE FRotator FRotator::GetDenormalized( ) const
+FORCEINLINE FRotator FRotator::GetDenormalized() const
 {
 	FRotator Rot = *this;
 	Rot.Pitch	= ClampAxis(Rot.Pitch);
@@ -607,20 +603,20 @@ FORCEINLINE FRotator FRotator::GetDenormalized( ) const
 }
 
 
-FORCEINLINE void FRotator::Normalize( )
+FORCEINLINE void FRotator::Normalize()
 {
 	Yaw = NormalizeAxis(Yaw);
 	Pitch = NormalizeAxis(Pitch);
 	Roll = NormalizeAxis(Roll);
 }
 
-FORCEINLINE FString FRotator::ToString( ) const
+FORCEINLINE FString FRotator::ToString() const
 {
 	return FString::Printf(TEXT("P=%f Y=%f R=%f"), Pitch, Yaw, Roll );
 }
 
 
-FORCEINLINE FString FRotator::ToCompactString( ) const
+FORCEINLINE FString FRotator::ToCompactString() const
 {
 	if( IsNearlyZero() )
 	{
@@ -668,7 +664,7 @@ FORCEINLINE bool FRotator::InitFromString( const FString& InSourceString )
 }
 
 
-FORCEINLINE bool FRotator::ContainsNaN( ) const
+FORCEINLINE bool FRotator::ContainsNaN() const
 {
 	return (FMath::IsNaN(Pitch) || !FMath::IsFinite(Pitch) 
 			|| FMath::IsNaN(Yaw) || !FMath::IsFinite(Yaw) 
diff --git a/Engine/Source/Runtime/Core/Public/Math/ScaleMatrix.h b/Engine/Source/Runtime/Core/Public/Math/ScaleMatrix.h
--- a/Engine/Source/Runtime/Core/Public/Math/ScaleMatrix.h
+++ b/Engine/Source/Runtime/Core/Public/Math/ScaleMatrix.h
@@ -1,9 +1,5 @@
 // Copyright 1998-2015 Epic Games, Inc. All Rights Reserved.
 
-/*=============================================================================
-	ScaleMatrix.h: Declares the FScaleMatrix class.
-=============================================================================*/
-
 #pragma once
 
 
diff --git a/Engine/Source/Runtime/Core/Public/Math/ScaleRotationTranslationMatrix.h b/Engine/Source/Runtime/Core/Public/Math/ScaleRotationTranslationMatrix.h
--- a/Engine/Source/Runtime/Core/Public/Math/ScaleRotationTranslationMatrix.h
+++ b/Engine/Source/Runtime/Core/Public/Math/ScaleRotationTranslationMatrix.h
@@ -2,20 +2,24 @@
 
 #pragma once
 
+
 /** Combined Scale rotation and translation matrix */
-class FScaleRotationTranslationMatrix : public FMatrix
+class FScaleRotationTranslationMatrix
+	: public FMatrix
 {
 public:
 
-	/** Constructor
-	*
-	* @param Scale scale to apply to matrix
-	* @param Rot rotation
-	* @param Origin translation to apply
-	*/
+	/**
+	 * Constructor.
+	 *
+	 * @param Scale scale to apply to matrix
+	 * @param Rot rotation
+	 * @param Origin translation to apply
+	 */
 	FScaleRotationTranslationMatrix(const FVector& Scale, const FRotator& Rot, const FVector& Origin);
 };
 
+
 FORCEINLINE FScaleRotationTranslationMatrix::FScaleRotationTranslationMatrix(const FVector& Scale, const FRotator& Rot, const FVector& Origin)
 {
 	const float	SR	= FMath::Sin(Rot.Roll * PI / 180.f);
@@ -44,4 +48,4 @@ FORCEINLINE FScaleRotationTranslationMatrix::FScaleRotationTranslationMatrix(con
 	M[3][1]	= Origin.Y;
 	M[3][2]	= Origin.Z;
 	M[3][3]	= 1.f;
-}
\ No newline at end of file
+}
diff --git a/Engine/Source/Runtime/Core/Public/Math/Sphere.h b/Engine/Source/Runtime/Core/Public/Math/Sphere.h
--- a/Engine/Source/Runtime/Core/Public/Math/Sphere.h
+++ b/Engine/Source/Runtime/Core/Public/Math/Sphere.h
@@ -1,9 +1,5 @@
 // Copyright 1998-2015 Epic Games, Inc. All Rights Reserved.
 
-/*=============================================================================
-	Sphere.h: Declares the FSphere class.
-=============================================================================*/
-
 #pragma once
 
 
@@ -22,10 +18,8 @@ public:
 
 public:
 
-	/**
-	 * Default constructor (no initialization).
-	 */
-	FSphere( ) { }
+	/** Default constructor (no initialization). */
+	FSphere() { }
 
 	/**
 	 * Creates and initializes a new sphere.
diff --git a/Engine/Source/Runtime/Core/Public/Math/TwoVectors.h b/Engine/Source/Runtime/Core/Public/Math/TwoVectors.h
--- a/Engine/Source/Runtime/Core/Public/Math/TwoVectors.h
+++ b/Engine/Source/Runtime/Core/Public/Math/TwoVectors.h
@@ -1,9 +1,5 @@
 // Copyright 1998-2015 Epic Games, Inc. All Rights Reserved.
 
-/*=============================================================================
-	TwoVectors.h: Declares the FTwoVectors class.
-=============================================================================*/
-
 #pragma once
 
 
@@ -14,18 +10,16 @@ class FTwoVectors
 {
 public:
 
-	/**  Holds the first vector. */
+	/** Holds the first vector. */
 	FVector	v1;
 
-	/**  Holds the second vector. */
+	/** Holds the second vector. */
 	FVector	v2;
 
 public:
 
-	/**
-	 * Default constructor.
-	 */
-	FORCEINLINE	FTwoVectors( );
+	/** Default constructor. */
+	FORCEINLINE	FTwoVectors();
 
 	/**
 	 * Creates and initializes a new instance with the specified vectors.
@@ -126,7 +120,7 @@ public:
 	 *
 	 * @return A negated copy of the pair.
 	 */
-	FORCEINLINE FTwoVectors operator-( ) const;
+	FORCEINLINE FTwoVectors operator-() const;
 
 	// Assignment operators.
 
@@ -182,8 +176,7 @@ public:
 	 * Get a specific component from the pair.
 	 *
 	 * @param i The index of the component, even indices are for the first vector,
-	 * odd ones are for the second. Returns index 5 if out of range.
-	 *
+	 *			odd ones are for the second. Returns index 5 if out of range.
 	 * @return Reference to the specified component.
 	 */
     float& operator[]( int32 i );
@@ -195,30 +188,29 @@ public:
 	 *
 	 * @return The maximum value of all the vector coordinates.
 	 */
-	float GetMax( ) const;
+	float GetMax() const;
 
 	/**
 	 * Get the minimum value of all the vector coordinates.
 	 *
 	 * @return The minimum value of all the vector coordinates.
 	 */
-	float GetMin( ) const;
+	float GetMin() const;
 
 	/**
 	 * Get a textual representation of this two-vector.
 	 *
 	 * @return A string describing the two-vector.
 	 */
-	FString ToString( ) const;
+	FString ToString() const;
 
 public:
 
 	/**
 	 * Serializes the two-vector.
 	 *
-	 * @param Ar - The archive to serialize into.
-	 * @param TwoVectors - The two-vector to serialize.
-	 *
+	 * @param Ar The archive to serialize into.
+	 * @param TwoVectors The two-vector to serialize.
 	 * @return Reference to the Archive after serialization.
 	 */
 	friend FArchive& operator<<( FArchive& Ar, FTwoVectors& TwoVectors )
@@ -237,7 +229,7 @@ FORCEINLINE FTwoVectors operator*( float Scale, const FTwoVectors& V )
 }
 
 
-FORCEINLINE	FTwoVectors::FTwoVectors( ) :
+FORCEINLINE	FTwoVectors::FTwoVectors() :
 	v1(0.0f),
 	v2(0.0f)
 { }
@@ -329,7 +321,7 @@ FORCEINLINE bool FTwoVectors::Equals( const FTwoVectors& V, float Tolerance ) co
 }
 
 
-FORCEINLINE FTwoVectors FTwoVectors::operator-( ) const
+FORCEINLINE FTwoVectors FTwoVectors::operator-() const
 {
 	return FTwoVectors(
 		FVector(-v1),
@@ -393,7 +385,7 @@ FORCEINLINE FTwoVectors FTwoVectors::operator/=( const FTwoVectors& V )
 }
 
 
-FORCEINLINE float FTwoVectors::GetMax( ) const
+FORCEINLINE float FTwoVectors::GetMax() const
 {
 	const float MaxMax = FMath::Max(FMath::Max(v1.X, v1.Y), v1.Z);
 	const float MaxMin = FMath::Max(FMath::Max(v2.X, v2.Y), v2.Z);
@@ -402,7 +394,7 @@ FORCEINLINE float FTwoVectors::GetMax( ) const
 }
 
 
-FORCEINLINE float FTwoVectors::GetMin( ) const
+FORCEINLINE float FTwoVectors::GetMin() const
 {
 	const float MinMax = FMath::Min(FMath::Min(v1.X, v1.Y), v1.Z);
 	const float MinMin = FMath::Min(FMath::Min(v2.X, v2.Y), v2.Z);
@@ -428,7 +420,7 @@ FORCEINLINE float& FTwoVectors::operator[]( int32 i )
 }
 
 
-FORCEINLINE FString FTwoVectors::ToString( ) const
+FORCEINLINE FString FTwoVectors::ToString() const
 {
 	return FString::Printf(TEXT("V1=(%s) V2=(%s)"), *v1.ToString(), *v2.ToString());
-}
\ No newline at end of file
+}
diff --git a/Engine/Source/Runtime/Core/Public/Math/UnrealMath.h b/Engine/Source/Runtime/Core/Public/Math/UnrealMath.h
--- a/Engine/Source/Runtime/Core/Public/Math/UnrealMath.h
+++ b/Engine/Source/Runtime/Core/Public/Math/UnrealMath.h
@@ -1,10 +1,7 @@
 // Copyright 1998-2015 Epic Games, Inc. All Rights Reserved.
 
-/*=============================================================================
-	UnMath.h: Unreal math Types and inline implementation 
-=============================================================================*/
-
 #pragma once
+
 #include "Color.h"
 #include "ColorList.h"
 #include "IntPoint.h"
@@ -52,11 +49,11 @@
 #include "ScalarRegister.h"
 #include "ConvexHull2d.h"
 
+
 // FVector2D Implementation
 FORCEINLINE FVector2D::FVector2D( const FVector& V )
-: X(V.X), Y(V.Y)
-{
-}
+	: X(V.X), Y(V.Y)
+{ }
 
 inline FVector FVector2D::SphericalToUnitCartesian() const
 {
diff --git a/Engine/Source/Runtime/Core/Public/Math/UnrealMathDirectX.h b/Engine/Source/Runtime/Core/Public/Math/UnrealMathDirectX.h
--- a/Engine/Source/Runtime/Core/Public/Math/UnrealMathDirectX.h
+++ b/Engine/Source/Runtime/Core/Public/Math/UnrealMathDirectX.h
@@ -1,10 +1,5 @@
 // Copyright 1998-2015 Epic Games, Inc. All Rights Reserved.
 
-/*=============================================================================
-	UnrealMathDirectX.h: DirectX specific vector intrinsics
-
-=============================================================================*/
-
 #pragma once
 
 #if __cplusplus_cli
diff --git a/Engine/Source/Runtime/Core/Public/Math/UnrealMathFPU.h b/Engine/Source/Runtime/Core/Public/Math/UnrealMathFPU.h
--- a/Engine/Source/Runtime/Core/Public/Math/UnrealMathFPU.h
+++ b/Engine/Source/Runtime/Core/Public/Math/UnrealMathFPU.h
@@ -1,10 +1,5 @@
 // Copyright 1998-2015 Epic Games, Inc. All Rights Reserved.
 
-/*=============================================================================
-	UnMathFpu.h: Fallback FPU vector intrinsics
-
-=============================================================================*/
-
 #pragma once
 
 
diff --git a/Engine/Source/Runtime/Core/Public/Math/UnrealMathNeon.h b/Engine/Source/Runtime/Core/Public/Math/UnrealMathNeon.h
--- a/Engine/Source/Runtime/Core/Public/Math/UnrealMathNeon.h
+++ b/Engine/Source/Runtime/Core/Public/Math/UnrealMathNeon.h
@@ -1,10 +1,5 @@
 // Copyright 1998-2015 Epic Games, Inc. All Rights Reserved.
 
-/*=============================================================================
-	UnMathNeon.h: ARM-NEON-specific vector intrinsics
-
-=============================================================================*/
-
 #pragma once
 
 // Include the intrinsic functions header
diff --git a/Engine/Source/Runtime/Core/Public/Math/UnrealMathSSE.h b/Engine/Source/Runtime/Core/Public/Math/UnrealMathSSE.h
--- a/Engine/Source/Runtime/Core/Public/Math/UnrealMathSSE.h
+++ b/Engine/Source/Runtime/Core/Public/Math/UnrealMathSSE.h
@@ -1,10 +1,5 @@
 // Copyright 1998-2015 Epic Games, Inc. All Rights Reserved.
 
-/*=============================================================================
-	UnMathSSE.h: Intel SSE specific vector intrinsics
-
-=============================================================================*/
-
 #pragma once
 
 #if __cplusplus_cli
@@ -15,6 +10,7 @@
 // We require SSE2
 #include <emmintrin.h>
 
+
 /*=============================================================================
  *	Helpers:
  *============================================================================*/
diff --git a/Engine/Source/Runtime/Core/Public/Math/UnrealMathUtility.h b/Engine/Source/Runtime/Core/Public/Math/UnrealMathUtility.h
--- a/Engine/Source/Runtime/Core/Public/Math/UnrealMathUtility.h
+++ b/Engine/Source/Runtime/Core/Public/Math/UnrealMathUtility.h
@@ -1,13 +1,11 @@
 // Copyright 1998-2015 Epic Games, Inc. All Rights Reserved.
 
-/*=============================================================================
-	UnrealMathUtility.h: Unreal math utility routines
-=============================================================================*/
-
 #pragma once
+
 #include "HAL/Platform.h"
 #include "HAL/PlatformMath.h"
 
+
 //#define IMPLEMENT_ASSIGNMENT_OPERATOR_MANUALLY
 
 // Assert on non finite numbers. Used to track NaNs.
diff --git a/Engine/Source/Runtime/Core/Public/Math/UnrealMatrix.h b/Engine/Source/Runtime/Core/Public/Math/UnrealMatrix.h
--- a/Engine/Source/Runtime/Core/Public/Math/UnrealMatrix.h
+++ b/Engine/Source/Runtime/Core/Public/Math/UnrealMatrix.h
@@ -1,8 +1,6 @@
 // Copyright 1998-2015 Epic Games, Inc. All Rights Reserved.
 
 /*=============================================================================
-	UnMath.inl: Unreal inlined math functions
-
 	NOTE: This file should ONLY be included by UnrealMath.h!
 =============================================================================*/
 
@@ -11,11 +9,11 @@
 #include "OutputDevices.h"
 #include "AssertionMacros.h"
 
+
 /**
  * FMatrix inline functions.
  */
 
-
 // Constructors.
 
 FORCEINLINE FMatrix::FMatrix()
diff --git a/Engine/Source/Runtime/Core/Public/Math/Vector.h b/Engine/Source/Runtime/Core/Public/Math/Vector.h
--- a/Engine/Source/Runtime/Core/Public/Math/Vector.h
+++ b/Engine/Source/Runtime/Core/Public/Math/Vector.h
@@ -1,9 +1,5 @@
 // Copyright 1998-2015 Epic Games, Inc. All Rights Reserved.
 
-/*=============================================================================
-	Vector.h: Declares the FVector class.
-=============================================================================*/
-
 #pragma once
 
 
@@ -45,9 +41,7 @@ public:
 	FORCEINLINE void DiagnosticCheckNaN() const {}
 #endif
 
-	/**
-	 * Default constructor (no initialization).
-	 */
+	/** Default constructor (no initialization). */
 	FORCEINLINE FVector();
 
 	/**
@@ -114,7 +108,6 @@ public:
 	* Copy another FVector into this one
 	*
 	* @param Other The other vector.
-	*
 	* @return Reference to vector after copy.
 	*/
 	FORCEINLINE FVector& operator=(const FVector& Other);
@@ -709,9 +702,9 @@ public:
 	 * but it won't change the direction of the Z axis.
 	 * All axes will be normalized.
 	 *
-	 * @param XAxis - The input basis' XAxis, and upon return the orthonormal basis' XAxis.
-	 * @param YAxis - The input basis' YAxis, and upon return the orthonormal basis' YAxis.
-	 * @param ZAxis - The input basis' ZAxis, and upon return the orthonormal basis' ZAxis.
+	 * @param XAxis The input basis' XAxis, and upon return the orthonormal basis' XAxis.
+	 * @param YAxis The input basis' YAxis, and upon return the orthonormal basis' YAxis.
+	 * @param ZAxis The input basis' ZAxis, and upon return the orthonormal basis' ZAxis.
 	 */
 	static CORE_API void CreateOrthonormalBasis(FVector& XAxis,FVector& YAxis,FVector& ZAxis);
 
@@ -931,7 +924,6 @@ public:
  *
  * @param Scale Scaling factor.
  * @param V Vector to scale.
- *
  * @return Result of multiplication.
  */
 FORCEINLINE FVector operator*( float Scale, const FVector& V )
@@ -943,7 +935,6 @@ FORCEINLINE FVector operator*( float Scale, const FVector& V )
  * Creates a hash value from a FVector. 
  *
  * @param Vector the vector to create a hash value for
- *
  * @return The hash value from the components
  */
 FORCEINLINE uint32 GetTypeHash(const FVector& Vector)
@@ -956,7 +947,6 @@ FORCEINLINE uint32 GetTypeHash(const FVector& Vector)
  * Creates a hash value from a FVector2D. 
  *
  * @param Vector the vector to create a hash value for
- *
  * @return The hash value from the components
  */
 FORCEINLINE uint32 GetTypeHash(const FVector2D& Vector)
@@ -981,7 +971,6 @@ FORCEINLINE uint32 GetTypeHash(const FVector2D& Vector)
  * @param Mins 3D Point defining the lower values of the axis of the bound box
  * @param Max 3D Point defining the lower values of the axis of the bound box
  * @param Point 3D position of interest
- * 
  * @return the distance from the Point to the bounding box.
  */
 FORCEINLINE float ComputeSquaredDistanceFromBoxToPoint( const FVector& Mins, const FVector& Maxs, const FVector& Point )
@@ -1021,6 +1010,7 @@ FORCEINLINE float ComputeSquaredDistanceFromBoxToPoint( const FVector& Mins, con
 	return DistSquared;
 }
 
+
 FORCEINLINE FVector::FVector( const FVector2D V, float InZ )
 	: X(V.X), Y(V.Y), Z(InZ)
 {
@@ -1028,7 +1018,6 @@ FORCEINLINE FVector::FVector( const FVector2D V, float InZ )
 }
 
 
-
 inline FVector FVector::RotateAngleAxis( const float AngleDeg, const FVector& Axis ) const
 {
 	const float S	= FMath::Sin(AngleDeg * PI / 180.f);
diff --git a/Engine/Source/Runtime/Core/Public/Math/Vector2D.h b/Engine/Source/Runtime/Core/Public/Math/Vector2D.h
--- a/Engine/Source/Runtime/Core/Public/Math/Vector2D.h
+++ b/Engine/Source/Runtime/Core/Public/Math/Vector2D.h
@@ -1,9 +1,5 @@
 // Copyright 1998-2015 Epic Games, Inc. All Rights Reserved.
 
-/*=============================================================================
-	Vector2D.h: Declares the FVector2D class.
-=============================================================================*/
-
 #pragma once
 
 
@@ -31,7 +27,7 @@ public:
 	/**
 	 * Default constructor (no initialization).
 	 */
-	FORCEINLINE FVector2D( ) { }
+	FORCEINLINE FVector2D() { }
 
 	/**
 	 * Constructor using initial values for each component.
@@ -342,35 +338,35 @@ public:
 	 *
 	 * @return The maximum value of the vector's components.
 	 */
-	float GetMax( ) const;
+	float GetMax() const;
 
 	/**
 	 * Get the maximum absolute value of the vector's components.
 	 *
 	 * @return The maximum absolute value of the vector's components.
 	 */
-	float GetAbsMax( ) const;
+	float GetAbsMax() const;
 
 	/**
 	 * Get the minimum value of the vector's components.
 	 *
 	 * @return The minimum value of the vector's components.
 	 */
-	float GetMin( ) const;
+	float GetMin() const;
 
 	/**
 	 * Get the length (magnitude) of this vector.
 	 *
 	 * @return The length of this vector.
 	 */
-	float Size( ) const;
+	float Size() const;
 
 	/**
 	 * Get the squared length of this vector.
 	 *
 	 * @return The squared length of this vector.
 	 */
-	float SizeSquared( ) const;
+	float SizeSquared() const;
 
 	/**
 	 * Gets a normalized copy of the vector, checking it is safe to do so based on the length.
@@ -405,14 +401,14 @@ public:
 	 *
 	 * @return true if vector is exactly zero, otherwise false.
 	 */
-	bool IsZero( ) const;
+	bool IsZero() const;
 
 	/**
 	 * Get this vector as an Int Point.
 	 *
 	 * @return New Int Point from this vector.
 	 */
-	FIntPoint IntPoint( ) const;
+	FIntPoint IntPoint() const;
 
 	/**
 	 * Creates a copy of this vector with both axes clamped to the given range.
@@ -433,7 +429,7 @@ public:
 	 *
 	 * @return Text describing the vector.
 	 */
-	FString ToString( ) const;
+	FString ToString() const;
 
 	/**
 	 * Initialize this Vector based on an FString. The String is expected to contain X=, Y=.
@@ -459,12 +455,12 @@ public:
 	}
 
 #if ENABLE_NAN_DIAGNOSTIC
-	FORCEINLINE void DiagnosticCheckNaN( ) const
+	FORCEINLINE void DiagnosticCheckNaN() const
 	{
 		checkf(!ContainsNaN(), TEXT("FVector contains NaN: %s"), *ToString());
 	}
 #else
-	FORCEINLINE void DiagnosticCheckNaN( ) const {}
+	FORCEINLINE void DiagnosticCheckNaN() const {}
 #endif
 
 	/**
@@ -485,7 +481,7 @@ public:
 	CORE_API bool NetSerialize( FArchive& Ar, class UPackageMap* Map, bool& bOutSuccess );
 
 	/** Converts spherical coordinates on the unit sphere into a Cartesian unit length vector. */
-	inline FVector SphericalToUnitCartesian( ) const;
+	inline FVector SphericalToUnitCartesian() const;
 };
 
 
diff --git a/Engine/Source/Runtime/Core/Public/Math/Vector2DHalf.h b/Engine/Source/Runtime/Core/Public/Math/Vector2DHalf.h
--- a/Engine/Source/Runtime/Core/Public/Math/Vector2DHalf.h
+++ b/Engine/Source/Runtime/Core/Public/Math/Vector2DHalf.h
@@ -1,9 +1,5 @@
 // Copyright 1998-2015 Epic Games, Inc. All Rights Reserved.
 
-/*=============================================================================
-	Vector2DHalf.h: Declares the FVector2DHalf class.
-=============================================================================*/
-
 #pragma once
 
 
@@ -20,12 +16,11 @@ struct FVector2DHalf
 
 public:
 
-	/**
-	 * Default Constructor (no initialization).
-	 */
-	FORCEINLINE FVector2DHalf( ) { }
+	/** Default Constructor (no initialization). */
+	FORCEINLINE FVector2DHalf() { }
 
-	/** Constructor 
+	/**
+	 * Constructor.
 	 *
 	 * InX half float X value
 	 * Iny half float Y value
@@ -54,10 +49,8 @@ public:
 	 */
  	FVector2DHalf& operator=( const FVector2D& Vector2D );
 
-	/**
-	 * Implicit conversion operator for conversion to FVector2D.
-	 */
-	operator FVector2D( ) const;
+	/** Implicit conversion operator for conversion to FVector2D. */
+	operator FVector2D() const;
 
 public:
 
@@ -66,7 +59,7 @@ public:
 	 *
 	 * @return Text describing the vector.
 	 */
-	FString ToString( ) const;
+	FString ToString() const;
 
 public:
 
@@ -91,6 +84,7 @@ FORCEINLINE FVector2DHalf::FVector2DHalf( const FFloat16& InX, const FFloat16& I
  	:	X(InX), Y(InY)
 { }
 
+
 FORCEINLINE FVector2DHalf::FVector2DHalf( float InX, float InY )
 	:	X(InX), Y(InY)
 { }
@@ -110,13 +104,13 @@ FORCEINLINE FVector2DHalf& FVector2DHalf::operator=( const FVector2D& Vector2D )
 }
 
 
-FORCEINLINE FString FVector2DHalf::ToString( ) const
+FORCEINLINE FString FVector2DHalf::ToString() const
 {
 	return FString::Printf(TEXT("X=%3.3f Y=%3.3f"), (float)X, (float)Y );
 }
 
 
-FORCEINLINE FVector2DHalf::operator FVector2D( ) const
+FORCEINLINE FVector2DHalf::operator FVector2D() const
 {
 	return FVector2D((float)X,(float)Y);
 }
diff --git a/Engine/Source/Runtime/Core/Public/Math/Vector4.h b/Engine/Source/Runtime/Core/Public/Math/Vector4.h
--- a/Engine/Source/Runtime/Core/Public/Math/Vector4.h
+++ b/Engine/Source/Runtime/Core/Public/Math/Vector4.h
@@ -1,9 +1,5 @@
 // Copyright 1998-2015 Epic Games, Inc. All Rights Reserved.
 
-/*=============================================================================
-	Vector.h: Declares the FVector class.
-=============================================================================*/
-
 #pragma once
 
 
@@ -74,7 +70,6 @@ public:
 	 * Access a specific component of the vector.
 	 *
 	 * @param ComponentIndex The index of the component.
-	 *
 	 * @return Reference to the desired component.
 	 */
 	FORCEINLINE float& operator[]( int32 ComponentIndex );
@@ -83,7 +78,6 @@ public:
 	 * Access a specific component of the vector.
 	 *
 	 * @param ComponentIndex The index of the component.
-	 *
 	 * @return Copy of the desired component.
 	 */
 	FORCEINLINE float operator[]( int32 ComponentIndex ) const;
@@ -95,13 +89,12 @@ public:
 	 *
 	 * @return A negated copy of the vector.
 	 */
-	FORCEINLINE FVector4 operator-( ) const;
+	FORCEINLINE FVector4 operator-() const;
 
 	/**
 	 * Gets the result of adding a vector to this.
 	 *
 	 * @param V The vector to add.
-	 *
 	 * @return The result of vector addition.
 	 */
 	FORCEINLINE FVector4 operator+( const FVector4& V ) const;
@@ -272,7 +265,7 @@ public:
 	 *
 	 * @return Text describing the vector.
 	 */
-	FString ToString( ) const;
+	FString ToString() const;
 
 	/**
 	 * Initialize this Vector based on an FString. The String is expected to contain X=, Y=, Z=, W=.
@@ -302,7 +295,7 @@ public:
 	FORCEINLINE FVector4 GetUnsafeNormal3() const;
 
 	DEPRECATED(4.7, "Deprecated due to unclear name, use GetUnsafeNormal3 instead.")
-	FORCEINLINE FVector4 UnsafeNormal3( ) const;
+	FORCEINLINE FVector4 UnsafeNormal3() const;
 
 	/**
 	 * Return the FRotator corresponding to the direction that the vector
@@ -311,7 +304,7 @@ public:
 	 *
 	 * @return The FRotator of the vector's direction.
 	 */
-	CORE_API FRotator Rotation( ) const;
+	CORE_API FRotator Rotation() const;
 
 	/**
 	 * Set all of the vectors coordinates.
@@ -328,7 +321,7 @@ public:
 	 *
 	 * @return The length of this vector.
 	 */
-	float Size3( ) const;
+	float Size3() const;
 
 	/**
 	 * Get the squared length of this vector not taking W component into account.
@@ -353,12 +346,12 @@ public:
 	void FindBestAxisVectors3( FVector4& Axis1, FVector4& Axis2 ) const;
 
 #if ENABLE_NAN_DIAGNOSTIC
-	FORCEINLINE void DiagnosticCheckNaN( ) const
+	FORCEINLINE void DiagnosticCheckNaN() const
 	{
 		checkf(!ContainsNaN(), TEXT("FVector contains NaN: %s"), *ToString());
 	}
 #else
-	FORCEINLINE void DiagnosticCheckNaN( ) const { }
+	FORCEINLINE void DiagnosticCheckNaN() const { }
 #endif
 
 public:
@@ -377,6 +370,7 @@ public:
 
 } GCC_ALIGN(16);
 
+
 /**
  * Creates a hash value from a FVector4.
  *
@@ -467,7 +461,7 @@ FORCEINLINE void FVector4::Set( float InX, float InY, float InZ, float InW )
 }
 
 
-FORCEINLINE FVector4 FVector4::operator-( ) const
+FORCEINLINE FVector4 FVector4::operator-() const
 {
 	return FVector4( -X, -Y, -Z, -W );
 }
@@ -547,7 +541,7 @@ FORCEINLINE bool FVector4::Equals( const FVector4& V, float Tolerance ) const
 }
 
 
-FORCEINLINE FString FVector4::ToString( ) const
+FORCEINLINE FString FVector4::ToString() const
 {
 	return FString::Printf(TEXT("X=%3.3f Y=%3.3f Z=%3.3f W=%3.3f"), X, Y, Z, W);
 }
@@ -586,7 +580,7 @@ FORCEINLINE FVector4 FVector4::SafeNormal(float Tolerance) const
 }
 
 
-FORCEINLINE FVector4 FVector4::GetUnsafeNormal3( ) const
+FORCEINLINE FVector4 FVector4::GetUnsafeNormal3() const
 {
 	const float Scale = FMath::InvSqrt(X*X+Y*Y+Z*Z);
 	return FVector4( X*Scale, Y*Scale, Z*Scale, 0.0f );
@@ -599,13 +593,13 @@ FORCEINLINE FVector4 FVector4::UnsafeNormal3() const
 }
 
 
-FORCEINLINE float FVector4::Size3( ) const
+FORCEINLINE float FVector4::Size3() const
 {
 	return FMath::Sqrt( X*X + Y*Y + Z*Z );
 }
 
 
-FORCEINLINE float FVector4::SizeSquared3( ) const
+FORCEINLINE float FVector4::SizeSquared3() const
 {
 	return X*X + Y*Y + Z*Z;
 }
@@ -617,7 +611,7 @@ FORCEINLINE bool FVector4::IsUnit3( float LengthSquaredTolerance ) const
 }
 
 
-FORCEINLINE bool FVector4::ContainsNaN( ) const
+FORCEINLINE bool FVector4::ContainsNaN() const
 {
 	return (FMath::IsNaN(X) || !FMath::IsFinite(X) || 
 			FMath::IsNaN(Y) || !FMath::IsFinite(Y) ||
