FText key stablization

FText DisplayString shared ptrs are persisted across property edits. This preserves their entry in the live table and also updates the string that the entry points to. We call GetString() to update the hashed value as persisting the DisplayString means that its lcoation in the table is preserved.
On serialization, a valid key is still found in the live table for the display string, so the key is preserved.

UE-5350 - As a translator, I would like the editor tutorial text to preserve it's history between changes in order to greatly improve the translation workflow.

diff --git a/Engine/Source/Runtime/Core/Private/Internationalization/Text.cpp b/Engine/Source/Runtime/Core/Private/Internationalization/Text.cpp
--- a/Engine/Source/Runtime/Core/Private/Internationalization/Text.cpp
+++ b/Engine/Source/Runtime/Core/Private/Internationalization/Text.cpp
@@ -843,4 +843,37 @@ bool FTextSnapshot::IsDisplayStringEqualTo(const FText& InText) const
 	return HistoryRevision == InHistoryRevision && DisplayStringPtr.IsValid() && DisplayStringPtr->Equals(InText.ToString(), ESearchCase::CaseSensitive);
 }
 
+FScopedTextIdentityPreserver::FScopedTextIdentityPreserver(FText& InTextToPersist)
+	: TextToPersist(InTextToPersist)
+{
+	// Empty display strings can't have a namespace or key.
+	if (!InTextToPersist.DisplayString->IsEmpty())
+	{ 
+		// Save off namespace and key to be restored later.
+		FTextLocalizationManager::Get().FindKeyNamespaceFromDisplayString(InTextToPersist.DisplayString, Namespace, Key);
+	}
+}
+
+FScopedTextIdentityPreserver::~FScopedTextIdentityPreserver()
+{
+	// If we don't have a key, then the old identity wasn't valid and shouldn't be preserved.
+	if (Key.IsValid())
+	{
+		// Get the text's new source string.
+		const FString* SourceString = FTextInspector::GetSourceString(TextToPersist);
+
+		// Without a source string, we can't possibly preserve the identity. If the text we're preserving identity for can't possibly have an identity anymore, this class shouldn't be used on this text.
+		check(SourceString);
+
+		// Make the history so that it can attempt find and serialize the identity for this source string when the history itself is serialized.
+		TextToPersist.History = MakeShareable(new FTextHistory_Base(*SourceString));
+
+		// Create/update the display string instance for this identity in the text localization manager...
+		const TSharedRef<FString, ESPMode::ThreadSafe> DisplayString = FTextLocalizationManager::Get().GetString(Namespace.IsValid() ? *Namespace : TEXT(""), *Key, SourceString);
+		
+		// ... and set it so that the text's history's serialization will properly map this display string instance back to the identity we intended to preserve and restore.
+		TextToPersist.DisplayString = DisplayString;
+	}
+}
+
 #undef LOCTEXT_NAMESPACE
diff --git a/Engine/Source/Runtime/Core/Public/Internationalization/Text.h b/Engine/Source/Runtime/Core/Public/Internationalization/Text.h
--- a/Engine/Source/Runtime/Core/Public/Internationalization/Text.h
+++ b/Engine/Source/Runtime/Core/Public/Internationalization/Text.h
@@ -381,6 +381,7 @@ public:
 	friend class FTextHistory_NamedFormat;
 	friend class FTextHistory_ArgumentDataFormat;
 	friend class FTextHistory_OrderedFormat;
+	friend class FScopedTextIdentityPreserver;
 
 #if !UE_ENABLE_ICU
 	friend class FLegacyTextHelper;
@@ -563,4 +564,16 @@ private:
 	int32 IndentCount;
 };
 
+class CORE_API FScopedTextIdentityPreserver
+{
+public:
+	FScopedTextIdentityPreserver(FText& InTextToPersist);
+	~FScopedTextIdentityPreserver();
+
+private:
+	TSharedPtr< FString, ESPMode::ThreadSafe > Namespace;
+	TSharedPtr< FString, ESPMode::ThreadSafe > Key;
+	FText& TextToPersist;
+};
+
 Expose_TNameOf(FText)
diff --git a/Engine/Source/Runtime/CoreUObject/Private/UObject/UTextProperty.cpp b/Engine/Source/Runtime/CoreUObject/Private/UObject/UTextProperty.cpp
--- a/Engine/Source/Runtime/CoreUObject/Private/UObject/UTextProperty.cpp
+++ b/Engine/Source/Runtime/CoreUObject/Private/UObject/UTextProperty.cpp
@@ -106,6 +106,7 @@ const TCHAR* UTextProperty::ImportText_Internal( const TCHAR* Buffer, void* Data
 		BufferAsString = FString(Temp);
 	}
 
+	FScopedTextIdentityPreserver Preserver(*Text);
 	*Text = FText::FromString( FString(BufferAsString) );
 
 	return Buffer;
