removed unused feature (flag to see if cvar changed) - there are better methods to do the same - will mention in documentation

diff --git a/Engine/Source/Runtime/Core/Private/HAL/ConsoleManager.cpp b/Engine/Source/Runtime/Core/Private/HAL/ConsoleManager.cpp
--- a/Engine/Source/Runtime/Core/Private/HAL/ConsoleManager.cpp
+++ b/Engine/Source/Runtime/Core/Private/HAL/ConsoleManager.cpp
@@ -142,8 +142,6 @@ public:
 
 		// only change on main thread
 
-		Flags = (EConsoleVariableFlags)((uint32)Flags | ECVF_Changed);
-
 		Flags = (EConsoleVariableFlags)(((uint32)Flags & ~ECVF_SetByMask) | SetBy);
 
 		OnChangedCallback.ExecuteIfBound(this);
@@ -1362,13 +1360,6 @@ void FConsoleManager::Test()
 		check(VarE->GetFloat() == RefE);
 		check(VarE->GetString() == FString(TEXT("2.1")));
 
-		// test changed
-
-		check(!VarA->TestFlags(ECVF_Changed));
-		check(!VarB->TestFlags(ECVF_Changed));
-		check(!VarD->TestFlags(ECVF_Changed));
-		check(!VarE->TestFlags(ECVF_Changed));
-
 		// call Set(string)
 
 		VarA->Set(TEXT("3.1"), ECVF_SetByConsoleVariablesIni);
@@ -1378,13 +1369,6 @@ void FConsoleManager::Test()
 
 		check(GConsoleVariableCallbackTestCounter == 1);
 
-		// test changed
-
-		check(VarA->TestFlags(ECVF_Changed));
-		check(VarB->TestFlags(ECVF_Changed));
-		check(VarD->TestFlags(ECVF_Changed));
-		check(VarE->TestFlags(ECVF_Changed));
-
 		// verify Set()
 
 		check(VarA->GetString() == FString(TEXT("3")));
@@ -1427,9 +1411,7 @@ void FConsoleManager::Test()
 			// note: exact comparison fails in Win32 release
 			check(FMath::IsNearlyEqual(VarC->GetFloat(), 1.23f, KINDA_SMALL_NUMBER));
 			check(VarC->GetString() == FString(TEXT("1.23")));
-			check(!VarC->TestFlags(ECVF_Changed));
 			VarC->Set(TEXT("3.1"), ECVF_SetByConsole);
-			check(VarC->TestFlags(ECVF_Changed));
 			check(VarC->GetString() == FString(TEXT("3.1")));
 			UnregisterConsoleObject(TEXT("TestNameC"), false);
 			check(!IConsoleManager::Get().FindConsoleVariable(TEXT("TestNameC")));
diff --git a/Engine/Source/Runtime/Core/Public/HAL/IConsoleManager.h b/Engine/Source/Runtime/Core/Public/HAL/IConsoleManager.h
--- a/Engine/Source/Runtime/Core/Public/HAL/IConsoleManager.h
+++ b/Engine/Source/Runtime/Core/Public/HAL/IConsoleManager.h
@@ -18,13 +18,11 @@
  *   <COMMAND>	 				print the current state of the console variable
  *   <COMMAND> x 				set and print the new state of the console variable
  *
- * The ECVF_Changed flag can be used to detect if the state was changes through the Set() method (float& and int& can change outside of this).
- * 
  * All variables support auto completion. The single line help that can show up there is currently not connected to the help as the help text
  * is expected to be multi line.
  * The former Exec() system can be used to access the console variables.
  * Use console variables only in main thread.
- * The state of console variables is not network synchronized or serialized (load/save). The plan is to allow to set the state in external files (game/platform/engige/local).
+ * The state of console variables is not network synchronized or serialized (load/save). The plan is to allow to set the state in external files (game/platform/engine/local).
  */
 
 /**
@@ -41,18 +39,6 @@ enum EConsoleVariableFlags
 	 * Then they are are hidden in the console and cannot be changed by the user.
 	 */
 	ECVF_Cheat = 0x1,
-	/**
-	 * By default, after creation this flag is not set.
-	 * Whenever the value was changes through the Set functions the flag is set.
-	 * To detect changes:
-	 * if(CVar->TestFlags(ECVF_Changed))
-	 * {
-	 *     CVar->ClearFlags(ECVF_Changed);
-	 *     ... 
-	 * }
-	 * The old value is not stored as this can be easily done by hand and only when needed.
-	 */
-	ECVF_Changed = 0x2,
 	/**
 	 * Console variables cannot be changed by the user (from console).
 	 * Changing from C++ or ini is still possible.
