Support for XBoxOne and many variations of XBox360 controllers on Mac
#717
#721

diff --git a/Engine/Source/Runtime/Core/Private/Mac/HIDInputInterface.cpp b/Engine/Source/Runtime/Core/Private/Mac/HIDInputInterface.cpp
--- a/Engine/Source/Runtime/Core/Private/Mac/HIDInputInterface.cpp
+++ b/Engine/Source/Runtime/Core/Private/Mac/HIDInputInterface.cpp
@@ -175,10 +175,9 @@ void HIDInputInterface::FHIDDeviceInfo::SetupMappings()
 		RightAnalogYMapping = kHIDUsage_GD_Rz;
 		RightTriggerAnalogMapping = kHIDUsage_GD_Ry;
 	}
-	else if ((VendorID == 0x45e && (ProductID == 0x28e || ProductID == 0x719)) // Original Microsoft controller
-			 || (VendorID == 0xe6f && ProductID == 0x401)) // GameStop version
+	else if (VendorID == 0x45e && ProductID == 0x2d1)
 	{
-		// Xbox 360 Controller
+		// Xbox One Controller
 		ButtonsMapping[1]	= 0;	// A
 		ButtonsMapping[2]	= 1;	// B
 		ButtonsMapping[3]	= 2;	// X
@@ -187,8 +186,8 @@ void HIDInputInterface::FHIDDeviceInfo::SetupMappings()
 		ButtonsMapping[6]	= 5;	// Right Shoulder
 		ButtonsMapping[7]	= 8;	// Left Thumbstick
 		ButtonsMapping[8]	= 9;	// Right Thumbstick
-		ButtonsMapping[9]	= 6;	// Start
-		ButtonsMapping[10]	= 7;	// Back
+		ButtonsMapping[9]	= 7;	// Menu
+		ButtonsMapping[10]	= 6;	// View
 		ButtonsMapping[12]	= 12;	// DPad Up
 		ButtonsMapping[13]	= 13;	// DPad Down
 		ButtonsMapping[14]	= 14;	// DPad Left
@@ -196,33 +195,150 @@ void HIDInputInterface::FHIDDeviceInfo::SetupMappings()
 
 		LeftAnalogXMapping = kHIDUsage_GD_X;
 		LeftAnalogYMapping = kHIDUsage_GD_Y;
-		LeftTriggerAnalogMapping = kHIDUsage_GD_Z;
-		RightAnalogXMapping = kHIDUsage_GD_Rx;
-		RightAnalogYMapping = kHIDUsage_GD_Ry;
-		RightTriggerAnalogMapping = kHIDUsage_GD_Rz;
+		LeftTriggerAnalogMapping = kHIDUsage_GD_Ry;
+		RightAnalogXMapping = kHIDUsage_GD_Z;
+		RightAnalogYMapping = kHIDUsage_GD_Rx;
+		RightTriggerAnalogMapping = kHIDUsage_GD_Rz;	
 	}
 	else
 	{
-		// Generic (based on Logitech RumblePad 2)
-		ButtonsMapping[1]	= 2;	// X
-		ButtonsMapping[2]	= 0;	// A
-		ButtonsMapping[3]	= 1;	// B
-		ButtonsMapping[4]	= 3;	// Y
-		ButtonsMapping[5]	= 4;	// Left Shoulder
-		ButtonsMapping[6]	= 5;	// Right Shoulder
-		ButtonsMapping[7]	= 10;	// Left Trigger
-		ButtonsMapping[8]	= 11;	// Right Trigger
-		ButtonsMapping[9]	= 7;	// Back
-		ButtonsMapping[10]	= 6;	// Start
-		ButtonsMapping[11]	= 8;	// Left Thumbstick
-		ButtonsMapping[12]	= 9;	// Right Thumbstick
+		// Array generated from [REDACTED_PROFILE_URL]/360Controller/blob/master/360Controller/Info.plist (with additions)
+		static const FXBox360ControllerID XBox360Controllers[] =
+		{
+			{ 0x0E6F, 0x0113 }, 	// AfterglowGamepadforXbox360
+			{ 0x0E6F, 0x0213 }, 	// AfterglowGamepadforXbox360
+			{ 0x1BAD, 0xF900 }, 	// AfterglowGamepadforXbox360
+			{ 0x0738, 0xCB29 }, 	// AviatorforXbox360PC
+			{ 0x15E4, 0x3F10 }, 	// BatarangwiredcontrollerXBOX
+			{ 0x146B, 0x0601 }, 	// BigbenControllerBB7201
+			{ 0x0738, 0xF401 }, 	// Controller
+			{ 0x0E6F, 0xF501 }, 	// Controller
+			{ 0x1430, 0xF801 }, 	// Controller
+			{ 0x1BAD, 0x028E }, 	// Controller
+			{ 0x1BAD, 0xFA01 }, 	// Controller
+			{ 0x12AB, 0x0004 }, 	// DDRUniverse2Mat
+			{ 0x24C6, 0x5B00 }, 	// Ferrari458Racingwheel
+			{ 0x1430, 0x4734 }, 	// GH4Guitar
+			{ 0x046D, 0xC21D }, 	// GamepadF310
+			{ 0x0E6F, 0x0301 }, 	// GamepadforXbox360
+			{ 0x0E6F, 0x0401 }, 	// GamepadforXbox360Z
+			{ 0x12AB, 0x0302 }, 	// GamepadforXbox360ZZ
+			{ 0x1BAD, 0xF902 }, 	// GamepadforXbox360ZZZ
+			{ 0x1BAD, 0xF901 }, 	// GamestopXbox360Controller
+			{ 0x1430, 0x474C }, 	// GuitarHeroforPCMAC
+			{ 0x1BAD, 0xF501 }, 	// HORIPADEX2TURBO
+			{ 0x1BAD, 0x0003 }, 	// HarmonixDrumKitforXbox360
+			{ 0x1BAD, 0x0002 }, 	// HarmonixGuitarforXbox360
+			{ 0x0F0D, 0x000A }, 	// HoriCoDOA4FightStick
+			{ 0x0F0D, 0x000D }, 	// HoriFightingStickEx2
+			{ 0x0F0D, 0x0016 }, 	// HoriRealArcadeProEx
+			{ 0x24C6, 0x5501 }, 	// HoriRealArcadeProVXSA
+			{ 0x24C6, 0x5506 }, 	// HoriSOULCALIBURVStick
+			{ 0x1BAD, 0xF02D }, 	// JoytechNeoSe
+			{ 0x162E, 0xBEEF }, 	// JoytechNeoSeTake2
+			{ 0x046D, 0xC242 }, 	// LogitechChillStream
+			{ 0x046D, 0xC21E }, 	// LogitechF510
+			{ 0x1BAD, 0xFD01 }, 	// MadCatz360
+			{ 0x0738, 0x4740 }, 	// MadCatzBeatPad
+			{ 0x1BAD, 0xF025 }, 	// MadCatzCallofDutyGamePad
+			{ 0x1BAD, 0xF027 }, 	// MadCatzFPSProGamePad
+			{ 0x1BAD, 0xF021 }, 	// MadCatzGhostReconFSGamePad
+			{ 0x0738, 0x4736 }, 	// MadCatzMicroConGamePadPro
+			{ 0x1BAD, 0xF036 }, 	// MadCatzMicroConGamePadProZ
+			{ 0x0738, 0x9871 }, 	// MadCatzPortableDrumKit
+			{ 0x0738, 0x4728 }, 	// MadCatzStreetFighterIVFightPad
+			{ 0x0738, 0x4718 }, 	// MadCatzStreetFighterIVFightStickSE
+			{ 0x0738, 0x4716 }, 	// MadCatzXbox360Controller
+			{ 0x0738, 0x4726 }, 	// MadCatzXbox360Controller
+			{ 0x0738, 0xBEEF }, 	// MadCatzXbox360Controller
+			{ 0x1BAD, 0xF016 }, 	// MadCatzXbox360Controller
+			{ 0x0738, 0xB726 }, 	// MadCatzXboxcontrollerMW2
+			{ 0x045E, 0x028E }, 	// MicrosoftXbox360Controller
+			{ 0x045E, 0x0719 }, 	// MicrosoftXbox360Controller
+			{ 0x12AB, 0x0301 }, 	// PDPAFTERGLOWAX1
+			{ 0x0E6F, 0x0105 }, 	// PDPDancePad
+			{ 0x0E6F, 0x0201 }, 	// PelicanTSZ360Pad
+			{ 0x15E4, 0x3F00 }, 	// PowerAMiniProElite
+			{ 0x24C6, 0x5300 }, 	// PowerAMiniProEliteGlow
+			{ 0x1BAD, 0xF504 }, 	// REALARCADEPROEX
+			{ 0x1BAD, 0xF502 }, 	// REALARCADEProVX
+			{ 0x1689, 0xFD00 }, 	// RazerOnza
+			{ 0x1689, 0xFD01 }, 	// RazerOnzaTournamentEdition
+			{ 0x1430, 0x4748 }, 	// RedOctaneGuitarHeroXplorer
+			{ 0x0E6F, 0x011F }, 	// RockCandyGamepadforXbox360
+			{ 0x12AB, 0x0006 }, 	// RockRevolutionforXbox360
+			{ 0x0738, 0xCB02 }, 	// SaitekCyborgRumblePadPCXbox360
+			{ 0x0738, 0xCB03 }, 	// SaitekP3200RumblePadPCXbox360
+			{ 0x1BAD, 0xF028 }, 	// StreetFighterIVFightPad
+			{ 0x0738, 0x4738 }, 	// StreetFighterIVFightStickTE
+			{ 0x0738, 0xF738 }, 	// SuperSFIVFightStickTES
+			{ 0x1BAD, 0xF903 }, 	// TronXbox360controller
+			{ 0x1BAD, 0x5500 }, 	// USBGamepad
+			{ 0x1BAD, 0xF906 }, 	// XB360MortalKombatFightStick
+			{ 0x15E4, 0x3F0A }, 	// XboxAirflowiredcontroller
+			{ 0x0E6F, 0x0401 },		// GameStop XBox 360 Controller
+		};
+
+		bool bIsXbox360Controller = false;
+
+		for (int i = 0; i < ARRAY_COUNT(XBox360Controllers); ++i)
+		{
+			const FXBox360ControllerID ControllerID = XBox360Controllers[i];
+			if (ControllerID.VendorID == VendorID && ControllerID.ProductID == ProductID)
+			{
+				bIsXbox360Controller = true;
+				break;
+			}
+		}
 
-		LeftAnalogXMapping = kHIDUsage_GD_X;
-		LeftAnalogYMapping = kHIDUsage_GD_Y;
-		LeftTriggerAnalogMapping = kHIDUsage_GD_Rx;
-		RightAnalogXMapping = kHIDUsage_GD_Z;
-		RightAnalogYMapping = kHIDUsage_GD_Rz;
-		RightTriggerAnalogMapping = kHIDUsage_GD_Ry;
+		if (bIsXbox360Controller)
+		{
+			// Xbox 360 Controller
+			ButtonsMapping[1]	= 0;	// A
+			ButtonsMapping[2]	= 1;	// B
+			ButtonsMapping[3]	= 2;	// X
+			ButtonsMapping[4]	= 3;	// Y
+			ButtonsMapping[5]	= 4;	// Left Shoulder
+			ButtonsMapping[6]	= 5;	// Right Shoulder
+			ButtonsMapping[7]	= 8;	// Left Thumbstick
+			ButtonsMapping[8]	= 9;	// Right Thumbstick
+			ButtonsMapping[9]	= 6;	// Start
+			ButtonsMapping[10]	= 7;	// Back
+			ButtonsMapping[12]	= 12;	// DPad Up
+			ButtonsMapping[13]	= 13;	// DPad Down
+			ButtonsMapping[14]	= 14;	// DPad Left
+			ButtonsMapping[15]	= 15;	// DPad Right
+
+			LeftAnalogXMapping = kHIDUsage_GD_X;
+			LeftAnalogYMapping = kHIDUsage_GD_Y;
+			LeftTriggerAnalogMapping = kHIDUsage_GD_Z;
+			RightAnalogXMapping = kHIDUsage_GD_Rx;
+			RightAnalogYMapping = kHIDUsage_GD_Ry;
+			RightTriggerAnalogMapping = kHIDUsage_GD_Rz;
+		}
+		else
+		{
+			// Generic (based on Logitech RumblePad 2)
+			ButtonsMapping[1]	= 2;	// X
+			ButtonsMapping[2]	= 0;	// A
+			ButtonsMapping[3]	= 1;	// B
+			ButtonsMapping[4]	= 3;	// Y
+			ButtonsMapping[5]	= 4;	// Left Shoulder
+			ButtonsMapping[6]	= 5;	// Right Shoulder
+			ButtonsMapping[7]	= 10;	// Left Trigger
+			ButtonsMapping[8]	= 11;	// Right Trigger
+			ButtonsMapping[9]	= 7;	// Back
+			ButtonsMapping[10]	= 6;	// Start
+			ButtonsMapping[11]	= 8;	// Left Thumbstick
+			ButtonsMapping[12]	= 9;	// Right Thumbstick
+
+			LeftAnalogXMapping = kHIDUsage_GD_X;
+			LeftAnalogYMapping = kHIDUsage_GD_Y;
+			LeftTriggerAnalogMapping = kHIDUsage_GD_Rx;
+			RightAnalogXMapping = kHIDUsage_GD_Z;
+			RightAnalogYMapping = kHIDUsage_GD_Rz;
+			RightTriggerAnalogMapping = kHIDUsage_GD_Ry;
+		}
 	}
 }
 
diff --git a/Engine/Source/Runtime/Core/Private/Mac/HIDInputInterface.h b/Engine/Source/Runtime/Core/Private/Mac/HIDInputInterface.h
--- a/Engine/Source/Runtime/Core/Private/Mac/HIDInputInterface.h
+++ b/Engine/Source/Runtime/Core/Private/Mac/HIDInputInterface.h
@@ -62,6 +62,12 @@ private:
 		void SetupMappings();
 	};
 
+	struct FXBox360ControllerID
+	{
+		const int32 VendorID;
+		const int32 ProductID;
+	};
+
 	struct FControllerState
 	{
 		/** Last frame's button states, so we only send events on edges */
