Initial revision of Future/Promise/Async implementation

diff --git a/Engine/Source/Runtime/Core/Private/Tests/Async/AsyncTest.cpp b/Engine/Source/Runtime/Core/Private/Tests/Async/AsyncTest.cpp
new file mode 100644
--- /dev/null
+++ b/Engine/Source/Runtime/Core/Private/Tests/Async/AsyncTest.cpp
@@ -0,0 +1,37 @@
+// Copyright 1998-2015 Epic Games, Inc. All Rights Reserved.
+
+#include "CorePrivatePCH.h"
+#include "Function.h"
+#include "Async.h"
+#include "AutomationTest.h"
+#include "Future.h"
+
+
+IMPLEMENT_SIMPLE_AUTOMATION_TEST(FAsyncTest, "Core.Async.Async", EAutomationTestFlags::ATF_Editor)
+
+
+bool FAsyncTest::RunTest(const FString& Parameters)
+{
+	auto Task = []()
+	{
+		return 123;
+	};
+
+	// task graph task
+	{
+		auto Future = Async<int>(EAsyncExecution::TaskGraph, Task);
+		int Result = Future.Get();
+
+		TestEqual(TEXT("Task graph task must return expected value"), Result, 123);
+	}
+
+	// thread task
+	{
+		auto Future = Async<int>(EAsyncExecution::Thread, Task);
+		int Result = Future.Get();
+
+		TestEqual(TEXT("Task graph task must return expected value"), Result, 123);
+	}
+
+	return true;
+}
diff --git a/Engine/Source/Runtime/Core/Public/Async/Async.h b/Engine/Source/Runtime/Core/Public/Async/Async.h
new file mode 100644
--- /dev/null
+++ b/Engine/Source/Runtime/Core/Public/Async/Async.h
@@ -0,0 +1,196 @@
+// Copyright 1998-2015 Epic Games, Inc. All Rights Reserved.
+
+#pragma once
+
+#include "Future.h"
+#include "TaskGraphInterfaces.h"
+
+
+/**
+ * Enumerates available asynchronous execution methods.
+ */
+enum class EAsyncExecution
+{
+	/** Execute in Task Graph (for short running tasks). */
+	TaskGraph,
+
+	/** Execute in separate thread (for long running tasks). */
+	Thread,
+
+	/** Execute in queued thread pool. */
+	ThreadPool
+};
+
+
+/**
+ * Template for asynchronous functions that are executed in the Task Graph system.
+ */
+template<typename ResultType>
+class TAsyncTask
+{
+public:
+
+	/**
+	 * Creates and initializes a new instance.
+	 *
+	 * @param InFunction The function to execute asynchronously.
+	 * @param InPromise The promise object used to return the function's result.
+	 */
+	TAsyncTask(const TFunction<ResultType()>& InFunction, TPromise<ResultType>&& InPromise)
+		: Function(MoveTemp(InFunction))
+		, Promise(MoveTemp(InPromise))
+	{ }
+
+public:
+
+	/**
+	 * Performs the actual task.
+	 *
+	 * @param CurrentThread The thread that this task is executing on.
+	 * @param MyCompletionGraphEvent The completion event.
+	 */
+	void DoTask(ENamedThreads::Type CurrentThread, const FGraphEventRef& MyCompletionGraphEvent)
+	{
+		Promise.SetValue(Function());
+	}
+
+	/**
+	 * Returns the name of the thread that this task should run on.
+	 *
+	 * @return Always run on any thread.
+	 */
+	ENamedThreads::Type GetDesiredThread()
+	{
+		return ENamedThreads::AnyThread;
+	}
+
+	/**
+	 * Gets the future that will hold the asynchronous result.
+	 *
+	 * @return A TFuture object.
+	 */
+	TFuture<ResultType> GetFuture()
+	{
+		return Promise.GetFuture();
+	}
+
+	/**
+	 * Gets the task's stats tracking identifier.
+	 *
+	 * @return Stats identifier.
+	 */
+	TStatId GetStatId() const
+	{
+		return GET_STATID(STAT_TaskGraph_OtherTasks);
+	}
+
+	/**
+	 * Gets the mode for tracking subsequent tasks.
+	 *
+	 * @return Always track subsequent tasks.
+	 */
+	static ESubsequentsMode::Type GetSubsequentsMode()
+	{
+		return ESubsequentsMode::TrackSubsequents;
+	}
+
+private:
+
+	/** The function to execute on the Task Graph. */
+	TFunction<ResultType()> Function;
+
+	/** The promise to assign the result to. */
+	TPromise<ResultType> Promise;
+};
+
+
+/**
+ * Template for asynchronous functions that are executed in a separate thread.
+ */
+template<typename ResultType>
+class TAsyncRunnable
+	: public FRunnable
+{
+public:
+
+	/**
+	 * Creates and initializes a new instance.
+	 *
+	 * @param InFunction The function to execute asynchronously.
+	 * @param InPromise The promise object used to return the function's result.
+	 * @param InThreadFuture The thread that is running this task.
+	 */
+	TAsyncRunnable(const TFunction<ResultType()>& InFunction, TPromise<ResultType>&& InPromise, TFuture<FRunnableThread*>&& InThreadFuture)
+		: Function(InFunction)
+		, Promise(MoveTemp(InPromise))
+		, ThreadFuture(MoveTemp(InThreadFuture))
+	{ }
+
+public:
+
+	// FRunnable interface
+
+	virtual uint32 Run() override
+	{
+		Promise.SetValue(Function());
+
+		FRunnableThread* Thread = ThreadFuture.Get();
+
+		delete Thread;
+		delete this;
+
+		return 0;
+	}
+
+private:
+
+	/** The function to execute on the Task Graph. */
+	TFunction<ResultType()> Function;
+
+	/** The promise to assign the result to. */
+	TPromise<ResultType> Promise;
+
+	/** The thread running this task. */
+	TFuture<FRunnableThread*> ThreadFuture;
+};
+
+
+/**
+ * Executes a given function asynchronously.
+ *
+ * @param Execution The execution method to use, i.e. on Task Graph or in a separate thread.
+ * @param Function The function to execute.
+ * @result A TFuture object that will receive the return value from the function.
+ */
+template<typename ResultType>
+TFuture<ResultType> Async(EAsyncExecution Execution, TFunction<ResultType()> Function)
+{
+	TPromise<ResultType> Promise;
+	TFuture<ResultType> Future = Promise.GetFuture();
+
+	switch (Execution)
+	{
+	case EAsyncExecution::TaskGraph:
+		{
+			TGraphTask<TAsyncTask<ResultType>>::CreateTask().ConstructAndDispatchWhenReady(Function, MoveTemp(Promise));
+		}
+		break;
+	
+	case EAsyncExecution::Thread:
+		{
+			TPromise<FRunnableThread*> ThreadPromise;
+			TAsyncRunnable<ResultType>* Runnable = new TAsyncRunnable<ResultType>(Function, MoveTemp(Promise), ThreadPromise.GetFuture());
+			FRunnableThread* RunnableThread = FRunnableThread::Create(Runnable, TEXT("TAsync"));
+
+			check(RunnableThread != nullptr)
+
+			ThreadPromise.SetValue(RunnableThread);
+		}
+		break;
+
+	default:
+		check(false) // not implemented yet!
+	}
+
+	return MoveTemp(Future);
+}
diff --git a/Engine/Source/Runtime/Core/Public/Async/Future.h b/Engine/Source/Runtime/Core/Public/Async/Future.h
new file mode 100644
--- /dev/null
+++ b/Engine/Source/Runtime/Core/Public/Async/Future.h
@@ -0,0 +1,1000 @@
+// Copyright 1998-2015 Epic Games, Inc. All Rights Reserved.
+
+#pragma once
+
+
+/* TFutureBase
+*****************************************************************************/
+
+/**
+ * Base class for the internal state of asynchronous return values (futures).
+ */
+class FFutureState
+{
+public:
+
+	/** Default constructor. */
+	FFutureState()
+		: CompletionEvent(FPlatformProcess::CreateSynchEvent(true))
+		, Complete(false)
+	{ }
+
+	/** Destructor. */
+	~FFutureState()
+	{
+		delete CompletionEvent;
+	}
+
+public:
+
+	/**
+	 * Checks whether the asynchronous result has been set.
+	 *
+	 * @return true if the result has been set, false otherwise.
+	 * @see WaitFor
+	 */
+	bool IsComplete() const
+	{
+		return Complete;
+	}
+
+	/**
+	 * Blocks the calling thread until the future result is available.
+	 *
+	 * @param Duration The maximum time span to wait for the future result.
+	 * @result true if the result is available, false otherwise.
+	 * @see IsComplete
+	 */
+	bool WaitFor(const FTimespan& Duration)
+	{
+		FScopeLock Lock(&CriticalSection);
+
+		if (CompletionEvent->Wait(Duration))
+		{
+			return true;
+		}
+
+		return false;
+	}
+
+protected:
+
+	/** Notifies any waiting threads that the result is available. */
+	void MarkComplete()
+	{
+		Complete = true;
+		CompletionEvent->Trigger();
+	}
+
+private:
+
+	/** Holds an event signaling that the result is available. */
+	FEvent* CompletionEvent;
+
+	/** Critical section for locking access to this class. */
+	mutable FCriticalSection CriticalSection;
+
+	/** Whether the asynchronous result is available. */
+	bool Complete;
+};
+
+
+/**
+ * Implements the internal state of asynchronous return values (futures).
+ */
+template<typename InternalResultType>
+class TFutureState
+	: public FFutureState
+{
+public:
+
+	/**
+	 * Gets the result (will block the calling thread until the result is available).
+	 *
+	 * @return The result value.
+	 * @see SetResult
+	 */
+	InternalResultType& GetResult()
+	{
+		while (!IsComplete())
+		{
+			WaitFor(FTimespan::MaxValue());
+		}
+
+		return Result;
+	}
+
+	/**
+	 * Sets the result and notifies any waiting threads.
+	 *
+	 * @param InResult The result to set.
+	 * @see GetResult
+	 */
+	void SetResult(const InternalResultType& InResult)
+	{
+		check(!IsComplete())
+
+		Result = InResult;
+		MarkComplete();
+	}
+
+	/**
+	 * Sets the result and notifies any waiting threads (from rvalue).
+	 *
+	 * @param InResult The result to set.
+	 * @see GetResult
+	 */
+	void SetResult(InternalResultType&& InResult)
+	{
+		check(!IsComplete())
+
+		Result = InResult;
+		MarkComplete();
+	}
+
+private:
+
+	/** Holds the asynchronous result. */
+	InternalResultType Result;
+};
+
+
+/**
+ * Abstract base template for futures and shared futures.
+ */
+template<typename InternalResultType>
+class TFutureBase
+{
+public:
+
+	/**
+	 * Checks whether this future object has a valid state.
+	 *
+	 * @return true if the state is valid, false otherwise.
+	 */
+	bool IsValid() const
+	{
+		return State.IsValid();
+	}
+
+	/**
+	 * Blocks the calling thread until the future result is available.
+	 *
+	 * Note that this method may block forever if the result is never set. Use
+	 * the WaitFor or WaitUntil methods to specify a maximum timeout for the wait.
+	 *
+	 * @see WaitFor, WaitUntil
+	 */
+	void Wait()
+	{
+		while (!WaitFor(FTimespan::MaxValue()));
+	}
+
+	/**
+	 * Blocks the calling thread until the future result is available or the specified duration is exceeded.
+	 *
+	 * @param Duration The maximum time span to wait for the future result.
+	 * @result true if the result is available, false otherwise.
+	 * @see Wait, WaitUntil
+	 */
+	bool WaitFor(const FTimespan& Duration)
+	{
+		return State->WaitFor(Duration);
+	}
+
+	/**
+	 * Blocks the calling thread until the future result is available or the specified time is hit.
+	 *
+	 * @param Time The time until to wait for the future result (in UTC).
+	 * @result true if the result is available, false otherwise.
+	 * @see Wait, WaitUntil
+	 */
+	bool WaitUntil(const FDateTime& Time)
+	{
+		return WaitFor(Time - FDateTime::UtcNow());
+	}
+
+protected:
+
+	typedef TSharedPtr<TFutureState<InternalResultType>, ESPMode::ThreadSafe> StateType;
+
+	/** Default constructor. */
+	TFutureBase() { }
+
+	/**
+	 * Creates and initializes a new instance.
+	 *
+	 * @param InState The shared state to initialize with.
+	 */
+	TFutureBase(const StateType& InState)
+		: State(InState)
+	{ }
+
+	/**
+	 * Move constructor.
+	 *
+	 * @param Other The future holding the shared state to move.
+	 */
+	TFutureBase(TFutureBase&& Other)
+		: State(MoveTemp(Other.State))
+	{
+		Other.State.Reset();
+	}
+
+	/** Hidden destructor. */
+	~TFutureBase() { }
+
+protected:
+
+	/** Move assignment operator. */
+	TFutureBase& operator=(TFutureBase&& Other)
+	{
+		State = MoveTemp(Other.State);
+		Other.State.Reset();
+	}
+
+protected:
+
+	/**
+	 * Gets the shared state object.
+	 *
+	 * @return The shared state object.
+	 */
+	const StateType& GetState()
+	{
+		// if you hit this assertion then your future has an invalid state.
+		// this happens if you have an uninitialized future or if you moved
+		// the future it to another instance.
+		check(State.IsValid())
+
+		return State;
+	}
+
+private:
+
+	/** Holds the future's state. */
+	StateType State;
+};
+
+
+/* TFuture
+*****************************************************************************/
+
+template<typename ResultType> class TSharedFuture;
+
+
+/**
+ * Template for unshared futures.
+ */
+template<typename ResultType>
+class TFuture
+	: public TFutureBase<ResultType>
+	, public FNoncopyable
+{
+	typedef TFutureBase<ResultType> BaseType;
+
+public:
+
+	/** Default constructor. */
+	TFuture() { }
+
+	/**
+	 * Creates and initializes a new instance.
+	 *
+	 * @param InState The shared state to initialize with.
+	 */
+	TFuture(const BaseType::StateType& InState)
+		: BaseType(InState)
+	{ }
+
+	/** Move constructor. */
+	TFuture(TFuture&& Other)
+		: BaseType(MoveTemp(Other))
+	{ }
+
+	/** Destructor. */
+	~TFuture() { }
+
+public:
+
+	/** Move assignment operator. */
+	TFuture& operator=(TFuture&& Other)
+	{
+		BaseType::operator=(MoveTemp(Other));
+		return *this;
+	}
+
+public:
+
+	/**
+	 * Gets the future's result.
+	 *
+	 * @return The result.
+	 */
+	ResultType Get()
+	{
+		return MoveTemp(GetState()->GetResult());
+	}
+
+	/**
+	 * Moves this future's state into a shared future.
+	 *
+	 * @return The shared future object.
+	 */
+	TSharedFuture<ResultType> Share()
+	{
+		return TSharedFuture<ResultType>(MoveTemp(*this));
+	}
+};
+
+
+/**
+ * Template for unshared futures (specialization for reference types).
+ */
+template<typename ResultType>
+class TFuture<ResultType&>
+	: public TFutureBase<ResultType*>
+	, public FNoncopyable
+{
+	typedef TFutureBase<ResultType*> BaseType;
+
+public:
+
+	/** Default constructor. */
+	TFuture() { }
+
+	/**
+	 * Creates and initializes a new instance.
+	 *
+	 * @param InState The shared state to initialize with.
+	 */
+	TFuture(const BaseType::StateType& InState)
+		: BaseType(InState)
+	{ }
+
+	/** Move constructor. */
+	TFuture(TFuture&& Other)
+		: BaseType(MoveTemp(Other))
+	{ }
+
+	/** Destructor. */
+	~TFuture() { }
+
+public:
+
+	/** Move assignment operator. */
+	TFuture& operator=(TFuture&& Other)
+	{
+		BaseType::operator=(MoveTemp(Other));
+		return *this;
+	}
+
+public:
+
+	/**
+	 * Gets the future's result.
+	 *
+	 * @return The result.
+	 */
+	ResultType& Get()
+	{
+		return *GetState()->GetResult();
+	}
+
+	/**
+	 * Moves this future's state into a shared future.
+	 *
+	 * @return The shared future object.
+	 */
+	TSharedFuture<ResultType&> Share()
+	{
+		return TSharedFuture<ResultType&>(MoveTemp(*this));
+	}
+};
+
+
+/**
+ * Template for unshared futures (specialization for void).
+ */
+template<>
+class TFuture<void>
+	: public TFutureBase<int>
+	, public FNoncopyable
+{
+	typedef TFutureBase<int> BaseType;
+
+public:
+
+	/** Default constructor. */
+	TFuture() { }
+
+	/**
+	 * Creates and initializes a new instance.
+	 *
+	 * @param InState The shared state to initialize with.
+	 */
+	TFuture(const BaseType::StateType& InState)
+		: BaseType(InState)
+	{ }
+
+	/** Move constructor. */
+	TFuture(TFuture&& Other)
+		: BaseType(MoveTemp(Other))
+	{ }
+
+	/** Destructor. */
+	~TFuture() { }
+
+public:
+
+	/** Move assignment operator. */
+	TFuture& operator=(TFuture&& Other)
+	{
+		BaseType::operator=(MoveTemp(Other));
+		return *this;
+	}
+
+public:
+
+	/**
+	 * Gets the future's result.
+	 *
+	 * @return The result.
+	 */
+	void Get()
+	{
+		GetState()->GetResult();
+	}
+
+	/**
+	 * Moves this future's state into a shared future.
+	 *
+	 * @return The shared future object.
+	 */
+	TSharedFuture<void> Share();
+};
+
+
+/* TSharedFuture
+*****************************************************************************/
+
+/**
+ * Template for shared futures.
+ */
+template<typename ResultType>
+class TSharedFuture
+	: public TFutureBase<ResultType>
+{
+	typedef TFutureBase<ResultType> BaseType;
+
+public:
+
+	/** Default constructor. */
+	TSharedFuture() { }
+
+	/**
+	 * Creates and initializes a new instance.
+	 *
+	 * @param InState The shared state to initialize from.
+	 */
+	TSharedFuture(const BaseType::StateType& InState)
+		: BaseType(InState)
+	{ }
+
+	/**
+	 * Creates and initializes a new instances from a future object.
+	 *
+	 * @param Future The future object to initialize from.
+	 */
+	TSharedFuture(TFuture<ResultType>&& Future)
+		: BaseType(MoveTemp(Future))
+	{ }
+
+	/** Copy constructor. */
+	TSharedFuture(const TSharedFuture& Other)
+		: BaseType(Other)
+	{ }
+
+	/** Move constructor. */
+	TSharedFuture(TSharedFuture&& Other)
+		: BaseType(MoveTemp(Other))
+	{ }
+
+	/** Destructor. */
+	~TSharedFuture() { }
+
+public:
+
+	/** Copy assignment operator. */
+	TSharedFuture& operator=(const TSharedFuture& Other)
+	{
+		BaseType::operator=(Other);
+		return *this;
+	}
+
+	/** Move assignment operator. */
+	TSharedFuture& operator=(TSharedFuture&& Other)
+	{
+		BaseType::operator=(MoveTemp(Other));
+		return *this;
+	}
+
+public:
+
+	/**
+	 * Gets the future's result.
+	 *
+	 * @return The result.
+	 */
+	ResultType Get()
+	{
+		return MoveTemp(GetState()->GetResult());
+	}
+};
+
+
+/**
+ * Template for shared futures (specialization for reference types).
+ */
+template<typename ResultType>
+class TSharedFuture<ResultType&>
+	: public TFutureBase<ResultType*>
+{
+	typedef TFutureBase<ResultType*> BaseType;
+
+public:
+
+	/** Default constructor. */
+	TSharedFuture() { }
+
+	/**
+	 * Creates and initializes a new instance.
+	 *
+	 * @param InState The shared state to initialize from.
+	 */
+	TSharedFuture(const BaseType::StateType& InState)
+		: BaseType(InState)
+	{ }
+
+	/**
+	* Creates and initializes a new instances from a future object.
+	*
+	* @param Future The future object to initialize from.
+	*/
+	TSharedFuture(TFuture&& Future)
+		: BaseType(MoveTemp(Future))
+	{ }
+
+	/** Copy constructor. */
+	TSharedFuture(const TSharedFuture& Other)
+		: BaseType(Other)
+	{ }
+
+	/** Move constructor. */
+	TSharedFuture(TSharedFuture&& Other)
+		: BaseType(MoveTemp(Other))
+	{ }
+
+	/** Destructor. */
+	~TSharedFuture() { }
+
+public:
+
+	/** Copy assignment operator. */
+	TSharedFuture& operator=(const TSharedFuture& Other)
+	{
+		BaseType::operator=(MoveTemp(Other));
+		return *this;
+	}
+
+	/** Move assignment operator. */
+	TSharedFuture& operator=(TSharedFuture&& Other)
+	{
+		BaseType::operator=(MoveTemp(Other));
+		return *this;
+	}
+
+public:
+
+	/**
+	 * Gets the future's result.
+	 *
+	 * @return The result.
+	 */
+	ResultType& Get()
+	{
+		return *GetState()->GetResult();
+	}
+};
+
+
+/**
+ * Template for shared futures (specialization for void).
+ */
+template<>
+class TSharedFuture<void>
+	: public TFutureBase<int>
+{
+	typedef TFutureBase<int> BaseType;
+
+public:
+
+	/** Default constructor. */
+	TSharedFuture() { }
+
+	/**
+	 * Creates and initializes a new instance from shared state.
+	 *
+	 * @param InState The shared state to initialize from.
+	 */
+	TSharedFuture(const BaseType::StateType& InState)
+		: BaseType(InState)
+	{ }
+
+	/**
+	 * Creates and initializes a new instances from a future object.
+	 *
+	 * @param Future The future object to initialize from.
+	 */
+	TSharedFuture(TFuture<void>&& Future)
+		: BaseType(MoveTemp(Future))
+	{ }
+
+	/** Copy constructor. */
+	TSharedFuture(const TSharedFuture& Other)
+		: BaseType(Other)
+	{ }
+
+	/** Move constructor. */
+	TSharedFuture(TSharedFuture&& Other)
+		: BaseType(MoveTemp(Other))
+	{ }
+
+	/** Destructor. */
+	~TSharedFuture() { }
+
+public:
+
+	/** Copy assignment operator. */
+	TSharedFuture& operator=(const TSharedFuture& Other)
+	{
+		BaseType::operator=(Other);
+		return *this;
+	}
+
+	/** Move assignment operator. */
+	TSharedFuture& operator=(TSharedFuture&& Other)
+	{
+		BaseType::operator=(MoveTemp(Other));
+		return *this;
+	}
+
+public:
+
+	/**
+	 * Gets the future's result.
+	 *
+	 * @return The result.
+	 */
+	void Get()
+	{
+		GetState()->GetResult();
+	}
+};
+
+
+inline TSharedFuture<void> TFuture<void>::Share()
+{
+	return TSharedFuture<void>(MoveTemp(*this));
+}
+
+
+/* TPromise
+*****************************************************************************/
+
+template<typename InternalResultType>
+class TPromiseBase
+	: FNoncopyable
+{
+	typedef TSharedPtr<TFutureState<InternalResultType>, ESPMode::ThreadSafe> StateType;
+
+public:
+
+	/** Default constructor. */
+	TPromiseBase()
+		: State(MakeShareable(new TFutureState<InternalResultType>))
+	{ }
+
+	/**
+	 * Move constructor.
+	 *
+	 * @param Other The promise holding the shared state to move.
+	 */
+	TPromiseBase(TPromiseBase&& Other)
+		: State(MoveTemp(Other.State))
+	{
+		Other.State.Reset();
+	}
+
+public:
+
+	/** Move assignment operator. */
+	TPromiseBase& operator=(TPromiseBase&& Other)
+	{
+		State = Other.State;
+		Other.State.Reset();
+	}
+
+protected:
+
+	/** Destructor. */
+	~TPromiseBase()
+	{
+		if (State.IsValid())
+		{
+			// if you hit this assertion then your promise never had its result
+			// value set. broken promises are considered programming errors.
+			check(State->IsComplete())
+		}
+	}
+
+	/**
+	 * Gets the shared state object.
+	 *
+	 * @return The shared state object.
+	 */
+	const StateType& GetState()
+	{
+		// if you hit this assertion then your promise has an invalid state.
+		// this happens if you move the promise to another instance.
+		check(State.IsValid())
+
+			return State;
+	}
+
+private:
+
+	/** Holds the shared state object. */
+	StateType State;
+};
+
+
+/**
+ * Template for promises.
+ */
+template<typename ResultType>
+class TPromise
+	: public TPromiseBase<ResultType>
+{
+	typedef public TPromiseBase<ResultType> BaseType;
+
+public:
+
+	/** Default constructor (creates a new shared state). */
+	TPromise()
+		: BaseType()
+		, FutureRetrieved(false)
+	{ }
+
+	/**
+	 * Move constructor.
+	 *
+	 * @param Other The promise holding the shared state to move.
+	 */
+	TPromise(TPromise&& Other)
+		: BaseType(MoveTemp(Other))
+	{ }
+
+public:
+
+	/**
+	 * Move assignment operator.
+	 *
+	 * @param Other The promise holding the shared state to move.
+	 */
+	TPromise& operator=(TPromise&& Other)
+	{
+		BaseType::operator=(MoveTemp(Other));
+		FutureRetrieved = MoveTemp(Other.FutureRetrieved);
+
+		return this;
+	}
+
+public:
+
+	/**
+	 * Gets a TFuture object associated with the shared state of this promise.
+	 *
+	 * @return The TFuture object.
+	 */
+	TFuture<ResultType> GetFuture()
+	{
+		check(!FutureRetrieved)
+			FutureRetrieved = true;
+
+		return TFuture<ResultType>(GetState());
+	}
+
+	/**
+	 * Sets the promised result.
+	 *
+	 * The result must be set only once. An assertion will
+	 * be triggered if this method is called a second time.
+	 *
+	 * @param Result The result value to set.
+	 */
+	void SetValue(const ResultType& Result)
+	{
+		GetState()->SetResult(Result);
+	}
+
+	/**
+	 * Sets the promised result (from rvalue).
+	 *
+	 * The result must be set only once. An assertion will
+	 * be triggered if this method is called a second time.
+	 *
+	 * @param Result The result value to set.
+	 */
+	void SetValue(ResultType&& Result)
+	{
+		GetState()->SetResult(MoveTemp(Result));
+	}
+
+private:
+
+	/** Whether a future has already been retrieved from this promise. */
+	bool FutureRetrieved;
+};
+
+
+/**
+ * Template for promises (specialization for reference types).
+ */
+template<typename ResultType>
+class TPromise<ResultType&>
+	: public TPromiseBase<ResultType*>
+{
+	typedef TPromiseBase<ResultType*> BaseType;
+
+public:
+
+	/** Default constructor (creates a new shared state). */
+	TPromise()
+		: BaseType()
+		, FutureRetrieved(false)
+	{ }
+
+	/**
+	 * Move constructor.
+	 *
+	 * @param Other The promise holding the shared state to move.
+	 */
+	TPromise(TPromise&& Other)
+		: BaseType(MoveTemp(Other))
+		, FutureRetrieved(MoveTemp(Other.FutureRetrieved))
+	{ }
+
+public:
+
+	/**
+	 * Move assignment operator.
+	 *
+	 * @param Other The promise holding the shared state to move.
+	 */
+	TPromise& operator=(TPromise&& Other)
+	{
+		BaseType::operator=(MoveTemp(Other));
+		FutureRetrieved = MoveTemp(Other.FutureRetrieved);
+
+		return this;
+	}
+
+public:
+
+	/**
+	 * Gets a TFuture object associated with the shared state of this promise.
+	 *
+	 * @return The TFuture object.
+	 */
+	TFuture<ResultType&> GetFuture()
+	{
+		check(!FutureRetrieved)
+			FutureRetrieved = true;
+
+		return TFuture<ResultType&>(GetState());
+	}
+
+	/**
+	 * Sets the promised result.
+	 *
+	 * The result must be set only once. An assertion will
+	 * be triggered if this method is called a second time.
+	 *
+	 * @param Result The result value to set.
+	 */
+	void SetValue(ResultType& Result)
+	{
+		GetState()->SetResult(Result);
+	}
+
+private:
+
+	/** Whether a future has already been retrieved from this promise. */
+	bool FutureRetrieved;
+};
+
+
+/**
+ * Template for promises (specialization for void results).
+ */
+template<>
+class TPromise<void>
+	: public TPromiseBase<int>
+{
+	typedef TPromiseBase<int> BaseType;
+
+public:
+
+	/** Default constructor (creates a new shared state). */
+	TPromise()
+		: BaseType()
+		, FutureRetrieved(false)
+	{ }
+
+	/**
+	 * Move constructor.
+	 *
+	 * @param Other The promise holding the shared state to move.
+	 */
+	TPromise(TPromise&& Other)
+		: BaseType(MoveTemp(Other))
+		, FutureRetrieved(false)
+	{ }
+
+public:
+
+	/**
+	 * Move assignment operator.
+	 *
+	 * @param Other The promise holding the shared state to move.
+	 */
+	TPromise& operator=(TPromise&& Other)
+	{
+		BaseType::operator=(MoveTemp(Other));
+		FutureRetrieved = MoveTemp(Other.FutureRetrieved);
+
+		return *this;
+	}
+
+public:
+
+	/**
+	 * Gets a TFuture object associated with the shared state of this promise.
+	 *
+	 * @return The TFuture object.
+	 */
+	TFuture<void> GetFuture()
+	{
+		check(!FutureRetrieved)
+			FutureRetrieved = true;
+
+		return TFuture<void>(GetState());
+	}
+
+	/**
+	 * Sets the promised result.
+	 *
+	 * The result must be set only once. An assertion will
+	 * be triggered if this method is called a second time.
+	 */
+	void SetValue()
+	{
+		GetState()->SetResult(0);
+	}
+
+private:
+
+	/** Whether a future has already been retrieved from this promise. */
+	bool FutureRetrieved;
+};
