Initial revision of ISlowTask, TSlowResult implementation

diff --git a/Engine/Source/Runtime/Core/Public/Async/ISlowTask.h b/Engine/Source/Runtime/Core/Public/Async/ISlowTask.h
new file mode 100644
--- /dev/null
+++ b/Engine/Source/Runtime/Core/Public/Async/ISlowTask.h
@@ -0,0 +1,62 @@
+// Copyright 1998-2015 Epic Games, Inc. All Rights Reserved.
+
+#pragma once
+
+
+/**
+ * Interface for slow running tasks.
+ *
+ * A slow running task is a unit of work that may take a considerable amount of time to
+ * complete, i.e. several seconds, minutes or even hours. This interface provides mechanisms
+ * for tracking and aborting such tasks.
+ *
+ * Classes that implement this interface should make a best effort to provide a completion
+ * percentage via the GetPercentageComplete() method. While the task is still pending, the
+ * returned value should be unset. After the task completed, the return value should be 1.0.
+ *
+ * Slow running tasks may also provide a status text that can be polled by the user interface.
+ */
+class ISlowTask
+{
+public:
+
+	/** Abort this task. */
+	virtual void Abort() = 0;
+
+	/**
+	 * Gets the task's completion percentage, if available (0.0 to 1.0, or unset).
+	 *
+	 * @return Completion percentage value.
+	 * @see IsComplete
+	 */
+	virtual TOptional<float> GetPercentageComplete() = 0;
+
+	/**
+	 * Gets the task's current status text.
+	 *
+	 * @return Status text.
+	 */
+	virtual FText GetStatusText() = 0;
+
+	/** A delegate that is executed after the task is complete. */
+	DECLARE_EVENT_OneParam(ISlowTask, FOnCompleted, bool /*Success*/)
+	virtual FOnCompleted& OnComplete() = 0;
+
+	/** A delegate that is executed after the task is complete. */
+	DECLARE_EVENT_OneParam(ISlowTask, FOnStatusChanged, FText /*StatusText*/)
+	virtual FOnStatusChanged& OnStatusChanged() = 0;
+
+public:
+
+	/**
+	 * Checks whether this task is complete.
+	 *
+	 * @return true if the task is complete, false otherwise.
+	 * @see GetPercentageComplete
+	 */
+	bool IsComplete() const
+	{
+		TOptional<float> PercentageComplete = GetPercentageComplete();
+		return (PercentageComplete.IsSet() && (PercentageComplete.Get() == 1.0f));
+	}
+};
diff --git a/Engine/Source/Runtime/Core/Public/Async/SlowResult.h b/Engine/Source/Runtime/Core/Public/Async/SlowResult.h
new file mode 100644
--- /dev/null
+++ b/Engine/Source/Runtime/Core/Public/Async/SlowResult.h
@@ -0,0 +1,32 @@
+// Copyright 1998-2015 Epic Games, Inc. All Rights Reserved.
+
+#pragma once
+
+#include "Future.h"
+#include "ISlowTask.h"
+
+
+/**
+ * Template for slow asynchronous results.
+ */
+template<typename ResultType)
+struct TSlowResult
+{
+	/** Holds the asynchronous return value. */
+	TFuture<ResultType> Future;
+
+	/** Pointer to the asynchronous task that is computing the result, or nullptr if the result was returned synchronously. */
+	TSharedPtr<ISlowTask> Task;
+
+
+	/**
+	* Creates and initializes a new instance.
+	*
+	* @param InFuture The future object to set.
+	* @param InTask The slow task to set.
+	*/
+	TSlowResult(TFuture<ResultType>&& InFuture, const TSharedRef<ISlowTask>& InTask)
+		: Future(InFuture)
+		, Task(InTask)
+	{ }
+};
diff --git a/Engine/Source/Runtime/Core/Public/Misc/ILifeCycle.h b/Engine/Source/Runtime/Core/Public/Misc/ILifeCycle.h
new file mode 100644
--- /dev/null
+++ b/Engine/Source/Runtime/Core/Public/Misc/ILifeCycle.h
@@ -0,0 +1,36 @@
+// Copyright 1998-2015 Epic Games, Inc. All Rights Reserved.
+
+#pragma once
+
+
+/**
+ * Interface for controlling an object’s start/stop life cycle.
+ */
+class ILifeCycle
+{
+public:
+
+	/**
+	 * Checks whether the object has been started.
+	 *
+	 * @return true if started, false otherwise.
+	 * @see Start, Stop
+	 */
+	virtual bool IsStarted() const = 0;
+
+	/**
+	 * Starts the object.
+	 *
+	 * @return true if it was started, false otherwise.
+	 * @see IsStarted, Stop
+	 */
+	virtual bool Start() = 0;
+
+	/**
+	 * Stops the object.
+	 *
+	 * @return true if it was stopped, false otherwise.
+	 * @see IsStarted, Start
+	 */
+	virtual bool Stop() = 0;
+};
