Changed implementation of TTypeToString and TTypeFromString to use global namespaced functions instead of explicit template instantiations.
This also removes the default implementations that simply fail assertions at runtime.

diff --git a/Engine/Source/Runtime/Core/Private/Containers/String.cpp b/Engine/Source/Runtime/Core/Private/Containers/String.cpp
--- a/Engine/Source/Runtime/Core/Private/Containers/String.cpp
+++ b/Engine/Source/Runtime/Core/Private/Containers/String.cpp
@@ -489,39 +489,7 @@ bool FString::IsNumeric() const
 		return 0;
 	}
 
-	TCHAR C = (*this)[0];
-	
-	if( C == '-' || C == '+' || C =='.' || FChar::IsDigit( C ) )
-	{
-		bool HasDot = (C == '.');
-
-		for( int32 i=1; i<Len(); i++ )
-		{
-			C = (*this)[i];
-
-			if( C == '.' )
-			{
-				if( HasDot )
-				{
-					return 0;
-				}
-				else
-				{
-					HasDot = 1;
-				}
-			}
-			else if( !FChar::IsDigit(C) )
-			{
-				return 0;
-			}
-		}
-
-		return 1;
-	}
-	else
-	{
-		return 0;
-	}
+	return FCString::IsNumeric(Data.GetData());
 }
 
 /**
diff --git a/Engine/Source/Runtime/Core/Private/Stats/StatsCommand.cpp b/Engine/Source/Runtime/Core/Private/Stats/StatsCommand.cpp
--- a/Engine/Source/Runtime/Core/Private/Stats/StatsCommand.cpp
+++ b/Engine/Source/Runtime/Core/Private/Stats/StatsCommand.cpp
@@ -16,23 +16,19 @@ DECLARE_CYCLE_STAT(TEXT("HUD Group"),STAT_HUDGroup,STATGROUP_StatSystem);
 DECLARE_CYCLE_STAT(TEXT("Accumulate"),STAT_Accumulate,STATGROUP_StatSystem);
 DECLARE_CYCLE_STAT(TEXT("GetFlatAggregates"),STAT_GetFlatAggregates,STATGROUP_StatSystem);
 
-template<>
-struct TTypeFromString<EStatCompareBy::Type>
+void FromString( EStatCompareBy::Type& OutValue, const TCHAR* Buffer )
 {
-	static void FromString( EStatCompareBy::Type& OutValue, const TCHAR* Buffer )
-	{
-		OutValue = EStatCompareBy::Sum;
+	OutValue = EStatCompareBy::Sum;
 
-		if (FCString::Stricmp(Buffer, TEXT("CallCount")) == 0)
-		{
-			OutValue = EStatCompareBy::CallCount;
-		}
-		else if (FCString::Stricmp(Buffer, TEXT("Name")) == 0)
-		{
-			OutValue = EStatCompareBy::Name;
-		}
+	if (FCString::Stricmp(Buffer, TEXT("CallCount")) == 0)
+	{
+		OutValue = EStatCompareBy::CallCount;
 	}
-};
+	else if (FCString::Stricmp(Buffer, TEXT("Name")) == 0)
+	{
+		OutValue = EStatCompareBy::Name;
+	}
+}
 
 struct FGroupFilter : public IItemFiler
 {
diff --git a/Engine/Source/Runtime/Core/Public/Containers/UnrealString.h b/Engine/Source/Runtime/Core/Public/Containers/UnrealString.h
--- a/Engine/Source/Runtime/Core/Public/Containers/UnrealString.h
+++ b/Engine/Source/Runtime/Core/Public/Containers/UnrealString.h
@@ -1566,168 +1566,89 @@ inline int32 HexToBytes( const FString& HexString, uint8* OutBytes )
 	return NumBytes - 1;
 }
 
-/** A little helper to avoid trying to compile a printf with types that are not accepted by printf **/
-
-template<typename T, bool TIsNumeric>
-struct TTypeToString_Internal
+/** Namespace that houses lexical conversion for various types. User defined conversions can be implemented externally */
+namespace LexicalConversion
 {
-	static FString ToString(T Value)
-	{
-		check(0); // you are asking us to convert a non-numeric type to a string
-		return FString();
-	}
-
-	static FString ToSanitizedString(T Value)
-	{
-		check(0); // you are asking us to convert a non-numeric type to a string
-		return FString();
-	}
-};
-
-template<typename T>
-struct TTypeToString_Internal<T, true>
-{
-	static FString ToString(T Value)
+	/**
+	 *	Expected functions in this namespace are as follows:
+	 *		static bool		TryParseString(T& OutValue, const TCHAR* Buffer);
+	 *		static void 	FromString(T& OutValue, const TCHAR* Buffer);
+	 *		static FString	ToString(const T& OutValue);
+	 *
+	 *	Implement custom functionality externally.
+	 */
+
+	/** Covert a string buffer to intrinsic types */
+	inline void FromString(int8& OutValue, 		const TCHAR* Buffer)	{	OutValue = FCString::Atoi(Buffer);		}
+	inline void FromString(int16& OutValue,		const TCHAR* Buffer)	{	OutValue = FCString::Atoi(Buffer);		}
+	inline void FromString(int32& OutValue,		const TCHAR* Buffer)	{	OutValue = FCString::Atoi(Buffer);		}
+	inline void FromString(int64& OutValue,		const TCHAR* Buffer)	{	OutValue = FCString::Atoi64(Buffer);	}
+	inline void FromString(uint8& OutValue,		const TCHAR* Buffer)	{	OutValue = FCString::Atoi(Buffer);		}
+	inline void FromString(uint16& OutValue, 	const TCHAR* Buffer)	{	OutValue = FCString::Atoi(Buffer);		}
+	inline void FromString(uint32& OutValue, 	const TCHAR* Buffer)	{	OutValue = FCString::Atoi64(Buffer);	}	//64 because this unsigned and so Atoi might overflow
+	inline void FromString(uint64& OutValue, 	const TCHAR* Buffer)	{	OutValue = FCString::Strtoui64(Buffer, NULL, 0); }
+	inline void FromString(float& OutValue,		const TCHAR* Buffer)	{	OutValue = FCString::Atof(Buffer);		}
+	inline void FromString(double& OutValue, 	const TCHAR* Buffer)	{	OutValue = FCString::Atod(Buffer);		}
+	inline void FromString(bool& OutValue, 		const TCHAR* Buffer)	{	OutValue = FCString::ToBool(Buffer);	}
+
+	/** Convert numeric types to a string */
+	template<typename T>
+	typename TEnableIf<TIsArithmeticType<T>::Value, FString>::Type
+		ToString(const T& Value)
 	{
 		return FString::Printf( TFormatSpecifier<T>::GetFormatSpecifier(), Value );
 	}
-	
-	static FString ToSanitizedString(T Value)
-	{
-		return ToString(Value);
-	}
-};
 
-template<>
-struct TTypeToString_Internal<float, true>
-{
-	static FString ToString(float Value)
+	/** Helper template to convert to sanitized strings */
+	template<typename T>
+	FString ToSanitizedString(const T& Value)
 	{
-		return FString::Printf( TFormatSpecifier<float>::GetFormatSpecifier(), Value );
+		return ToString(Value);
 	}
 
-	static FString ToSanitizedString(float Value)
+	/** Specialized for floats */
+	template<>
+	inline FString ToSanitizedString<float>(const float& Value)
 	{
 		return FString::SanitizeFloat( Value );
 	}
-};
-
-template<typename T>
-struct TTypeToString : public TTypeToString_Internal<T, TIsArithmeticType<T>::Value>
-{
-};
 
-/** A little helper to convert from a string to known numeric types **/
 
-template<typename T>
-struct TTypeFromString
-{
-	static void FromString(T& OutValue, const TCHAR* String)
-	{
-		check(0); // you are asking us to convert a string to unknown type
-	}
-};
-
-template<>
-struct TTypeFromString<int8>
-{
-	static void FromString(int8& OutValue, const TCHAR* Buffer)
-	{
-		OutValue = FCString::Atoi(Buffer);
-	}
-};
-
-template<>
-struct TTypeFromString<int16>
-{
-	static void FromString(int16& OutValue, const TCHAR* Buffer)
-	{
-		OutValue = FCString::Atoi(Buffer);
-	}
-};
-
-template<>
-struct TTypeFromString<int32>
-{
-	static void FromString(int32& OutValue, const TCHAR* Buffer)
+	/** Parse a string into this type, returning whether it was successful */
+	/** Specialization for arithmetic types */
+	template<typename T>
+	static typename TEnableIf<TIsArithmeticType<T>::Value, bool>::Type
+		TryParseString(T& OutValue, const TCHAR* Buffer)
 	{
-		OutValue = FCString::Atoi(Buffer);
-	}
-};
-
-template<>
-struct TTypeFromString<int64>
-{
-	static void FromString(int64& OutValue, const TCHAR* Buffer)
-	{
-		OutValue = FCString::Atoi64(Buffer);
-	}
-};
-
-template<>
-struct TTypeFromString<uint8>
-{
-	static void FromString(uint8& OutValue, const TCHAR* Buffer)
-	{
-		OutValue = FCString::Atoi(Buffer);
-	}
-};
-
-template<>
-struct TTypeFromString<uint16>
-{
-	static void FromString(uint16& OutValue, const TCHAR* Buffer)
-	{
-		OutValue = FCString::Atoi(Buffer);
-	}
-};
-
-template<>
-struct TTypeFromString<uint32>
-{
-	static void FromString(uint32& OutValue, const TCHAR* Buffer)
-	{
-		OutValue = FCString::Atoi64(Buffer); //64 because this unsigned and so Atoi might overflow
-	}
-};
-
-template<>
-struct TTypeFromString<uint64>
-{
-	static void FromString(uint64& OutValue, const TCHAR* Buffer)
-	{
-		OutValue =  FCString::Strtoui64(Buffer, NULL, 0);
+		if (FCString::IsNumeric(Buffer))
+		{
+			FromString(OutValue, Buffer);
+			return true;
+		}
+		return false;
 	}
-};
-
-template<>
-struct TTypeFromString<float>
-{
-	static void FromString(float& OutValue, const TCHAR* Buffer)
+	
+	/** Try and parse a bool - always returns true */
+	static bool TryParseString(bool& OutValue, const TCHAR* Buffer)
 	{
-		OutValue = FCString::Atof(Buffer);
+		FromString(OutValue, Buffer);
+		return true;
 	}
-};
+}
 
-template<>
-struct TTypeFromString<double>
+/** Shorthand legacy use for LexicalConversion functions */
+template<typename T>
+struct TTypeToString
 {
-	static void FromString(double& OutValue, const TCHAR* Buffer)
-	{
-		OutValue = FCString::Atod(Buffer);
-	}
+	static FString ToString(const T& Value)				{ return LexicalConversion::ToString(Value); }
+	static FString ToSanitizedString(const T& Value)	{ return LexicalConversion::ToSanitizedString(Value); }
 };
-
-template<>
-struct TTypeFromString<bool>
+template<typename T>
+struct TTypeFromString
 {
-	static void FromString(bool& OutValue, const TCHAR* Buffer)
-	{
-		OutValue = FCString::ToBool(Buffer);
-	}
+	static void FromString(T& Value, const TCHAR* Buffer) { return LexicalConversion::FromString(Value, Buffer); }
 };
 
-
 /*----------------------------------------------------------------------------
 	Special archivers.
 ----------------------------------------------------------------------------*/
diff --git a/Engine/Source/Runtime/Core/Public/Misc/CString.h b/Engine/Source/Runtime/Core/Public/Misc/CString.h
--- a/Engine/Source/Runtime/Core/Public/Misc/CString.h
+++ b/Engine/Source/Runtime/Core/Public/Misc/CString.h
@@ -30,6 +30,39 @@ struct TCString
 	 * @param Str - string that will be checked
 	 **/
 	static FORCEINLINE bool IsPureAnsi(const CharType* Str);
+
+	/**
+	 * Returns whether this string contains only numeric characters 
+	 * @param Str - string that will be checked
+	 **/
+	static bool IsNumeric(const CharType* Str)
+	{
+		if (*Str == '-' || *Str == '+')
+		{
+			Str++;
+		}
+
+		bool bHasDot = false;
+		while (*Str != '\0')
+		{
+			if (*Str == '.')
+			{
+				if (bHasDot)
+				{
+					return false;
+				}
+				bHasDot = true;
+			}
+			else if (!FChar::IsDigit(*Str))
+			{
+				return false;
+			}
+			
+			++Str;
+		}
+
+		return true;
+	}
 	
 	/**
 	 * strcpy wrapper
diff --git a/Engine/Source/Runtime/Core/Public/Stats/StatsData.h b/Engine/Source/Runtime/Core/Public/Stats/StatsData.h
--- a/Engine/Source/Runtime/Core/Public/Stats/StatsData.h
+++ b/Engine/Source/Runtime/Core/Public/Stats/StatsData.h
@@ -37,22 +37,25 @@ struct CORE_API FStatConstants
 	static const FString ThreadNameMarker;
 };
 
-template<>
-struct TTypeFromString<FName>
+namespace LexicalConversion
 {
-	static void FromString( FName& OutValue, const TCHAR* Buffer )
+	inline void FromString(FName& OutValue, const TCHAR* Buffer )
 	{
 		OutValue = FName( Buffer );
 	}
-};
+}
 
+/** Parse a typed value into the specified out parameter.
+ * 	Expects to find a FromString function that takes a reference to T. Defaults are provided in the LexicalConversion namespace.
+ */
 template <typename T>
 void ParseTypedValue( const TCHAR* Stream, const TCHAR* Match, T& Out )
 {
 	TCHAR Temp[64] = TEXT( "" );
 	if( FParse::Value( Stream, Match, Temp, ARRAY_COUNT( Temp ) ) )
 	{
-		TTypeFromString<T>::FromString( Out, Temp );
+		using namespace LexicalConversion;
+		FromString( Out, Temp );
 	}
 }
 
diff --git a/Engine/Source/Runtime/CoreUObject/Public/UObject/UnrealType.h b/Engine/Source/Runtime/CoreUObject/Public/UObject/UnrealType.h
--- a/Engine/Source/Runtime/CoreUObject/Public/UObject/UnrealType.h
+++ b/Engine/Source/Runtime/CoreUObject/Public/UObject/UnrealType.h
@@ -865,18 +865,6 @@ public:
 		return TNameOf<TCppType>::GetName();
 	}
 
-	/** Get the value of the property at an address to a string, so far only numeric types are supported */
-	static FString ToString(void const* A)
-	{
-		return TTypeToString<TCppType>::ToString(GetPropertyValue(A));
-	}
-
-	/** Set the value of the property at an address from a string, so far only numeric types are supported */
-	static void FromString(void* A, const TCHAR* Buffer)
-	{
-		return TTypeFromString<TCppType>::FromString(*GetPropertyValuePtr(A), Buffer);
-	}
-
 	/** Convert the address of a value of the property to the proper type */
 	static FORCEINLINE TCppType const* GetPropertyValuePtr(void const* A)
 	{
@@ -1239,11 +1227,11 @@ public:
 	}
 	virtual void SetNumericPropertyValueFromString(void* Data, TCHAR const* Value) const override
 	{
-		TTypeFundamentals::FromString(Data, Value);
+		LexicalConversion::FromString(*TTypeFundamentals::GetPropertyValuePtr(Data), Value);
 	}
 	virtual FString GetNumericPropertyValueToString(void const* Data) const override
 	{
-		return TTypeFundamentals::ToString(Data);
+		return LexicalConversion::ToString(TTypeFundamentals::GetPropertyValue(Data));
 	}
 	virtual int64 GetSignedIntPropertyValue(void const* Data) const override
 	{
diff --git a/Engine/Source/Runtime/Engine/Private/StatsRender2.cpp b/Engine/Source/Runtime/Engine/Private/StatsRender2.cpp
--- a/Engine/Source/Runtime/Engine/Private/StatsRender2.cpp
+++ b/Engine/Source/Runtime/Engine/Private/StatsRender2.cpp
@@ -56,19 +56,15 @@ static FStatFont GStatFonts[(int32)EStatFontTypes::NumFonts] =
 	FStatFont( 48, 12, 2 ),
 };
 
-template<>
-struct TTypeFromString<EStatFontTypes>
+void FromString( EStatFontTypes& OutValue, const TCHAR* Buffer )
 {
-	static void FromString( EStatFontTypes& OutValue, const TCHAR* Buffer )
-	{
-		OutValue = EStatFontTypes::Small;
+	OutValue = EStatFontTypes::Small;
 
-		if( FCString::Stricmp( Buffer, TEXT( "Tiny" ) ) == 0 )
-		{
-			OutValue = EStatFontTypes::Tiny;
-		}
+	if( FCString::Stricmp( Buffer, TEXT( "Tiny" ) ) == 0 )
+	{
+		OutValue = EStatFontTypes::Tiny;
 	}
-};
+}
 
 /** Holds various parameters used for rendering stats. */
 struct FStatRenderGlobals
