Added ability to specify units on UProperty types through a Units= meta tag
Added ability to adaptively change display units to a more appropriate range on numerical property editors. This behaviour can be disabled with a AllowUnitRangeAdaption=False meta tag.
Spin boxes now accept a numeric type interface which allows customization of numeric handling.

diff --git a/Engine/Source/Editor/EditorStyle/Private/EditorStyleClasses.cpp b/Engine/Source/Editor/EditorStyle/Private/EditorStyleClasses.cpp
--- a/Engine/Source/Editor/EditorStyle/Private/EditorStyleClasses.cpp
+++ b/Engine/Source/Editor/EditorStyle/Private/EditorStyleClasses.cpp
@@ -1,7 +1,7 @@
 // Copyright 1998-2015 Epic Games, Inc. All Rights Reserved.
 
 #include "EditorStylePrivatePCH.h"
-
+#include "UnitConversion.h"
 
 /* UEditorStyleSettings interface
  *****************************************************************************/
@@ -13,6 +13,7 @@ UEditorStyleSettings::UEditorStyleSettings( const FObjectInitializer& ObjectInit
 	InactiveSelectionColor = FLinearColor(0.25f, 0.25f, 0.25f);
 	PressedSelectionColor = FLinearColor(0.701f, 0.225f, 0.003f);
 
+	bDisplayPropertyUnits = true;
 	bShowFriendlyNames = true;
 	LogTimestampMode = ELogTimes::None;
 }
@@ -26,10 +27,14 @@ void UEditorStyleSettings::PostEditChangeProperty(struct FPropertyChangedEvent&
 
 	const FName Name = (PropertyChangedEvent.Property != nullptr) ? PropertyChangedEvent.Property->GetFName() : NAME_None;
 
-	if (Name == FName(TEXT("bEnableWindowAnimations")))
+	if (Name == GET_MEMBER_NAME_CHECKED(UEditorStyleSettings, bEnableWindowAnimations))
 	{
 		FSlateApplication::Get().EnableMenuAnimations(bEnableWindowAnimations);
 	}
+	else if (Name == GET_MEMBER_NAME_CHECKED(UEditorStyleSettings, bDisplayPropertyUnits))
+	{
+		FUnitConversion::bIsUnitDisplayEnabled = bDisplayPropertyUnits;
+	}
 
 //	if (!FUnrealEdMisc::Get().IsDeletePreferences())
 	{
@@ -39,4 +44,10 @@ void UEditorStyleSettings::PostEditChangeProperty(struct FPropertyChangedEvent&
 	SettingChangedEvent.Broadcast(Name);
 }
 
+
+void UEditorStyleSettings::PostInitProperties()
+{
+	Super::PostInitProperties();
+	FUnitConversion::bIsUnitDisplayEnabled = bDisplayPropertyUnits;
+}
 #endif
diff --git a/Engine/Source/Editor/EditorStyle/Public/Classes/EditorStyleSettings.h b/Engine/Source/Editor/EditorStyle/Public/Classes/EditorStyleSettings.h
--- a/Engine/Source/Editor/EditorStyle/Public/Classes/EditorStyleSettings.h
+++ b/Engine/Source/Editor/EditorStyle/Public/Classes/EditorStyleSettings.h
@@ -70,6 +70,10 @@ public:
 	UPROPERTY(EditAnywhere, config, Category=UserInterface, meta=(DisplayName="Show Friendly Variable Names"))
 	uint32 bShowFriendlyNames:1;
 
+	/** When set, numerical units will be displayed on applicable properties */
+	UPROPERTY(EditAnywhere, config, Category=UserInterface, meta=(DisplayName="Display Units on Numerical Properties"))
+	uint32 bDisplayPropertyUnits:1;
+
 	/** When enabled, the Editor Preferences and Project Settings menu items in the main menu will be expanded with sub-menus for each settings section. */
 	UPROPERTY(EditAnywhere, config, Category=UserInterface, AdvancedDisplay)
 	uint32 bExpandConfigurationMenus:1;
@@ -94,6 +98,7 @@ protected:
 
 #if WITH_EDITOR
 	virtual void PostEditChangeProperty( struct FPropertyChangedEvent& PropertyChangedEvent ) override;
+	virtual void PostInitProperties() override;
 #endif
 
 private:
diff --git a/Engine/Source/Editor/PropertyEditor/Private/UserInterface/PropertyEditor/SPropertyEditorNumeric.h b/Engine/Source/Editor/PropertyEditor/Private/UserInterface/PropertyEditor/SPropertyEditorNumeric.h
--- a/Engine/Source/Editor/PropertyEditor/Private/UserInterface/PropertyEditor/SPropertyEditorNumeric.h
+++ b/Engine/Source/Editor/PropertyEditor/Private/UserInterface/PropertyEditor/SPropertyEditorNumeric.h
@@ -4,6 +4,9 @@
 #include "PropertyEditorConstants.h"
 #include "SNumericEntryBox.h"
 
+#include "UnitConversion.h"
+#include "NumericUnitTypeInterface.inl"
+
 #define LOCTEXT_NAMESPACE "PropertyEditor"
 
 template <typename NumericType>
@@ -89,6 +92,25 @@ public:
 		const bool bAllowSpin = (!ObjectPropertyNode || (1 == ObjectPropertyNode->GetNumObjects()))
 			&& !PropertyNode->GetProperty()->GetBoolMetaData("NoSpinbox");
 
+		// Set up the correct type interface if we want to display units on the property editor
+		TSharedPtr<INumericTypeInterface<NumericType>> TypeInterface;
+		if (FUnitConversion::bIsUnitDisplayEnabled)
+		{
+			const FString& Units = InPropertyEditor->GetProperty()->GetMetaData(TEXT("Units"));
+			auto PropertyUnits = FUnitConversion::UnitFromString(*Units);
+
+			if (PropertyUnits.IsSet())
+			{
+				bool bAllowUnitRangeAdaption = true;
+				const auto& AllowRageAdaptionString = Property->GetMetaData(TEXT("AllowUnitRangeAdaption"));
+				if (!AllowRageAdaptionString.IsEmpty())
+				{
+					LexicalConversion::FromString(bAllowUnitRangeAdaption, *AllowRageAdaptionString);
+				}
+				TypeInterface = MakeShareable( new TNumericUnitTypeInterface<NumericType>(PropertyUnits.GetValue(), bAllowUnitRangeAdaption) );
+			}
+		}
+
 		ChildSlot
 			[
 				SAssignNew(PrimaryWidget, SNumericEntryBox<NumericType>)
@@ -107,6 +129,7 @@ public:
 				.OnValueCommitted(this, &SPropertyEditorNumeric<NumericType>::OnValueCommitted)
 				.OnBeginSliderMovement(this, &SPropertyEditorNumeric<NumericType>::OnBeginSliderMovement)
 				.OnEndSliderMovement(this, &SPropertyEditorNumeric<NumericType>::OnEndSliderMovement)
+				.TypeInterface(TypeInterface)
 			];
 
 		SetEnabled(TAttribute<bool>(this, &SPropertyEditorNumeric<NumericType>::CanEdit));
diff --git a/Engine/Source/Programs/UnrealHeaderTool/Private/HeaderParser.cpp b/Engine/Source/Programs/UnrealHeaderTool/Private/HeaderParser.cpp
--- a/Engine/Source/Programs/UnrealHeaderTool/Private/HeaderParser.cpp
+++ b/Engine/Source/Programs/UnrealHeaderTool/Private/HeaderParser.cpp
@@ -11,6 +11,7 @@
 #include "DefaultValueHelper.h"
 #include "IScriptGeneratorPluginInterface.h"
 #include "Manifest.h"
+#include "UnitConversion.h"
 
 /*-----------------------------------------------------------------------------
 	Constants & declarations.
@@ -5770,6 +5771,18 @@ void FHeaderParser::ValidateMetaDataFormat(UField* Field, const FString& InKey,
 			FError::Throwf(TEXT("'%s' metadata was '%s' but it must be %s or %s"), *InKey, *InValue, *ExperimentalValue, *EarlyAccessValue);
 		}
 	}
+	else if (InKey == TEXT("Units"))
+	{
+		// Check for numeric property
+		if (!Cast<UNumericProperty>(Field))
+		{
+			FError::Throwf(TEXT("'Units' meta data can only be applied to numeric properties"));
+		}
+		else if (!FUnitConversion::UnitFromString(*InValue))
+		{
+			FError::Throwf(TEXT("Unrecognized units (%s) specified for numeric property '%s'"), *InValue, *Field->GetDisplayNameText().ToString());
+		}
+	}
 }
 
 // Validates the metadata, then adds it to the class data
diff --git a/Engine/Source/Runtime/Core/Private/Math/UnitConversion.cpp b/Engine/Source/Runtime/Core/Private/Math/UnitConversion.cpp
new file mode 100644
--- /dev/null
+++ b/Engine/Source/Runtime/Core/Private/Math/UnitConversion.cpp
@@ -0,0 +1,359 @@
+// Copyright 1998-2015 Epic Games, Inc. All Rights Reserved.
+
+#include "CorePrivatePCH.h"
+#include "UnitConversion.h"
+
+FUnitConversion::FParseCandidate FUnitConversion::ParseCandidates[] = {
+	
+	{ TEXT("Micrometers"),			EUnit::Micrometers },			{ TEXT("um"),		EUnit::Micrometers }, 			{ TEXT("\u00B5m"),	EUnit::Micrometers },
+	{ TEXT("Millimeters"),			EUnit::Millimeters },			{ TEXT("mm"),		EUnit::Millimeters },
+	{ TEXT("Centimeters"),			EUnit::Centimeters },			{ TEXT("cm"),		EUnit::Centimeters },
+	{ TEXT("Meters"),				EUnit::Meters },				{ TEXT("m"),		EUnit::Meters },
+	{ TEXT("Kilometers"),			EUnit::Kilometers },			{ TEXT("km"),		EUnit::Kilometers },
+	{ TEXT("Inches"),				EUnit::Inches },				{ TEXT("in"),		EUnit::Inches },
+	{ TEXT("Feet"),					EUnit::Feet },					{ TEXT("ft"),		EUnit::Feet },
+	{ TEXT("Yards"),				EUnit::Yards },					{ TEXT("yd"),		EUnit::Yards },
+	{ TEXT("Miles"),				EUnit::Miles },					{ TEXT("mi"),		EUnit::Miles },
+	{ TEXT("Lightyears"),			EUnit::Lightyears },			{ TEXT("ly"),		EUnit::Lightyears },
+
+	{ TEXT("Degrees"),				EUnit::Degrees },				{ TEXT("deg"),		EUnit::Degrees },				{ TEXT("\u00B0"),	EUnit::Degrees },
+	{ TEXT("Radians"),				EUnit::Radians },				{ TEXT("rad"),		EUnit::Radians },
+		
+	{ TEXT("MetersPerSecond"),		EUnit::MetersPerSecond },		{ TEXT("m/s"),		EUnit::MetersPerSecond },
+	{ TEXT("KilometersPerHour"),	EUnit::KilometersPerHour },		{ TEXT("km/h"),		EUnit::KilometersPerHour },		{ TEXT("kmph"),		EUnit::KilometersPerHour },
+	{ TEXT("MilesPerHour"),			EUnit::MilesPerHour },			{ TEXT("mi/h"),		EUnit::MilesPerHour },			{ TEXT("mph"),		EUnit::MilesPerHour },
+		
+	{ TEXT("Celsius"),				EUnit::Celsius },				{ TEXT("C"),		EUnit::Celsius },				{ TEXT("degC"),		EUnit::Celsius },			{ TEXT("\u00B0C"),		EUnit::Celsius },
+	{ TEXT("Farenheit"),			EUnit::Farenheit },				{ TEXT("F"),		EUnit::Farenheit },				{ TEXT("degF"),		EUnit::Farenheit },			{ TEXT("\u00B0F"),		EUnit::Farenheit },
+	{ TEXT("Kelvin"),				EUnit::Kelvin },				{ TEXT("K"),		EUnit::Kelvin },
+					
+	{ TEXT("Micrograms"),			EUnit::Micrograms },			{ TEXT("ug"),		EUnit::Micrograms }, 			{ TEXT("\u00B5g"),		EUnit::Micrograms },
+	{ TEXT("Milligrams"),			EUnit::Milligrams },			{ TEXT("mg"),		EUnit::Milligrams },
+	{ TEXT("Grams"),				EUnit::Grams },					{ TEXT("g"),		EUnit::Grams },
+	{ TEXT("Kilograms"),			EUnit::Kilograms },				{ TEXT("kg"),		EUnit::Kilograms },
+	{ TEXT("MetricTons"),			EUnit::MetricTons },			{ TEXT("t"),		EUnit::MetricTons },
+	{ TEXT("Ounces"),				EUnit::Ounces },				{ TEXT("oz"),		EUnit::Ounces },
+	{ TEXT("Pounds"),				EUnit::Pounds },				{ TEXT("lb"),		EUnit::Pounds },
+	{ TEXT("Stones"),				EUnit::Stones },				{ TEXT("st"),		EUnit::Stones },
+
+	{ TEXT("Newtons"),				EUnit::Newtons },				{ TEXT("N"),		EUnit::Newtons },
+	{ TEXT("PoundsForce"),			EUnit::PoundsForce },			{ TEXT("lbf"),		EUnit::PoundsForce },
+	{ TEXT("KilogramsForce"),		EUnit::KilogramsForce },		{ TEXT("kgf"),		EUnit::KilogramsForce },
+
+	{ TEXT("Hertz"),				EUnit::Hertz },					{ TEXT("Hz"),		EUnit::Hertz },
+	{ TEXT("Kilohertz"),			EUnit::Kilohertz },				{ TEXT("KHz"),		EUnit::Kilohertz },
+	{ TEXT("Megahertz"),			EUnit::Megahertz },				{ TEXT("MHz"),		EUnit::Megahertz },
+	{ TEXT("Gigahertz"),			EUnit::Gigahertz },				{ TEXT("GHz"),		EUnit::Gigahertz },
+	{ TEXT("RevolutionsPerMinute"),	EUnit::RevolutionsPerMinute },	{ TEXT("rpm"),		EUnit::RevolutionsPerMinute },
+
+	{ TEXT("Bytes"),				EUnit::Bytes },					{ TEXT("B"),		EUnit::Bytes },
+	{ TEXT("Kilobytes"),			EUnit::Kilobytes },				{ TEXT("KB"),		EUnit::Kilobytes },
+	{ TEXT("Megabytes"),			EUnit::Megabytes },				{ TEXT("MB"),		EUnit::Megabytes },
+	{ TEXT("Gigabytes"),			EUnit::Gigabytes },				{ TEXT("GB"),		EUnit::Gigabytes },
+	{ TEXT("Terabytes"),			EUnit::Terabytes },				{ TEXT("TB"),		EUnit::Terabytes },
+
+	{ TEXT("Lumens"),				EUnit::Lumens },				{ TEXT("lm"),		EUnit::Lumens },
+
+	{ TEXT("Milliseconds"),			EUnit::Milliseconds },			{ TEXT("ms"),		EUnit::Milliseconds },
+	{ TEXT("Seconds"),				EUnit::Seconds },				{ TEXT("s"),		EUnit::Seconds },
+	{ TEXT("Minutes"),				EUnit::Minutes },				{ TEXT("min"),		EUnit::Minutes },
+	{ TEXT("Hours"),				EUnit::Hours },					{ TEXT("hrs"),		EUnit::Hours },
+	{ TEXT("Days"),					EUnit::Days },					{ TEXT("dy"),		EUnit::Days },
+	{ TEXT("Months"),				EUnit::Months },				{ TEXT("mth"),		EUnit::Months },
+	{ TEXT("Years"),				EUnit::Years },					{ TEXT("yr"),		EUnit::Years },
+};
+
+const TCHAR* const FUnitConversion::DisplayStrings[] = {
+	TEXT("\u00B5m"),			TEXT("mm"),					TEXT("cm"),					TEXT("m"),					TEXT("km"),
+	TEXT("in"),					TEXT("ft"),					TEXT("yd"),					TEXT("mi"),
+	TEXT("ly"),
+
+	TEXT("\u00B0"), TEXT("rad"),
+
+	TEXT("m/s"), TEXT("km/h"), TEXT("mi/h"),
+
+	TEXT("\u00B0C"), TEXT("\u00B0F"), TEXT("K"),
+
+	TEXT("\u00B5g"), TEXT("mg"), TEXT("g"), TEXT("kg"), TEXT("t"),
+	TEXT("oz"), TEXT("lb"), TEXT("st"),
+
+	TEXT("N"), TEXT("lbf"), TEXT("kgf"),
+
+	TEXT("Hz"), TEXT("KHz"), TEXT("MHz"), TEXT("GHz"), TEXT("rpm"),
+
+	TEXT("B"), TEXT("KB"), TEXT("MB"), TEXT("GB"), TEXT("TB"),
+
+	TEXT("lm"),
+
+	TEXT("ms"), TEXT("s"), TEXT("min"), TEXT("hr"), TEXT("dy"), TEXT("mth"), TEXT("yr"),
+};
+
+const FUnitConversion::EUnitType FUnitConversion::Types[] = {
+	EUnitType::Distance,	EUnitType::Distance,	EUnitType::Distance,	EUnitType::Distance,	EUnitType::Distance,
+	EUnitType::Distance,	EUnitType::Distance,	EUnitType::Distance,	EUnitType::Distance,
+	EUnitType::Distance,
+
+	EUnitType::Angle,		EUnitType::Angle,
+
+	EUnitType::Speed,		EUnitType::Speed, 		EUnitType::Speed,
+
+	EUnitType::Temperature,	EUnitType::Temperature,	EUnitType::Temperature,
+
+	EUnitType::Mass,		EUnitType::Mass,		EUnitType::Mass,		EUnitType::Mass,		EUnitType::Mass,
+	EUnitType::Mass,		EUnitType::Mass,		EUnitType::Mass,
+
+	EUnitType::Force,		EUnitType::Force,		EUnitType::Force,
+
+	EUnitType::Frequency,	EUnitType::Frequency,	EUnitType::Frequency,	EUnitType::Frequency,	EUnitType::Frequency,
+
+	EUnitType::DataSize,	EUnitType::DataSize,	EUnitType::DataSize,	EUnitType::DataSize,	EUnitType::DataSize,
+
+	EUnitType::LuminousFlux,
+
+	EUnitType::Time,		EUnitType::Time,		EUnitType::Time,		EUnitType::Time,		EUnitType::Time,		EUnitType::Time,		EUnitType::Time,
+};
+
+/** Globally check whether we should display units in on user interfaces */
+bool FUnitConversion::bIsUnitDisplayEnabled = true;
+
+/** Get the unit abbreviation the specified unit type */
+const TCHAR* FUnitConversion::GetUnitDisplayString(EUnit Unit)
+{
+	static_assert(ARRAY_COUNT(FUnitConversion::Types) == (uint8)EUnit::Unspecified, "Type array does not match size of unit enum");
+	static_assert(ARRAY_COUNT(FUnitConversion::DisplayStrings) == (uint8)EUnit::Unspecified, "Display String array does not match size of unit enum");
+	
+	if (Unit != EUnit::Unspecified)
+	{
+		return DisplayStrings[(uint8)Unit];
+	}
+	return nullptr;
+}
+
+/** Helper function to find a unit from a string (name or abbreviation) */
+TOptional<EUnit> FUnitConversion::UnitFromString(const TCHAR* UnitString)
+{
+	if (!UnitString || *UnitString == '\0')
+	{
+		return TOptional<EUnit>();
+	}
+
+	for (int32 Index = 0; Index < ARRAY_COUNT(ParseCandidates); ++Index)
+	{
+		if (FCString::Stricmp(UnitString, ParseCandidates[Index].String) == 0)
+		{
+			return ParseCandidates[Index].Unit;
+		}
+	}
+
+	return TOptional<EUnit>();
+}
+
+double FUnitConversion::DistanceUnificationFactor(EUnit From)
+{
+	// Convert to meters
+	double Factor = 1;
+	switch (From)
+	{
+		case EUnit::Micrometers:		return 0.000001;
+		case EUnit::Millimeters:		return 0.001;
+		case EUnit::Centimeters:		return 0.01;
+		case EUnit::Kilometers:			return 1000;
+
+		case EUnit::Lightyears:			return 9.4605284e15;
+
+		case EUnit::Miles:				Factor *= 1760;				// fallthrough
+		case EUnit::Yards:				Factor *= 3;				// fallthrough
+		case EUnit::Feet:				Factor *= 12;				// fallthrough
+		case EUnit::Inches:				Factor /= 39.3700787;		// fallthrough
+		default: 						return Factor;				// return
+	}
+}
+
+double FUnitConversion::AngleUnificationFactor(EUnit From)
+{
+	// Convert to degrees
+	switch (From)
+	{
+		case EUnit::Radians:			return (180 / PI);
+		default: 						return 1;
+	}
+}
+
+double FUnitConversion::SpeedUnificationFactor(EUnit From)
+{
+	// Convert to km/h
+	switch (From)
+	{
+		case EUnit::MetersPerSecond:	return 3.6;
+		case EUnit::MilesPerHour:		return DistanceUnificationFactor(EUnit::Miles) / 1000;
+		default: 						return 1;
+	}
+}
+
+double FUnitConversion::MassUnificationFactor(EUnit From)
+{
+	double Factor = 1;
+	// Convert to grams
+	switch (From)
+	{
+		case EUnit::Micrograms:			return 0.000001;
+		case EUnit::Milligrams:			return 0.001;
+		case EUnit::Kilograms:			return 1000;
+		case EUnit::MetricTons:			return 1000000;
+
+		case EUnit::Stones:				Factor *= 14;		// fallthrough
+		case EUnit::Pounds:				Factor *= 16;		// fallthrough
+		case EUnit::Ounces:				Factor *= 28.3495;	// fallthrough
+		default: 						return Factor;		// return
+	}
+}
+
+double FUnitConversion::ForceUnificationFactor(EUnit From)
+{
+	// Convert to Newtons
+	switch (From)
+	{
+		case EUnit::PoundsForce:		return 4.44822162;
+		case EUnit::KilogramsForce:		return 9.80665;
+		default: 						return 1;
+	}
+}
+
+double FUnitConversion::FrequencyUnificationFactor(EUnit From)
+{
+	// Convert to KHz
+	switch (From)
+	{
+		case EUnit::Hertz:				return 0.001;
+		case EUnit::Megahertz:			return 1000;
+		case EUnit::Gigahertz:			return 1000000;
+
+		case EUnit::RevolutionsPerMinute:	return 0.001/60;
+
+		default: 						return 1;
+	}
+}
+
+double FUnitConversion::DataSizeUnificationFactor(EUnit From)
+{
+	// Convert to MB
+	switch (From)
+	{
+		case EUnit::Bytes:				return 1.0/(1024*1024);
+		case EUnit::Kilobytes:			return 1.0/1024;
+		case EUnit::Gigabytes:			return 1024;
+		case EUnit::Terabytes:			return 1024*1024;
+
+		default: 						return 1;
+	}
+}
+
+double FUnitConversion::TimeUnificationFactor(EUnit From)
+{
+	// Convert to hours
+	double Factor = 1;
+	switch (From)
+	{
+		case EUnit::Months:				return (365.242 * 24) / 12;
+
+		case EUnit::Years:				Factor *= 365.242;	// fallthrough
+		case EUnit::Days:				Factor *= 24;		// fallthrough
+										return Factor;
+
+		case EUnit::Milliseconds:		Factor /= 1000;		// fallthrough
+		case EUnit::Seconds:			Factor /= 60;		// fallthrough
+		case EUnit::Minutes:			Factor /= 60;		// fallthrough
+										return Factor;
+
+		default: 						return 1;
+	}
+}
+
+TOptional<const TArray<FUnitConversion::FQuantizationInfo>*> FUnitConversion::GetQuantizationBounds(EUnit Unit)
+{
+	struct FStaticBounds
+	{
+		TArray<FQuantizationInfo> MetricDistance;
+		TArray<FQuantizationInfo> ImperialDistance;
+
+		TArray<FQuantizationInfo> MetricMass;
+		TArray<FQuantizationInfo> ImperialMass;
+
+		TArray<FQuantizationInfo> Frequency;
+		TArray<FQuantizationInfo> DataSize;
+
+		TArray<FQuantizationInfo> Time;
+
+		FStaticBounds()
+		{
+			MetricDistance.Emplace(EUnit::Micrometers,	1000);
+			MetricDistance.Emplace(EUnit::Millimeters,	10);
+			MetricDistance.Emplace(EUnit::Centimeters,	100);
+			MetricDistance.Emplace(EUnit::Meters,			1000);
+			MetricDistance.Emplace(EUnit::Kilometers,		0);
+
+			ImperialDistance.Emplace(EUnit::Inches,	12);
+			ImperialDistance.Emplace(EUnit::Feet,		3);
+			ImperialDistance.Emplace(EUnit::Yards,	1760);
+			ImperialDistance.Emplace(EUnit::Miles,	0);
+
+			MetricMass.Emplace(EUnit::Micrograms,	1000);
+			MetricMass.Emplace(EUnit::Milligrams,	1000);
+			MetricMass.Emplace(EUnit::Grams,		1000);
+			MetricMass.Emplace(EUnit::Kilograms,	1000);
+			MetricMass.Emplace(EUnit::MetricTons,	0);
+
+			ImperialMass.Emplace(EUnit::Ounces,	16);
+			ImperialMass.Emplace(EUnit::Pounds,	14);
+			ImperialMass.Emplace(EUnit::Stones,	0);
+
+			Frequency.Emplace(EUnit::Hertz,		1000);
+			Frequency.Emplace(EUnit::Kilohertz,	1000);
+			Frequency.Emplace(EUnit::Megahertz,	1000);
+			Frequency.Emplace(EUnit::Gigahertz,	0);
+
+			DataSize.Emplace(EUnit::Bytes,		1000);
+			DataSize.Emplace(EUnit::Kilobytes,	1000);
+			DataSize.Emplace(EUnit::Megabytes,	1000);
+			DataSize.Emplace(EUnit::Gigabytes,	1000);
+			DataSize.Emplace(EUnit::Terabytes,	0);
+
+			Time.Emplace(EUnit::Milliseconds,		1000);
+			Time.Emplace(EUnit::Seconds,			60);
+			Time.Emplace(EUnit::Minutes,			60);
+			Time.Emplace(EUnit::Hours,			24);
+			Time.Emplace(EUnit::Days,				365.242f / 12);
+			Time.Emplace(EUnit::Months,			12);
+			Time.Emplace(EUnit::Years,			0);
+		}
+	};
+
+	static FStaticBounds Bounds;
+
+	switch (Unit)
+	{
+	case EUnit::Micrometers: case EUnit::Millimeters: case EUnit::Centimeters: case EUnit::Meters: case EUnit::Kilometers:
+		return &Bounds.MetricDistance;
+
+	case EUnit::Inches: case EUnit::Feet: case EUnit::Yards: case EUnit::Miles:
+		return &Bounds.ImperialDistance;
+
+	case EUnit::Micrograms: case EUnit::Milligrams: case EUnit::Grams: case EUnit::Kilograms: case EUnit::MetricTons:
+		return &Bounds.MetricMass;
+
+	case EUnit::Ounces: case EUnit::Pounds: case EUnit::Stones:
+		return &Bounds.ImperialMass;
+
+	case EUnit::Hertz: case EUnit::Kilohertz: case EUnit::Megahertz: case EUnit::Gigahertz: case EUnit::RevolutionsPerMinute:
+		return &Bounds.Frequency;
+
+	case EUnit::Bytes: case EUnit::Kilobytes: case EUnit::Megabytes: case EUnit::Gigabytes: case EUnit::Terabytes:
+		return &Bounds.DataSize;
+
+	case EUnit::Milliseconds: case EUnit::Seconds: case EUnit::Minutes: case EUnit::Hours: case EUnit::Days: case EUnit::Months: case EUnit::Years:
+		return &Bounds.Time;
+
+	default:
+		return TOptional<const TArray<FQuantizationInfo>*>();
+	}
+}
\ No newline at end of file
diff --git a/Engine/Source/Runtime/Core/Private/Tests/Math/UnitConversionTests.cpp b/Engine/Source/Runtime/Core/Private/Tests/Math/UnitConversionTests.cpp
new file mode 100644
--- /dev/null
+++ b/Engine/Source/Runtime/Core/Private/Tests/Math/UnitConversionTests.cpp
@@ -0,0 +1,100 @@
+// Copyright 1998-2015 Epic Games, Inc. All Rights Reserved.
+
+#include "CorePrivatePCH.h"
+#include "UnitConversion.h"
+
+
+IMPLEMENT_SIMPLE_AUTOMATION_TEST(FUnitUnitTests, "Core.Math.Unit Conversion", EAutomationTestFlags::ATF_SmokeTest)
+
+bool IsRoughlyEqual(double One, double Two, float Epsilon)
+{
+	return FMath::Abs(One-Two) < Epsilon;
+}
+
+bool FUnitUnitTests::RunTest(const FString& Parameters)
+{
+	struct FTestStruct
+	{
+		double Source;
+		double ExpectedResult;
+		double AccuracyEpsilon;
+
+		EUnit FromUnit, ToUnit;
+	};
+	static FTestStruct Tests[] = {
+		{ 0.025,	2.3651e11,	1e7, 	EUnit::Lightyears, 			EUnit::Kilometers	 		},
+		{ 0.5,		80467.2,	0.1,	EUnit::Miles, 				EUnit::Centimeters	 		},
+		{ 0.2,		182.88,		0.01,	EUnit::Yards, 				EUnit::Millimeters	 		},
+		{ 0.2,		60960,		0.1,	EUnit::Feet, 				EUnit::Micrometers	 		},
+		{ 10,		0.254,		0.001,	EUnit::Inches, 				EUnit::Meters		 		},
+		{ 0.75,		2460.6299,	1e-4,	EUnit::Kilometers, 			EUnit::Feet					},
+		{ 1,		39.37,		0.01,	EUnit::Meters, 				EUnit::Inches				},
+		{ 2750,		27.5,		1e-6,	EUnit::Centimeters, 		EUnit::Meters				},
+		{ 1000,		1.0936,		1e-4,	EUnit::Millimeters, 		EUnit::Yards		 		},
+		{ 2000,		0.0787,		1e-4,	EUnit::Micrometers, 		EUnit::Inches		 		},
+
+		{ 90,		PI/2,		1e-6,	EUnit::Degrees, 			EUnit::Radians		 		},
+		{ PI,		180,		1e-6,	EUnit::Radians, 			EUnit::Degrees		 		},
+
+		{ 12,		43.2,		0.1,	EUnit::MetersPerSecond,		EUnit::KilometersPerHour	},
+		{ 1,		0.6214,		1e-4,	EUnit::KilometersPerHour,	EUnit::MilesPerHour			},
+		{ 15,		6.7056,		1e-4,	EUnit::MilesPerHour,		EUnit::MetersPerSecond		},
+
+		{ 100, 		212,		0.1,	EUnit::Celsius,				EUnit::Farenheit			},
+		{ 400, 		477.594,	1e-3,	EUnit::Farenheit,			EUnit::Kelvin				},
+		{ 72, 		-201.15,	0.01,	EUnit::Kelvin,				EUnit::Celsius				},
+
+		{ 1000, 	3.5274e-5,	1e-6,	EUnit::Micrograms,			EUnit::Ounces,				},
+		{ 1000,		1,			0.1,	EUnit::Milligrams,			EUnit::Grams,				},
+		{ 200,		0.4409,		1e-4,	EUnit::Grams,				EUnit::Pounds,				},
+		{ 0.15,		150000,		0.1,	EUnit::Kilograms,			EUnit::Milligrams,			},
+		{ 1,		157.473,	1e-3,	EUnit::MetricTons,			EUnit::Stones,				},
+		{ 0.001,	28349.5,	0.1,	EUnit::Ounces,				EUnit::Micrograms,			},
+		{ 500,		226.796,	1e-3,	EUnit::Pounds,				EUnit::Kilograms,			},
+		{ 100,		0.6350,		1e-4,	EUnit::Stones,				EUnit::MetricTons,			},
+
+		{ 100,		10.1972,	1e-4,	EUnit::Newtons,				EUnit::KilogramsForce,		},
+		{ 15,		66.7233,	1e-4,	EUnit::PoundsForce,			EUnit::Newtons,				},
+		{ 2,		4.4092,		1e-4,	EUnit::KilogramsForce,		EUnit::PoundsForce,			},
+		
+		{ 1000,		1,			0.1,	EUnit::Hertz,				EUnit::Kilohertz,			},
+		{ 0.25,		250*60,		1e-3,	EUnit::Kilohertz,			EUnit::RevolutionsPerMinute,},
+		{ 1000,		1,			1e-3,	EUnit::Megahertz,			EUnit::Gigahertz,			},
+		{ 0.001,	1000000,	1e-3,	EUnit::Gigahertz,			EUnit::Hertz,				},
+		{ 100,		100.0/60,	1e-3,	EUnit::RevolutionsPerMinute,EUnit::Hertz,				},
+		
+		{ 1024,		1,			1e-3,	EUnit::Bytes,				EUnit::Kilobytes,			},
+		{ 1.5,		1536,		1e-3,	EUnit::Kilobytes,			EUnit::Bytes,				},
+		{ 1000,		9.5367e-4,	1e-5,	EUnit::Megabytes,			EUnit::Terabytes,			},
+		{ 0.5,		512,		1e-3,	EUnit::Gigabytes,			EUnit::Megabytes,			},
+		{ 0.25,		256,		1e-3,	EUnit::Terabytes,			EUnit::Gigabytes,			},
+		
+		{ 10000,	0.166667,	1e-6,	EUnit::Milliseconds,		EUnit::Minutes,				},
+		{ 0.5,		500,		1e-6,	EUnit::Seconds,				EUnit::Milliseconds,		},
+		{ 30,		60*30,		1e-6,	EUnit::Minutes,				EUnit::Seconds,				},
+		{ 5,		5.0/24,		1e-6,	EUnit::Hours,				EUnit::Days,				},
+		{ 0.75,		18,			1e-6,	EUnit::Days,				EUnit::Hours,				},
+		{ 3,		0.25,		1e-6,	EUnit::Months,				EUnit::Years,				},
+		{ 0.5,		6,			1e-6,	EUnit::Years,				EUnit::Months,				},
+
+	};
+
+	bool bSuccess = true;
+	for (auto& Test : Tests)
+	{
+		const double ActualResult = FUnitConversion::Convert(Test.Source, Test.FromUnit, Test.ToUnit);
+		if (!IsRoughlyEqual(ActualResult, Test.ExpectedResult, Test.AccuracyEpsilon))
+		{
+			bSuccess = false;
+
+			const TCHAR* FromUnitString	= FUnitConversion::GetUnitDisplayString(Test.FromUnit);
+			const TCHAR* ToUnitString 	= FUnitConversion::GetUnitDisplayString(Test.ToUnit);
+
+			AddError(FString::Printf(TEXT("Conversion from %s to %s was incorrect. %.3f%s !~= %.3f%s"),
+				FromUnitString, ToUnitString,
+				Test.Source, FromUnitString, ActualResult, ToUnitString)
+			);
+		}
+	}
+	return bSuccess;
+}
\ No newline at end of file
diff --git a/Engine/Source/Runtime/Core/Public/Math/NumericLimits.h b/Engine/Source/Runtime/Core/Public/Math/NumericLimits.h
--- a/Engine/Source/Runtime/Core/Public/Math/NumericLimits.h
+++ b/Engine/Source/Runtime/Core/Public/Math/NumericLimits.h
@@ -34,38 +34,10 @@
  *****************************************************************************/
 
 /**
- * Helper class to maps numeric type to it's limits
+ * Helper class to map a numeric type to its limits
  */
 template <typename NumericType>
-struct TNumericLimits 
-{
-	/*
-	 *	@return minimum value
-	 */
-	static NumericType Min()
-	{
-		check(0); // min value of unknown type
-		return 0;
-	}
-	
-	/*
-	 *	@return maximum value
-	 */
-	static NumericType Max()
-	{
-		check(0); // max value of unknown type		
-		return 0;
-	}
-
-	/*
-	 *	@return most negative value
-	 */
-	static NumericType Lowest()
-	{
-		check(0); // lowest value of unknown type		
-		return 0;
-	}
-};
+struct TNumericLimits;
 
 
 /**
diff --git a/Engine/Source/Runtime/Core/Public/Math/UnitConversion.h b/Engine/Source/Runtime/Core/Public/Math/UnitConversion.h
new file mode 100644
--- /dev/null
+++ b/Engine/Source/Runtime/Core/Public/Math/UnitConversion.h
@@ -0,0 +1,468 @@
+// Copyright 1998-2015 Epic Games, Inc. All Rights Reserved.
+
+#pragma once
+
+#include "Optional.h"
+#include "Function.h"
+
+/** Enum *must* be zero-indexed and sequential */
+/** Enum *must* match the mirrored enum that exists in KismetMathLibrary.h for the purposes of UObject reflection */
+enum class EUnit
+{
+	/** Scalar distance/length units */
+	Micrometers, Millimeters, Centimeters, Meters, Kilometers,
+	Inches, Feet, Yards, Miles,
+	Lightyears,
+
+	/** Angular units */
+	Degrees, Radians,
+
+	/** Speed units */
+	MetersPerSecond, KilometersPerHour, MilesPerHour,
+
+	/** Temperature units */
+	Celsius, Farenheit, Kelvin,
+
+	/** Mass units */
+	Micrograms, Milligrams, Grams, Kilograms, MetricTons,
+	Ounces, Pounds, Stones,
+
+	/** Force units */
+	Newtons, PoundsForce, KilogramsForce,
+
+	/** Frequency units */
+	Hertz, Kilohertz, Megahertz, Gigahertz, RevolutionsPerMinute,
+
+	/** Data Size units */
+	Bytes, Kilobytes, Megabytes, Gigabytes, Terabytes,
+
+	/** Luminous flux units */
+	Lumens,
+
+	/** Time units */
+	Milliseconds, Seconds, Minutes, Hours, Days, Months, Years,
+
+	/** Symbolic entry, not specifyable on meta data */
+	Unspecified
+};
+
+template<typename NumericType> struct FNumericUnit;
+
+struct CORE_API FUnitConversion
+{
+	/** Globally check whether we should display units in on user interfaces */
+	static bool bIsUnitDisplayEnabled;
+
+	/** Check whether it is possible to convert a number between the two specified units */
+	static bool AreUnitsCompatible(EUnit From, EUnit To)
+	{
+		return From == EUnit::Unspecified || To == EUnit::Unspecified || Types[(uint8)From] == Types[(uint8)To];
+	}
+
+	/** Convert the specified number from one unit to another. Does nothing if the units are incompatible. */
+	template<typename T>
+	static T Convert(T InValue, EUnit From, EUnit To)
+	{
+		if (!AreUnitsCompatible(From, To))
+		{
+			return InValue;
+		}
+		else if (From == EUnit::Unspecified || To == EUnit::Unspecified)
+		{
+			return InValue;
+		}
+
+		switch(Types[(uint8)From])
+		{
+			case EUnitType::Distance:			return InValue * DistanceUnificationFactor(From)		* (1.0 / DistanceUnificationFactor(To));
+			case EUnitType::Angle:				return InValue * AngleUnificationFactor(From) 			* (1.0 / AngleUnificationFactor(To));
+			case EUnitType::Speed:				return InValue * SpeedUnificationFactor(From) 			* (1.0 / SpeedUnificationFactor(To));
+			case EUnitType::Mass:				return InValue * MassUnificationFactor(From) 			* (1.0 / MassUnificationFactor(To));
+			case EUnitType::Force:				return InValue * ForceUnificationFactor(From) 			* (1.0 / ForceUnificationFactor(To));
+			case EUnitType::Frequency:			return InValue * FrequencyUnificationFactor(From) 		* (1.0 / FrequencyUnificationFactor(To));
+			case EUnitType::DataSize:			return InValue * DataSizeUnificationFactor(From) 		* (1.0 / DataSizeUnificationFactor(To));
+			case EUnitType::LuminousFlux:		return InValue;
+			case EUnitType::Time:				return InValue * TimeUnificationFactor(From) 			* (1.0 / TimeUnificationFactor(To));			
+			// Temperature conversion is not just a simple multiplication, so needs special treatment
+			case EUnitType::Temperature:
+			{
+				double NewValue = InValue;
+				// Put it into kelvin
+				switch (From)
+				{
+					case EUnit::Celsius:			NewValue = NewValue + 273.15;					break;
+					case EUnit::Farenheit:			NewValue = (NewValue + 459.67) * 5.f/9.f;		break;
+					default: 																		break;
+				}
+				// And out again
+				switch (To)
+				{
+					case EUnit::Celsius:			return NewValue - 273.15;
+					case EUnit::Farenheit:			return NewValue * 9.f/5.f - 459.67;
+					default: 						return NewValue;
+				}
+			}
+
+			default:							return InValue;
+		}
+	}
+
+	/** Get the display string for the the specified unit type */
+	static const TCHAR* GetUnitDisplayString(EUnit Unit);
+
+	/** Helper function to find a unit from a string (name or display string) */
+	static TOptional<EUnit> UnitFromString(const TCHAR* UnitString);
+
+private:
+
+	/** Find the common quantization factor for the specified distance unit. Quantizes to Meters. */
+	static double DistanceUnificationFactor(EUnit From);
+	/** Find the common quantization factor for the specified angular unit. Quantizes to Degrees. */
+	static double AngleUnificationFactor(EUnit From);
+	/** Find the common quantization factor for the specified speed unit. Quantizes to km/h. */
+	static double SpeedUnificationFactor(EUnit From);
+	/** Find the common quantization factor for the specified temperature unit. Quantizes to Kelvin. */
+	static double TemperatureUnificationFactor(EUnit From);
+	/** Find the common quantization factor for the specified mass unit. Quantizes to Grams. */
+	static double MassUnificationFactor(EUnit From);
+	/** Find the common quantization factor for the specified force unit. Quantizes to Newtons. */
+	static double ForceUnificationFactor(EUnit From);
+	/** Find the common quantization factor for the specified frequency unit. Quantizes to KHz. */
+	static double FrequencyUnificationFactor(EUnit From);
+	/** Find the common quantization factor for the specified data size unit. Quantizes to MB. */
+	static double DataSizeUnificationFactor(EUnit From);
+	/** Find the common quantization factor for the specified time unit. Quantizes to hours. */
+	static double TimeUnificationFactor(EUnit From);
+
+public:
+
+	/** Structure used to define the factor required to get from one unit type to the next. */
+	struct FQuantizationInfo
+	{
+		/** The unit to which this factor applies */
+		EUnit Units;
+		/** The factor by which to multiply to get to the next unit in this range */
+		float Factor;
+
+		/** Constructor */
+		FQuantizationInfo(EUnit InUnit, float InFactor) : Units(InUnit), Factor(InFactor) {}
+	};
+
+	/** Find the quantization bounds for the specified unit, if any */
+	static TOptional<const TArray<FQuantizationInfo>*> GetQuantizationBounds(EUnit Unit);
+
+	/** Quantizes this number to the most appropriate unit for user friendly presentation (e.g. 1000m returns 1km). */
+	template<typename T>
+	static FNumericUnit<T> QuantizeUnitsToBestFit(T Value, EUnit Units)
+	{
+		auto OptionalBounds = FUnitConversion::GetQuantizationBounds(Units);
+		if (!OptionalBounds.IsSet())
+		{
+			return FNumericUnit<T>(Value, Units);
+		}
+
+		const auto& Bounds = *OptionalBounds.GetValue();
+
+		const int32 CurrentUnitIndex = (uint8)Units - (uint8)Bounds[0].Units;
+
+		EUnit NewUnits = Units;
+		double NewValue = Value;
+
+		if (FMath::Abs(NewValue) > 1)
+		{
+			// Large number? Try larger units
+			for (int32 Index = CurrentUnitIndex; Index < Bounds.Num(); ++Index)
+			{
+				if (Bounds[Index].Factor == 0)
+				{
+					break;
+				}
+
+				const auto Tmp = NewValue / Bounds[Index].Factor;
+
+				if (FMath::Abs(Tmp) < 1)
+				{
+					break;
+				}
+
+				NewValue = Tmp;
+				NewUnits = Bounds[Index + 1].Units;
+			}
+		}
+		else if (NewValue != 0)
+		{
+			// Small number? Try smaller units
+			for (int32 Index = CurrentUnitIndex - 1; Index >= 0; --Index)
+			{
+				NewValue *= Bounds[Index].Factor;
+				NewUnits = Bounds[Index].Units;
+
+				if (FMath::Abs(NewValue) > 1)
+				{
+					break;
+				}
+			}
+		}
+
+		return FNumericUnit<T>(NewValue, NewUnits);
+	}
+
+private:
+
+	/** Private enumeration that specifies particular classes of unit */
+	enum class EUnitType
+	{
+		Distance, Angle, Speed, Temperature, Mass, Force, Frequency, DataSize, LuminousFlux, Time
+	};
+
+	/** Structure used to match units when parsing */
+	struct FParseCandidate
+	{
+		const TCHAR* String;
+		EUnit Unit;
+	};
+	/** Arbitrary length array used for parsing input strings into EUnit types */
+	static FParseCandidate	 	ParseCandidates[];
+
+	/** Static array of display strings that directly map to EUnit enumerations */
+	static const TCHAR* const 	DisplayStrings[];
+
+	/** Static array of unit types that directly map to EUnit enumerations */
+	static const EUnitType 		Types[];
+};
+
+
+/**
+ * FNumericUnit is a numeric type that wraps the templated type, whilst a specified unit.
+ * It handles conversion to/from related units automatically. The units are considered not to contribute to the type's state, and as such should be considered immutable once set.
+ */
+template<typename NumericType>
+struct FNumericUnit
+{
+	/** The numeric (scalar) value */
+	NumericType Value;
+	/** The associated units for the value. Can never change once set to anything other than EUnit::Unspecified. */
+	const EUnit Units;
+
+	/** Constructors */
+	FNumericUnit() : Units(EUnit::Unspecified) {}
+	FNumericUnit(const NumericType& InValue, EUnit InUnits = EUnit::Unspecified) : Value(InValue), Units(InUnits) {}
+
+	/** Copy construction/assignment from the same type */
+	FNumericUnit(const FNumericUnit& Other) : Units(EUnit::Unspecified)							{	(*this) = Other;								}
+	FNumericUnit& operator=(const FNumericUnit& Other)											{ 	CopyValueWithConversion(Other);	return *this; 	}
+
+	/** Templated Copy construction/assignment from differing numeric types. Relies on implicit conversion of the two numeric types. */
+	template<typename OtherType> FNumericUnit(const FNumericUnit<OtherType>& Other)				{	(*this) = Other;								}
+	template<typename OtherType> FNumericUnit& operator=(const FNumericUnit<OtherType>& Other)	{	CopyValueWithConversion(Other);	return *this;	}
+
+	/** Convert this quantity to a different unit */
+	TOptional<FNumericUnit<NumericType>> ConvertTo(EUnit ToUnits) const
+	{
+		if (Units == EUnit::Unspecified)
+		{
+			return FNumericUnit(Value, ToUnits);
+		}
+		else if (FUnitConversion::AreUnitsCompatible(Units, ToUnits))
+		{
+			return FNumericUnit<NumericType>(FUnitConversion::Convert(Value, Units, ToUnits), ToUnits);
+		}
+
+		return TOptional<FNumericUnit<NumericType>>();
+	}
+
+public:
+
+	/** Quantizes this number to the most appropriate unit for user friendly presentation (e.g. 1000m returns 1km). */
+	FNumericUnit<NumericType> QuantizeUnitsToBestFit() const
+	{
+		return FUnitConversion::QuantizeUnitsToBestFit(Value, Units);
+	}
+
+	/** Parse a numeric unit from a string */
+	static TOptional<FNumericUnit<NumericType>> TryParseString(const TCHAR* InSource)
+	{
+		TOptional<FNumericUnit<NumericType>> Result;
+		if (!InSource || !*InSource)
+		{
+			return Result;
+		}
+
+		const TCHAR* NumberEnd = nullptr;
+		if (!ExtractNumberBoundary(InSource, NumberEnd))
+		{
+			return Result;
+		}
+
+		NumericType NewValue;
+		LexicalConversion::FromString(NewValue, InSource);
+
+		// Now parse the units
+		while(FChar::IsWhitespace(*NumberEnd)) ++NumberEnd;
+
+		if (*NumberEnd == '\0')
+		{
+			// No units
+			Result.Emplace(NewValue);
+		}
+		else
+		{
+			// If the string specifies units, they must map to something that exists for this function to succeed
+			auto NewUnits = FUnitConversion::UnitFromString(NumberEnd);
+			if (NewUnits)
+			{
+				Result.Emplace(NewValue, NewUnits.GetValue());
+			}
+		}
+
+		return Result;
+	}
+
+public:
+	/** Global arithmetic operators for number types. Deals with conversion from related units correctly. */
+	template<typename OtherType>
+	friend bool operator==(const FNumericUnit& LHS, const FNumericUnit<OtherType>& RHS)
+	{
+		if (LHS.Units != EUnit::Unspecified && RHS.Units != EUnit::Unspecified)
+		{
+			if (LHS.Units == RHS.Units)
+			{
+				return LHS.Value == RHS.Value;
+			}
+			else if (FUnitConversion::AreUnitsCompatible(LHS.Units, RHS.Units))
+			{
+				return LHS.Value == FUnitConversion::Convert(RHS.Value, RHS.Units, LHS.Units);
+			}
+			else
+			{
+				// Invalid conversion
+				return false;
+			}
+		}
+		else
+		{
+			return LHS.Value == RHS.Value;
+		}
+	}
+
+	template<typename OtherType>
+	friend bool operator!=(const FNumericUnit& LHS, const OtherType& RHS)
+	{
+		return !(LHS == RHS);
+	}
+
+private:
+	/** Conversion to the numeric type disabled as coupled with implicit construction from NumericType can easily lead to loss of associated units. */
+	operator const NumericType&() const;
+
+	/** Copy another unit into this one, taking account of its units, and applying necessary conversion */
+	template<typename OtherType>
+	void CopyValueWithConversion(const FNumericUnit<OtherType>& Other)
+	{
+		if (Units != EUnit::Unspecified && Other.Units != EUnit::Unspecified)
+		{
+			if (Units == Other.Units)
+			{
+				Value = Other.Value;
+			}
+			else if (FUnitConversion::AreUnitsCompatible(Units, Other.Units))
+			{
+				Value = FUnitConversion::Convert(Other.Value, Other.Units, Units);
+			}
+			else
+			{
+				// Invalid conversion - assignment invalid
+			}
+		}
+		else
+		{
+			// If our units haven't been specified, we take on the units of the rhs
+			if (Units == EUnit::Unspecified)
+			{
+				// This is the only time we ever change units. Const_cast is 'acceptible' here since the units haven't been specified yet.
+				const_cast<EUnit&>(Units) = Other.Units;
+			}
+
+			Value = Other.Value;
+		}
+	}
+
+	/** Given a string, skip past whitespace, then any numeric characters. Set End pointer to the end of the last numeric character. */
+	static bool ExtractNumberBoundary(const TCHAR* Start, const TCHAR*& End)
+	{
+		while(FChar::IsWhitespace(*Start)) ++Start;
+
+		End = Start;
+		if (*End == '-' || *End == '+')
+		{
+			End++;
+		}
+
+		bool bHasDot = false;
+		while (FChar::IsDigit(*End) || *End == '.')
+		{
+			if (*End == '.')
+			{
+				if (bHasDot)
+				{
+					return false;
+				}
+				bHasDot = true;
+			}
+			++End;
+		}
+
+		return true;
+	}
+};
+
+
+template <typename NumericType>
+struct TNumericLimits<FNumericUnit<NumericType>> : public TNumericLimits<NumericType>
+{ };
+
+namespace LexicalConversion
+{
+	template<typename T>
+	FString ToString(const FNumericUnit<T>& NumericUnit)
+	{
+		FString String = ToString(NumericUnit.Value);
+		String += TEXT(" ");
+		String += FUnitConversion::GetUnitDisplayString(NumericUnit.Units);
+
+		return String;
+	}
+
+	template<typename T>
+	FString ToSanitizedString(const FNumericUnit<T>& NumericUnit)
+	{
+		FString String = ToSanitizedString(NumericUnit.Value);
+		String += TEXT(" ");
+		String += FUnitConversion::GetUnitDisplayString(NumericUnit.Units);
+
+		return String;
+	}
+
+	template<typename T>
+	void FromString(FNumericUnit<T>& OutValue, const TCHAR* String)
+	{
+		auto Parsed = FNumericUnit<T>::TryParseString(String);
+		if (Parsed)
+		{
+			OutValue = Parsed.GetValue();
+		}
+	}
+	
+	template<typename T>
+	bool TryParseString(FNumericUnit<T>& OutValue, const TCHAR* String)
+	{
+		auto Parsed = FNumericUnit<T>::TryParseString(String);
+		if (Parsed)
+		{
+			OutValue = Parsed.GetValue();
+			return true;
+		}
+
+		return false;
+	}
+}
\ No newline at end of file
diff --git a/Engine/Source/Runtime/Slate/Public/Widgets/Input/NumericTypeInterface.h b/Engine/Source/Runtime/Slate/Public/Widgets/Input/NumericTypeInterface.h
new file mode 100644
--- /dev/null
+++ b/Engine/Source/Runtime/Slate/Public/Widgets/Input/NumericTypeInterface.h
@@ -0,0 +1,94 @@
+// Copyright 1998-2015 Epic Games, Inc. All Rights Reserved.
+
+#pragma once
+
+/** Interface to provide specific functionality for dealing with a numeric type. Currently includes string conversion functionality. */
+template<typename NumericType>
+struct INumericTypeInterface
+{
+	virtual ~INumericTypeInterface() {}
+
+	/** Convert the type to/from a string */
+	virtual FString ToString(const NumericType& Value) const = 0;
+	virtual TOptional<NumericType> FromString(const FString& InString, const NumericType& InCurrentValue) = 0;
+
+	/** Check whether the typed character is valid */
+	virtual bool IsCharacterValid(TCHAR InChar) const = 0;
+};
+
+/** Default numeric type interface */
+template<typename NumericType>
+struct TDefaultNumericTypeInterface : INumericTypeInterface<NumericType>
+{
+	/** Convert the type to/from a string */
+	virtual FString ToString(const NumericType& Value) const override
+	{
+		return LexicalConversion::ToSanitizedString(Value);
+	}
+	virtual TOptional<NumericType> FromString(const FString& InString, const NumericType& InCurrentValue) override
+	{
+		NumericType NewValue;
+		bool bEvalResult = LexicalConversion::TryParseString( NewValue, *InString );
+		if (!bEvalResult)
+		{
+			float FloatValue = 0.f;
+			bEvalResult = FMath::Eval( *InString, FloatValue  );
+			NewValue = FloatValue;
+		}
+
+		return bEvalResult ? NewValue : TOptional<NumericType>();
+	}
+
+	/** Check whether the typed character is valid */
+	virtual bool IsCharacterValid(TCHAR InChar) const override;	
+};
+
+template<typename T>
+bool TDefaultNumericTypeInterface<T>::IsCharacterValid(TCHAR InChar) const
+{
+	static FString ValidChars(TEXT("1234567890-.")); // for integers and doubles
+	int32 OutUnused;
+	return ValidChars.FindChar( InChar, OutUnused );
+}
+
+/** Specialized for floats */
+template<>
+inline bool TDefaultNumericTypeInterface<float>::IsCharacterValid(TCHAR InChar) const
+{
+	static FString ValidChars(TEXT("1234567890()-+=\\/.,*^%%"));
+	int32 OutUnused;
+	return ValidChars.FindChar( InChar, OutUnused );		
+}
+
+/** Forward declaration of types defined in UnitConversion.h */
+enum class EUnit;
+template<typename> struct FNumericUnit;
+
+/**
+ * Numeric interface that specifies how to interact with a number in a specific unit.
+ * Include NumericUnitTypeInterface.inl for symbol definitions.
+ */
+template<typename NumericType>
+struct TNumericUnitTypeInterface : TDefaultNumericTypeInterface<NumericType>
+{
+	/** The units that we should treat this value as */
+	EUnit Units;
+
+	/** True if we are allowed to display this unit to the user as a more relevant related unit type (e.g. 10000 m shows as 10 km) */
+	bool bAllowUnitRangeAdaption;
+
+	/** Constructor */
+	TNumericUnitTypeInterface(EUnit InUnits, bool bInAllowUnitRangeAdaption = true);
+
+	/** Convert this type to a string */
+	virtual FString ToString(const NumericType& Value) const override;
+
+	/** Attempt to parse a numeral with our units from the specified string. */
+	virtual TOptional<NumericType> FromString(const FString& ValueString, const NumericType& InCurrentValue) override;
+
+	/** Called by this class to quantize the specified value to the best unit */
+	virtual FNumericUnit<NumericType> QuantizeUnitsToBestFit(const NumericType& InValue, EUnit InUnits) const;
+
+	/** Check whether the specified typed character is valid */
+	virtual bool IsCharacterValid(TCHAR InChar) const override;
+};
\ No newline at end of file
diff --git a/Engine/Source/Runtime/Slate/Public/Widgets/Input/NumericUnitTypeInterface.inl b/Engine/Source/Runtime/Slate/Public/Widgets/Input/NumericUnitTypeInterface.inl
new file mode 100644
--- /dev/null
+++ b/Engine/Source/Runtime/Slate/Public/Widgets/Input/NumericUnitTypeInterface.inl
@@ -0,0 +1,62 @@
+// Copyright 1998-2014 Epic Games, Inc. All Rights Reserved.
+
+#pragma once
+
+#include "UnitConversion.h"
+
+template<typename NumericType>
+TNumericUnitTypeInterface<NumericType>::TNumericUnitTypeInterface(EUnit InUnits, bool bInAllowUnitRangeAdaption)
+	: Units(InUnits), bAllowUnitRangeAdaption(bInAllowUnitRangeAdaption)
+{}
+
+template<typename NumericType>
+FNumericUnit<NumericType> TNumericUnitTypeInterface<NumericType>::QuantizeUnitsToBestFit(const NumericType& InValue, EUnit InUnits) const
+{
+	return FUnitConversion::QuantizeUnitsToBestFit(InValue, InUnits);
+}
+
+template<typename NumericType>
+FString TNumericUnitTypeInterface<NumericType>::ToString(const NumericType& Value) const
+{
+	using namespace LexicalConversion;
+
+	const FNumericUnit<NumericType> Default(Value, Units);
+	return ToSanitizedString(bAllowUnitRangeAdaption ? QuantizeUnitsToBestFit(Default.Value, Default.Units) : Default);
+}
+
+template<typename NumericType>
+TOptional<NumericType> TNumericUnitTypeInterface<NumericType>::FromString(const FString& InString, const NumericType& InCurrentValue)
+{
+	using namespace LexicalConversion;
+
+	// Always parse in as a double, to allow for input of higher-order units with decimal numerals into integral types (eg, inputting 0.5km as 500m)
+	FNumericUnit<double> NewValue;
+	bool bEvalResult = TryParseString( NewValue, *InString ) && FUnitConversion::AreUnitsCompatible( NewValue.Units, Units );
+	if (bEvalResult)
+	{
+		// Convert the number into the correct units
+		EUnit SourceUnits = NewValue.Units;
+		if (SourceUnits == EUnit::Unspecified && bAllowUnitRangeAdaption)
+		{
+			// Use the display units of the current value
+			SourceUnits = QuantizeUnitsToBestFit(InCurrentValue, Units).Units;
+		}
+		return FUnitConversion::Convert(NewValue.Value, SourceUnits, Units);
+	}
+	else
+	{
+		float FloatValue = 0.f;
+		if (FMath::Eval( *InString, FloatValue  ))
+		{
+			return FloatValue;
+		}
+	}
+
+	return TOptional<NumericType>();
+}
+
+template<typename NumericType>
+bool TNumericUnitTypeInterface<NumericType>::IsCharacterValid(TCHAR InChar) const
+{
+	return true;
+}
\ No newline at end of file
diff --git a/Engine/Source/Runtime/Slate/Public/Widgets/Input/SNumericEntryBox.h b/Engine/Source/Runtime/Slate/Public/Widgets/Input/SNumericEntryBox.h
--- a/Engine/Source/Runtime/Slate/Public/Widgets/Input/SNumericEntryBox.h
+++ b/Engine/Source/Runtime/Slate/Public/Widgets/Input/SNumericEntryBox.h
@@ -86,7 +86,9 @@ public:
 		/** Called right after the slider handle is released by the user */
 		SLATE_EVENT( FOnValueChanged, OnEndSliderMovement )		
 		/** Menu extender for right-click context menu */
-		SLATE_EVENT( FMenuExtensionDelegate, ContextMenuExtender )		
+		SLATE_EVENT( FMenuExtensionDelegate, ContextMenuExtender )
+		/** Provide custom type conversion functionality to this spin box */
+		SLATE_ARGUMENT( TSharedPtr< INumericTypeInterface<NumericType> >, TypeInterface )
 
 	SLATE_END_ARGS()
 	SNumericEntryBox()
@@ -109,6 +111,8 @@ public:
 		BorderImageFocused = &InArgs._EditableTextBoxStyle->BackgroundImageFocused;
 		const FMargin& TextMargin = InArgs._EditableTextBoxStyle->Padding;
 		
+		Interface = InArgs._TypeInterface.IsValid() ? InArgs._TypeInterface : MakeShareable( new TDefaultNumericTypeInterface<NumericType> );
+
 		TSharedPtr<SWidget> FinalWidget;
 		if( bAllowSpin )
 		{
@@ -127,7 +131,8 @@ public:
 				.SliderExponent(InArgs._SliderExponent)
 				.OnBeginSliderMovement(InArgs._OnBeginSliderMovement)
 				.OnEndSliderMovement(InArgs._OnEndSliderMovement)
-				.MinDesiredWidth(InArgs._MinDesiredValueWidth);
+				.MinDesiredWidth(InArgs._MinDesiredValueWidth)
+				.TypeInterface(Interface);
 		}
 
 		// Always create an editable text box.  In the case of an undetermined value being passed in, we cant use the spinbox.
@@ -295,8 +300,7 @@ private:
 			// If the value was set convert it to a string, otherwise the value cannot be determined
 			if( Value.IsSet() == true )
 			{
-				auto CurrentValue = Value.GetValue();
-				NewText = FText::FromString(TTypeToString<NumericType>::ToSanitizedString(CurrentValue));
+				NewText = FText::FromString(Interface->ToString(Value.GetValue()));
 			}
 			else
 			{
@@ -377,34 +381,19 @@ private:
 		{
 			if( bCommit )
 			{
-				bool bEvalResult = false;
-				NumericType NumericValue = 0;
-
-				// Convert string to an underlying type in case text is not math equation
-				// Try to parse equation otherwise
-				if ( NewValue.IsNumeric() )
+				TOptional<NumericType> NumericValue = Interface->FromString(NewValue.ToString(), OnGetValueForSpinBox());
+				if( NumericValue.IsSet() )
 				{
-					TTypeFromString<NumericType>::FromString( NumericValue, *NewValue.ToString() );
-					bEvalResult = true;
-				}
-				else
-				{
-					// Only evaluate equations on commit or else they could fail because the equation is still being typed
-					float Value = 0.f;
-					bEvalResult = FMath::Eval( NewValue.ToString(), Value );
-					NumericValue = NumericType(Value);
-				}
-				
-				if( bEvalResult )
-				{
-					OnValueCommitted.ExecuteIfBound(NumericValue, CommitInfo );
+					OnValueCommitted.ExecuteIfBound(NumericValue.GetValue(), CommitInfo );
 				}
 			}
-			else if( NewValue.IsNumeric() )
+			else
 			{
-				NumericType Value = 0;
-				TTypeFromString<NumericType>::FromString( Value, *NewValue.ToString() );
-				OnValueChanged.ExecuteIfBound( Value );
+				NumericType NumericValue;
+				if( LexicalConversion::TryParseString(NumericValue, *NewValue.ToString()) )
+				{
+					OnValueChanged.ExecuteIfBound( NumericValue );
+				}
 			}
 		}
 	}
@@ -461,6 +450,8 @@ private:
 	const FSlateBrush* BorderImageFocused;
 	/** Prevents the value portion of the control from being smaller than desired in certain cases. */
 	TAttribute<float> MinDesiredValueWidth;
+	/** Type interface that defines how we should deal with the templated numeric type. Always valid after construction. */
+	TSharedPtr< INumericTypeInterface<NumericType> > Interface;
 };
 
 
diff --git a/Engine/Source/Runtime/Slate/Public/Widgets/Input/SSpinBox.h b/Engine/Source/Runtime/Slate/Public/Widgets/Input/SSpinBox.h
--- a/Engine/Source/Runtime/Slate/Public/Widgets/Input/SSpinBox.h
+++ b/Engine/Source/Runtime/Slate/Public/Widgets/Input/SSpinBox.h
@@ -2,6 +2,7 @@
 
 #pragma once
 
+#include "NumericTypeInterface.h"
 
 /**
  * A Slate SpinBox resembles traditional spin boxes in that it is a widget that provides
@@ -50,7 +51,7 @@ public:
 		/** The maximum value that can be specified by using the slider, defaults to MaxValue */
 		SLATE_ATTRIBUTE( TOptional< NumericType >, MaxSliderValue )
 		/** Delta to increment the value as the slider moves.  If not specified will determine automatically */
-		SLATE_ATTRIBUTE( NumericType, Delta )
+		SLATE_ARGUMENT( NumericType, Delta )
 		/** Use exponential scale for the slider */
 		SLATE_ATTRIBUTE( float, SliderExponent )
 		/** Font used to display text in the slider */
@@ -71,6 +72,8 @@ public:
 		SLATE_ATTRIBUTE( bool, SelectAllTextOnCommit )
 		/** Minimum width that a spin box should be */
 		SLATE_ATTRIBUTE( float, MinDesiredWidth )
+		/** Provide custom type conversion functionality to this spin box */
+		SLATE_ARGUMENT( TSharedPtr< INumericTypeInterface<NumericType> >, TypeInterface )
 
 	SLATE_END_ARGS()
 
@@ -90,6 +93,7 @@ public:
 		Style = InArgs._Style;
 
 		ForegroundColor = InArgs._Style->ForegroundColor;
+		Interface = InArgs._TypeInterface.IsValid() ? InArgs._TypeInterface : MakeShareable( new TDefaultNumericTypeInterface<NumericType> );
 
 		ValueAttribute = InArgs._Value;
 		OnValueChanged = InArgs._OnValueChanged;
@@ -112,7 +116,7 @@ public:
 		InternalValue = 0;
 		PreDragValue = 0;
 
-		Delta = InArgs._Delta.Get();
+		Delta = InArgs._Delta;
 	
 		BackgroundHoveredBrush = &InArgs._Style->HoveredBackgroundBrush;
 		BackgroundBrush = &InArgs._Style->BackgroundBrush;
@@ -152,7 +156,7 @@ public:
 				.Font(InArgs._Font)
 				.SelectAllTextWhenFocused( true )
 				.Text( this, &SSpinBox<NumericType>::GetValueAsText )
-				.OnIsTypedCharValid_Static(TCharacterFilter<NumericType>::IsValid)
+				.OnIsTypedCharValid(this, &SSpinBox<NumericType>::IsCharacterValid)
 				.OnTextCommitted( this, &SSpinBox<NumericType>::TextField_OnTextCommitted )
 				.ClearKeyboardFocusOnCommit( InArgs._ClearKeyboardFocusOnCommit )
 				.SelectAllTextOnCommit( InArgs._SelectAllTextOnCommit )
@@ -514,7 +518,7 @@ protected:
 			CurrentValue = (NumericType)Snap(CurrentValue, Delta);
 		}
 		
-		return TTypeToString<NumericType>::ToSanitizedString(CurrentValue);
+		return Interface->ToString(CurrentValue);
 	}
 
 	/** @return the value being observed by the spinbox as FText - todo: spinbox FText support (reimplement me) */
@@ -536,25 +540,10 @@ protected:
 			ExitTextMode();
 		}
 
-		NumericType NewValue = 0;
-		bool bEvalResult = false;
-		FString NewString = NewText.ToString();
-		
-		if (NewString.IsNumeric())
-		{
-			TTypeFromString<NumericType>::FromString( NewValue, *NewString );
-			bEvalResult = true;
-		}
-		else
-		{
-			float FloatValue = 0.f;
-			bEvalResult = FMath::Eval( *NewString, FloatValue  );
-			NewValue = NumericType(FloatValue);
-		}
-				
-		if (bEvalResult)
+		TOptional<NumericType> NewValue = Interface->FromString(NewText.ToString(), ValueAttribute.Get());
+		if (NewValue.IsSet())
 		{
-			CommitValue( NewValue, CommittedViaTypeIn, CommitInfo );
+			CommitValue( NewValue.GetValue(), CommittedViaTypeIn, CommitInfo );
 		}
 	}
 
@@ -665,31 +654,6 @@ protected:
 		double Result = FMath::GridSnap(InValue, InGrid);
 		return FMath::Clamp<double>(Result, TNumericLimits<NumericType>::Lowest(), TNumericLimits<NumericType>::Max());
 	}
-	
-	/**
-	 * Helper structure to check if a given character is supported by underlying numeric type
-	 */
-	template<typename T, typename U = void>
-	struct TCharacterFilter
-	{
-		static bool IsValid(const TCHAR InChar)
-		{
-			static FString ValidChars(TEXT("1234567890-.")); // for integers and doubles
-			int32 OutUnused;
-			return ValidChars.FindChar( InChar, OutUnused );
-		}
-	};
-
-	template<typename U> // nested partial specialization is allowed
-	struct TCharacterFilter<float, U>
-	{
-		static bool IsValid(const TCHAR InChar)
-		{
-			static FString ValidChars(TEXT("1234567890()-+=\\/.,*^%%"));
-			int32 OutUnused;
-			return ValidChars.FindChar( InChar, OutUnused );		
-		}
-	};
 
 private:
 	TAttribute<NumericType> ValueAttribute;
@@ -700,6 +664,9 @@ private:
 	TSharedPtr<STextBlock> TextBlock;
 	TSharedPtr<SEditableText> EditableText;
 
+	/** Interface that defines conversion functionality for the templated type */
+	TSharedPtr< INumericTypeInterface<NumericType> > Interface;
+
 	/** True when no range is specified, spinner can be spun indefinitely */
 	bool bUnlimitedSpinRange;
 	void UpdateIsSpinRangeUnlimited()
@@ -729,6 +696,12 @@ private:
 		return FMath::Max(0.0f, MinDesiredWidth.Get() - Style->ArrowsImage.ImageSize.X);
 	}
 
+	/** Check whether a typed character is valid */
+	bool IsCharacterValid(TCHAR InChar) const
+	{
+		return Interface->IsCharacterValid(InChar);
+	}
+
 	/** Whether the user is dragging the slider */
 	bool bDragging;
 	
diff --git a/Engine/Source/Runtime/UMG/Private/Components/SpinBox.cpp b/Engine/Source/Runtime/UMG/Private/Components/SpinBox.cpp
--- a/Engine/Source/Runtime/UMG/Private/Components/SpinBox.cpp
+++ b/Engine/Source/Runtime/UMG/Private/Components/SpinBox.cpp
@@ -21,7 +21,7 @@ USpinBox::USpinBox(const FObjectInitializer& ObjectInitializer)
 	MaxValue = Defaults._MaxValue.Get().Get(0.0f);
 	MinSliderValue = Defaults._MinSliderValue.Get().Get(0.0f);
 	MaxSliderValue = Defaults._MaxSliderValue.Get().Get(0.0f);
-	Delta = Defaults._Delta.Get();
+	Delta = Defaults._Delta;
 	SliderExponent = Defaults._SliderExponent.Get();
 	MinDesiredWidth = Defaults._MinDesiredWidth.Get();
 	ClearKeyboardFocusOnCommit = Defaults._ClearKeyboardFocusOnCommit.Get();
