Removing erroneous comment

diff --git a/Engine/Source/Runtime/Core/Public/Math/UnitConversion.h b/Engine/Source/Runtime/Core/Public/Math/UnitConversion.h
--- a/Engine/Source/Runtime/Core/Public/Math/UnitConversion.h
+++ b/Engine/Source/Runtime/Core/Public/Math/UnitConversion.h
@@ -6,7 +6,6 @@
 #include "Function.h"
 
 /** Enum *must* be zero-indexed and sequential */
-/** Enum *must* match the mirrored enum that exists in KismetMathLibrary.h for the purposes of UObject reflection */
 enum class EUnit
 {
 	/** Scalar distance/length units */
