Fixed Clang comipile errors in Async Framework

diff --git a/Engine/Source/Runtime/Core/Public/Async/Future.h b/Engine/Source/Runtime/Core/Public/Async/Future.h
--- a/Engine/Source/Runtime/Core/Public/Async/Future.h
+++ b/Engine/Source/Runtime/Core/Public/Async/Future.h
@@ -231,6 +231,7 @@ protected:
 	{
 		State = MoveTemp(Other.State);
 		Other.State.Reset();
+		return *this;
 	}
 
 protected:
@@ -283,7 +284,7 @@ public:
 	 *
 	 * @param InState The shared state to initialize with.
 	 */
-	TFuture(const BaseType::StateType& InState)
+	TFuture(const typename BaseType::StateType& InState)
 		: BaseType(InState)
 	{ }
 
@@ -313,7 +314,7 @@ public:
 	 */
 	ResultType Get()
 	{
-		return MoveTemp(GetState()->GetResult());
+		return MoveTemp(this->GetState()->GetResult());
 	}
 
 	/**
@@ -348,7 +349,7 @@ public:
 	 *
 	 * @param InState The shared state to initialize with.
 	 */
-	TFuture(const BaseType::StateType& InState)
+	TFuture(const typename BaseType::StateType& InState)
 		: BaseType(InState)
 	{ }
 
@@ -378,7 +379,7 @@ public:
 	 */
 	ResultType& Get()
 	{
-		return *GetState()->GetResult();
+		return *this->GetState()->GetResult();
 	}
 
 	/**
@@ -477,7 +478,7 @@ public:
 	 *
 	 * @param InState The shared state to initialize from.
 	 */
-	TSharedFuture(const BaseType::StateType& InState)
+	TSharedFuture(const typename BaseType::StateType& InState)
 		: BaseType(InState)
 	{ }
 
@@ -528,7 +529,7 @@ public:
 	 */
 	ResultType Get()
 	{
-		return MoveTemp(GetState()->GetResult());
+		return MoveTemp(this->GetState()->GetResult());
 	}
 };
 
@@ -552,7 +553,7 @@ public:
 	 *
 	 * @param InState The shared state to initialize from.
 	 */
-	TSharedFuture(const BaseType::StateType& InState)
+	TSharedFuture(const typename BaseType::StateType& InState)
 		: BaseType(InState)
 	{ }
 
@@ -561,7 +562,7 @@ public:
 	*
 	* @param Future The future object to initialize from.
 	*/
-	TSharedFuture(TFuture&& Future)
+	TSharedFuture(TFuture<ResultType>&& Future)
 		: BaseType(MoveTemp(Future))
 	{ }
 
@@ -603,7 +604,7 @@ public:
 	 */
 	ResultType& Get()
 	{
-		return *GetState()->GetResult();
+		return *this->GetState()->GetResult();
 	}
 };
 
@@ -640,11 +641,6 @@ public:
 		: BaseType(MoveTemp(Future))
 	{ }
 
-	/** Copy constructor. */
-	TSharedFuture(const TSharedFuture& Other)
-		: BaseType(Other)
-	{ }
-
 	/** Move constructor. */
 	TSharedFuture(TSharedFuture&& Other)
 		: BaseType(MoveTemp(Other))
@@ -655,13 +651,6 @@ public:
 
 public:
 
-	/** Copy assignment operator. */
-	TSharedFuture& operator=(const TSharedFuture& Other)
-	{
-		BaseType::operator=(Other);
-		return *this;
-	}
-
 	/** Move assignment operator. */
 	TSharedFuture& operator=(TSharedFuture&& Other)
 	{
@@ -723,6 +712,7 @@ public:
 	{
 		State = Other.State;
 		Other.State.Reset();
+		return *this;
 	}
 
 protected:
@@ -766,10 +756,10 @@ template<typename ResultType>
 class TPromise
 	: public TPromiseBase<ResultType>
 {
-	typedef public TPromiseBase<ResultType> BaseType;
-
 public:
 
+	typedef TPromiseBase<ResultType> BaseType;
+
 	/** Default constructor (creates a new shared state). */
 	TPromise()
 		: BaseType()
@@ -812,7 +802,7 @@ public:
 		check(!FutureRetrieved)
 			FutureRetrieved = true;
 
-		return TFuture<ResultType>(GetState());
+		return TFuture<ResultType>(this->GetState());
 	}
 
 	/**
@@ -825,7 +815,7 @@ public:
 	 */
 	void SetValue(const ResultType& Result)
 	{
-		GetState()->SetResult(Result);
+		this->GetState()->SetResult(Result);
 	}
 
 	/**
@@ -838,7 +828,7 @@ public:
 	 */
 	void SetValue(ResultType&& Result)
 	{
-		GetState()->SetResult(MoveTemp(Result));
+		this->GetState()->SetResult(MoveTemp(Result));
 	}
 
 private:
@@ -902,7 +892,7 @@ public:
 		check(!FutureRetrieved)
 			FutureRetrieved = true;
 
-		return TFuture<ResultType&>(GetState());
+		return TFuture<ResultType&>(this->GetState());
 	}
 
 	/**
@@ -915,7 +905,7 @@ public:
 	 */
 	void SetValue(ResultType& Result)
 	{
-		GetState()->SetResult(Result);
+		this->GetState()->SetResult(Result);
 	}
 
 private:
