Async: const-correctness for getters

diff --git a/Engine/Source/Runtime/Core/Public/Async/Future.h b/Engine/Source/Runtime/Core/Public/Async/Future.h
--- a/Engine/Source/Runtime/Core/Public/Async/Future.h
+++ b/Engine/Source/Runtime/Core/Public/Async/Future.h
@@ -45,7 +45,7 @@ public:
 	 * @result true if the result is available, false otherwise.
 	 * @see IsComplete
 	 */
-	bool WaitFor(const FTimespan& Duration)
+	bool WaitFor(const FTimespan& Duration) const
 	{
 		FScopeLock Lock(&CriticalSection);
 
@@ -94,7 +94,7 @@ public:
 	 * @return The result value.
 	 * @see SetResult
 	 */
-	InternalResultType& GetResult()
+	const InternalResultType& GetResult() const
 	{
 		while (!IsComplete())
 		{
@@ -165,7 +165,7 @@ public:
 	 *
 	 * @see WaitFor, WaitUntil
 	 */
-	void Wait()
+	void Wait() const
 	{
 		while (!WaitFor(FTimespan::MaxValue()));
 	}
@@ -177,7 +177,7 @@ public:
 	 * @result true if the result is available, false otherwise.
 	 * @see Wait, WaitUntil
 	 */
-	bool WaitFor(const FTimespan& Duration)
+	bool WaitFor(const FTimespan& Duration) const
 	{
 		return State->WaitFor(Duration);
 	}
@@ -189,7 +189,7 @@ public:
 	 * @result true if the result is available, false otherwise.
 	 * @see Wait, WaitUntil
 	 */
-	bool WaitUntil(const FDateTime& Time)
+	bool WaitUntil(const FDateTime& Time) const
 	{
 		return WaitFor(Time - FDateTime::UtcNow());
 	}
@@ -241,7 +241,7 @@ protected:
 	 *
 	 * @return The shared state object.
 	 */
-	const StateType& GetState()
+	const StateType& GetState() const
 	{
 		// if you hit this assertion then your future has an invalid state.
 		// this happens if you have an uninitialized future or if you moved
@@ -312,7 +312,7 @@ public:
 	 *
 	 * @return The result.
 	 */
-	ResultType Get()
+	ResultType Get() const
 	{
 		return MoveTemp(this->GetState()->GetResult());
 	}
@@ -377,7 +377,7 @@ public:
 	 *
 	 * @return The result.
 	 */
-	ResultType& Get()
+	ResultType& Get() const
 	{
 		return *this->GetState()->GetResult();
 	}
@@ -442,7 +442,7 @@ public:
 	 *
 	 * @return The result.
 	 */
-	void Get()
+	void Get() const
 	{
 		GetState()->GetResult();
 	}
@@ -527,7 +527,7 @@ public:
 	 *
 	 * @return The result.
 	 */
-	ResultType Get()
+	ResultType Get() const
 	{
 		return MoveTemp(this->GetState()->GetResult());
 	}
@@ -602,7 +602,7 @@ public:
 	 *
 	 * @return The result.
 	 */
-	ResultType& Get()
+	ResultType& Get() const
 	{
 		return *this->GetState()->GetResult();
 	}
@@ -665,7 +665,7 @@ public:
 	 *
 	 * @return The result.
 	 */
-	void Get()
+	void Get() const
 	{
 		GetState()->GetResult();
 	}
@@ -739,7 +739,7 @@ protected:
 		// this happens if you move the promise to another instance.
 		check(State.IsValid())
 
-			return State;
+		return State;
 	}
 
 private:
@@ -800,7 +800,7 @@ public:
 	TFuture<ResultType> GetFuture()
 	{
 		check(!FutureRetrieved)
-			FutureRetrieved = true;
+		FutureRetrieved = true;
 
 		return TFuture<ResultType>(this->GetState());
 	}
@@ -890,7 +890,7 @@ public:
 	TFuture<ResultType&> GetFuture()
 	{
 		check(!FutureRetrieved)
-			FutureRetrieved = true;
+		FutureRetrieved = true;
 
 		return TFuture<ResultType&>(this->GetState());
 	}
@@ -967,7 +967,7 @@ public:
 	TFuture<void> GetFuture()
 	{
 		check(!FutureRetrieved)
-			FutureRetrieved = true;
+		FutureRetrieved = true;
 
 		return TFuture<void>(GetState());
 	}
