Async: Slow task API improvements

diff --git a/Engine/Source/Runtime/Core/Public/Async/ISlowTask.h b/Engine/Source/Runtime/Core/Public/Async/ISlowTask.h
--- a/Engine/Source/Runtime/Core/Public/Async/ISlowTask.h
+++ b/Engine/Source/Runtime/Core/Public/Async/ISlowTask.h
@@ -3,6 +3,25 @@
 #pragma once
 
 
+/**
+ * Enumerates possible states of slow running tasks.
+ */
+enum class ESlowTaskStates
+{
+	/** Task has been aborted. */
+	Aborted,
+
+	/** Task has completed execution. */
+	Completed,
+
+	/** Task execution is still pending. */
+	Pending,
+
+	/** Task is executing. */
+	Running,
+};
+
+
 /**
  * Interface for slow running tasks.
  *
@@ -11,8 +30,9 @@
  * for tracking and aborting such tasks.
  *
  * Classes that implement this interface should make a best effort to provide a completion
- * percentage via the GetPercentageComplete() method. While the task is still pending, the
- * returned value should be unset. After the task completed, the return value should be 1.0.
+ * percentage via the GetPercentageComplete() method. If this value is expensive to compute,
+ * implementors should regularly cache the values. If a completion percentage cannot be
+ * computed for some reason, the return value should be unset.
  *
  * Slow running tasks may also provide a status text that can be polled by the user interface.
  */
@@ -38,6 +58,15 @@ public:
 	 */
 	virtual FText GetStatusText() = 0;
 
+	/**
+	 * Gets the current state of the task.
+	 *
+	 * @return Task state.
+	 */
+	virtual ESlowTaskStates GetTaskState() = 0;
+
+public:
+
 	/** A delegate that is executed after the task is complete. */
 	DECLARE_EVENT_OneParam(ISlowTask, FOnCompleted, bool /*Success*/)
 	virtual FOnCompleted& OnComplete() = 0;
@@ -45,18 +74,4 @@ public:
 	/** A delegate that is executed after the task is complete. */
 	DECLARE_EVENT_OneParam(ISlowTask, FOnStatusChanged, FText /*StatusText*/)
 	virtual FOnStatusChanged& OnStatusChanged() = 0;
-
-public:
-
-	/**
-	 * Checks whether this task is complete.
-	 *
-	 * @return true if the task is complete, false otherwise.
-	 * @see GetPercentageComplete
-	 */
-	bool IsComplete() const
-	{
-		TOptional<float> PercentageComplete = GetPercentageComplete();
-		return (PercentageComplete.IsSet() && (PercentageComplete.Get() == 1.0f));
-	}
 };
