Async: Replaced FNonCopyable with hidden constructors/operators for better debugging messages

diff --git a/Engine/Source/Runtime/Core/Public/Async/Future.h b/Engine/Source/Runtime/Core/Public/Async/Future.h
--- a/Engine/Source/Runtime/Core/Public/Async/Future.h
+++ b/Engine/Source/Runtime/Core/Public/Async/Future.h
@@ -270,7 +270,6 @@ template<typename ResultType> class TSharedFuture;
 template<typename ResultType>
 class TFuture
 	: public TFutureBase<ResultType>
-	, public FNoncopyable
 {
 	typedef TFutureBase<ResultType> BaseType;
 
@@ -326,6 +325,14 @@ public:
 	{
 		return TSharedFuture<ResultType>(MoveTemp(*this));
 	}
+
+private:
+
+	/** Hidden copy constructor (futures cannot be copied). */
+	TFuture(const TFuture&);
+
+	/** Hidden copy assignment (futures cannot be copied). */
+	TFuture& operator=(const TFuture&);
 };
 
 
@@ -335,7 +342,6 @@ public:
 template<typename ResultType>
 class TFuture<ResultType&>
 	: public TFutureBase<ResultType*>
-	, public FNoncopyable
 {
 	typedef TFutureBase<ResultType*> BaseType;
 
@@ -391,6 +397,14 @@ public:
 	{
 		return TSharedFuture<ResultType&>(MoveTemp(*this));
 	}
+
+private:
+
+	/** Hidden copy constructor (futures cannot be copied). */
+	TFuture(const TFuture&);
+
+	/** Hidden copy assignment (futures cannot be copied). */
+	TFuture& operator=(const TFuture&);
 };
 
 
@@ -400,7 +414,6 @@ public:
 template<>
 class TFuture<void>
 	: public TFutureBase<int>
-	, public FNoncopyable
 {
 	typedef TFutureBase<int> BaseType;
 
@@ -453,6 +466,14 @@ public:
 	 * @return The shared future object.
 	 */
 	TSharedFuture<void> Share();
+
+private:
+
+	/** Hidden copy constructor (futures cannot be copied). */
+	TFuture(const TFuture&);
+
+	/** Hidden copy assignment (futures cannot be copied). */
+	TFuture& operator=(const TFuture&);
 };
 
 
@@ -584,7 +605,7 @@ public:
 	/** Copy assignment operator. */
 	TSharedFuture& operator=(const TSharedFuture& Other)
 	{
-		BaseType::operator=(MoveTemp(Other));
+		BaseType::operator=(Other);
 		return *this;
 	}
 
@@ -641,6 +662,11 @@ public:
 		: BaseType(MoveTemp(Future))
 	{ }
 
+	/** Copy constructor. */
+/*	TSharedFuture(const TSharedFuture& Other)
+		: BaseType(Other)
+	{ }*/
+
 	/** Move constructor. */
 	TSharedFuture(TSharedFuture&& Other)
 		: BaseType(MoveTemp(Other))
@@ -651,6 +677,13 @@ public:
 
 public:
 
+	/** Copy assignment operator. */
+/*	TSharedFuture& operator=(const TSharedFuture& Other)
+	{
+		BaseType::operator=(Other);
+		return *this;
+	}*/
+
 	/** Move assignment operator. */
 	TSharedFuture& operator=(TSharedFuture&& Other)
 	{
