Async: Move semantics for TFunction parameters

diff --git a/Engine/Source/Runtime/Core/Public/Async/Async.h b/Engine/Source/Runtime/Core/Public/Async/Async.h
--- a/Engine/Source/Runtime/Core/Public/Async/Async.h
+++ b/Engine/Source/Runtime/Core/Public/Async/Async.h
@@ -36,7 +36,7 @@ public:
 	 * @param InFunction The function to execute asynchronously.
 	 * @param InPromise The promise object used to return the function's result.
 	 */
-	TAsyncTask(const TFunction<ResultType()>& InFunction, TPromise<ResultType>&& InPromise)
+	TAsyncTask(TFunction<ResultType()>&& InFunction, TPromise<ResultType>&& InPromise)
 		: Function(MoveTemp(InFunction))
 		, Promise(MoveTemp(InPromise))
 	{ }
@@ -120,7 +120,7 @@ public:
 	 * @param InPromise The promise object used to return the function's result.
 	 * @param InThreadFuture The thread that is running this task.
 	 */
-	TAsyncRunnable(const TFunction<ResultType()>& InFunction, TPromise<ResultType>&& InPromise, TFuture<FRunnableThread*>&& InThreadFuture)
+	TAsyncRunnable(TFunction<ResultType()>&& InFunction, TPromise<ResultType>&& InPromise, TFuture<FRunnableThread*>&& InThreadFuture)
 		: Function(InFunction)
 		, Promise(MoveTemp(InPromise))
 		, ThreadFuture(MoveTemp(InThreadFuture))
@@ -172,14 +172,14 @@ TFuture<ResultType> Async(EAsyncExecution Execution, TFunction<ResultType()> Fun
 	{
 	case EAsyncExecution::TaskGraph:
 		{
-			TGraphTask<TAsyncTask<ResultType>>::CreateTask().ConstructAndDispatchWhenReady(Function, MoveTemp(Promise));
+			TGraphTask<TAsyncTask<ResultType>>::CreateTask().ConstructAndDispatchWhenReady(MoveTemp(Function), MoveTemp(Promise));
 		}
 		break;
 	
 	case EAsyncExecution::Thread:
 		{
 			TPromise<FRunnableThread*> ThreadPromise;
-			TAsyncRunnable<ResultType>* Runnable = new TAsyncRunnable<ResultType>(Function, MoveTemp(Promise), ThreadPromise.GetFuture());
+			TAsyncRunnable<ResultType>* Runnable = new TAsyncRunnable<ResultType>(MoveTemp(Function), MoveTemp(Promise), ThreadPromise.GetFuture());
 			FRunnableThread* RunnableThread = FRunnableThread::Create(Runnable, TEXT("TAsync"));
 
 			check(RunnableThread != nullptr)
