Async: Fixed missing semicolons

diff --git a/Engine/Source/Runtime/Core/Public/Async/Async.h b/Engine/Source/Runtime/Core/Public/Async/Async.h
--- a/Engine/Source/Runtime/Core/Public/Async/Async.h
+++ b/Engine/Source/Runtime/Core/Public/Async/Async.h
@@ -182,14 +182,14 @@ TFuture<ResultType> Async(EAsyncExecution Execution, TFunction<ResultType()> Fun
 			TAsyncRunnable<ResultType>* Runnable = new TAsyncRunnable<ResultType>(MoveTemp(Function), MoveTemp(Promise), ThreadPromise.GetFuture());
 			FRunnableThread* RunnableThread = FRunnableThread::Create(Runnable, TEXT("TAsync"));
 
-			check(RunnableThread != nullptr)
+			check(RunnableThread != nullptr);
 
 			ThreadPromise.SetValue(RunnableThread);
 		}
 		break;
 
 	default:
-		check(false) // not implemented yet!
+		check(false); // not implemented yet!
 	}
 
 	return MoveTemp(Future);
diff --git a/Engine/Source/Runtime/Core/Public/Async/Future.h b/Engine/Source/Runtime/Core/Public/Async/Future.h
--- a/Engine/Source/Runtime/Core/Public/Async/Future.h
+++ b/Engine/Source/Runtime/Core/Public/Async/Future.h
@@ -112,7 +112,7 @@ public:
 	 */
 	void SetResult(const InternalResultType& InResult)
 	{
-		check(!IsComplete())
+		check(!IsComplete());
 
 		Result = InResult;
 		MarkComplete();
@@ -126,7 +126,7 @@ public:
 	 */
 	void SetResult(InternalResultType&& InResult)
 	{
-		check(!IsComplete())
+		check(!IsComplete());
 
 		Result = InResult;
 		MarkComplete();
@@ -246,7 +246,7 @@ protected:
 		// if you hit this assertion then your future has an invalid state.
 		// this happens if you have an uninitialized future or if you moved
 		// the future it to another instance.
-		check(State.IsValid())
+		check(State.IsValid());
 
 		return State;
 	}
@@ -757,7 +757,7 @@ protected:
 		{
 			// if you hit this assertion then your promise never had its result
 			// value set. broken promises are considered programming errors.
-			check(State->IsComplete())
+			check(State->IsComplete());
 		}
 	}
 
@@ -770,7 +770,7 @@ protected:
 	{
 		// if you hit this assertion then your promise has an invalid state.
 		// this happens if you move the promise to another instance.
-		check(State.IsValid())
+		check(State.IsValid());
 
 		return State;
 	}
@@ -832,7 +832,7 @@ public:
 	 */
 	TFuture<ResultType> GetFuture()
 	{
-		check(!FutureRetrieved)
+		check(!FutureRetrieved);
 		FutureRetrieved = true;
 
 		return TFuture<ResultType>(this->GetState());
@@ -922,7 +922,7 @@ public:
 	 */
 	TFuture<ResultType&> GetFuture()
 	{
-		check(!FutureRetrieved)
+		check(!FutureRetrieved);
 		FutureRetrieved = true;
 
 		return TFuture<ResultType&>(this->GetState());
@@ -999,7 +999,7 @@ public:
 	 */
 	TFuture<void> GetFuture()
 	{
-		check(!FutureRetrieved)
+		check(!FutureRetrieved);
 		FutureRetrieved = true;
 
 		return TFuture<void>(GetState());
