Reduced accuracy for angular unit conversion unit test

The problem here seemed to be that PI at runtime is only precise to 1e-7, which seemed to be exacerbated by the unit conversion that resulted in a slightly too large error for the unit test. I've reduced the epsilon for this test a little, and improved the unit test reporting output.

diff --git a/Engine/Source/Runtime/Core/Private/Tests/Math/UnitConversionTests.cpp b/Engine/Source/Runtime/Core/Private/Tests/Math/UnitConversionTests.cpp
--- a/Engine/Source/Runtime/Core/Private/Tests/Math/UnitConversionTests.cpp
+++ b/Engine/Source/Runtime/Core/Private/Tests/Math/UnitConversionTests.cpp
@@ -8,7 +8,7 @@ IMPLEMENT_SIMPLE_AUTOMATION_TEST(FUnitUnitTests, "Core.Math.Unit Conversion", EA
 
 bool IsRoughlyEqual(double One, double Two, float Epsilon)
 {
-	return FMath::Abs(One-Two) < Epsilon;
+	return FMath::Abs(One-Two) <= Epsilon;
 }
 
 bool FUnitUnitTests::RunTest(const FString& Parameters)
@@ -33,8 +33,8 @@ bool FUnitUnitTests::RunTest(const FString& Parameters)
 		{ 1000,		1.0936,		1e-4,	EUnit::Millimeters, 		EUnit::Yards		 		},
 		{ 2000,		0.0787,		1e-4,	EUnit::Micrometers, 		EUnit::Inches		 		},
 
-		{ 90,		PI/2,		1e-6,	EUnit::Degrees, 			EUnit::Radians		 		},
-		{ PI,		180,		1e-6,	EUnit::Radians, 			EUnit::Degrees		 		},
+		{ 90,		PI/2,		1e-3,	EUnit::Degrees, 			EUnit::Radians		 		},
+		{ PI,		180,		1e-3,	EUnit::Radians, 			EUnit::Degrees		 		},
 
 		{ 12,		43.2,		0.1,	EUnit::MetersPerSecond,		EUnit::KilometersPerHour	},
 		{ 1,		0.6214,		1e-4,	EUnit::KilometersPerHour,	EUnit::MilesPerHour			},
@@ -90,9 +90,12 @@ bool FUnitUnitTests::RunTest(const FString& Parameters)
 			const TCHAR* FromUnitString	= FUnitConversion::GetUnitDisplayString(Test.FromUnit);
 			const TCHAR* ToUnitString 	= FUnitConversion::GetUnitDisplayString(Test.ToUnit);
 
-			AddError(FString::Printf(TEXT("Conversion from %s to %s was incorrect. %.3f%s !~= %.3f%s"),
+			AddError(FString::Printf(TEXT("Conversion from %s to %s was incorrect. Converting %.10f%s to %s resulted in %.15f%s, expected %.15f%s (threshold = %.15f)"),
 				FromUnitString, ToUnitString,
-				Test.Source, FromUnitString, ActualResult, ToUnitString)
+				Test.Source, FromUnitString, ToUnitString,
+				ActualResult, ToUnitString,
+				Test.ExpectedResult, ToUnitString,
+				Test.AccuracyEpsilon)
 			);
 		}
 	}
