Fix CIS in VS2012. Weirdness with operators in TFunction.

diff --git a/Engine/Source/Runtime/Core/Public/Templates/Function.h b/Engine/Source/Runtime/Core/Public/Templates/Function.h
--- a/Engine/Source/Runtime/Core/Public/Templates/Function.h
+++ b/Engine/Source/Runtime/Core/Public/Templates/Function.h
@@ -645,10 +645,16 @@ public:
 	SAFE_BOOL_OPERATORS(TFunction)
 
 private:
+
+
+	UE4Function_Private::IFunction_OwnedObject* Func;
+};
+
 	/**
 	 * Nullptr equality operator.
 	 */
-	friend bool operator==(TYPE_OF_NULLPTR, const TFunction& Func)
+	template <typename FuncType>
+	bool operator==(TYPE_OF_NULLPTR, const TFunction<FuncType>& Func)
 	{
 		return !Func;
 	}
@@ -656,7 +662,8 @@ private:
 	/**
 	 * Nullptr equality operator.
 	 */
-	friend bool operator==(const TFunction& Func, TYPE_OF_NULLPTR)
+	template <typename FuncType>
+	bool operator==(const TFunction<FuncType>& Func, TYPE_OF_NULLPTR)
 	{
 		return !Func;
 	}
@@ -664,7 +671,8 @@ private:
 	/**
 	 * Nullptr inequality operator.
 	 */
-	friend bool operator!=(TYPE_OF_NULLPTR, const TFunction& Func)
+	template <typename FuncType>
+	bool operator!=(TYPE_OF_NULLPTR, const TFunction<FuncType>& Func)
 	{
 		return (bool)Func;
 	}
@@ -672,10 +680,8 @@ private:
 	/**
 	 * Nullptr inequality operator.
 	 */
-	friend bool operator!=(const TFunction& Func, TYPE_OF_NULLPTR)
+	template <typename FuncType>
+	bool operator!=(const TFunction<FuncType>& Func, TYPE_OF_NULLPTR)
 	{
 		return (bool)Func;
-	}
-
-	UE4Function_Private::IFunction_OwnedObject* Func;
-};
+	}
\ No newline at end of file
