Stats - Fixed a leak

diff --git a/Engine/Source/Runtime/Core/Private/HAL/PThreadRunnableThread.cpp b/Engine/Source/Runtime/Core/Private/HAL/PThreadRunnableThread.cpp
new file mode 100644
--- /dev/null
+++ b/Engine/Source/Runtime/Core/Private/HAL/PThreadRunnableThread.cpp
@@ -0,0 +1,43 @@
+// Copyright 1998-2014 Epic Games, Inc. All Rights Reserved.
+
+#include "CorePrivatePCH.h"
+
+#if	PLATFORM_USE_PTHREADS
+
+#include "PThreadRunnableThread.h"
+
+uint32 FRunnableThreadPThread::Run()
+{
+	ThreadIsRunning = true;
+	// Assume we'll fail init
+	uint32 ExitCode = 1;
+	check(Runnable);
+
+	// Initialize the runnable object
+	if (Runnable->Init() == true)
+	{
+		// Initialization has completed, release the sync event
+		ThreadInitSyncEvent->Trigger();
+		// Now run the task that needs to be done
+		ExitCode = Runnable->Run();
+		// Allow any allocated resources to be cleaned up
+		Runnable->Exit();
+	}
+	else
+	{
+		// Initialization has failed, release the sync event
+		ThreadInitSyncEvent->Trigger();
+	}
+
+	// Clean ourselves up without waiting
+	ThreadIsRunning = false;
+
+#if STATS
+	FThreadStats::Shutdown();
+#endif
+
+	return ExitCode;
+}
+
+
+#endif // PLATFORM_USE_PTHREADS
\ No newline at end of file
diff --git a/Engine/Source/Runtime/Core/Private/HAL/PThreadRunnableThread.h b/Engine/Source/Runtime/Core/Private/HAL/PThreadRunnableThread.h
--- a/Engine/Source/Runtime/Core/Private/HAL/PThreadRunnableThread.h
+++ b/Engine/Source/Runtime/Core/Private/HAL/PThreadRunnableThread.h
@@ -224,33 +224,7 @@ protected:
 	 * The real thread entry point. It calls the Init/Run/Exit methods on
 	 * the runnable object
 	 */
-	uint32 Run()
-	{
-		ThreadIsRunning = true;
-		// Assume we'll fail init
-		uint32 ExitCode = 1;
-		check(Runnable);
-
-		// Initialize the runnable object
-		if (Runnable->Init() == true)
-		{
-			// Initialization has completed, release the sync event
-			ThreadInitSyncEvent->Trigger();
-			// Now run the task that needs to be done
-			ExitCode = Runnable->Run();
-			// Allow any allocated resources to be cleaned up
-			Runnable->Exit();
-		}
-		else
-		{
-			// Initialization has failed, release the sync event
-			ThreadInitSyncEvent->Trigger();
-		}
-
-		// Clean ourselves up without waiting
-		ThreadIsRunning = false;
-		return ExitCode;
-	}
+	uint32 Run();
 
 public:
 	FRunnableThreadPThread()
