FText key stablization & key editing

FText DisplayString shared ptrs are persisted across property edits. This preserves their entry in the live table and also updates the string that the entry points to. We call GetString() to update the hashed value as persisting the DisplayString means that its location in the table is preserved.
On serialization, a valid key is still found in the live table for the display string, so the key is preserved.
Added funtionality to the property editor to allow key editing and re-generation so any key conflicts can be resolved in the future.

UE-5350 - As a translator, I would like the editor tutorial text to preserve it's history between changes in order to greatly improve the translation workflow.

diff --git a/Engine/Source/Editor/EditorStyle/Private/SlateEditorStyle.cpp b/Engine/Source/Editor/EditorStyle/Private/SlateEditorStyle.cpp
--- a/Engine/Source/Editor/EditorStyle/Private/SlateEditorStyle.cpp
+++ b/Engine/Source/Editor/EditorStyle/Private/SlateEditorStyle.cpp
@@ -2969,6 +2969,7 @@ void FSlateEditorStyle::FStyle::SetupPropertyEditorStyles()
 		Set( "PropertyWindow.Button_PickAsset", new IMAGE_BRUSH( "Icons/pillarray_12x", Icon12x12 ) );
 		Set( "PropertyWindow.Button_PickActor", new IMAGE_BRUSH( "Icons/hiererchy_16x", Icon12x12 ) );
 		Set( "PropertyWindow.Button_PickActorInteractive", new IMAGE_BRUSH( "Icons/eyedropper_16px", Icon12x12 ) );
+		Set( "PropertyWindow.Button_Refresh", new IMAGE_BRUSH("Icons/refresh_12x", Icon12x12 ) );
 
 		Set( "PropertyWindow.WindowBorder", new BOX_BRUSH( "Common/GroupBorder", FMargin(4.0f/16.0f) ) );
 		Set( "DetailsView.NameChangeCommitted", new BOX_BRUSH( "Common/EditableTextSelectionBackground", FMargin(4.f/16.f) ) );
diff --git a/Engine/Source/Editor/IntroTutorials/Private/TutorialStructCustomization.cpp b/Engine/Source/Editor/IntroTutorials/Private/TutorialStructCustomization.cpp
--- a/Engine/Source/Editor/IntroTutorials/Private/TutorialStructCustomization.cpp
+++ b/Engine/Source/Editor/IntroTutorials/Private/TutorialStructCustomization.cpp
@@ -7,6 +7,7 @@
 #include "ISlateMetaData.h"
 #include "TutorialMetaData.h"
 #include "SDockTab.h"
+#include "PropertyCustomizationHelpers.h"
 
 #define LOCTEXT_NAMESPACE "TutorialStructCustomization"
 
@@ -165,6 +166,14 @@ void FTutorialContentCustomization::CustomizeHeader( TSharedRef<class IPropertyH
 				.OnTextCommitted(FOnTextCommitted::CreateStatic(&Local::OnTextCommitted, TextProperty))
 				.OnTextChanged(FOnTextChanged::CreateStatic(&Local::OnTextChanged, TextProperty))
 			]
+			+SHorizontalBox::Slot()
+			.AutoWidth()
+			.VAlign(VAlign_Center)
+			.HAlign(HAlign_Center)
+			.Padding(2.0f, 0.0f, 0.0f, 0.0f)
+			[
+				PropertyCustomizationHelpers::MakeTextLocalizationButton(TextProperty.ToSharedRef())
+			]
 		]
 	];
 }
diff --git a/Engine/Source/Editor/PropertyEditor/Private/PropertyCustomizationHelpers.cpp b/Engine/Source/Editor/PropertyEditor/Private/PropertyCustomizationHelpers.cpp
--- a/Engine/Source/Editor/PropertyEditor/Private/PropertyCustomizationHelpers.cpp
+++ b/Engine/Source/Editor/PropertyEditor/Private/PropertyCustomizationHelpers.cpp
@@ -1051,5 +1051,193 @@ void FMaterialList::AddMaterialItem( FDetailWidgetRow& Row, int32 CurrentSlot, c
 	];
 }
 
+
+TSharedRef<SWidget> PropertyCustomizationHelpers::MakeTextLocalizationButton(TSharedRef<IPropertyHandle> InPropertyHandle)
+{
+	struct Local
+	{
+		static void DirtyPropertyPackage(TSharedRef<IPropertyHandle> InPropertyHandle)
+		{
+			// dirty packages as they will need to be re-serialized with the correct key
+			TArray<UObject*> Objects;
+			InPropertyHandle->GetOuterObjects(Objects);
+			if (Objects.Num() > 0 && Objects[0]->GetOutermost() != nullptr)
+			{
+				Objects[0]->GetOutermost()->MarkPackageDirty();
+			}
+		}
+
+		static FReply HandleGenerateKeyClicked(TSharedRef<IPropertyHandle> InPropertyHandle)
+		{
+			FText DisplayText;
+			if (InPropertyHandle->GetValueAsDisplayText(DisplayText) == FPropertyAccess::Success)
+			{
+				TSharedRef< FString, ESPMode::ThreadSafe > DisplayString = FTextInspector::GetSharedDisplayString(DisplayText);
+				TSharedPtr< FString, ESPMode::ThreadSafe > Namespace;
+				TSharedPtr< FString, ESPMode::ThreadSafe > Key;
+				FTextLocalizationManager::Get().FindKeyNamespaceFromDisplayString(DisplayString, Namespace, Key);
+
+				if (Key.IsValid())
+				{
+					*Key.Get() = FGuid::NewGuid().ToString();
+					DirtyPropertyPackage(InPropertyHandle);
+				}
+			}
+
+			return FReply::Unhandled();
+		}
+
+		static void HandleKeyTextCommitted(const FText& InText, ETextCommit::Type InCommitType, TSharedRef<IPropertyHandle> InPropertyHandle)
+		{
+			FText DisplayText;
+			if (InPropertyHandle->GetValueAsDisplayText(DisplayText) == FPropertyAccess::Success)
+			{
+				TSharedRef< FString, ESPMode::ThreadSafe > DisplayString = FTextInspector::GetSharedDisplayString(DisplayText);
+				TSharedPtr< FString, ESPMode::ThreadSafe > Namespace;
+				TSharedPtr< FString, ESPMode::ThreadSafe > Key;
+				FTextLocalizationManager::Get().FindKeyNamespaceFromDisplayString(DisplayString, Namespace, Key);
+
+				if (Key.IsValid())
+				{
+					*Key.Get() = InText.ToString();
+					DirtyPropertyPackage(InPropertyHandle);
+				}
+			}
+		}
+
+		static void HandleNamespaceTextCommitted(const FText& InText, ETextCommit::Type InCommitType, TSharedRef<IPropertyHandle> InPropertyHandle)
+		{
+			FText DisplayText;
+			if (InPropertyHandle->GetValueAsDisplayText(DisplayText) == FPropertyAccess::Success)
+			{
+				TSharedRef< FString, ESPMode::ThreadSafe > DisplayString = FTextInspector::GetSharedDisplayString(DisplayText);
+				TSharedPtr< FString, ESPMode::ThreadSafe > Namespace;
+				TSharedPtr< FString, ESPMode::ThreadSafe > Key;
+				FTextLocalizationManager::Get().FindKeyNamespaceFromDisplayString(DisplayString, Namespace, Key);
+
+				if (Namespace.IsValid())
+				{
+					*Namespace.Get() = InText.ToString();
+
+					// dirty packages as they will need to be re-serialized with the correct namespace
+					TArray<UObject*> Objects;
+					InPropertyHandle->GetOuterObjects(Objects);
+					if (Objects.Num() > 0 && Objects[0]->GetOutermost() != nullptr)
+					{
+						Objects[0]->GetOutermost()->MarkPackageDirty();
+					}
+				}
+			}
+		}
+
+		static TSharedRef<SWidget> GetLocalizationMenuContent(TSharedRef<IPropertyHandle> InPropertyHandle)
+		{
+			FMenuBuilder MenuContentBuilder(true, NULL);
+			{
+				FText DisplayText;
+				if (InPropertyHandle->GetValueAsDisplayText(DisplayText) == FPropertyAccess::Success)
+				{
+					TSharedRef< FString, ESPMode::ThreadSafe > DisplayString = FTextInspector::GetSharedDisplayString(DisplayText);
+					TSharedPtr< FString, ESPMode::ThreadSafe > Namespace;
+					TSharedPtr< FString, ESPMode::ThreadSafe > Key;
+					FTextLocalizationManager::Get().FindKeyNamespaceFromDisplayString(DisplayString, Namespace, Key);
+
+					MenuContentBuilder.BeginSection(TEXT("Localization"), LOCTEXT("LocalizationSectionHeading", "Localization"));
+					{
+						if (!Namespace.IsValid() && !Key.IsValid())
+						{
+							MenuContentBuilder.AddWidget(SNew(SSpacer), LOCTEXT("NoLocWarning", "No localization available"));
+						}
+						else
+						{
+							TSharedRef<SGridPanel> GridPanel = SNew(SGridPanel);
+
+							if (Key.IsValid())
+							{
+								GridPanel->AddSlot(0, 0)
+									.VAlign(VAlign_Center)
+									.Padding(1.0f)
+									[
+										SNew(STextBlock)
+										.Text(LOCTEXT("KeyLabel", "Key"))
+										.ToolTipText(LOCTEXT("KeyTooltip", "The localization key of the text property"))
+									];
+
+								auto KeyLambda = [=](){ return FText::FromString(*Key.Get()); };
+
+								GridPanel->AddSlot(1, 0)
+									.VAlign(VAlign_Center)
+									.Padding(1.0f)
+									[
+										SNew(SHorizontalBox)
+										+ SHorizontalBox::Slot()
+										.FillWidth(1.0f)
+										[
+											SNew(SEditableTextBox)
+											.Text_Lambda(KeyLambda)
+											.OnTextCommitted(FOnTextCommitted::CreateStatic(&Local::HandleKeyTextCommitted, InPropertyHandle))
+										]
+										+ SHorizontalBox::Slot()
+											.AutoWidth()
+											.HAlign(HAlign_Center)
+											.Padding(2.0f, 0.0f, 0.0f, 0.0f)
+											[
+												SNew(SButton)
+												.ToolTipText(LOCTEXT("RefreshKeyTooltip", "Generate a new random key"))
+												.ButtonStyle(FEditorStyle::Get(), "HoverHintOnly")
+												.OnClicked(FOnClicked::CreateStatic(&Local::HandleGenerateKeyClicked, InPropertyHandle))
+												.Content()
+												[
+													SNew(SImage)
+													.Image(FEditorStyle::GetBrush("PropertyWindow.Button_Refresh"))
+												]
+											]
+									];
+							}
+
+							if (Namespace.IsValid())
+							{
+								auto NamespaceLambda = [=](){ return FText::FromString(*Namespace.Get()); };
+
+								GridPanel->AddSlot(0, 1)
+									.VAlign(VAlign_Center)
+									.Padding(1.0f, 1.0f, 5.0f, 1.0f)
+									[
+										SNew(STextBlock)
+										.Text(LOCTEXT("NamespaceLabel", "Namespace"))
+										.ToolTipText(LOCTEXT("NamespaceTooltip", "The localization namespace of the text property"))
+									];
+
+								GridPanel->AddSlot(1, 1)
+									.VAlign(VAlign_Center)
+									.Padding(1.0f)
+									[
+										SNew(SEditableTextBox)
+										.Text_Lambda(NamespaceLambda)
+										.OnTextCommitted(FOnTextCommitted::CreateStatic(&Local::HandleNamespaceTextCommitted, InPropertyHandle))
+									];
+							}
+
+							MenuContentBuilder.AddWidget(GridPanel, FText::GetEmpty());
+						}
+					}
+
+					MenuContentBuilder.EndSection();
+				}
+			}
+
+			return MenuContentBuilder.MakeWidget();
+		}
+	};
+
+	return SNew(SComboButton)
+		.ToolTipText(LOCTEXT("LocalizationUtilitiesTooltip", "Localization Utilities"))
+		.ButtonStyle(FEditorStyle::Get(), "HoverHintOnly")
+		.ContentPadding(2)
+		.ForegroundColor(FSlateColor::UseForeground())
+		.HasDownArrow(true)
+		.OnGetMenuContent(FOnGetContent::CreateStatic(&Local::GetLocalizationMenuContent, InPropertyHandle));
+}
+
 #undef LOCTEXT_NAMESPACE
 
diff --git a/Engine/Source/Editor/PropertyEditor/Private/UserInterface/PropertyEditor/SPropertyEditorText.cpp b/Engine/Source/Editor/PropertyEditor/Private/UserInterface/PropertyEditor/SPropertyEditorText.cpp
--- a/Engine/Source/Editor/PropertyEditor/Private/UserInterface/PropertyEditor/SPropertyEditorText.cpp
+++ b/Engine/Source/Editor/PropertyEditor/Private/UserInterface/PropertyEditor/SPropertyEditorText.cpp
@@ -13,38 +13,67 @@ void SPropertyEditorText::Construct( const FArguments& InArgs, const TSharedRef<
 	PropertyEditor = InPropertyEditor;
 
 	const UProperty* Property = InPropertyEditor->GetProperty();
+
+	TSharedPtr<SHorizontalBox> HorizontalBox;
+
 	bIsMultiLine = Property->GetBoolMetaData("MultiLine");
 	if(bIsMultiLine)
 	{
 		ChildSlot
 		[
-			SAssignNew( PrimaryWidget, SMultiLineEditableTextBox)
-			.Text( InPropertyEditor, &FPropertyEditor::GetValueAsText )
-			.Font( InArgs._Font )
-			.SelectAllTextWhenFocused( false )
-			.ClearKeyboardFocusOnCommit(false)
-			.OnTextCommitted( this, &SPropertyEditorText::OnTextCommitted )
-			.SelectAllTextOnCommit( false )
-			.IsReadOnly(this, &SPropertyEditorText::IsReadOnly)
-			.AutoWrapText(true)
-			.ModiferKeyForNewLine(EModifierKey::Shift)
+			SAssignNew(HorizontalBox, SHorizontalBox)
+			+SHorizontalBox::Slot()
+			.FillWidth(1.0f)
+			[
+				SAssignNew(PrimaryWidget, SMultiLineEditableTextBox)
+				.Text(InPropertyEditor, &FPropertyEditor::GetValueAsText)
+				.Font(InArgs._Font)
+				.SelectAllTextWhenFocused(false)
+				.ClearKeyboardFocusOnCommit(false)
+				.OnTextCommitted(this, &SPropertyEditorText::OnTextCommitted)
+				.SelectAllTextOnCommit(false)
+				.IsReadOnly(this, &SPropertyEditorText::IsReadOnly)
+				.AutoWrapText(true)
+				.ModiferKeyForNewLine(EModifierKey::Shift)
+			]
 		];
 	}
 	else
 	{
 		ChildSlot
 		[
-			SAssignNew( PrimaryWidget, SEditableTextBox )
-			.Text( InPropertyEditor, &FPropertyEditor::GetValueAsText )
-			.Font( InArgs._Font )
-			.SelectAllTextWhenFocused( true )
-			.ClearKeyboardFocusOnCommit(false)
-			.OnTextCommitted( this, &SPropertyEditorText::OnTextCommitted )
-			.SelectAllTextOnCommit( true )
-			.IsReadOnly(this, &SPropertyEditorText::IsReadOnly)
+			SAssignNew(HorizontalBox, SHorizontalBox)
+			+SHorizontalBox::Slot()
+			.FillWidth(1.0f)
+			[
+				SAssignNew( PrimaryWidget, SEditableTextBox )
+				.Text( InPropertyEditor, &FPropertyEditor::GetValueAsText )
+				.Font( InArgs._Font )
+				.SelectAllTextWhenFocused( true )
+				.ClearKeyboardFocusOnCommit(false)
+				.OnTextCommitted( this, &SPropertyEditorText::OnTextCommitted )
+				.SelectAllTextOnCommit( true )
+				.IsReadOnly(this, &SPropertyEditorText::IsReadOnly)
+			]
 		];
 	}
 
+	if (InPropertyEditor->PropertyIsA(UTextProperty::StaticClass()))
+	{
+		const TSharedRef< FPropertyNode > PropertyNode = InPropertyEditor->GetPropertyNode();
+		const TSharedRef< IPropertyHandle > PropertyHandle = PropertyEditor->GetPropertyHandle();
+
+		// text properties require some extra controls to allow for localization functionality
+		HorizontalBox->AddSlot()
+			.AutoWidth()
+			.VAlign(VAlign_Center)
+			.HAlign(HAlign_Center)
+			.Padding(2.0f, 0.0f, 0.0f, 0.0f)
+			[
+				PropertyCustomizationHelpers::MakeTextLocalizationButton(PropertyHandle)
+			];
+	}
+
 	if( InPropertyEditor->PropertyIsA( UObjectPropertyBase::StaticClass() ) )
 	{
 		// Object properties should display their entire text in a tooltip
diff --git a/Engine/Source/Editor/PropertyEditor/Public/PropertyCustomizationHelpers.h b/Engine/Source/Editor/PropertyEditor/Public/PropertyCustomizationHelpers.h
--- a/Engine/Source/Editor/PropertyEditor/Public/PropertyCustomizationHelpers.h
+++ b/Engine/Source/Editor/PropertyEditor/Public/PropertyCustomizationHelpers.h
@@ -40,6 +40,9 @@ namespace PropertyCustomizationHelpers
 
 	/** Returns a list of factories which can be used to create new assets, based on the supplied class */
 	PROPERTYEDITOR_API TArray<UFactory*> GetNewAssetFactoriesForClasses(const TArray<const UClass*>& Classes);
+
+	/** Build a combo button allowing access to text property localization utilities */
+	PROPERTYEDITOR_API TSharedRef<SWidget> MakeTextLocalizationButton(TSharedRef<IPropertyHandle> InPropertyHandle);
 }
 
 /** Delegate used to get a generic object */
diff --git a/Engine/Source/Runtime/Core/Private/Internationalization/Text.cpp b/Engine/Source/Runtime/Core/Private/Internationalization/Text.cpp
--- a/Engine/Source/Runtime/Core/Private/Internationalization/Text.cpp
+++ b/Engine/Source/Runtime/Core/Private/Internationalization/Text.cpp
@@ -16,7 +16,7 @@
 
 #define LOCTEXT_NAMESPACE "Core.Text"
 
-FArchive& operator<<( FArchive& Ar, FFormatArgumentData& Value )
+FArchive& operator<<(FArchive& Ar, FFormatArgumentData& Value)
 {
 	Ar << Value.ArgumentName;
 	Ar << Value.ArgumentValue;
@@ -54,6 +54,11 @@ const FString& FTextInspector::GetDisplayString(const FText& Text)
 	return Text.DisplayString.Get();
 }
 
+const TSharedRef<FString, ESPMode::ThreadSafe> FTextInspector::GetSharedDisplayString(const FText& Text)
+{
+	return Text.DisplayString;
+}
+
 int32 FTextInspector::GetFlags(const FText& Text)
 {
 	return Text.Flags;
@@ -111,82 +116,82 @@ FFormatArgumentValue::FFormatArgumentValue(const FText& Value)
 FFormatArgumentValue::FFormatArgumentValue(const FFormatArgumentValue& Source)
 	: Type(Source.Type)
 {
-	switch(Type)
+	switch (Type)
 	{
 	case EFormatArgumentType::Int:
-		{
-			IntValue = Source.IntValue;
-		}
+	{
+		IntValue = Source.IntValue;
+	}
 		break;
 	case EFormatArgumentType::UInt:
-		{
-			UIntValue = Source.UIntValue;
-		}
+	{
+		UIntValue = Source.UIntValue;
+	}
 		break;
 	case EFormatArgumentType::Float:
-		{
-			FloatValue = Source.FloatValue;
-		}
+	{
+		FloatValue = Source.FloatValue;
+	}
 	case EFormatArgumentType::Double:
-		{
-			DoubleValue = Source.DoubleValue;
-		}
+	{
+		DoubleValue = Source.DoubleValue;
+	}
 		break;
 	case EFormatArgumentType::Text:
-		{
-			TextValue = new FText(*Source.TextValue);
-		}
+	{
+		TextValue = new FText(*Source.TextValue);
+	}
 		break;
 	}
 }
 
 FFormatArgumentValue::~FFormatArgumentValue()
 {
-	if(Type == EFormatArgumentType::Text)
+	if (Type == EFormatArgumentType::Text)
 	{
 		delete TextValue;
 	}
 }
 
-FArchive& operator<<( FArchive& Ar, FFormatArgumentValue& Value )
+FArchive& operator<<(FArchive& Ar, FFormatArgumentValue& Value)
 {
 	int8 TypeAsInt8 = Value.Type;
 	Ar << TypeAsInt8;
 	Value.Type = (EFormatArgumentType::Type)TypeAsInt8;
 
-	switch(Value.Type)
+	switch (Value.Type)
 	{
 	case EFormatArgumentType::Double:
-		{
-			Ar << Value.DoubleValue;
-			break;
-		}
+	{
+		Ar << Value.DoubleValue;
+		break;
+	}
 	case EFormatArgumentType::Float:
-		{
-			Ar << Value.FloatValue;
-			break;
-		}
+	{
+		Ar << Value.FloatValue;
+		break;
+	}
 	case EFormatArgumentType::Int:
-		{
-			Ar << Value.IntValue;
-			break;
-		}
+	{
+		Ar << Value.IntValue;
+		break;
+	}
 	case EFormatArgumentType::UInt:
-		{
-			Ar << Value.UIntValue;
-			break;
-		}
+	{
+		Ar << Value.UIntValue;
+		break;
+	}
 	case EFormatArgumentType::Text:
+	{
+		if (Ar.IsLoading())
 		{
-			if(Ar.IsLoading())
-			{
-				Value.TextValue = new FText;
-			}
-			Ar << *Value.TextValue;
-			break;
+			Value.TextValue = new FText;
 		}
+		Ar << *Value.TextValue;
+		break;
+	}
 	}
-	
+
 	return Ar;
 }
 
@@ -235,22 +240,22 @@ const FText FText::UnEscapedCloseBraceOutsideOfArgumentBlock = LOCTEXT("Error_Un
 const FText FText::SerializationFailureError = LOCTEXT("Error_SerializationFailure", "ERR: Transient text cannot be serialized \"{0}\".");
 
 FText::FText()
-	: DisplayString( StaticEmptyText.DisplayString )
+	: DisplayString(StaticEmptyText.DisplayString)
 	, Flags(0)
 {
 }
 
-FText::FText( EInitToEmptyString )
-	: DisplayString( new FString() )
+FText::FText(EInitToEmptyString)
+	: DisplayString(new FString())
 	, Flags(0)
 {
 }
 
 #if PLATFORM_COMPILER_HAS_DEFAULTED_FUNCTIONS
-	FText::FText(const FText& Other) = default;
-	FText::FText(FText&& Other) = default;
-	FText& FText::operator=(const FText& Other) = default;
-	FText& FText::operator=(FText&& Other) = default;
+FText::FText(const FText& Other) = default;
+FText::FText(FText&& Other) = default;
+FText& FText::operator=(const FText& Other) = default;
+FText& FText::operator=(FText&& Other) = default;
 #else
 FText::FText(const FText& Source)
 	: DisplayString(Source.DisplayString)
@@ -285,15 +290,15 @@ FText& FText::operator=(FText&& Source)
 }
 #endif
 
-FText::FText( FString InSourceString )
-	: DisplayString( new FString( MoveTemp(InSourceString) ))
+FText::FText(FString InSourceString)
+	: DisplayString(new FString(MoveTemp(InSourceString)))
 	, Flags(0)
 {
 	History = MakeShareable(new FTextHistory_Base(DisplayString));
 }
 
-FText::FText( FString InSourceString, FString InNamespace, FString InKey, int32 InFlags )
-	: DisplayString( FTextLocalizationManager::Get().GetString(InNamespace, InKey, &InSourceString) )
+FText::FText(FString InSourceString, FString InNamespace, FString InKey, int32 InFlags)
+	: DisplayString(FTextLocalizationManager::Get().GetString(InNamespace, InKey, &InSourceString))
 	, Flags(InFlags)
 {
 	History = MakeShareable(new FTextHistory_Base(InSourceString));
@@ -307,7 +312,7 @@ bool FText::IsEmptyOrWhitespace() const
 		return true;
 	}
 
-	for( const TCHAR Character : DisplayString.Get() )
+	for (const TCHAR Character : DisplayString.Get())
 	{
 		if (!IsWhitespace(Character))
 		{
@@ -318,14 +323,14 @@ bool FText::IsEmptyOrWhitespace() const
 	return true;
 }
 
-FText FText::TrimPreceding( const FText& InText )
+FText FText::TrimPreceding(const FText& InText)
 {
 	FString TrimmedString = InText.ToString();
 	{
 		int32 StartPos = 0;
-		while ( StartPos < TrimmedString.Len() )
+		while (StartPos < TrimmedString.Len())
 		{
-			if( !FText::IsWhitespace( TrimmedString[StartPos] ) )
+			if (!FText::IsWhitespace(TrimmedString[StartPos]))
 			{
 				break;
 			}
@@ -333,14 +338,14 @@ FText FText::TrimPreceding( const FText& InText )
 			++StartPos;
 		}
 
-		TrimmedString = TrimmedString.Right( TrimmedString.Len() - StartPos );
+		TrimmedString = TrimmedString.Right(TrimmedString.Len() - StartPos);
 	}
 
-	FText NewText = FText( MoveTemp( TrimmedString ) );
+	FText NewText = FText(MoveTemp(TrimmedString));
 
 	if (!GIsEditor)
 	{
-		if( (InText.Flags & (1 << ETextFlag::CultureInvariant)) != 0 )
+		if ((InText.Flags & (1 << ETextFlag::CultureInvariant)) != 0)
 		{
 			NewText.Flags = NewText.Flags | ETextFlag::Transient;
 		}
@@ -353,14 +358,14 @@ FText FText::TrimPreceding( const FText& InText )
 	return NewText;
 }
 
-FText FText::TrimTrailing( const FText& InText )
+FText FText::TrimTrailing(const FText& InText)
 {
 	FString TrimmedString = InText.ToString();
 	{
 		int32 EndPos = TrimmedString.Len() - 1;
-		while( EndPos >= 0 )
+		while (EndPos >= 0)
 		{
-			if( !FText::IsWhitespace( TrimmedString[EndPos] ) )
+			if (!FText::IsWhitespace(TrimmedString[EndPos]))
 			{
 				break;
 			}
@@ -368,14 +373,14 @@ FText FText::TrimTrailing( const FText& InText )
 			EndPos--;
 		}
 
-		TrimmedString = TrimmedString.Left( EndPos + 1 );
+		TrimmedString = TrimmedString.Left(EndPos + 1);
 	}
 
-	FText NewText = FText( MoveTemp ( TrimmedString ) );
+	FText NewText = FText(MoveTemp(TrimmedString));
 
 	if (!GIsEditor)
 	{
-		if( (InText.Flags & (1 << ETextFlag::CultureInvariant)) != 0 )
+		if ((InText.Flags & (1 << ETextFlag::CultureInvariant)) != 0)
 		{
 			NewText.Flags = NewText.Flags & ETextFlag::Transient;
 		}
@@ -388,14 +393,14 @@ FText FText::TrimTrailing( const FText& InText )
 	return NewText;
 }
 
-FText FText::TrimPrecedingAndTrailing( const FText& InText )
+FText FText::TrimPrecedingAndTrailing(const FText& InText)
 {
 	FString TrimmedString = InText.ToString();
 	{
 		int32 StartPos = 0;
-		while ( StartPos < TrimmedString.Len() )
+		while (StartPos < TrimmedString.Len())
 		{
-			if( !FText::IsWhitespace( TrimmedString[StartPos] ) )
+			if (!FText::IsWhitespace(TrimmedString[StartPos]))
 			{
 				break;
 			}
@@ -404,9 +409,9 @@ FText FText::TrimPrecedingAndTrailing( const FText& InText )
 		}
 
 		int32 EndPos = TrimmedString.Len() - 1;
-		while( EndPos >= 0 )
+		while (EndPos >= 0)
 		{
-			if( !FText::IsWhitespace( TrimmedString[EndPos] ) )
+			if (!FText::IsWhitespace(TrimmedString[EndPos]))
 			{
 				break;
 			}
@@ -415,14 +420,14 @@ FText FText::TrimPrecedingAndTrailing( const FText& InText )
 		}
 
 		const int32 Len = (EndPos + 1) - StartPos;
-		TrimmedString = TrimmedString.Mid( StartPos, Len );
+		TrimmedString = TrimmedString.Mid(StartPos, Len);
 	}
 
-	FText NewText = FText( MoveTemp( TrimmedString ) );
+	FText NewText = FText(MoveTemp(TrimmedString));
 
 	if (!GIsEditor)
 	{
-		if( (InText.Flags & (1 << ETextFlag::CultureInvariant)) != 0 )
+		if ((InText.Flags & (1 << ETextFlag::CultureInvariant)) != 0)
 		{
 			NewText.Flags = NewText.Flags | ETextFlag::Transient;
 		}
@@ -435,14 +440,14 @@ FText FText::TrimPrecedingAndTrailing( const FText& InText )
 	return NewText;
 }
 
-FText FText::Format(const FText& Fmt,const FText& v1)
+FText FText::Format(const FText& Fmt, const FText& v1)
 {
 	FFormatOrderedArguments Arguments;
 	Arguments.Add(v1);
 	return FText::Format(Fmt, Arguments);
 }
 
-FText FText::Format(const FText& Fmt,const FText& v1,const FText& v2)
+FText FText::Format(const FText& Fmt, const FText& v1, const FText& v2)
 {
 	FFormatOrderedArguments Arguments;
 	Arguments.Add(v1);
@@ -450,7 +455,7 @@ FText FText::Format(const FText& Fmt,const FText& v1,const FText& v2)
 	return FText::Format(Fmt, Arguments);
 }
 
-FText FText::Format(const FText& Fmt,const FText& v1,const FText& v2,const FText& v3)
+FText FText::Format(const FText& Fmt, const FText& v1, const FText& v2, const FText& v3)
 {
 	FFormatOrderedArguments Arguments;
 	Arguments.Add(v1);
@@ -459,7 +464,7 @@ FText FText::Format(const FText& Fmt,const FText& v1,const FText& v2,const FText
 	return FText::Format(Fmt, Arguments);
 }
 
-FText FText::Format(const FText& Fmt,const FText& v1,const FText& v2,const FText& v3,const FText& v4)
+FText FText::Format(const FText& Fmt, const FText& v1, const FText& v2, const FText& v3, const FText& v4)
 {
 	FFormatOrderedArguments Arguments;
 	Arguments.Add(v1);
@@ -492,20 +497,20 @@ DEF_ASNUMBER_CAST(uint64, int64_t)
 #undef DEF_ASNUMBER_CAST
 
 /**
- * Generate an FText that represents the passed number as currency in the current culture
- */
+* Generate an FText that represents the passed number as currency in the current culture
+*/
 #define DEF_ASCURRENCY_CAST(T1, T2) FText FText::AsCurrency(T1 Val, const FString& CurrencyCode, const FNumberFormattingOptions* const Options, const FCulturePtr& TargetCulture) { return FText::AsCurrencyTemplate<T1, T2>(Val, CurrencyCode, Options, TargetCulture); }
 #define DEF_ASCURRENCY(T) DEF_ASCURRENCY_CAST(T, T)
 DEF_ASCURRENCY(float)
-	DEF_ASCURRENCY(double)
-	DEF_ASCURRENCY(int8)
-	DEF_ASCURRENCY(int16)
-	DEF_ASCURRENCY(int32)
-	DEF_ASCURRENCY_CAST(int64, int64_t)
-	DEF_ASCURRENCY(uint8)
-	DEF_ASCURRENCY(uint16)
-	DEF_ASCURRENCY_CAST(uint32, int64_t)
-	DEF_ASCURRENCY_CAST(uint64, int64_t)
+DEF_ASCURRENCY(double)
+DEF_ASCURRENCY(int8)
+DEF_ASCURRENCY(int16)
+DEF_ASCURRENCY(int32)
+DEF_ASCURRENCY_CAST(int64, int64_t)
+DEF_ASCURRENCY(uint8)
+DEF_ASCURRENCY(uint16)
+DEF_ASCURRENCY_CAST(uint32, int64_t)
+DEF_ASCURRENCY_CAST(uint64, int64_t)
 #undef DEF_ASCURRENCY
 #undef DEF_ASCURRENCY_CAST
 
@@ -520,26 +525,26 @@ DEF_ASPERCENT(float)
 #undef DEF_ASPERCENT
 #undef DEF_ASPERCENT_CAST
 
-bool FText::FindText( const FString& Namespace, const FString& Key, FText& OutText, const FString* const SourceString )
+bool FText::FindText(const FString& Namespace, const FString& Key, FText& OutText, const FString* const SourceString)
 {
-	TSharedPtr< FString, ESPMode::ThreadSafe > FoundString = FTextLocalizationManager::Get().FindString( Namespace, Key );
+	TSharedPtr< FString, ESPMode::ThreadSafe > FoundString = FTextLocalizationManager::Get().FindString(Namespace, Key);
 
-	if ( FoundString.IsValid() )
+	if (FoundString.IsValid())
 	{
-		OutText = FText( SourceString ? *SourceString : FString(), Namespace, Key );
+		OutText = FText(SourceString ? *SourceString : FString(), Namespace, Key);
 	}
 
 	return FoundString.IsValid();
 }
 
-CORE_API FArchive& operator<<( FArchive& Ar, FText& Value )
+CORE_API FArchive& operator<<(FArchive& Ar, FText& Value)
 {
 	//When duplicating, the CDO is used as the template, then values for the instance are assigned.
 	//If we don't duplicate the string, the CDO and the instance are both pointing at the same thing.
 	//This would result in all subsequently duplicated objects stamping over formerly duplicated ones.
 
 	// Older FText's stored their "SourceString", that is now stored in a history class so move it there
-	if(Ar.IsLoading() && Ar.UE4Ver() < VER_UE4_FTEXT_HISTORY)
+	if (Ar.IsLoading() && Ar.UE4Ver() < VER_UE4_FTEXT_HISTORY)
 	{
 		FString SourceStringToImplantIntoHistory;
 		Ar << SourceStringToImplantIntoHistory;
@@ -547,15 +552,15 @@ CORE_API FArchive& operator<<( FArchive& Ar, FText& Value )
 		Value.History = MakeShareable(new FTextHistory_Base(SourceStringToImplantIntoHistory));
 
 		// Namespaces and keys are no longer stored in the FText, we need to read them in and discard
-		if( Ar.UE4Ver() >= VER_UE4_ADDED_NAMESPACE_AND_KEY_DATA_TO_FTEXT )
+		if (Ar.UE4Ver() >= VER_UE4_ADDED_NAMESPACE_AND_KEY_DATA_TO_FTEXT)
 		{
 			TSharedPtr< FString, ESPMode::ThreadSafe > Namespace;
 			TSharedPtr< FString, ESPMode::ThreadSafe > Key;
 
-			Namespace = MakeShareable( new FString() );
+			Namespace = MakeShareable(new FString());
 			Ar << *Namespace;
 
-			Key = MakeShareable( new FString() );
+			Key = MakeShareable(new FString());
 			Ar << *Key;
 
 			// Get the DisplayString using the namespace, key, and source string.
@@ -565,18 +570,18 @@ CORE_API FArchive& operator<<( FArchive& Ar, FText& Value )
 		}
 	}
 
-	if(Ar.IsPersistent())
+	if (Ar.IsPersistent())
 	{
 		Value.Flags &= ~(ETextFlag::ConvertedProperty); // Remove conversion flag before saving.
 	}
 	Ar << Value.Flags;
 
-	if( Ar.UE4Ver() >= VER_UE4_FTEXT_HISTORY )
+	if (Ar.UE4Ver() >= VER_UE4_FTEXT_HISTORY)
 	{
-		if(Ar.IsSaving())
+		if (Ar.IsSaving())
 		{
 			// If there is no history, mark it, otherwise the history will serialize it's type
-			if(!Value.History.IsValid())
+			if (!Value.History.IsValid())
 			{
 				int8 NoHistory = INDEX_NONE;
 				Ar << NoHistory;
@@ -589,73 +594,73 @@ CORE_API FArchive& operator<<( FArchive& Ar, FText& Value )
 			Ar << HistoryType;
 
 			// Create the history class based on the serialized type
-			switch((ETextHistoryType)HistoryType)
+			switch ((ETextHistoryType)HistoryType)
 			{
 			case ETextHistoryType::Base:
-				{
-					Value.History = MakeShareable(new FTextHistory_Base);
-					break;
-				}
+			{
+				Value.History = MakeShareable(new FTextHistory_Base);
+				break;
+			}
 			case ETextHistoryType::NamedFormat:
-				{
-					Value.History = MakeShareable(new FTextHistory_NamedFormat);
-					break;
-				}
+			{
+				Value.History = MakeShareable(new FTextHistory_NamedFormat);
+				break;
+			}
 			case ETextHistoryType::OrderedFormat:
-				{
-					Value.History = MakeShareable(new FTextHistory_OrderedFormat);
-					break;
-				}
+			{
+				Value.History = MakeShareable(new FTextHistory_OrderedFormat);
+				break;
+			}
 			case ETextHistoryType::ArgumentFormat:
-				{
-					Value.History = MakeShareable(new FTextHistory_ArgumentDataFormat);
-					break;
-				}
+			{
+				Value.History = MakeShareable(new FTextHistory_ArgumentDataFormat);
+				break;
+			}
 			case ETextHistoryType::AsNumber:
-				{
-					Value.History = MakeShareable(new FTextHistory_AsNumber);
-					break;
-				}
+			{
+				Value.History = MakeShareable(new FTextHistory_AsNumber);
+				break;
+			}
 			case ETextHistoryType::AsPercent:
-				{
-					Value.History = MakeShareable(new FTextHistory_AsPercent);
-					break;
-				}
+			{
+				Value.History = MakeShareable(new FTextHistory_AsPercent);
+				break;
+			}
 			case ETextHistoryType::AsCurrency:
-				{
-					Value.History = MakeShareable(new FTextHistory_AsCurrency);
-					break;
-				}
+			{
+				Value.History = MakeShareable(new FTextHistory_AsCurrency);
+				break;
+			}
 			case ETextHistoryType::AsDate:
-				{
-					Value.History = MakeShareable(new FTextHistory_AsDate);
-					break;
-				}
+			{
+				Value.History = MakeShareable(new FTextHistory_AsDate);
+				break;
+			}
 			case ETextHistoryType::AsTime:
-				{
-					Value.History = MakeShareable(new FTextHistory_AsTime);
-					break;
-				}
+			{
+				Value.History = MakeShareable(new FTextHistory_AsTime);
+				break;
+			}
 			case ETextHistoryType::AsDateTime:
-				{
-					Value.History = MakeShareable(new FTextHistory_AsDateTime);
-					break;
-				}
+			{
+				Value.History = MakeShareable(new FTextHistory_AsDateTime);
+				break;
+			}
 			default:
-				{
-					Value.History.Reset();
-				}
+			{
+				Value.History.Reset();
+			}
 			}
 		}
 
-		if(Value.History.IsValid())
+		if (Value.History.IsValid())
 		{
 			Value.History->Serialize(Ar);
 			Value.History->SerializeForDisplayString(Ar, Value.DisplayString);
 		}
 	}
 
-	if(Ar.IsLoading())
+	if (Ar.IsLoading())
 	{
 		Value.Rebuild();
 	}
@@ -666,7 +671,7 @@ CORE_API FArchive& operator<<( FArchive& Ar, FText& Value )
 	//	UE_LOG( LogText, Error, TEXT("Loaded an FText from a persistent archive but lacking a key (Namespace:%s, Source:%s)."), Value.Namespace.Get() ? **Value.Namespace : TEXT(""), Value.SourceString.Get() ? **Value.SourceString : TEXT("") );
 	//}
 
-	if( Value.ShouldGatherForLocalization() )
+	if (Value.ShouldGatherForLocalization())
 	{
 		Ar.ThisRequiresLocalizationGather();
 	}
@@ -675,15 +680,15 @@ CORE_API FArchive& operator<<( FArchive& Ar, FText& Value )
 }
 
 #if WITH_EDITOR
-FText FText::ChangeKey( FString Namespace, FString Key, const FText& Text )
+FText FText::ChangeKey(FString Namespace, FString Key, const FText& Text)
 {
-	return FText( *Text.History->GetSourceString().Get(), MoveTemp( Namespace ), MoveTemp( Key ) );
+	return FText(*Text.History->GetSourceString().Get(), MoveTemp(Namespace), MoveTemp(Key));
 }
 #endif
 
 FText FText::CreateNumericalText(FString InSourceString)
 {
-	FText NewText = FText( MoveTemp( InSourceString ) );
+	FText NewText = FText(MoveTemp(InSourceString));
 	if (!GIsEditor)
 	{
 		NewText.Flags |= ETextFlag::Transient;
@@ -693,7 +698,7 @@ FText FText::CreateNumericalText(FString InSourceString)
 
 FText FText::CreateChronologicalText(FString InSourceString)
 {
-	FText NewText = FText( MoveTemp( InSourceString ) );
+	FText NewText = FText(MoveTemp(InSourceString));
 	if (!GIsEditor)
 	{
 		NewText.Flags |= ETextFlag::Transient;
@@ -701,21 +706,21 @@ FText FText::CreateChronologicalText(FString InSourceString)
 	return NewText;
 }
 
-FText FText::FromName( const FName& Val) 
+FText FText::FromName(const FName& Val)
 {
-	FText NewText = FText( Val.ToString() );
+	FText NewText = FText(Val.ToString());
 
 	if (!GIsEditor)
 	{
 		NewText.Flags |= ETextFlag::CultureInvariant;
 	}
 
-	return NewText; 
+	return NewText;
 }
 
-FText FText::FromString( FString String )
+FText FText::FromString(FString String)
 {
-	FText NewText = FText( MoveTemp(String) );
+	FText NewText = FText(MoveTemp(String));
 
 	if (!GIsEditor)
 	{
@@ -725,17 +730,17 @@ FText FText::FromString( FString String )
 	return NewText;
 }
 
-FText FText::AsCultureInvariant( FString String )
+FText FText::AsCultureInvariant(FString String)
 {
-	FText NewText = FText( String );
+	FText NewText = FText(String);
 	NewText.Flags |= ETextFlag::CultureInvariant;
 
 	return NewText;
 }
 
-FText FText::AsCultureInvariant( FText Text )
+FText FText::AsCultureInvariant(FText Text)
 {
-	FText NewText = FText( Text );
+	FText NewText = FText(Text);
 	NewText.Flags |= ETextFlag::CultureInvariant;
 
 	return NewText;
@@ -750,7 +755,7 @@ const FString& FText::ToString() const
 
 FString FText::BuildSourceString() const
 {
-	if(History.IsValid())
+	if (History.IsValid())
 	{
 		return History->ToText(true).ToString();
 	}
@@ -759,7 +764,7 @@ FString FText::BuildSourceString() const
 
 void FText::Rebuild() const
 {
-	if(History.IsValid())
+	if (History.IsValid())
 	{
 		History->Rebuild(DisplayString);
 	}
@@ -771,9 +776,9 @@ bool FText::ShouldGatherForLocalization() const
 
 	auto IsAllWhitespace = [](FString& String) -> bool
 	{
-		for(int32 i = 0; i < String.Len(); ++i)
+		for (int32 i = 0; i < String.Len(); ++i)
 		{
-			if( !FText::IsWhitespace( String[i] ) )
+			if (!FText::IsWhitespace(String[i]))
 			{
 				return false;
 			}
@@ -786,10 +791,10 @@ bool FText::ShouldGatherForLocalization() const
 
 TSharedPtr< FString, ESPMode::ThreadSafe > FText::GetSourceString() const
 {
-	if(History.IsValid())
+	if (History.IsValid())
 	{
 		TSharedPtr< FString, ESPMode::ThreadSafe > SourceString = History->GetSourceString();
-		if(SourceString.IsValid())
+		if (SourceString.IsValid())
 		{
 			return SourceString;
 		}
@@ -798,7 +803,7 @@ TSharedPtr< FString, ESPMode::ThreadSafe > FText::GetSourceString() const
 	return DisplayString;
 }
 
-bool FText::IdenticalTo( const FText& Other ) const
+bool FText::IdenticalTo(const FText& Other) const
 {
 	// If both instances point to the same string, then both instances are considered identical
 	// This is fast as it skips a lexical compare, however it can also return true for two instances that have identical strings, but in different pointers
@@ -843,4 +848,37 @@ bool FTextSnapshot::IsDisplayStringEqualTo(const FText& InText) const
 	return HistoryRevision == InHistoryRevision && DisplayStringPtr.IsValid() && DisplayStringPtr->Equals(InText.ToString(), ESearchCase::CaseSensitive);
 }
 
+FScopedTextIdentityPreserver::FScopedTextIdentityPreserver(FText& InTextToPersist)
+	: TextToPersist(InTextToPersist)
+{
+	// Empty display strings can't have a namespace or key.
+	if (!InTextToPersist.DisplayString->IsEmpty())
+	{
+		// Save off namespace and key to be restored later.
+		FTextLocalizationManager::Get().FindKeyNamespaceFromDisplayString(InTextToPersist.DisplayString, Namespace, Key);
+	}
+}
+
+FScopedTextIdentityPreserver::~FScopedTextIdentityPreserver()
+{
+	// If we don't have a key, then the old identity wasn't valid and shouldn't be preserved.
+	if (Key.IsValid())
+	{
+		// Get the text's new source string.
+		const FString* SourceString = FTextInspector::GetSourceString(TextToPersist);
+
+		// Without a source string, we can't possibly preserve the identity. If the text we're preserving identity for can't possibly have an identity anymore, this class shouldn't be used on this text.
+		check(SourceString);
+
+		// Make the history so that it can attempt find and serialize the identity for this source string when the history itself is serialized.
+		TextToPersist.History = MakeShareable(new FTextHistory_Base(*SourceString));
+
+		// Create/update the display string instance for this identity in the text localization manager...
+		const TSharedRef<FString, ESPMode::ThreadSafe> DisplayString = FTextLocalizationManager::Get().GetString(Namespace.IsValid() ? *Namespace : TEXT(""), *Key, SourceString);
+
+		// ... and set it so that the text's history's serialization will properly map this display string instance back to the identity we intended to preserve and restore.
+		TextToPersist.DisplayString = DisplayString;
+	}
+}
+
 #undef LOCTEXT_NAMESPACE
diff --git a/Engine/Source/Runtime/Core/Private/Internationalization/TextLocalizationManager.cpp b/Engine/Source/Runtime/Core/Private/Internationalization/TextLocalizationManager.cpp
--- a/Engine/Source/Runtime/Core/Private/Internationalization/TextLocalizationManager.cpp
+++ b/Engine/Source/Runtime/Core/Private/Internationalization/TextLocalizationManager.cpp
@@ -467,7 +467,7 @@ void FTextLocalizationManager::UpdateLiveTable(const TArray<FLocalizationEntryTr
 						);
 					LiveKeyTable.Add( Key, NewLiveEntry );
 
-					ReverseLiveTable.Add(NewLiveEntry.String, FNamespaceKeyEntry( NamespaceName.IsEmpty()? nullptr : MakeShareable( new FString( NamespaceName ) ), Key.IsEmpty()? nullptr : MakeShareable( new FString( Key ) )));
+					ReverseLiveTable.Add(NewLiveEntry.String, FNamespaceKeyEntry( MakeShareable( new FString( NamespaceName ) ), MakeShareable( new FString( Key ) )));
 				}
 			}
 		}
@@ -638,7 +638,7 @@ TSharedRef<FString, ESPMode::ThreadSafe> FTextLocalizationManager::GetString(con
 
 			LiveKeyTable->Add( Key, NewEntry );
 
-			ReverseLiveTable.Add(NewEntry.String, FNamespaceKeyEntry( Namespace.IsEmpty()? nullptr : MakeShareable( new FString( Namespace ) ), Key.IsEmpty()? nullptr : MakeShareable( new FString( Key ) )));
+			ReverseLiveTable.Add(NewEntry.String, FNamespaceKeyEntry( MakeShareable( new FString( Namespace ) ), MakeShareable( new FString( Key ) )));
 
 		}
 
diff --git a/Engine/Source/Runtime/Core/Public/Internationalization/Text.h b/Engine/Source/Runtime/Core/Public/Internationalization/Text.h
--- a/Engine/Source/Runtime/Core/Public/Internationalization/Text.h
+++ b/Engine/Source/Runtime/Core/Public/Internationalization/Text.h
@@ -18,9 +18,9 @@ namespace ETextFlag
 {
 	enum Type
 	{
-		Transient = (1<<0),
-		CultureInvariant = (1<<1),
-		ConvertedProperty = (1<<2)
+		Transient = (1 << 0),
+		CultureInvariant = (1 << 1),
+		ConvertedProperty = (1 << 2)
 	};
 }
 
@@ -95,22 +95,22 @@ struct CORE_API FNumberFormattingOptions
 	FNumberFormattingOptions();
 
 	bool UseGrouping;
-	FNumberFormattingOptions& SetUseGrouping( bool InValue ){ UseGrouping = InValue; return *this; }
+	FNumberFormattingOptions& SetUseGrouping(bool InValue){ UseGrouping = InValue; return *this; }
 
 	ERoundingMode RoundingMode;
-	FNumberFormattingOptions& SetRoundingMode( ERoundingMode InValue ){ RoundingMode = InValue; return *this; }
+	FNumberFormattingOptions& SetRoundingMode(ERoundingMode InValue){ RoundingMode = InValue; return *this; }
 
 	int32 MinimumIntegralDigits;
-	FNumberFormattingOptions& SetMinimumIntegralDigits( int32 InValue ){ MinimumIntegralDigits = InValue; return *this; }
+	FNumberFormattingOptions& SetMinimumIntegralDigits(int32 InValue){ MinimumIntegralDigits = InValue; return *this; }
 
 	int32 MaximumIntegralDigits;
-	FNumberFormattingOptions& SetMaximumIntegralDigits( int32 InValue ){ MaximumIntegralDigits = InValue; return *this; }
+	FNumberFormattingOptions& SetMaximumIntegralDigits(int32 InValue){ MaximumIntegralDigits = InValue; return *this; }
 
 	int32 MinimumFractionalDigits;
-	FNumberFormattingOptions& SetMinimumFractionalDigits( int32 InValue ){ MinimumFractionalDigits = InValue; return *this; }
+	FNumberFormattingOptions& SetMinimumFractionalDigits(int32 InValue){ MinimumFractionalDigits = InValue; return *this; }
 
 	int32 MaximumFractionalDigits;
-	FNumberFormattingOptions& SetMaximumFractionalDigits( int32 InValue ){ MaximumFractionalDigits = InValue; return *this; }
+	FNumberFormattingOptions& SetMaximumFractionalDigits(int32 InValue){ MaximumFractionalDigits = InValue; return *this; }
 
 	friend FArchive& operator<<(FArchive& Ar, FNumberFormattingOptions& Value);
 };
@@ -129,86 +129,86 @@ public:
 public:
 
 	FText();
-	FText( const FText& Source );
+	FText(const FText& Source);
 	FText(FText&& Source);
 
 	FText& operator=(const FText& Source);
 	FText& operator=(FText&& Source);
 
 	/**
-	 * Generate an FText that represents the passed number in the current culture
-	 */
-	static FText AsNumber(float Val,	const FNumberFormattingOptions* const Options = NULL, const FCulturePtr& TargetCulture = NULL);
-	static FText AsNumber(double Val,	const FNumberFormattingOptions* const Options = NULL, const FCulturePtr& TargetCulture = NULL);
-	static FText AsNumber(int8 Val,		const FNumberFormattingOptions* const Options = NULL, const FCulturePtr& TargetCulture = NULL);
-	static FText AsNumber(int16 Val,	const FNumberFormattingOptions* const Options = NULL, const FCulturePtr& TargetCulture = NULL);
-	static FText AsNumber(int32 Val,	const FNumberFormattingOptions* const Options = NULL, const FCulturePtr& TargetCulture = NULL);
-	static FText AsNumber(int64 Val,	const FNumberFormattingOptions* const Options = NULL, const FCulturePtr& TargetCulture = NULL);
-	static FText AsNumber(uint8 Val,	const FNumberFormattingOptions* const Options = NULL, const FCulturePtr& TargetCulture = NULL);
-	static FText AsNumber(uint16 Val,	const FNumberFormattingOptions* const Options = NULL, const FCulturePtr& TargetCulture = NULL);
-	static FText AsNumber(uint32 Val,	const FNumberFormattingOptions* const Options = NULL, const FCulturePtr& TargetCulture = NULL);
-	static FText AsNumber(uint64 Val,	const FNumberFormattingOptions* const Options = NULL, const FCulturePtr& TargetCulture = NULL);
-	static FText AsNumber(long Val,		const FNumberFormattingOptions* const Options = NULL, const FCulturePtr& TargetCulture = NULL);
+	* Generate an FText that represents the passed number in the current culture
+	*/
+	static FText AsNumber(float Val, const FNumberFormattingOptions* const Options = NULL, const FCulturePtr& TargetCulture = NULL);
+	static FText AsNumber(double Val, const FNumberFormattingOptions* const Options = NULL, const FCulturePtr& TargetCulture = NULL);
+	static FText AsNumber(int8 Val, const FNumberFormattingOptions* const Options = NULL, const FCulturePtr& TargetCulture = NULL);
+	static FText AsNumber(int16 Val, const FNumberFormattingOptions* const Options = NULL, const FCulturePtr& TargetCulture = NULL);
+	static FText AsNumber(int32 Val, const FNumberFormattingOptions* const Options = NULL, const FCulturePtr& TargetCulture = NULL);
+	static FText AsNumber(int64 Val, const FNumberFormattingOptions* const Options = NULL, const FCulturePtr& TargetCulture = NULL);
+	static FText AsNumber(uint8 Val, const FNumberFormattingOptions* const Options = NULL, const FCulturePtr& TargetCulture = NULL);
+	static FText AsNumber(uint16 Val, const FNumberFormattingOptions* const Options = NULL, const FCulturePtr& TargetCulture = NULL);
+	static FText AsNumber(uint32 Val, const FNumberFormattingOptions* const Options = NULL, const FCulturePtr& TargetCulture = NULL);
+	static FText AsNumber(uint64 Val, const FNumberFormattingOptions* const Options = NULL, const FCulturePtr& TargetCulture = NULL);
+	static FText AsNumber(long Val, const FNumberFormattingOptions* const Options = NULL, const FCulturePtr& TargetCulture = NULL);
 
 	/**
-	 * Generate an FText that represents the passed number as currency in the current culture
-	 */
-	static FText AsCurrency(float Val,  const FString& CurrencyCode = FString(), const FNumberFormattingOptions* const Options = NULL, const FCulturePtr& TargetCulture = NULL);
+	* Generate an FText that represents the passed number as currency in the current culture
+	*/
+	static FText AsCurrency(float Val, const FString& CurrencyCode = FString(), const FNumberFormattingOptions* const Options = NULL, const FCulturePtr& TargetCulture = NULL);
 	static FText AsCurrency(double Val, const FString& CurrencyCode = FString(), const FNumberFormattingOptions* const Options = NULL, const FCulturePtr& TargetCulture = NULL);
-	static FText AsCurrency(int8 Val,   const FString& CurrencyCode = FString(), const FNumberFormattingOptions* const Options = NULL, const FCulturePtr& TargetCulture = NULL);
-	static FText AsCurrency(int16 Val,  const FString& CurrencyCode = FString(), const FNumberFormattingOptions* const Options = NULL, const FCulturePtr& TargetCulture = NULL);
-	static FText AsCurrency(int32 Val,  const FString& CurrencyCode = FString(), const FNumberFormattingOptions* const Options = NULL, const FCulturePtr& TargetCulture = NULL);
-	static FText AsCurrency(int64 Val,  const FString& CurrencyCode = FString(), const FNumberFormattingOptions* const Options = NULL, const FCulturePtr& TargetCulture = NULL);
-	static FText AsCurrency(uint8 Val,  const FString& CurrencyCode = FString(), const FNumberFormattingOptions* const Options = NULL, const FCulturePtr& TargetCulture = NULL);
+	static FText AsCurrency(int8 Val, const FString& CurrencyCode = FString(), const FNumberFormattingOptions* const Options = NULL, const FCulturePtr& TargetCulture = NULL);
+	static FText AsCurrency(int16 Val, const FString& CurrencyCode = FString(), const FNumberFormattingOptions* const Options = NULL, const FCulturePtr& TargetCulture = NULL);
+	static FText AsCurrency(int32 Val, const FString& CurrencyCode = FString(), const FNumberFormattingOptions* const Options = NULL, const FCulturePtr& TargetCulture = NULL);
+	static FText AsCurrency(int64 Val, const FString& CurrencyCode = FString(), const FNumberFormattingOptions* const Options = NULL, const FCulturePtr& TargetCulture = NULL);
+	static FText AsCurrency(uint8 Val, const FString& CurrencyCode = FString(), const FNumberFormattingOptions* const Options = NULL, const FCulturePtr& TargetCulture = NULL);
 	static FText AsCurrency(uint16 Val, const FString& CurrencyCode = FString(), const FNumberFormattingOptions* const Options = NULL, const FCulturePtr& TargetCulture = NULL);
 	static FText AsCurrency(uint32 Val, const FString& CurrencyCode = FString(), const FNumberFormattingOptions* const Options = NULL, const FCulturePtr& TargetCulture = NULL);
 	static FText AsCurrency(uint64 Val, const FString& CurrencyCode = FString(), const FNumberFormattingOptions* const Options = NULL, const FCulturePtr& TargetCulture = NULL);
-	static FText AsCurrency(long Val,   const FString& CurrencyCode = FString(), const FNumberFormattingOptions* const Options = NULL, const FCulturePtr& TargetCulture = NULL);
+	static FText AsCurrency(long Val, const FString& CurrencyCode = FString(), const FNumberFormattingOptions* const Options = NULL, const FCulturePtr& TargetCulture = NULL);
 
 	/**
-	 * Generate an FText that represents the passed number as a percentage in the current culture
-	 */
-	static FText AsPercent(float Val,	const FNumberFormattingOptions* const Options = NULL, const FCulturePtr& TargetCulture = NULL);
-	static FText AsPercent(double Val,	const FNumberFormattingOptions* const Options = NULL, const FCulturePtr& TargetCulture = NULL);
+	* Generate an FText that represents the passed number as a percentage in the current culture
+	*/
+	static FText AsPercent(float Val, const FNumberFormattingOptions* const Options = NULL, const FCulturePtr& TargetCulture = NULL);
+	static FText AsPercent(double Val, const FNumberFormattingOptions* const Options = NULL, const FCulturePtr& TargetCulture = NULL);
 
 	/**
-	 * Generate an FText that represents the passed number as a date and/or time in the current culture
-	 */
+	* Generate an FText that represents the passed number as a date and/or time in the current culture
+	*/
 	static FText AsDate(const FDateTime& DateTime, const EDateTimeStyle::Type DateStyle = EDateTimeStyle::Default, const FString& TimeZone = TEXT(""), const FCulturePtr& TargetCulture = NULL);
 	static FText AsDateTime(const FDateTime& DateTime, const EDateTimeStyle::Type DateStyle = EDateTimeStyle::Default, const EDateTimeStyle::Type TimeStyle = EDateTimeStyle::Default, const FString& TimeZone = TEXT(""), const FCulturePtr& TargetCulture = NULL);
 	static FText AsTime(const FDateTime& DateTime, const EDateTimeStyle::Type TimeStyle = EDateTimeStyle::Default, const FString& TimeZone = TEXT(""), const FCulturePtr& TargetCulture = NULL);
 	static FText AsTimespan(const FTimespan& Timespan, const FCulturePtr& TargetCulture = NULL);
 
 	/**
-	 * Generate an FText that represents the passed number as a memory size in the current culture
-	 */
+	* Generate an FText that represents the passed number as a memory size in the current culture
+	*/
 	static FText AsMemory(SIZE_T NumBytes, const FNumberFormattingOptions* const Options = NULL, const FCulturePtr& TargetCulture = NULL);
 
 	/**
-	 * Attempts to find an existing FText using the representation found in the loc tables for the specified namespace and key
-	 * @return true if OutText was properly set; otherwise false and OutText will be untouched
-	 */
-	static bool FindText( const FString& Namespace, const FString& Key, FText& OutText, const FString* const SourceString = nullptr );
+	* Attempts to find an existing FText using the representation found in the loc tables for the specified namespace and key
+	* @return true if OutText was properly set; otherwise false and OutText will be untouched
+	*/
+	static bool FindText(const FString& Namespace, const FString& Key, FText& OutText, const FString* const SourceString = nullptr);
 
 	/**
-	 * Generate an FText representing the pass name
-	 */
-	static FText FromName( const FName& Val);
-	
+	* Generate an FText representing the pass name
+	*/
+	static FText FromName(const FName& Val);
+
 	/**
-	 * Generate an FText representing the passed in string
-	 */
-	static FText FromString( FString String );
+	* Generate an FText representing the passed in string
+	*/
+	static FText FromString(FString String);
 
 	/**
-	 * Generate a culture invariant FText representing the passed in string
-	 */
-	static FText AsCultureInvariant( FString String );
+	* Generate a culture invariant FText representing the passed in string
+	*/
+	static FText AsCultureInvariant(FString String);
 
 	/**
-	 * Generate a culture invariant FText representing the passed in FText
-	 */
-	static FText AsCultureInvariant( FText Text );
+	* Generate a culture invariant FText representing the passed in FText
+	*/
+	static FText AsCultureInvariant(FText Text);
 
 	const FString& ToString() const;
 
@@ -220,19 +220,19 @@ public:
 		return DisplayString.Get().IsNumeric();
 	}
 
-	int32 CompareTo( const FText& Other, const ETextComparisonLevel::Type ComparisonLevel = ETextComparisonLevel::Default ) const;
-	int32 CompareToCaseIgnored( const FText& Other ) const;
+	int32 CompareTo(const FText& Other, const ETextComparisonLevel::Type ComparisonLevel = ETextComparisonLevel::Default) const;
+	int32 CompareToCaseIgnored(const FText& Other) const;
 
-	bool EqualTo( const FText& Other, const ETextComparisonLevel::Type ComparisonLevel = ETextComparisonLevel::Default ) const;
-	bool EqualToCaseIgnored( const FText& Other ) const;
+	bool EqualTo(const FText& Other, const ETextComparisonLevel::Type ComparisonLevel = ETextComparisonLevel::Default) const;
+	bool EqualToCaseIgnored(const FText& Other) const;
 
 	/**
-	 * Check to see if this FText is identical to the other FText
-	 * 
-	 * Note:	This doesn't compare the text, but only checks that the internal string pointers have the same target (which makes it very fast!)
-	 *			If you actually want to perform a lexical comparison, then you need to use EqualTo instead
-	 */
-	bool IdenticalTo( const FText& Other ) const;
+	* Check to see if this FText is identical to the other FText
+	*
+	* Note:	This doesn't compare the text, but only checks that the internal string pointers have the same target (which makes it very fast!)
+	*			If you actually want to perform a lexical comparison, then you need to use EqualTo instead
+	*/
+	bool IdenticalTo(const FText& Other) const;
 
 	class CORE_API FSortPredicate
 	{
@@ -256,24 +256,24 @@ public:
 	bool IsEmptyOrWhitespace() const;
 
 	/**
-	 * Removes whitespace characters from the front of the string.
-	 */
-	static FText TrimPreceding( const FText& );
+	* Removes whitespace characters from the front of the string.
+	*/
+	static FText TrimPreceding(const FText&);
 
 	/**
-	 * Removes trailing whitespace characters
-	 */
-	static FText TrimTrailing( const FText& );
+	* Removes trailing whitespace characters
+	*/
+	static FText TrimTrailing(const FText&);
 
 	/**
-	 * Does both of the above without needing to create an additional FText in the interim.
-	 */
-	static FText TrimPrecedingAndTrailing( const FText& );
+	* Does both of the above without needing to create an additional FText in the interim.
+	*/
+	static FText TrimPrecedingAndTrailing(const FText&);
 
 	/**
-	 * Check to see if the given character is considered whitespace by the current culture
-	 */
-	static bool IsWhitespace( const TCHAR Char );
+	* Check to see if the given character is considered whitespace by the current culture
+	*/
+	static bool IsWhitespace(const TCHAR Char);
 
 	static void GetFormatPatternParameters(const FText& Pattern, TArray<FString>& ParameterNames);
 
@@ -281,13 +281,13 @@ public:
 	static FText Format(const FText& Pattern, const FFormatOrderedArguments& Arguments);
 	static FText Format(const FText& Pattern, const TArray< struct FFormatArgumentData > InArguments);
 
-	static FText Format(const FText& Fmt,const FText& v1);
-	static FText Format(const FText& Fmt,const FText& v1,const FText& v2);
-	static FText Format(const FText& Fmt,const FText& v1,const FText& v2,const FText& v3);
-	static FText Format(const FText& Fmt,const FText& v1,const FText& v2,const FText& v3,const FText& v4);
+	static FText Format(const FText& Fmt, const FText& v1);
+	static FText Format(const FText& Fmt, const FText& v1, const FText& v2);
+	static FText Format(const FText& Fmt, const FText& v1, const FText& v2, const FText& v3);
+	static FText Format(const FText& Fmt, const FText& v1, const FText& v2, const FText& v3, const FText& v4);
 
-	static void SetEnableErrorCheckingResults(bool bEnable){bEnableErrorCheckingResults=bEnable;}
-	static bool GetEnableErrorCheckingResults(){return bEnableErrorCheckingResults;}
+	static void SetEnableErrorCheckingResults(bool bEnable){ bEnableErrorCheckingResults = bEnable; }
+	static bool GetEnableErrorCheckingResults(){ return bEnableErrorCheckingResults; }
 
 	static void SetSuppressWarnings(bool bSuppress){ bSuppressWarnings = bSuppress; }
 	static bool GetSuppressWarnings(){ return bSuppressWarnings; }
@@ -299,29 +299,29 @@ private:
 
 	/** Special constructor used to create StaticEmptyText without also allocating a history object */
 	enum class EInitToEmptyString : uint8 { Value };
-	explicit FText( EInitToEmptyString );
+	explicit FText(EInitToEmptyString);
 
-	explicit FText( FString InSourceString );
+	explicit FText(FString InSourceString);
 
-	explicit FText( FString InSourceString, FString InNamespace, FString InKey, int32 InFlags=0 );
+	explicit FText(FString InSourceString, FString InNamespace, FString InKey, int32 InFlags = 0);
 
-	friend CORE_API FArchive& operator<<( FArchive& Ar, FText& Value );
+	friend CORE_API FArchive& operator<<(FArchive& Ar, FText& Value);
 
 #if WITH_EDITOR
 	/**
-	 * Constructs a new FText with the SourceString of the specified text but with the specified namespace and key
-	 */
-	static FText ChangeKey( FString Namespace, FString Key, const FText& Text );
+	* Constructs a new FText with the SourceString of the specified text but with the specified namespace and key
+	*/
+	static FText ChangeKey(FString Namespace, FString Key, const FText& Text);
 #endif
 
 	/**
-	 * Generate an FText for a string formatted numerically.
-	 */
+	* Generate an FText for a string formatted numerically.
+	*/
 	static FText CreateNumericalText(FString InSourceString);
 
 	/**
-	 * Generate an FText for a string formatted from a date/time.
-	 */
+	* Generate an FText for a string formatted from a date/time.
+	*/
 	static FText CreateChronologicalText(FString InSourceString);
 
 	/** Returns the source string of the FText */
@@ -381,6 +381,7 @@ public:
 	friend class FTextHistory_NamedFormat;
 	friend class FTextHistory_ArgumentDataFormat;
 	friend class FTextHistory_OrderedFormat;
+	friend class FScopedTextIdentityPreserver;
 
 #if !UE_ENABLE_ICU
 	friend class FLegacyTextHelper;
@@ -424,6 +425,7 @@ public:
 	static const FString* GetKey(const FText& Text);
 	static const FString* GetSourceString(const FText& Text);
 	static const FString& GetDisplayString(const FText& Text);
+	static const TSharedRef<FString, ESPMode::ThreadSafe> GetSharedDisplayString(const FText& Text);
 	static int32 GetFlags(const FText& Text);
 };
 
@@ -451,7 +453,7 @@ struct CORE_API FFormatArgumentValue
 	FFormatArgumentValue(const FFormatArgumentValue& Source);
 	~FFormatArgumentValue();
 
-	friend FArchive& operator<<( FArchive& Ar, FFormatArgumentValue& Value );
+	friend FArchive& operator<<(FArchive& Ar, FFormatArgumentValue& Value);
 };
 
 struct FFormatArgumentData
@@ -459,7 +461,7 @@ struct FFormatArgumentData
 	FText ArgumentName;
 	FText ArgumentValue;
 
-	friend FArchive& operator<<( FArchive& Ar, FFormatArgumentData& Value );
+	friend FArchive& operator<<(FArchive& Ar, FFormatArgumentData& Value);
 };
 
 class CORE_API FTextBuilder
@@ -563,4 +565,16 @@ private:
 	int32 IndentCount;
 };
 
+class CORE_API FScopedTextIdentityPreserver
+{
+public:
+	FScopedTextIdentityPreserver(FText& InTextToPersist);
+	~FScopedTextIdentityPreserver();
+
+private:
+	TSharedPtr< FString, ESPMode::ThreadSafe > Namespace;
+	TSharedPtr< FString, ESPMode::ThreadSafe > Key;
+	FText& TextToPersist;
+};
+
 Expose_TNameOf(FText)
diff --git a/Engine/Source/Runtime/Core/Public/Internationalization/TextLocalizationManager.h b/Engine/Source/Runtime/Core/Public/Internationalization/TextLocalizationManager.h
--- a/Engine/Source/Runtime/Core/Public/Internationalization/TextLocalizationManager.h
+++ b/Engine/Source/Runtime/Core/Public/Internationalization/TextLocalizationManager.h
@@ -79,8 +79,8 @@ private:
 public:
 	static CORE_API FTextLocalizationManager& Get();
 
-	TSharedPtr<FString, ESPMode::ThreadSafe> FindString( const FString& Namespace, const FString& Key, const FString* const SourceString = nullptr );
-	TSharedRef<FString, ESPMode::ThreadSafe> GetString(const FString& Namespace, const FString& Key, const FString* const SourceString);
+	TSharedPtr<FString, ESPMode::ThreadSafe> CORE_API FindString(const FString& Namespace, const FString& Key, const FString* const SourceString = nullptr);
+	TSharedRef<FString, ESPMode::ThreadSafe> CORE_API GetString(const FString& Namespace, const FString& Key, const FString* const SourceString);
 
 	/**
 	 * Using an FText SourceString's, returns the associated Namespace and Key
diff --git a/Engine/Source/Runtime/CoreUObject/Private/UObject/UTextProperty.cpp b/Engine/Source/Runtime/CoreUObject/Private/UObject/UTextProperty.cpp
--- a/Engine/Source/Runtime/CoreUObject/Private/UObject/UTextProperty.cpp
+++ b/Engine/Source/Runtime/CoreUObject/Private/UObject/UTextProperty.cpp
@@ -4,21 +4,21 @@
 #include "UObject/UTextProperty.h"
 #include "PropertyHelper.h"
 
-bool UTextProperty::Identical( const void* A, const void* B, uint32 PortFlags ) const
+bool UTextProperty::Identical(const void* A, const void* B, uint32 PortFlags) const
 {
 	const TCppType ValueA = GetPropertyValue(A);
-	if ( B )
+	if (B)
 	{
 		const TCppType ValueB = GetPropertyValue(B);
 
-		if ( ValueA.IsCultureInvariant() != ValueB.IsCultureInvariant() || ValueA.IsTransient() != ValueB.IsTransient() )
+		if (ValueA.IsCultureInvariant() != ValueB.IsCultureInvariant() || ValueA.IsTransient() != ValueB.IsTransient())
 		{
 			//A culture variant text is never equal to a culture invariant text
 			//A transient text is never equal to a non-transient text
 			return false;
 		}
-		
-		if ( ValueA.IsCultureInvariant() == ValueB.IsCultureInvariant() || ValueA.IsTransient() == ValueB.IsTransient() )
+
+		if (ValueA.IsCultureInvariant() == ValueB.IsCultureInvariant() || ValueA.IsTransient() == ValueB.IsTransient())
 		{
 			//Culture invariant text don't have a namespace/key so we compare the source string
 			//Transient text don't have a namespace/key or source so we compare the display string
@@ -31,48 +31,48 @@ bool UTextProperty::Identical( const void* A, const void* B, uint32 PortFlags )
 		}
 		else
 		{
-			return	FTextInspector::GetNamespace(ValueA)	==	FTextInspector::GetNamespace(ValueB) &&
-					FTextInspector::GetKey(ValueA)			==	FTextInspector::GetKey(ValueB);
+			return	FTextInspector::GetNamespace(ValueA) == FTextInspector::GetNamespace(ValueB) &&
+				FTextInspector::GetKey(ValueA) == FTextInspector::GetKey(ValueB);
 		}
 	}
 
 	return FTextInspector::GetDisplayString(ValueA).IsEmpty();
 }
 
-void UTextProperty::SerializeItem( FArchive& Ar, void* Value, int32 MaxReadBytes, void const* Defaults ) const
+void UTextProperty::SerializeItem(FArchive& Ar, void* Value, int32 MaxReadBytes, void const* Defaults) const
 {
 	const TCppType PropertyValue = GetPropertyValue(Value);
-	if ( Ar.IsSaving() && Ar.IsPersistent() && PropertyValue.IsTransient() )
+	if (Ar.IsSaving() && Ar.IsPersistent() && PropertyValue.IsTransient())
 	{
 		// Convert to a string and back, to get it set as an invariant message instead of transient. If it stayed transient, error messages recursively expand
-		FString ErrorMessage = FText::Format( FText::SerializationFailureError, FText::FromString( FTextInspector::GetDisplayString(PropertyValue) ) ).ToString();
-		UE_LOG( LogProperty, Warning, TEXT("%s"), *ErrorMessage);
+		FString ErrorMessage = FText::Format(FText::SerializationFailureError, FText::FromString(FTextInspector::GetDisplayString(PropertyValue))).ToString();
+		UE_LOG(LogProperty, Warning, TEXT("%s"), *ErrorMessage);
 		SetPropertyValue(Value, FText::FromString(ErrorMessage));
 	}
 
 	Ar << *GetPropertyValuePtr(Value);
 }
 
-void UTextProperty::ExportTextItem( FString& ValueStr, const void* PropertyValue, const void* DefaultValue, UObject* Parent, int32 PortFlags, UObject* ExportRootScope ) const
+void UTextProperty::ExportTextItem(FString& ValueStr, const void* PropertyValue, const void* DefaultValue, UObject* Parent, int32 PortFlags, UObject* ExportRootScope) const
 {
 	const FText& TextValue = GetPropertyValue(PropertyValue);
 	const FString& StringValue = FTextInspector::GetDisplayString(TextValue);
-	if( !(PortFlags & PPF_Delimited) )
+	if (!(PortFlags & PPF_Delimited))
 	{
 		ValueStr += StringValue;
 	}
-	else if ( StringValue.Len() > 0 )
+	else if (StringValue.Len() > 0)
 	{
-		ValueStr += FString::Printf( TEXT("\"%s\""), *(StringValue.ReplaceCharWithEscapedChar()) );
+		ValueStr += FString::Printf(TEXT("\"%s\""), *(StringValue.ReplaceCharWithEscapedChar()));
 	}
 }
 
-const TCHAR* UTextProperty::ImportText_Internal( const TCHAR* Buffer, void* Data, int32 PortFlags, UObject* Parent, FOutputDevice* ErrorText ) const
+const TCHAR* UTextProperty::ImportText_Internal(const TCHAR* Buffer, void* Data, int32 PortFlags, UObject* Parent, FOutputDevice* ErrorText) const
 {
 	FText* Text = GetPropertyValuePtr(Data);
 
 	FString BufferAsString;
-	if( !(PortFlags & PPF_Delimited) )
+	if (!(PortFlags & PPF_Delimited))
 	{
 		BufferAsString = FString(Buffer);
 
@@ -106,7 +106,8 @@ const TCHAR* UTextProperty::ImportText_Internal( const TCHAR* Buffer, void* Data
 		BufferAsString = FString(Temp);
 	}
 
-	*Text = FText::FromString( FString(BufferAsString) );
+	FScopedTextIdentityPreserver Preserver(*Text);
+	*Text = FText::FromString(FString(BufferAsString));
 
 	return Buffer;
 }
