Delegate comparisons deprecated, lots of other associated code deprecated, and lots of warning fixups:

diff --git a/Engine/Source/Editor/BehaviorTreeEditor/Private/BehaviorTreeEditor.cpp b/Engine/Source/Editor/BehaviorTreeEditor/Private/BehaviorTreeEditor.cpp
--- a/Engine/Source/Editor/BehaviorTreeEditor/Private/BehaviorTreeEditor.cpp
+++ b/Engine/Source/Editor/BehaviorTreeEditor/Private/BehaviorTreeEditor.cpp
@@ -52,8 +52,8 @@ FBehaviorTreeEditor::FBehaviorTreeEditor()
 	}
 
 	// listen for package change events to update injected nodes
-	UPackage::PackageSavedEvent.AddRaw(this, &FBehaviorTreeEditor::OnPackageSaved);
-	FClassBrowseHelper::OnPackageListUpdated.AddRaw(this, &FBehaviorTreeEditor::OnClassListUpdated);
+	OnPackageSavedDelegateHandle     = UPackage::PackageSavedEvent.AddRaw(this, &FBehaviorTreeEditor::OnPackageSaved);
+	OnClassListUpdatedDelegateHandle = FClassBrowseHelper::OnPackageListUpdated.AddRaw(this, &FBehaviorTreeEditor::OnClassListUpdated);
 
 	bShowDecoratorRangeLower = false;
 	bShowDecoratorRangeSelf = false;
@@ -74,8 +74,8 @@ FBehaviorTreeEditor::~FBehaviorTreeEditor()
 		Editor->UnregisterForUndo( this );
 	}
 
-	UPackage::PackageSavedEvent.RemoveRaw(this, &FBehaviorTreeEditor::OnPackageSaved);
-	FClassBrowseHelper::OnPackageListUpdated.RemoveRaw(this, &FBehaviorTreeEditor::OnClassListUpdated);
+	UPackage::PackageSavedEvent.Remove(OnPackageSavedDelegateHandle);
+	FClassBrowseHelper::OnPackageListUpdated.Remove(OnClassListUpdatedDelegateHandle);
 
 	Debugger.Reset();
 }
diff --git a/Engine/Source/Editor/BehaviorTreeEditor/Private/BehaviorTreeEditor.h b/Engine/Source/Editor/BehaviorTreeEditor/Private/BehaviorTreeEditor.h
--- a/Engine/Source/Editor/BehaviorTreeEditor/Private/BehaviorTreeEditor.h
+++ b/Engine/Source/Editor/BehaviorTreeEditor/Private/BehaviorTreeEditor.h
@@ -323,6 +323,12 @@ private:
 	/** Whether the current selection is inherited, stored here so it can be accessed by our details customization */
 	bool bIsCurrentBlackboardEntryInherited;
 
+	/** Handle to the registered OnPackageSave delegate */
+	FDelegateHandle OnPackageSavedDelegateHandle;
+
+	/** Handle to the registered OnClassListUpdated delegate */
+	FDelegateHandle OnClassListUpdatedDelegateHandle;
+
 public:
 	/** Modes in mode switcher */
 	static const FName BehaviorTreeMode;
diff --git a/Engine/Source/Editor/BlueprintGraph/Private/BlueprintActionDatabase.cpp b/Engine/Source/Editor/BlueprintGraph/Private/BlueprintActionDatabase.cpp
--- a/Engine/Source/Editor/BlueprintGraph/Private/BlueprintActionDatabase.cpp
+++ b/Engine/Source/Editor/BlueprintGraph/Private/BlueprintActionDatabase.cpp
@@ -1170,7 +1170,7 @@ void FBlueprintActionDatabase::RefreshAssetActions(UObject* const AssetObject)
 		// blueprint
 		if (!bHadExistingEntry)
 		{
-			BlueprintAsset->OnChanged().AddStatic(&BlueprintActionDatabaseImpl::OnBlueprintChanged);
+			OnBlueprintChangedDelegateHandle = BlueprintAsset->OnChanged().AddStatic(&BlueprintActionDatabaseImpl::OnBlueprintChanged);
 		}
 	}
 
@@ -1226,7 +1226,7 @@ void FBlueprintActionDatabase::ClearAssetActions(UObject* const AssetObject)
 
 	if (UBlueprint* BlueprintAsset = Cast<UBlueprint>(AssetObject))
 	{
-		BlueprintAsset->OnChanged().RemoveStatic(&BlueprintActionDatabaseImpl::OnBlueprintChanged);
+		BlueprintAsset->OnChanged().Remove(OnBlueprintChangedDelegateHandle);
 	}
 
 	if (bHasEntry && (ActionList->Num() > 0) && !BlueprintActionDatabaseImpl::bIsInitializing)
diff --git a/Engine/Source/Editor/BlueprintGraph/Public/BlueprintActionDatabase.h b/Engine/Source/Editor/BlueprintGraph/Public/BlueprintActionDatabase.h
--- a/Engine/Source/Editor/BlueprintGraph/Public/BlueprintActionDatabase.h
+++ b/Engine/Source/Editor/BlueprintGraph/Public/BlueprintActionDatabase.h
@@ -153,4 +153,7 @@ private:
 	/** */
 	FOnDatabaseEntryUpdated EntryRefreshDelegate;
 	FOnDatabaseEntryUpdated EntryRemovedDelegate;
+
+	/** Handle to the registered OnBlueprintChanged delegate. */
+	FDelegateHandle OnBlueprintChangedDelegateHandle;
 };
diff --git a/Engine/Source/Editor/ClassViewer/Private/SClassViewer.cpp b/Engine/Source/Editor/ClassViewer/Private/SClassViewer.cpp
--- a/Engine/Source/Editor/ClassViewer/Private/SClassViewer.cpp
+++ b/Engine/Source/Editor/ClassViewer/Private/SClassViewer.cpp
@@ -428,6 +428,11 @@ private:
 private:
 	/** The "Object" class node that is used as a rooting point for the Class Viewer. */
 	TSharedPtr< FClassViewerNode > ObjectClassRoot;
+
+	/** Handles to various registered RequestPopulateClassHierarchy delegates */
+	FDelegateHandle OnFilesLoadedRequestPopulateClassHierarchyDelegateHandle;
+	FDelegateHandle OnBlueprintCompiledRequestPopulateClassHierarchyDelegateHandle;
+	FDelegateHandle OnClassPackageLoadedOrUnloadedRequestPopulateClassHierarchyDelegateHandle;
 };
 
 namespace ClassViewer
@@ -1592,7 +1597,7 @@ FClassHierarchy::FClassHierarchy()
 {
 	// Register with the Asset Registry to be informed when it is done loading up files.
 	FAssetRegistryModule& AssetRegistryModule = FModuleManager::GetModuleChecked<FAssetRegistryModule>(TEXT("AssetRegistry"));
-	AssetRegistryModule.Get().OnFilesLoaded().AddStatic( ClassViewer::Helpers::RequestPopulateClassHierarchy );
+	OnFilesLoadedRequestPopulateClassHierarchyDelegateHandle = AssetRegistryModule.Get().OnFilesLoaded().AddStatic( ClassViewer::Helpers::RequestPopulateClassHierarchy );
 	AssetRegistryModule.Get().OnAssetAdded().AddRaw( this, &FClassHierarchy::AddAsset);
 	AssetRegistryModule.Get().OnAssetRemoved().AddRaw( this, &FClassHierarchy::RemoveAsset );
 
@@ -1601,8 +1606,8 @@ FClassHierarchy::FClassHierarchy()
 	HotReloadSupport.OnHotReload().AddRaw( this, &FClassHierarchy::OnHotReload );
 
 	// Register to have Populate called when a Blueprint is compiled.
-	GEditor->OnBlueprintCompiled().AddStatic(ClassViewer::Helpers::RequestPopulateClassHierarchy);
-	GEditor->OnClassPackageLoadedOrUnloaded().AddStatic(ClassViewer::Helpers::RequestPopulateClassHierarchy);
+	OnBlueprintCompiledRequestPopulateClassHierarchyDelegateHandle            = GEditor->OnBlueprintCompiled().AddStatic(ClassViewer::Helpers::RequestPopulateClassHierarchy);
+	OnClassPackageLoadedOrUnloadedRequestPopulateClassHierarchyDelegateHandle = GEditor->OnClassPackageLoadedOrUnloaded().AddStatic(ClassViewer::Helpers::RequestPopulateClassHierarchy);
 
 	FModuleManager::Get().OnModulesChanged().AddStatic(&OnModulesChanged);
 }
@@ -1613,7 +1618,7 @@ FClassHierarchy::~FClassHierarchy()
 	if( FModuleManager::Get().IsModuleLoaded( TEXT("AssetRegistry") ) )
 	{
 		FAssetRegistryModule& AssetRegistryModule = FModuleManager::GetModuleChecked<FAssetRegistryModule>(TEXT("AssetRegistry"));
-		AssetRegistryModule.Get().OnFilesLoaded().RemoveStatic(ClassViewer::Helpers::RequestPopulateClassHierarchy);
+		AssetRegistryModule.Get().OnFilesLoaded().Remove(OnFilesLoadedRequestPopulateClassHierarchyDelegateHandle);
 		AssetRegistryModule.Get().OnAssetAdded().RemoveAll( this );
 		AssetRegistryModule.Get().OnAssetRemoved().RemoveAll( this );
 
@@ -1622,8 +1627,8 @@ FClassHierarchy::~FClassHierarchy()
 		HotReloadSupport.OnHotReload().RemoveAll( this );
 
 		// Unregister to have Populate called when a Blueprint is compiled.
-		GEditor->OnBlueprintCompiled().RemoveStatic(ClassViewer::Helpers::RequestPopulateClassHierarchy);
-		GEditor->OnClassPackageLoadedOrUnloaded().RemoveStatic(ClassViewer::Helpers::RequestPopulateClassHierarchy);
+		GEditor->OnBlueprintCompiled().Remove(OnBlueprintCompiledRequestPopulateClassHierarchyDelegateHandle);
+		GEditor->OnClassPackageLoadedOrUnloaded().Remove(OnClassPackageLoadedOrUnloadedRequestPopulateClassHierarchyDelegateHandle);
 	}
 
 	FModuleManager::Get().OnModulesChanged().RemoveAll(this);
diff --git a/Engine/Source/Editor/ContentBrowser/Private/AssetViewWidgets.cpp b/Engine/Source/Editor/ContentBrowser/Private/AssetViewWidgets.cpp
--- a/Engine/Source/Editor/ContentBrowser/Private/AssetViewWidgets.cpp
+++ b/Engine/Source/Editor/ContentBrowser/Private/AssetViewWidgets.cpp
@@ -153,7 +153,7 @@ void SAssetViewItem::Construct( const FArguments& InArgs )
 	SourceControlStateDelay = 0.0f;
 	bSourceControlStateRequested = false;
 
-	ISourceControlModule::Get().GetProvider().RegisterSourceControlStateChanged(FSourceControlStateChanged::FDelegate::CreateSP(this, &SAssetViewItem::HandleSourceControlStateChanged));
+	ISourceControlModule::Get().GetProvider().RegisterSourceControlStateChanged_Handle(FSourceControlStateChanged::FDelegate::CreateSP(this, &SAssetViewItem::HandleSourceControlStateChanged));
 
 	// Source control state may have already been cached, make sure the control is in sync with 
 	// cached state as the delegate is not going to be invoked again until source control state 
diff --git a/Engine/Source/Editor/ContentBrowser/Private/SourcesViewWidgets.cpp b/Engine/Source/Editor/ContentBrowser/Private/SourcesViewWidgets.cpp
--- a/Engine/Source/Editor/ContentBrowser/Private/SourcesViewWidgets.cpp
+++ b/Engine/Source/Editor/ContentBrowser/Private/SourcesViewWidgets.cpp
@@ -73,7 +73,7 @@ void SAssetTreeItem::Construct( const FArguments& InArgs )
 
 	if( InlineRenameWidget.IsValid() )
 	{
-		TreeItem.Pin()->OnRenamedRequestEvent.AddSP( InlineRenameWidget.Get(), &SInlineEditableTextBlock::EnterEditingMode );
+		EnterEditingModeDelegateHandle = TreeItem.Pin()->OnRenamedRequestEvent.AddSP( InlineRenameWidget.Get(), &SInlineEditableTextBlock::EnterEditingMode );
 	}
 }
 
@@ -81,7 +81,7 @@ SAssetTreeItem::~SAssetTreeItem()
 {
 	if( InlineRenameWidget.IsValid() )
 	{
-		TreeItem.Pin()->OnRenamedRequestEvent.RemoveSP( InlineRenameWidget.Get(), &SInlineEditableTextBlock::EnterEditingMode );
+		TreeItem.Pin()->OnRenamedRequestEvent.Remove( EnterEditingModeDelegateHandle );
 	}
 }
 
@@ -448,7 +448,7 @@ void SCollectionListItem::Construct( const FArguments& InArgs )
 	if(InlineRenameWidget.IsValid())
 	{
 		// This is broadcast when the context menu / input binding requests a rename
-		CollectionItem.Pin()->OnRenamedRequestEvent.AddSP(InlineRenameWidget.Get(), &SInlineEditableTextBlock::EnterEditingMode);
+		EnterEditingModeDelegateHandle = CollectionItem.Pin()->OnRenamedRequestEvent.AddSP(InlineRenameWidget.Get(), &SInlineEditableTextBlock::EnterEditingMode);
 	}
 }
 
@@ -456,7 +456,7 @@ SCollectionListItem::~SCollectionListItem()
 {
 	if(InlineRenameWidget.IsValid())
 	{
-		CollectionItem.Pin()->OnRenamedRequestEvent.RemoveSP( InlineRenameWidget.Get(), &SInlineEditableTextBlock::EnterEditingMode );
+		CollectionItem.Pin()->OnRenamedRequestEvent.Remove( EnterEditingModeDelegateHandle );
 	}
 }
 
diff --git a/Engine/Source/Editor/ContentBrowser/Private/SourcesViewWidgets.h b/Engine/Source/Editor/ContentBrowser/Private/SourcesViewWidgets.h
--- a/Engine/Source/Editor/ContentBrowser/Private/SourcesViewWidgets.h
+++ b/Engine/Source/Editor/ContentBrowser/Private/SourcesViewWidgets.h
@@ -128,6 +128,9 @@ private:
 
 	/** Widget to display the name of the asset item and allows for renaming */
 	TSharedPtr< SInlineEditableTextBlock > InlineRenameWidget;
+
+	/** Handle to the registered EnterEditingMode delegate. */
+	FDelegateHandle EnterEditingModeDelegateHandle;
 };
 
 /** A single item in the collection list. */
@@ -231,4 +234,7 @@ private:
 
 	/** Widget to display the name of the collection item and allows for renaming */
 	TSharedPtr< SInlineEditableTextBlock > InlineRenameWidget;
+
+	/** Handle to the registered EnterEditingMode delegate. */
+	FDelegateHandle EnterEditingModeDelegateHandle;
 };
diff --git a/Engine/Source/Editor/DetailCustomizations/Private/AnimSequenceDetails.cpp b/Engine/Source/Editor/DetailCustomizations/Private/AnimSequenceDetails.cpp
--- a/Engine/Source/Editor/DetailCustomizations/Private/AnimSequenceDetails.cpp
+++ b/Engine/Source/Editor/DetailCustomizations/Private/AnimSequenceDetails.cpp
@@ -213,7 +213,7 @@ void FAnimSequenceDetails::RegisterRetargetSourceChanged()
 	if (TargetSkeleton.IsValid() && !OnDelegateRetargetSourceChanged.IsBound())
 	{
 		OnDelegateRetargetSourceChanged = USkeleton::FOnRetargetSourceChanged::CreateSP( this, &FAnimSequenceDetails::DelegateRetargetSourceChanged );
-		TargetSkeleton->RegisterOnRetargetSourceChanged(OnDelegateRetargetSourceChanged);
+		OnDelegateRetargetSourceChangedDelegateHandle = TargetSkeleton->RegisterOnRetargetSourceChanged(OnDelegateRetargetSourceChanged);
 	}
 }
 
@@ -221,7 +221,7 @@ FAnimSequenceDetails::~FAnimSequenceDetails()
 {
 	if (TargetSkeleton.IsValid() && OnDelegateRetargetSourceChanged.IsBound())
 	{
-		TargetSkeleton->UnregisterOnRetargetSourceChanged(OnDelegateRetargetSourceChanged);
+		TargetSkeleton->UnregisterOnRetargetSourceChanged(OnDelegateRetargetSourceChangedDelegateHandle);
 	}
 }
 
diff --git a/Engine/Source/Editor/DetailCustomizations/Private/AnimSequenceDetails.h b/Engine/Source/Editor/DetailCustomizations/Private/AnimSequenceDetails.h
--- a/Engine/Source/Editor/DetailCustomizations/Private/AnimSequenceDetails.h
+++ b/Engine/Source/Editor/DetailCustomizations/Private/AnimSequenceDetails.h
@@ -46,6 +46,7 @@ private:
 	TSharedPtr<FString> GetRetargetSourceString(FName RetargetSourceName) const;
 
 	USkeleton::FOnRetargetSourceChanged OnDelegateRetargetSourceChanged;
+	FDelegateHandle OnDelegateRetargetSourceChangedDelegateHandle;
 	void RegisterRetargetSourceChanged();
 	void DelegateRetargetSourceChanged();
 };
diff --git a/Engine/Source/Editor/DetailCustomizations/Private/SkeletalMeshComponentDetails.cpp b/Engine/Source/Editor/DetailCustomizations/Private/SkeletalMeshComponentDetails.cpp
--- a/Engine/Source/Editor/DetailCustomizations/Private/SkeletalMeshComponentDetails.cpp
+++ b/Engine/Source/Editor/DetailCustomizations/Private/SkeletalMeshComponentDetails.cpp
@@ -279,7 +279,7 @@ void FSkeletalMeshComponentDetails::RegisterSkeletalMeshPropertyChanged(TWeakObj
 {
 	if(Mesh.IsValid() && OnSkeletalMeshPropertyChanged.IsBound())
 	{
-		Mesh->RegisterOnSkeletalMeshPropertyChanged(OnSkeletalMeshPropertyChanged);
+		OnSkeletalMeshPropertyChangedDelegateHandles.Add(Mesh.Get(), Mesh->RegisterOnSkeletalMeshPropertyChanged(OnSkeletalMeshPropertyChanged));
 	}
 }
 
@@ -287,7 +287,8 @@ void FSkeletalMeshComponentDetails::UnregisterSkeletalMeshPropertyChanged(TWeakO
 {
 	if(Mesh.IsValid())
 	{
-		Mesh->UnregisterOnSkeletalMeshPropertyChanged(OnSkeletalMeshPropertyChanged);
+		Mesh->UnregisterOnSkeletalMeshPropertyChanged(OnSkeletalMeshPropertyChangedDelegateHandles.FindRef(Mesh.Get()));
+		OnSkeletalMeshPropertyChangedDelegateHandles.Remove(Mesh.Get());
 	}
 }
 
@@ -297,7 +298,8 @@ void FSkeletalMeshComponentDetails::UnregisterAllMeshPropertyChangedCallers()
 	{
 		if(USkeletalMeshComponent* Mesh = Cast<USkeletalMeshComponent>(MeshIter->Get()))
 		{
-			Mesh->UnregisterOnSkeletalMeshPropertyChanged(OnSkeletalMeshPropertyChanged);
+			Mesh->UnregisterOnSkeletalMeshPropertyChanged(OnSkeletalMeshPropertyChangedDelegateHandles.FindRef(Mesh));
+			OnSkeletalMeshPropertyChangedDelegateHandles.Remove(Mesh);
 		}
 	}
 }
diff --git a/Engine/Source/Editor/DetailCustomizations/Private/SkeletalMeshComponentDetails.h b/Engine/Source/Editor/DetailCustomizations/Private/SkeletalMeshComponentDetails.h
--- a/Engine/Source/Editor/DetailCustomizations/Private/SkeletalMeshComponentDetails.h
+++ b/Engine/Source/Editor/DetailCustomizations/Private/SkeletalMeshComponentDetails.h
@@ -77,4 +77,7 @@ private:
 
 	/** The combo button for the class picker, Cached so we can close it when the user picks something */
 	TSharedPtr<SComboButton> ClassPickerComboButton;
+
+	/** Per-mesh handles to registered OnSkeletalMeshPropertyChanged delegates */
+	TMap<USkeletalMeshComponent*, FDelegateHandle> OnSkeletalMeshPropertyChangedDelegateHandles;
 };
diff --git a/Engine/Source/Editor/DeviceProfileEditor/Private/SDeviceProfileEditor.cpp b/Engine/Source/Editor/DeviceProfileEditor/Private/SDeviceProfileEditor.cpp
--- a/Engine/Source/Editor/DeviceProfileEditor/Private/SDeviceProfileEditor.cpp
+++ b/Engine/Source/Editor/DeviceProfileEditor/Private/SDeviceProfileEditor.cpp
@@ -311,9 +311,6 @@ SDeviceProfileEditor::~SDeviceProfileEditor()
 	if( DeviceProfileManager.IsValid() )
 	{
 		DeviceProfileManager->SaveProfiles();
-
-		// Unbind our delegates when destroyed
-		DeviceProfileManager->OnManagerUpdated().RemoveRaw( this, &SDeviceProfileEditor::RebuildPropertyTable );
 	}
 
 	if (TabManager.IsValid())
diff --git a/Engine/Source/Editor/DeviceProfileEditor/Private/SDeviceProfileSelectionPanel.cpp b/Engine/Source/Editor/DeviceProfileEditor/Private/SDeviceProfileSelectionPanel.cpp
--- a/Engine/Source/Editor/DeviceProfileEditor/Private/SDeviceProfileSelectionPanel.cpp
+++ b/Engine/Source/Editor/DeviceProfileEditor/Private/SDeviceProfileSelectionPanel.cpp
@@ -218,7 +218,7 @@ void SDeviceProfileSelectionPanel::Construct( const FArguments& InArgs, TWeakObj
 	OnDeviceProfileViewAlone = InArgs._OnDeviceProfileViewAlone;
 
 	// Hook up our regen function to keep track of device profile manager changes
-	DeviceProfileManager->OnManagerUpdated().AddRaw( this, &SDeviceProfileSelectionPanel::RegenerateProfileList );
+	RegenerateProfileListDelegateHandle = DeviceProfileManager->OnManagerUpdated().AddRaw( this, &SDeviceProfileSelectionPanel::RegenerateProfileList );
 
 	ChildSlot
 	[
@@ -267,7 +267,7 @@ SDeviceProfileSelectionPanel::~SDeviceProfileSelectionPanel()
 	if( DeviceProfileManager.IsValid() )
 	{
 		// Remove the delegate when we are destroyed
-		DeviceProfileManager->OnManagerUpdated().RemoveRaw( this, &SDeviceProfileSelectionPanel::RegenerateProfileList );
+		DeviceProfileManager->OnManagerUpdated().Remove( RegenerateProfileListDelegateHandle );
 	}
 }
 
diff --git a/Engine/Source/Editor/DeviceProfileEditor/Private/SDeviceProfileSelectionPanel.h b/Engine/Source/Editor/DeviceProfileEditor/Private/SDeviceProfileSelectionPanel.h
--- a/Engine/Source/Editor/DeviceProfileEditor/Private/SDeviceProfileSelectionPanel.h
+++ b/Engine/Source/Editor/DeviceProfileEditor/Private/SDeviceProfileSelectionPanel.h
@@ -79,6 +79,9 @@ private:
 
 	/** The profile selected from the current list. */
 	TWeakObjectPtr< UDeviceProfile > SelectedProfile;
+
+	/** Handle to the registered RegenerateProfileList delegate */
+	FDelegateHandle RegenerateProfileListDelegateHandle;
 };
 
 
diff --git a/Engine/Source/Editor/GameProjectGeneration/Private/GameProjectUtils.cpp b/Engine/Source/Editor/GameProjectGeneration/Private/GameProjectUtils.cpp
--- a/Engine/Source/Editor/GameProjectGeneration/Private/GameProjectUtils.cpp
+++ b/Engine/Source/Editor/GameProjectGeneration/Private/GameProjectUtils.cpp
@@ -24,6 +24,7 @@
 #include "GameFramework/GameMode.h"
 #include "HotReloadInterface.h"
 #include "SVerbChoiceDialog.h"
+#include "SourceCodeNavigation.h"
 
 
 #define LOCTEXT_NAMESPACE "GameProjectUtils"
diff --git a/Engine/Source/Editor/GameplayAbilitiesEditor/Private/GameplayAbilitiesEditorModule.cpp b/Engine/Source/Editor/GameplayAbilitiesEditor/Private/GameplayAbilitiesEditorModule.cpp
--- a/Engine/Source/Editor/GameplayAbilitiesEditor/Private/GameplayAbilitiesEditorModule.cpp
+++ b/Engine/Source/Editor/GameplayAbilitiesEditor/Private/GameplayAbilitiesEditorModule.cpp
@@ -40,6 +40,9 @@ private:
 
 	/** Node factory for abilities graph; Cached so it can be unregistered */
 	TSharedPtr<FGameplayAbilitiesGraphPanelNodeFactory> GameplayAbilitiesGraphPanelNodeFactory;
+
+	/** Handle to the registered GameplayTagTreeChanged delegate */
+	FDelegateHandle GameplayTagTreeChangedDelegateHandle;
 };
 
 IMPLEMENT_MODULE(FGameplayAbilitiesEditorModule, GameplayAbilitiesEditor)
@@ -73,7 +76,7 @@ void FGameplayAbilitiesEditorModule::StartupModule()
 
 	// Listen for changes to the gameplay tag tree so we can refresh blueprint actions for the GameplayCueEvent node
 	UGameplayTagsManager& GameplayTagsManager = IGameplayTagsModule::Get().GetGameplayTagsManager();
-	GameplayTagsManager.OnGameplayTagTreeChanged().AddStatic(&FGameplayAbilitiesEditorModule::GameplayTagTreeChanged);
+	GameplayTagTreeChangedDelegateHandle = GameplayTagsManager.OnGameplayTagTreeChanged().AddStatic(&FGameplayAbilitiesEditorModule::GameplayTagTreeChanged);
 }
 
 void FGameplayAbilitiesEditorModule::RegisterAssetTypeAction(IAssetTools& AssetTools, TSharedRef<IAssetTypeActions> Action)
@@ -138,6 +141,6 @@ void FGameplayAbilitiesEditorModule::ShutdownModule()
 	if ( IGameplayTagsModule::IsAvailable() )
 	{
 		UGameplayTagsManager& GameplayTagsManager = IGameplayTagsModule::Get().GetGameplayTagsManager();
-		GameplayTagsManager.OnGameplayTagTreeChanged().RemoveStatic(&FGameplayAbilitiesEditorModule::GameplayTagTreeChanged);
+		GameplayTagsManager.OnGameplayTagTreeChanged().Remove(GameplayTagTreeChangedDelegateHandle);
 	}
 }
diff --git a/Engine/Source/Editor/GraphEditor/Private/SGraphPanel.cpp b/Engine/Source/Editor/GraphEditor/Private/SGraphPanel.cpp
--- a/Engine/Source/Editor/GraphEditor/Private/SGraphPanel.cpp
+++ b/Engine/Source/Editor/GraphEditor/Private/SGraphPanel.cpp
@@ -95,7 +95,7 @@ void SGraphPanel::Construct( const SGraphPanel::FArguments& InArgs )
 
 	// Register for notifications
 	MyRegisteredGraphChangedDelegate = FOnGraphChanged::FDelegate::CreateSP(this, &SGraphPanel::OnGraphChanged);
-	this->GraphObj->AddOnGraphChangedHandler(MyRegisteredGraphChangedDelegate);
+	MyRegisteredGraphChangedDelegateHandle = this->GraphObj->AddOnGraphChangedHandler(MyRegisteredGraphChangedDelegate);
 	
 	ShowGraphStateOverlay = InArgs._ShowGraphStateOverlay;
 }
@@ -105,7 +105,7 @@ SGraphPanel::~SGraphPanel()
 	FEditorDelegates::BeginPIE.RemoveAll( this );
 	FEditorDelegates::EndPIE.RemoveAll( this );
 
-	this->GraphObj->RemoveOnGraphChangedHandler(MyRegisteredGraphChangedDelegate);
+	this->GraphObj->RemoveOnGraphChangedHandler(MyRegisteredGraphChangedDelegateHandle);
 }
 
 //////////////////////////////////////////////////////////////////////////
diff --git a/Engine/Source/Editor/GraphEditor/Private/SGraphPanel.h b/Engine/Source/Editor/GraphEditor/Private/SGraphPanel.h
--- a/Engine/Source/Editor/GraphEditor/Private/SGraphPanel.h
+++ b/Engine/Source/Editor/GraphEditor/Private/SGraphPanel.h
@@ -244,6 +244,7 @@ private:
 	TMap<const class UEdGraphNode*, int32> UserAddedNodes;
 
 	FOnGraphChanged::FDelegate MyRegisteredGraphChangedDelegate;
+	FDelegateHandle            MyRegisteredGraphChangedDelegateHandle;
 private:
 	/** Called when PIE begins */
 	void OnBeginPIE( const bool bIsSimulating );
diff --git a/Engine/Source/Editor/Kismet/Private/BlueprintEditor.cpp b/Engine/Source/Editor/Kismet/Private/BlueprintEditor.cpp
--- a/Engine/Source/Editor/Kismet/Private/BlueprintEditor.cpp
+++ b/Engine/Source/Editor/Kismet/Private/BlueprintEditor.cpp
@@ -1223,7 +1223,7 @@ void FBlueprintEditor::CommonInitialization(const TArray<UBlueprint*>& InitBluep
 	}
 
 	// Make sure we know when tabs become active to update details tab
-	FGlobalTabmanager::Get()->OnActiveTabChanged_Subscribe( FOnActiveTabChanged::FDelegate::CreateRaw(this, &FBlueprintEditor::OnActiveTabChanged) );
+	OnActiveTabChangedDelegateHandle = FGlobalTabmanager::Get()->OnActiveTabChanged_Subscribe( FOnActiveTabChanged::FDelegate::CreateRaw(this, &FBlueprintEditor::OnActiveTabChanged) );
 
 	if (InitBlueprints.Num() == 1)
 	{
@@ -1773,7 +1773,7 @@ FBlueprintEditor::~FBlueprintEditor()
 		GetBlueprintObj()->OnChanged().RemoveAll( this );
 	}
 
-	FGlobalTabmanager::Get()->OnActiveTabChanged_Unsubscribe( FOnActiveTabChanged::FDelegate::CreateRaw(this, &FBlueprintEditor::OnActiveTabChanged) );
+	FGlobalTabmanager::Get()->OnActiveTabChanged_Unsubscribe( OnActiveTabChangedDelegateHandle );
 
 	if (FEngineAnalytics::IsAvailable())
 	{
diff --git a/Engine/Source/Editor/Kismet/Private/BlueprintEditorModule.cpp b/Engine/Source/Editor/Kismet/Private/BlueprintEditorModule.cpp
--- a/Engine/Source/Editor/Kismet/Private/BlueprintEditorModule.cpp
+++ b/Engine/Source/Editor/Kismet/Private/BlueprintEditorModule.cpp
@@ -86,7 +86,9 @@ void FBlueprintEditorModule::StartupModule()
 		// Extend the level viewport context menu to handle blueprints
 		LevelViewportContextMenuBlueprintExtender = FLevelEditorModule::FLevelViewportMenuExtender_SelectedActors::CreateStatic(&ExtendLevelViewportContextMenuForBlueprints);
 		FLevelEditorModule& LevelEditorModule = FModuleManager::Get().LoadModuleChecked<FLevelEditorModule>("LevelEditor");
-		LevelEditorModule.GetAllLevelViewportContextMenuExtenders().Add(LevelViewportContextMenuBlueprintExtender);
+		auto& MenuExtenders = LevelEditorModule.GetAllLevelViewportContextMenuExtenders();
+		MenuExtenders.Add(LevelViewportContextMenuBlueprintExtender);
+		LevelViewportContextMenuBlueprintExtenderDelegateHandle = MenuExtenders.Last().GetHandle();
 	}
 
 	FMessageLogModule& MessageLogModule = FModuleManager::LoadModuleChecked<FMessageLogModule>("MessageLog");
@@ -123,7 +125,9 @@ void FBlueprintEditorModule::ShutdownModule()
 	if ( FModuleManager::Get().IsModuleLoaded( "LevelEditor" ) )
 	{
 		FLevelEditorModule& LevelEditorModule = FModuleManager::LoadModuleChecked<FLevelEditorModule>( "LevelEditor" );
-		LevelEditorModule.GetAllLevelViewportContextMenuExtenders().Remove(LevelViewportContextMenuBlueprintExtender);
+		LevelEditorModule.GetAllLevelViewportContextMenuExtenders().RemoveAll([&](const FLevelEditorModule::FLevelViewportMenuExtender_SelectedActors& Delegate) {
+			return Delegate.GetHandle() == LevelViewportContextMenuBlueprintExtenderDelegateHandle;
+		});
 	}
 
 	FMessageLogModule& MessageLogModule = FModuleManager::LoadModuleChecked<FMessageLogModule>("MessageLog");
diff --git a/Engine/Source/Editor/Kismet/Private/SBlueprintDiff.cpp b/Engine/Source/Editor/Kismet/Private/SBlueprintDiff.cpp
--- a/Engine/Source/Editor/Kismet/Private/SBlueprintDiff.cpp
+++ b/Engine/Source/Editor/Kismet/Private/SBlueprintDiff.cpp
@@ -340,6 +340,9 @@ private:
 
 	/** ListView of differences */
 	TSharedPtr<SListViewType> DiffList;
+
+	/** Handle to the registered OnGraphChanged delegate. */
+	FDelegateHandle OnGraphChangedDelegateHandle;
 };
 
 TSharedRef<SWidget>	FDiffResultItem::GenerateWidget() const
@@ -366,7 +369,7 @@ FListItemGraphToDiff::FListItemGraphToDiff( class SBlueprintDiff* InDiff, class
 	//need to know when it is modified
 	if(InGraphNew)
 	{
-		InGraphNew->AddOnGraphChangedHandler( FOnGraphChanged::FDelegate::CreateRaw(this, &FListItemGraphToDiff::OnGraphChanged));
+		OnGraphChangedDelegateHandle = InGraphNew->AddOnGraphChangedHandler( FOnGraphChanged::FDelegate::CreateRaw(this, &FListItemGraphToDiff::OnGraphChanged));
 	}
 
 	BuildDiffSourceArray();
@@ -376,7 +379,7 @@ FListItemGraphToDiff::~FListItemGraphToDiff()
 {
 	if(GraphNew)
 	{
-		GraphNew->RemoveOnGraphChangedHandler( FOnGraphChanged::FDelegate::CreateRaw(this, &FListItemGraphToDiff::OnGraphChanged));
+		GraphNew->RemoveOnGraphChangedHandler( OnGraphChangedDelegateHandle);
 	}
 }
 
diff --git a/Engine/Source/Editor/Kismet/Private/WorkflowOrientedApp/WorkflowTabManager.cpp b/Engine/Source/Editor/Kismet/Private/WorkflowOrientedApp/WorkflowTabManager.cpp
--- a/Engine/Source/Editor/Kismet/Private/WorkflowOrientedApp/WorkflowTabManager.cpp
+++ b/Engine/Source/Editor/Kismet/Private/WorkflowOrientedApp/WorkflowTabManager.cpp
@@ -371,12 +371,12 @@ void FDocumentTracker::RegisterDocumentFactory(TSharedPtr<class FDocumentTabFact
 FDocumentTracker::FDocumentTracker()
 {
 	// Make sure we know when tabs become active
-	FGlobalTabmanager::Get()->OnActiveTabChanged_Subscribe( FOnActiveTabChanged::FDelegate::CreateRaw( this, &FDocumentTracker::OnActiveTabChanged ) );
+	OnActiveTabChangedDelegateHandle = FGlobalTabmanager::Get()->OnActiveTabChanged_Subscribe( FOnActiveTabChanged::FDelegate::CreateRaw( this, &FDocumentTracker::OnActiveTabChanged ) );
 }
 
 FDocumentTracker::~FDocumentTracker()
 {
-	FGlobalTabmanager::Get()->OnActiveTabChanged_Unsubscribe( FOnActiveTabChanged::FDelegate::CreateRaw( this, &FDocumentTracker::OnActiveTabChanged ) );
+	FGlobalTabmanager::Get()->OnActiveTabChanged_Unsubscribe( OnActiveTabChangedDelegateHandle );
 }
 
 // Called by the global active tab changed callback; dispatches to individually registered callbacks
diff --git a/Engine/Source/Editor/Kismet/Public/BlueprintEditor.h b/Engine/Source/Editor/Kismet/Public/BlueprintEditor.h
--- a/Engine/Source/Editor/Kismet/Public/BlueprintEditor.h
+++ b/Engine/Source/Editor/Kismet/Public/BlueprintEditor.h
@@ -1061,6 +1061,9 @@ private:
 	UEdGraph* HasOpenActionMenu;
 	/** Used to nicely fade instruction text, when the context menu is opened. */
 	float InstructionsFadeCountdown;
+
+	/** Handle to the registered OnActiveTabChanged delegate */
+	FDelegateHandle OnActiveTabChangedDelegateHandle;
 };
 
 #undef LOCTEXT_NAMESPACE
diff --git a/Engine/Source/Editor/Kismet/Public/BlueprintEditorModule.h b/Engine/Source/Editor/Kismet/Public/BlueprintEditorModule.h
--- a/Engine/Source/Editor/Kismet/Public/BlueprintEditorModule.h
+++ b/Engine/Source/Editor/Kismet/Public/BlueprintEditorModule.h
@@ -152,4 +152,7 @@ private:
 	 * of the blueprint editor. 
 	 */
 	TSharedPtr<FUICommandList> SharedBlueprintEditorCommands;
+
+	/** Handle to a registered LevelViewportContextMenuBlueprintExtender delegate */
+	FDelegateHandle LevelViewportContextMenuBlueprintExtenderDelegateHandle;
 };
diff --git a/Engine/Source/Editor/LandscapeEditor/Private/LandscapeEdMode.cpp b/Engine/Source/Editor/LandscapeEditor/Private/LandscapeEdMode.cpp
--- a/Engine/Source/Editor/LandscapeEditor/Private/LandscapeEdMode.cpp
+++ b/Engine/Source/Editor/LandscapeEditor/Private/LandscapeEdMode.cpp
@@ -321,8 +321,8 @@ void FEdModeLandscape::Enter()
 	UpdateLandscapeList();
 	UpdateTargetList();
 
-	FEditorSupportDelegates::WorldChange.AddRaw(this, &FEdModeLandscape::OnWorldChange);
-	UMaterial::OnMaterialCompilationFinished().AddRaw(this, &FEdModeLandscape::OnMaterialCompilationFinished);
+	OnWorldChangeDelegateHandle                 = FEditorSupportDelegates::WorldChange.AddRaw(this, &FEdModeLandscape::OnWorldChange);
+	OnMaterialCompilationFinishedDelegateHandle = UMaterial::OnMaterialCompilationFinished().AddRaw(this, &FEdModeLandscape::OnMaterialCompilationFinished);
 
 	if (CurrentGizmoActor.IsValid())
 	{
@@ -450,8 +450,8 @@ void FEdModeLandscape::Enter()
 /** FEdMode: Called when the mode is exited */
 void FEdModeLandscape::Exit()
 {
-	FEditorSupportDelegates::WorldChange.RemoveRaw(this, &FEdModeLandscape::OnWorldChange);
-	UMaterial::OnMaterialCompilationFinished().RemoveRaw(this, &FEdModeLandscape::OnMaterialCompilationFinished);
+	FEditorSupportDelegates::WorldChange.Remove(OnWorldChangeDelegateHandle);
+	UMaterial::OnMaterialCompilationFinished().Remove(OnMaterialCompilationFinishedDelegateHandle);
 
 	// Restore real-time viewport state if we changed it
 	const bool bWantRealTime = false;
diff --git a/Engine/Source/Editor/Layers/Private/SLayersViewRow.cpp b/Engine/Source/Editor/Layers/Private/SLayersViewRow.cpp
--- a/Engine/Source/Editor/Layers/Private/SLayersViewRow.cpp
+++ b/Engine/Source/Editor/Layers/Private/SLayersViewRow.cpp
@@ -17,7 +17,7 @@ void SLayersViewRow::Construct(const FArguments& InArgs, TSharedRef< FLayerViewM
 
 SLayersViewRow::~SLayersViewRow()
 {
-	ViewModel->OnRenamedRequest().RemoveSP(InlineTextBlock.Get(), &SInlineEditableTextBlock::EnterEditingMode);
+	ViewModel->OnRenamedRequest().Remove(EnterEditingModeDelegateHandle);
 }
 
 TSharedRef< SWidget > SLayersViewRow::GenerateWidgetForColumn(const FName& ColumnID)
@@ -52,7 +52,7 @@ TSharedRef< SWidget > SLayersViewRow::GenerateWidgetForColumn(const FName& Colum
 			]
 		;
 
-		ViewModel->OnRenamedRequest().AddSP(InlineTextBlock.Get(), &SInlineEditableTextBlock::EnterEditingMode);
+		EnterEditingModeDelegateHandle = ViewModel->OnRenamedRequest().AddSP(InlineTextBlock.Get(), &SInlineEditableTextBlock::EnterEditingMode);
 	}
 	else if (ColumnID == LayersView::ColumnID_Visibility)
 	{
diff --git a/Engine/Source/Editor/Layers/Private/SLayersViewRow.h b/Engine/Source/Editor/Layers/Private/SLayersViewRow.h
--- a/Engine/Source/Editor/Layers/Private/SLayersViewRow.h
+++ b/Engine/Source/Editor/Layers/Private/SLayersViewRow.h
@@ -122,5 +122,8 @@ private:
 
 	/** Widget for displaying and editing the Layer name */
 	TSharedPtr< SInlineEditableTextBlock > InlineTextBlock;
+
+	/** Handle to the registered EnterEditingMode delegate */
+	FDelegateHandle EnterEditingModeDelegateHandle;
 };
 
diff --git a/Engine/Source/Editor/LevelEditor/Private/SLevelEditorBuildAndSubmit.cpp b/Engine/Source/Editor/LevelEditor/Private/SLevelEditorBuildAndSubmit.cpp
--- a/Engine/Source/Editor/LevelEditor/Private/SLevelEditorBuildAndSubmit.cpp
+++ b/Engine/Source/Editor/LevelEditor/Private/SLevelEditorBuildAndSubmit.cpp
@@ -126,7 +126,7 @@ SLevelEditorBuildAndSubmit::~SLevelEditorBuildAndSubmit()
 {
 	UPackage::PackageDirtyStateChangedEvent.RemoveAll(this);
 
-	ISourceControlModule::Get().GetProvider().UnregisterSourceControlStateChanged(FSourceControlStateChanged::FDelegate::CreateRaw(this, &SLevelEditorBuildAndSubmit::OnSourceControlStateChanged));
+	ISourceControlModule::Get().GetProvider().UnregisterSourceControlStateChanged_Handle(OnSourceControlStateChangedDelegateHandle);
 }
 
 void SLevelEditorBuildAndSubmit::Construct( const FArguments& InArgs, const TSharedRef< class ILevelEditor >& OwningLevelEditor )
@@ -276,7 +276,7 @@ void SLevelEditorBuildAndSubmit::Construct( const FArguments& InArgs, const TSha
 
 	UpdatePackagesList();
 
-	ISourceControlModule::Get().GetProvider().RegisterSourceControlStateChanged(FSourceControlStateChanged::FDelegate::CreateRaw(this, &SLevelEditorBuildAndSubmit::OnSourceControlStateChanged));
+	OnSourceControlStateChangedDelegateHandle = ISourceControlModule::Get().GetProvider().RegisterSourceControlStateChanged_Handle(FSourceControlStateChanged::FDelegate::CreateRaw(this, &SLevelEditorBuildAndSubmit::OnSourceControlStateChanged));
 
 	UPackage::PackageDirtyStateChangedEvent.AddRaw(this, &SLevelEditorBuildAndSubmit::OnEditorPackageModified);
 }
diff --git a/Engine/Source/Editor/LevelEditor/Private/SLevelEditorBuildAndSubmit.h b/Engine/Source/Editor/LevelEditor/Private/SLevelEditorBuildAndSubmit.h
--- a/Engine/Source/Editor/LevelEditor/Private/SLevelEditorBuildAndSubmit.h
+++ b/Engine/Source/Editor/LevelEditor/Private/SLevelEditorBuildAndSubmit.h
@@ -94,4 +94,7 @@ private:
 
 	/** Set by OnShowHideExtraPackagesSection() - keeps track of the visibility or the additional packages list */
 	bool bIsExtraPackagesSectionExpanded;
+
+	/** Handle to the registered OnSourceControlStateChanged delegate */
+	FDelegateHandle OnSourceControlStateChangedDelegateHandle;
 };
\ No newline at end of file
diff --git a/Engine/Source/Editor/Matinee/Private/Matinee.cpp b/Engine/Source/Editor/Matinee/Private/Matinee.cpp
--- a/Engine/Source/Editor/Matinee/Private/Matinee.cpp
+++ b/Engine/Source/Editor/Matinee/Private/Matinee.cpp
@@ -959,7 +959,7 @@ void FMatinee::InitMatinee(const EToolkitMode::Type Mode, const TSharedPtr< clas
 	bIsInitialized = true;
 
 	// register for any actor move change
-	GEngine->OnActorMoved().AddRaw(this, &FMatinee::OnActorMoved);
+	OnActorMovedDelegateHandle = GEngine->OnActorMoved().AddRaw(this, &FMatinee::OnActorMoved);
 
 	// register for any objects replaced
 	GEditor->OnObjectsReplaced().AddSP(this, &FMatinee::OnObjectsReplaced);
@@ -1968,7 +1968,7 @@ void FMatinee::OnClose()
 	}
 
 	// Unregister call back events
-	GEngine->OnActorMoved().RemoveRaw(this, &FMatinee::OnActorMoved);
+	GEngine->OnActorMoved().Remove(OnActorMovedDelegateHandle);
 	GEditor->OnObjectsReplaced().RemoveAll(this);
 
 	// Restore the perspective viewport audio settings when matinee closes.
diff --git a/Engine/Source/Editor/Matinee/Private/Matinee.h b/Engine/Source/Editor/Matinee/Private/Matinee.h
--- a/Engine/Source/Editor/Matinee/Private/Matinee.h
+++ b/Engine/Source/Editor/Matinee/Private/Matinee.h
@@ -1535,6 +1535,9 @@ private:
 	TMap<UInterpTrack*, AddKeyInfo> AddKeyInfoMap;
 	TMap<UInterpTrack*, int32> TrackToNewKeyIndexMap;
 
+	/** Handle to the registered OnActorMoved delegate */
+	FDelegateHandle OnActorMovedDelegateHandle;
+
 	/** For keeping track of the previously used camera, so we can detect cuts when playing back in editor mode */
 	ACameraActor* PreviousCamera;
 
diff --git a/Engine/Source/Editor/NewsFeed/Private/Implementation/NewsFeedCache.cpp b/Engine/Source/Editor/NewsFeed/Private/Implementation/NewsFeedCache.cpp
--- a/Engine/Source/Editor/NewsFeed/Private/Implementation/NewsFeedCache.cpp
+++ b/Engine/Source/Editor/NewsFeed/Private/Implementation/NewsFeedCache.cpp
@@ -22,7 +22,7 @@ FNewsFeedCache::FNewsFeedCache( )
 
 FNewsFeedCache::~FNewsFeedCache( )
 {
-	FTicker::GetCoreTicker().RemoveTicker(TickDelegate);
+	FTicker::GetCoreTicker().RemoveTicker(TickDelegateHandle);
 	SaveSettings();
 }
 
@@ -60,15 +60,15 @@ bool FNewsFeedCache::LoadFeed(  )
 	Icons.Empty();
 
 //	TitleFile = InTitleFile;
-//	TitleFile->AddOnEnumerateFilesCompleteDelegate(FOnEnumerateFilesCompleteDelegate::CreateRaw(this, &FNewsFeedCache::HandleEnumerateFilesComplete));
-//	TitleFile->AddOnReadFileCompleteDelegate(FOnReadFileCompleteDelegate::CreateRaw(this, &FNewsFeedCache::HandleReadFileComplete));
+//	TitleFile->AddOnEnumerateFilesCompleteDelegate_Handle(FOnEnumerateFilesCompleteDelegate::CreateRaw(this, &FNewsFeedCache::HandleEnumerateFilesComplete));
+//	TitleFile->AddOnReadFileCompleteDelegate_Handle(FOnReadFileCompleteDelegate::CreateRaw(this, &FNewsFeedCache::HandleReadFileComplete));
 //	TitleFile->ClearFiles();
 //	TitleFile->EnumerateFiles();
 	LoadTitleFile();
 
 	// auto reload after 15 minutes
-	FTicker::GetCoreTicker().RemoveTicker(TickDelegate);
-	FTicker::GetCoreTicker().AddTicker(TickDelegate, 15 * 60.0f);
+	FTicker::GetCoreTicker().RemoveTicker(TickDelegateHandle);
+	TickDelegateHandle = FTicker::GetCoreTicker().AddTicker(TickDelegate, 15 * 60.0f);
 
 	return true;
 }
@@ -147,8 +147,8 @@ void FNewsFeedCache::LoadTitleFile( )
 	// fetch news feed
 	if (TitleFile.IsValid())
 	{
-		TitleFile->AddOnEnumerateFilesCompleteDelegate(FOnEnumerateFilesCompleteDelegate::CreateRaw(this, &FNewsFeedCache::HandleEnumerateFilesComplete));
-		TitleFile->AddOnReadFileCompleteDelegate(FOnReadFileCompleteDelegate::CreateRaw(this, &FNewsFeedCache::HandleReadFileComplete));
+		TitleFile->AddOnEnumerateFilesCompleteDelegate_Handle(FOnEnumerateFilesCompleteDelegate::CreateRaw(this, &FNewsFeedCache::HandleEnumerateFilesComplete));
+		TitleFile->AddOnReadFileCompleteDelegate_Handle(FOnReadFileCompleteDelegate::CreateRaw(this, &FNewsFeedCache::HandleReadFileComplete));
 		TitleFile->ClearFiles();
 		TitleFile->EnumerateFiles();
 	}
diff --git a/Engine/Source/Editor/Persona/Private/Persona.cpp b/Engine/Source/Editor/Persona/Private/Persona.cpp
--- a/Engine/Source/Editor/Persona/Private/Persona.cpp
+++ b/Engine/Source/Editor/Persona/Private/Persona.cpp
@@ -187,7 +187,7 @@ FPersona::FPersona()
 {
 	// Register to be notified when properties are edited
 	OnPropertyChangedHandle = FCoreUObjectDelegates::FOnObjectPropertyChanged::FDelegate::CreateRaw(this, &FPersona::OnPropertyChanged);
-	FCoreUObjectDelegates::OnObjectPropertyChanged.Add(OnPropertyChangedHandle);
+	OnPropertyChangedHandleDelegateHandle = FCoreUObjectDelegates::OnObjectPropertyChanged.Add(OnPropertyChangedHandle);
 
 	GEditor->OnBlueprintPreCompile().AddRaw(this, &FPersona::OnBlueprintPreCompile);
 }
@@ -210,7 +210,7 @@ FPersona::~FPersona()
 		Viewport.Pin()->CleanupPersonaReferences();
 	}
 
-	FCoreUObjectDelegates::OnObjectPropertyChanged.Remove(OnPropertyChangedHandle);
+	FCoreUObjectDelegates::OnObjectPropertyChanged.Remove(OnPropertyChangedHandleDelegateHandle);
 
 	if(PreviewComponent)
 	{
diff --git a/Engine/Source/Editor/Persona/Private/Persona.h b/Engine/Source/Editor/Persona/Private/Persona.h
--- a/Engine/Source/Editor/Persona/Private/Persona.h
+++ b/Engine/Source/Editor/Persona/Private/Persona.h
@@ -638,4 +638,7 @@ private:
 
 	/** Sequence Browser **/
 	TWeakPtr<class SAnimationSequenceBrowser> SequenceBrowser;
+
+	/** Handle to the registered OnPropertyChangedHandle delegate */
+	FDelegateHandle OnPropertyChangedHandleDelegateHandle;
 };
diff --git a/Engine/Source/Editor/Persona/Private/SAnimNotifyPanel.cpp b/Engine/Source/Editor/Persona/Private/SAnimNotifyPanel.cpp
--- a/Engine/Source/Editor/Persona/Private/SAnimNotifyPanel.cpp
+++ b/Engine/Source/Editor/Persona/Private/SAnimNotifyPanel.cpp
@@ -3065,7 +3065,7 @@ void SAnimNotifyPanel::Construct(const FArguments& InArgs)
 	];
 
 	OnPropertyChangedHandle = FCoreUObjectDelegates::FOnObjectPropertyChanged::FDelegate::CreateSP(this, &SAnimNotifyPanel::OnPropertyChanged);
-	FCoreUObjectDelegates::OnObjectPropertyChanged.Add(OnPropertyChangedHandle);
+	OnPropertyChangedHandleDelegateHandle = FCoreUObjectDelegates::OnObjectPropertyChanged.Add(OnPropertyChangedHandle);
 
 	// Base notify classes used to search asset data for children.
 	NotifyClassNames.Add(TEXT("Class'/Script/Engine.AnimNotify'"));
@@ -3086,7 +3086,7 @@ SAnimNotifyPanel::~SAnimNotifyPanel()
 		PersonaPtr.Pin()->UnregisterOnGenericDelete(this);
 	}
 
-	FCoreUObjectDelegates::OnObjectPropertyChanged.Remove(OnPropertyChangedHandle);
+	FCoreUObjectDelegates::OnObjectPropertyChanged.Remove(OnPropertyChangedHandleDelegateHandle);
 }
 
 FReply SAnimNotifyPanel::InsertTrack(int32 TrackIndexToInsert)
diff --git a/Engine/Source/Editor/Persona/Private/SAnimNotifyPanel.h b/Engine/Source/Editor/Persona/Private/SAnimNotifyPanel.h
--- a/Engine/Source/Editor/Persona/Private/SAnimNotifyPanel.h
+++ b/Engine/Source/Editor/Persona/Private/SAnimNotifyPanel.h
@@ -280,4 +280,7 @@ private:
 
 	/** Classes that are known to be derived from blueprint state notifies */
 	TArray<FString> NotifyStateClassNames;
+
+	/** Handle to the registered OnPropertyChangedHandle delegate */
+	FDelegateHandle OnPropertyChangedHandleDelegateHandle;
 };
\ No newline at end of file
diff --git a/Engine/Source/Editor/Persona/Private/SAnimationBlendSpace1D.cpp b/Engine/Source/Editor/Persona/Private/SAnimationBlendSpace1D.cpp
--- a/Engine/Source/Editor/Persona/Private/SAnimationBlendSpace1D.cpp
+++ b/Engine/Source/Editor/Persona/Private/SAnimationBlendSpace1D.cpp
@@ -515,7 +515,7 @@ void SBlendSpace1DWidget::ResampleData()
 
 SBlendSpaceEditor1D::~SBlendSpaceEditor1D()
 {
-	FCoreUObjectDelegates::OnObjectPropertyChanged.Remove(OnPropertyChangedHandle);
+	FCoreUObjectDelegates::OnObjectPropertyChanged.Remove(OnPropertyChangedHandleDelegateHandle);
 }
 
 void SBlendSpaceEditor1D::Construct(const FArguments& InArgs)
@@ -527,7 +527,7 @@ void SBlendSpaceEditor1D::Construct(const FArguments& InArgs)
 
 	// Register to be notified when properties are edited
 	OnPropertyChangedHandle = FCoreUObjectDelegates::FOnObjectPropertyChanged::FDelegate::CreateRaw(this, &SBlendSpaceEditor1D::OnPropertyChanged);
-	FCoreUObjectDelegates::OnObjectPropertyChanged.Add(OnPropertyChangedHandle);
+	OnPropertyChangedHandleDelegateHandle = FCoreUObjectDelegates::OnObjectPropertyChanged.Add(OnPropertyChangedHandle);
 
 	bCachedDisplayVerticalValue = GetBlendSpace()->bDisplayEditorVertically;
 
diff --git a/Engine/Source/Editor/PropertyEditor/Private/Presentation/PropertyTable/PropertyTable.cpp b/Engine/Source/Editor/PropertyEditor/Private/Presentation/PropertyTable/PropertyTable.cpp
--- a/Engine/Source/Editor/PropertyEditor/Private/Presentation/PropertyTable/PropertyTable.cpp
+++ b/Engine/Source/Editor/PropertyEditor/Private/Presentation/PropertyTable/PropertyTable.cpp
@@ -92,7 +92,7 @@ void FPropertyTable::CreateColorPickerWindow( const TSharedRef< class FPropertyE
 
 void FPropertyTable::EnqueueDeferredAction( FSimpleDelegate DeferredAction )
 {
-	DeferredActions.AddUnique( DeferredAction );
+	DeferredActions.Add( DeferredAction );
 }
 
 TSharedPtr<class FAssetThumbnailPool> FPropertyTable::GetThumbnailPool() const
diff --git a/Engine/Source/Editor/PropertyEditor/Private/SDetailsViewBase.cpp b/Engine/Source/Editor/PropertyEditor/Private/SDetailsViewBase.cpp
--- a/Engine/Source/Editor/PropertyEditor/Private/SDetailsViewBase.cpp
+++ b/Engine/Source/Editor/PropertyEditor/Private/SDetailsViewBase.cpp
@@ -198,7 +198,7 @@ const FSlateBrush* SDetailsViewBase::OnGetFilterButtonImageResource() const
 
 void SDetailsViewBase::EnqueueDeferredAction(FSimpleDelegate& DeferredAction)
 {
-	DeferredActions.AddUnique(DeferredAction);
+	DeferredActions.Add(DeferredAction);
 }
 
 /**
diff --git a/Engine/Source/Editor/PropertyEditor/Private/SPropertyTreeViewImpl.cpp b/Engine/Source/Editor/PropertyEditor/Private/SPropertyTreeViewImpl.cpp
--- a/Engine/Source/Editor/PropertyEditor/Private/SPropertyTreeViewImpl.cpp
+++ b/Engine/Source/Editor/PropertyEditor/Private/SPropertyTreeViewImpl.cpp
@@ -512,7 +512,7 @@ void SPropertyTreeViewImpl::RestoreColumnWidths()
 
 void SPropertyTreeViewImpl::EnqueueDeferredAction( FSimpleDelegate& DeferredAction )
 {
-	DeferredActions.AddUnique( DeferredAction );
+	DeferredActions.Add( DeferredAction );
 }
 
 void SPropertyTreeViewImpl::SetFromExistingTree( TSharedPtr<FObjectPropertyNode> RootNode, TSharedPtr<FPropertyNode> PropertyToView )
diff --git a/Engine/Source/Editor/Sequencer/Private/SequencerModule.cpp b/Engine/Source/Editor/Sequencer/Private/SequencerModule.cpp
--- a/Engine/Source/Editor/Sequencer/Private/SequencerModule.cpp
+++ b/Engine/Source/Editor/Sequencer/Private/SequencerModule.cpp
@@ -8,6 +8,13 @@
 #include "SequencerAssetEditor.h"
 #include "SequencerObjectChangeListener.h"
 
+// We disable the deprecation warnings here because otherwise it'll complain about us
+// implementing RegisterTrackEditor and UnRegisterTrackEditor.  We know
+// that, but we only want it to complain if *others* implement or call these functions.
+//
+// These macros should be removed when those functions are removed.
+PRAGMA_DISABLE_DEPRECATION_WARNINGS
+
 /**
  * SequencerModule implementation (private)
  */
@@ -46,12 +53,26 @@ class FSequencerModule : public ISequencerModule
 
 	virtual void RegisterTrackEditor( FOnCreateTrackEditor InOnCreateTrackEditor ) override
 	{
-		TrackEditorDelegates.AddUnique( InOnCreateTrackEditor );
+		if (!TrackEditorDelegates.ContainsByPredicate([&](const FOnCreateTrackEditor& Delegate){ return Delegate.DEPRECATED_Compare(InOnCreateTrackEditor); }))
+		{
+			TrackEditorDelegates.Add(InOnCreateTrackEditor);
+		}
 	}
 
 	virtual void UnRegisterTrackEditor( FOnCreateTrackEditor InOnCreateTrackEditor ) override
 	{
-		TrackEditorDelegates.Remove( InOnCreateTrackEditor );
+		TrackEditorDelegates.RemoveAll( [&](const FOnCreateTrackEditor& Delegate){ return Delegate.DEPRECATED_Compare(InOnCreateTrackEditor); } );
+	}
+
+	virtual FDelegateHandle RegisterTrackEditor_Handle( FOnCreateTrackEditor InOnCreateTrackEditor ) override
+	{
+		TrackEditorDelegates.Add( InOnCreateTrackEditor );
+		return TrackEditorDelegates.Last().GetHandle();
+	}
+
+	virtual void UnRegisterTrackEditor_Handle( FDelegateHandle InHandle ) override
+	{
+		TrackEditorDelegates.RemoveAll( [=](const FOnCreateTrackEditor& Delegate){ return Delegate.GetHandle() == InHandle; } );
 	}
 
 	virtual void StartupModule() override
@@ -74,6 +95,7 @@ private:
 	TArray< FOnCreateTrackEditor > TrackEditorDelegates;
 };
 
+PRAGMA_ENABLE_DEPRECATION_WARNINGS
 
 
 IMPLEMENT_MODULE( FSequencerModule, Sequencer );
diff --git a/Engine/Source/Editor/Sequencer/Public/ISequencerModule.h b/Engine/Source/Editor/Sequencer/Public/ISequencerModule.h
--- a/Engine/Source/Editor/Sequencer/Public/ISequencerModule.h
+++ b/Engine/Source/Editor/Sequencer/Public/ISequencerModule.h
@@ -64,6 +64,7 @@ public:
 	 *
 	 * @param InOnCreateTrackEditor	Delegate to register
 	 */
+	DELEGATE_DEPRECATED("This function is deprecated - please replace any usage with RegisterTrackEditor_Handle.")
 	virtual void RegisterTrackEditor( FOnCreateTrackEditor InOnCreateTrackEditor ) = 0;
 
 	/** 
@@ -71,6 +72,23 @@ public:
 	 *
 	 * @param InOnCreateTrackEditor	Delegate to unregister
 	 */
+	DELEGATE_DEPRECATED("This function is deprecated - please replace any usage with UnRegisterTrackEditor_Handle, passing the result of RegisterTrackEditor_Handle.")
 	virtual void UnRegisterTrackEditor( FOnCreateTrackEditor InOnCreateTrackEditor ) = 0;
+
+	/** 
+	 * Registers a delegate that will create an editor for a track in each sequencer 
+	 *
+	 * @param InOnCreateTrackEditor	Delegate to register
+	 *
+	 * @return A handle to the newly-added delegate.
+	 */
+	virtual FDelegateHandle RegisterTrackEditor_Handle( FOnCreateTrackEditor InOnCreateTrackEditor ) = 0;
+
+	/** 
+	 * Unregisters a previously registered delegate for creating a track editor
+	 *
+	 * @param InHandle	Handle to the delegate to unregister
+	 */
+	virtual void UnRegisterTrackEditor_Handle( FDelegateHandle InHandle ) = 0;
 };
 
diff --git a/Engine/Source/Editor/UMGEditor/Private/UMGEditorModule.cpp b/Engine/Source/Editor/UMGEditor/Private/UMGEditorModule.cpp
--- a/Engine/Source/Editor/UMGEditor/Private/UMGEditorModule.cpp
+++ b/Engine/Source/Editor/UMGEditor/Private/UMGEditorModule.cpp
@@ -54,8 +54,8 @@ public:
 
 		// Register with the sequencer module that we provide auto-key handlers.
 		ISequencerModule& SequencerModule = FModuleManager::Get().LoadModuleChecked<ISequencerModule>("Sequencer");
-		SequencerModule.RegisterTrackEditor(FOnCreateTrackEditor::CreateStatic(&FMarginTrackEditor::CreateTrackEditor));
-		SequencerModule.RegisterTrackEditor(FOnCreateTrackEditor::CreateStatic(&F2DTransformTrackEditor::CreateTrackEditor));
+		MarginTrackEditorCreateTrackEditorHandle    = SequencerModule.RegisterTrackEditor_Handle(FOnCreateTrackEditor::CreateStatic(&FMarginTrackEditor::CreateTrackEditor));
+		TransformTrackEditorCreateTrackEditorHandle = SequencerModule.RegisterTrackEditor_Handle(FOnCreateTrackEditor::CreateStatic(&F2DTransformTrackEditor::CreateTrackEditor));
 	}
 
 	/** Called before the module is unloaded, right before the module object is destroyed. */
@@ -132,6 +132,9 @@ private:
 	TSharedPtr<FExtensibilityManager> MenuExtensibilityManager;
 	TSharedPtr<FExtensibilityManager> ToolBarExtensibilityManager;
 
+	FDelegateHandle MarginTrackEditorCreateTrackEditorHandle;
+	FDelegateHandle TransformTrackEditorCreateTrackEditorHandle;
+
 	/** All created asset type actions.  Cached here so that we can unregister it during shutdown. */
 	TArray< TSharedPtr<IAssetTypeActions> > CreatedAssetTypeActions;
 
diff --git a/Engine/Source/Editor/UnrealEd/Classes/AutoReimport/AutoReimportManager.h b/Engine/Source/Editor/UnrealEd/Classes/AutoReimport/AutoReimportManager.h
--- a/Engine/Source/Editor/UnrealEd/Classes/AutoReimport/AutoReimportManager.h
+++ b/Engine/Source/Editor/UnrealEd/Classes/AutoReimport/AutoReimportManager.h
@@ -29,11 +29,15 @@ public:
 	DECLARE_DELEGATE_RetVal_TwoParams( bool, FAutoReimportDelegate, UObject*, const FString&);
 
 	/** Add a handler for reimporting custom file types, ex AddReimportDelegate("xml", FAutoReimportDelegate::CreateStatic(&MyImporter::AutoReimportStuff))  */
-	void AddReimportDelegate(const FString& FileType, const FAutoReimportDelegate& Delegate);
+	FDelegateHandle AddReimportDelegate(const FString& FileType, const FAutoReimportDelegate& Delegate);
 
 	/** Remove a handler for reimporting a custom file type. */
+	DELEGATE_DEPRECATED("This overload of RemoveReimportDelegate is deprecated, instead pass the result of AddReimportDelegate.")
 	void RemoveReimportDelegate( const FString& FileType, const FAutoReimportDelegate& Delegate );
 
+	/** Remove a handler for reimporting a custom file type. */
+	void RemoveReimportDelegate( const FString& FileType, FDelegateHandle Handle );
+
 	/** Reimport the file */
 	static bool ReimportFile(UObject* Obj, const FString File);
 
@@ -100,6 +104,9 @@ private:
 	/*The Directories we are currently watching */
 	FDirectories WatchedDirectories;
 
+	/** Handle to the registered OnDirectoryChanged delegate. */
+	FDelegateHandle OnDirectoryChangedDelegateHandle;
+
 	//////////////////////////////////////////////////////////////////////////
 	// Auto Reimport handlers
 
diff --git a/Engine/Source/Editor/UnrealEd/Classes/Editor/EditorEngine.h b/Engine/Source/Editor/UnrealEd/Classes/Editor/EditorEngine.h
--- a/Engine/Source/Editor/UnrealEd/Classes/Editor/EditorEngine.h
+++ b/Engine/Source/Editor/UnrealEd/Classes/Editor/EditorEngine.h
@@ -2701,10 +2701,17 @@ private:
 	/** Gets the init values for worlds opened via Map_Load in the editor */
 	UWorld::InitializationValues GetEditorWorldInitializationValues() const;
 
+	FDelegateHandle OnLoginPIECompleteDelegateHandle;
+
+	TMap<FName, FDelegateHandle> OnLoginPIECompleteDelegateHandlesForPIEInstances;
+
 public:
 	// Launcher Worker
 	TSharedPtr<class ILauncherWorker> LauncherWorker;
 	
 	/** Function to run the Play On command for automation testing. */
 	void AutomationPlayUsingLauncher(const FString& InLauncherDeviceId);	
+
+private:
+	FTimerHandle CleanupPIEOnlineSessionsTimerHandle;
 };
diff --git a/Engine/Source/Editor/UnrealEd/Classes/Editor/UnrealEdEngine.h b/Engine/Source/Editor/UnrealEd/Classes/Editor/UnrealEdEngine.h
--- a/Engine/Source/Editor/UnrealEd/Classes/Editor/UnrealEdEngine.h
+++ b/Engine/Source/Editor/UnrealEd/Classes/Editor/UnrealEdEngine.h
@@ -740,4 +740,7 @@ protected:
 
 	/** Instance responsible for monitoring this editor's performance */
 	FPerformanceMonitor* PerformanceMonitor;
+
+	/** Handle to the registered OnMatineeEditorClosed delegate. */
+	FDelegateHandle OnMatineeEditorClosedDelegateHandle;
 };
diff --git a/Engine/Source/Editor/UnrealEd/Private/AutoReimportManager.cpp b/Engine/Source/Editor/UnrealEd/Private/AutoReimportManager.cpp
--- a/Engine/Source/Editor/UnrealEd/Private/AutoReimportManager.cpp
+++ b/Engine/Source/Editor/UnrealEd/Private/AutoReimportManager.cpp
@@ -19,25 +19,32 @@ UAutoReimportManager::UAutoReimportManager(const FObjectInitializer& ObjectIniti
 	ReimportTime = 0.0f;
 }
 
-void UAutoReimportManager::AddReimportDelegate( const FString& FileType, const FAutoReimportDelegate& Delegate )
+FDelegateHandle UAutoReimportManager::AddReimportDelegate( const FString& FileType, const FAutoReimportDelegate& Delegate )
 {
 	if(ReimportDelegates* Del = ExtensionToReimporter.Find(FileType))
 	{
 		Del->Add(Delegate);
+		return Del->Last().GetHandle();
 	}
-	else
+
+	ReimportDelegates Delegates;
+	Delegates.Add(Delegate);
+	return ExtensionToReimporter.Add(FileType, Delegates).Last().GetHandle();
+}
+
+void UAutoReimportManager::RemoveReimportDelegate( const FString& FileType, const FAutoReimportDelegate& InDelegate )
+{
+	if(ReimportDelegates* Del = ExtensionToReimporter.Find(FileType))
 	{
-		ReimportDelegates Delegates;
-		Delegates.Add(Delegate);
-		ExtensionToReimporter.Add(FileType, Delegates);
+		Del->RemoveAll([&](const FAutoReimportDelegate& Delegate){ return Delegate.DEPRECATED_Compare(InDelegate); });
 	}
 }
 
-void UAutoReimportManager::RemoveReimportDelegate( const FString& FileType, const FAutoReimportDelegate& Delegate )
+void UAutoReimportManager::RemoveReimportDelegate( const FString& FileType, FDelegateHandle Handle )
 {
 	if(ReimportDelegates* Del = ExtensionToReimporter.Find(FileType))
 	{
-		Del->Remove(Delegate);
+		Del->RemoveAll([=](const FAutoReimportDelegate& Delegate){ return Delegate.GetHandle() == Handle; });
 	}
 }
 
@@ -284,7 +291,7 @@ void UAutoReimportManager::ActivateAutoReimport( const FString& Directory)
 	{
 		if(Directory.IsEmpty() == false)
 		{
-			DirectoryWatcher->RegisterDirectoryChangedCallback(Directory, IDirectoryWatcher::FDirectoryChanged::CreateUObject(this, &UAutoReimportManager::OnDirectoryChanged));
+			DirectoryWatcher->RegisterDirectoryChangedCallback_Handle(Directory, IDirectoryWatcher::FDirectoryChanged::CreateUObject(this, &UAutoReimportManager::OnDirectoryChanged), OnDirectoryChangedDelegateHandle);
 		}
 	}
 }
@@ -297,7 +304,7 @@ void UAutoReimportManager::DeactivateAutoReimport( const FString& Directory )
 	{
 		if(Directory.IsEmpty() == false)
 		{
-			DirectoryWatcher->UnregisterDirectoryChangedCallback(Directory, IDirectoryWatcher::FDirectoryChanged::CreateUObject(this, &UAutoReimportManager::OnDirectoryChanged));
+			DirectoryWatcher->UnregisterDirectoryChangedCallback_Handle(Directory, OnDirectoryChangedDelegateHandle);
 		}
 	}
 }
diff --git a/Engine/Source/Editor/UnrealEd/Private/Kismet2/CompilerResultsLog.cpp b/Engine/Source/Editor/UnrealEd/Private/Kismet2/CompilerResultsLog.cpp
--- a/Engine/Source/Editor/UnrealEd/Private/Kismet2/CompilerResultsLog.cpp
+++ b/Engine/Source/Editor/UnrealEd/Private/Kismet2/CompilerResultsLog.cpp
@@ -15,6 +15,7 @@
 
 const FName FCompilerResultsLog::Name(TEXT("CompilerResultsLog"));
 FCompilerResultsLog* FCompilerResultsLog::CurrentEventTarget = nullptr;
+FDelegateHandle FCompilerResultsLog::GetGlobalModuleCompilerDumpDelegateHandle;
 
 //////////////////////////////////////////////////////////////////////////
 // FCompilerResultsLog
@@ -93,12 +94,12 @@ void FCompilerResultsLog::Register()
 	FMessageLogModule& MessageLogModule = FModuleManager::LoadModuleChecked<FMessageLogModule>("MessageLog");
 	MessageLogModule.RegisterLogListing(Name, LOCTEXT("CompilerLog", "Compiler Log"));
 
-	IHotReloadModule::Get().OnModuleCompilerFinished().AddStatic( &FCompilerResultsLog::GetGlobalModuleCompilerDump );
+	GetGlobalModuleCompilerDumpDelegateHandle = IHotReloadModule::Get().OnModuleCompilerFinished().AddStatic( &FCompilerResultsLog::GetGlobalModuleCompilerDump );
 }
 
 void FCompilerResultsLog::Unregister()
 {
-	IHotReloadModule::Get().OnModuleCompilerFinished().RemoveStatic( &FCompilerResultsLog::GetGlobalModuleCompilerDump );
+	IHotReloadModule::Get().OnModuleCompilerFinished().Remove( GetGlobalModuleCompilerDumpDelegateHandle );
 
 	FMessageLogModule& MessageLogModule = FModuleManager::LoadModuleChecked<FMessageLogModule>("MessageLog");
 	MessageLogModule.UnregisterLogListing(Name);
diff --git a/Engine/Source/Editor/UnrealEd/Private/PerformanceMonitor.cpp b/Engine/Source/Editor/UnrealEd/Private/PerformanceMonitor.cpp
--- a/Engine/Source/Editor/UnrealEd/Private/PerformanceMonitor.cpp
+++ b/Engine/Source/Editor/UnrealEd/Private/PerformanceMonitor.cpp
@@ -125,12 +125,12 @@ FPerformanceMonitor::FPerformanceMonitor()
 		}
 	});
 
-	IConsoleManager::Get().RegisterConsoleVariableSink(CVarDelegate);
+	CVarDelegateHandle = IConsoleManager::Get().RegisterConsoleVariableSink_Handle(CVarDelegate);
 }
 
 FPerformanceMonitor::~FPerformanceMonitor()
 {
-	IConsoleManager::Get().UnregisterConsoleVariableSink(CVarDelegate);	
+	IConsoleManager::Get().UnregisterConsoleVariableSink_Handle(CVarDelegateHandle);
 }
 
 bool FPerformanceMonitor::WillAutoScalabilityHelp() const
diff --git a/Engine/Source/Editor/UnrealEd/Private/PerformanceMonitor.h b/Engine/Source/Editor/UnrealEd/Private/PerformanceMonitor.h
--- a/Engine/Source/Editor/UnrealEd/Private/PerformanceMonitor.h
+++ b/Engine/Source/Editor/UnrealEd/Private/PerformanceMonitor.h
@@ -197,4 +197,6 @@ private:
 
 	/** Console variable delegate for checking when the console variables have changed */
 	FConsoleCommandDelegate CVarDelegate;
+
+	FConsoleVariableSinkHandle CVarDelegateHandle;
 };
\ No newline at end of file
diff --git a/Engine/Source/Editor/UnrealEd/Private/PlayLevel.cpp b/Engine/Source/Editor/UnrealEd/Private/PlayLevel.cpp
--- a/Engine/Source/Editor/UnrealEd/Private/PlayLevel.cpp
+++ b/Engine/Source/Editor/UnrealEd/Private/PlayLevel.cpp
@@ -196,7 +196,7 @@ void UEditorEngine::EndPlayMap()
 		// as a different delegate
 		FTimerDelegate DestroyTimer;
 		DestroyTimer.BindUObject(this, &UEditorEngine::CleanupPIEOnlineSessions, OnlineIdentifiers);
-		GetTimerManager()->SetTimer(DestroyTimer, 0.1f, false);
+		GetTimerManager()->SetTimer(CleanupPIEOnlineSessionsTimerHandle, DestroyTimer, 0.1f, false);
 	}
 	
 	{
@@ -2323,7 +2323,7 @@ void UEditorEngine::LoginPIEInstances(bool bAnyBlueprintErrors, bool bStartInSpe
 			Delegate.BindUObject(this, &UEditorEngine::OnLoginPIEComplete, DataStruct);
 
 			// Login first and continue the flow later
-			IdentityInt->AddOnLoginCompleteDelegate(0, Delegate);
+			OnLoginPIECompleteDelegateHandle = IdentityInt->AddOnLoginCompleteDelegate_Handle(0, Delegate);
 			IdentityInt->Login(0, AccountCreds);
 
 			ClientNum++;
@@ -2368,8 +2368,8 @@ void UEditorEngine::LoginPIEInstances(bool bAnyBlueprintErrors, bool bStartInSpe
 		FOnLoginCompleteDelegate Delegate;
 		Delegate.BindUObject(this, &UEditorEngine::OnLoginPIEComplete, DataStruct);
 
-		IdentityInt->ClearOnLoginCompleteDelegate(0, Delegate);
-		IdentityInt->AddOnLoginCompleteDelegate(0, Delegate);
+		IdentityInt->ClearOnLoginCompleteDelegate_Handle(0, OnLoginPIECompleteDelegateHandlesForPIEInstances.FindRef(OnlineIdentifier));
+		OnLoginPIECompleteDelegateHandlesForPIEInstances[OnlineIdentifier] = IdentityInt->AddOnLoginCompleteDelegate_Handle(0, Delegate);
 		IdentityInt->Login(0, AccountCreds);
 	}
 
@@ -2386,9 +2386,8 @@ void UEditorEngine::OnLoginPIEComplete(int32 LocalUserNum, bool bWasSuccessful,
 	IOnlineIdentityPtr IdentityInt = Online::GetIdentityInterface(OnlineIdentifier);
 
 	// Cleanup the login delegate before calling create below
-	FOnLoginCompleteDelegate Delegate;
-	Delegate.BindUObject(this, &UEditorEngine::OnLoginPIEComplete, DataStruct);
-	IdentityInt->ClearOnLoginCompleteDelegate(0, Delegate);
+	IdentityInt->ClearOnLoginCompleteDelegate_Handle(0, OnLoginPIECompleteDelegateHandle);
+	OnLoginPIECompleteDelegateHandlesForPIEInstances.Remove(OnlineIdentifier);
 
 	// Create the new world
 	CreatePIEWorldFromLogin(PieWorldContext, DataStruct.NetMode, DataStruct);
diff --git a/Engine/Source/Editor/UnrealEd/Private/ReferencedAssetsUtils.cpp b/Engine/Source/Editor/UnrealEd/Private/ReferencedAssetsUtils.cpp
--- a/Engine/Source/Editor/UnrealEd/Private/ReferencedAssetsUtils.cpp
+++ b/Engine/Source/Editor/UnrealEd/Private/ReferencedAssetsUtils.cpp
@@ -13,7 +13,7 @@ DEFINE_LOG_CATEGORY_STATIC(LogReferncedAssetsBrowser, Log, All);
  */
 FFindReferencedAssets::FFindReferencedAssets(void)
 {
-	FEditorDelegates::MapChange.AddRaw(this, &FFindReferencedAssets::OnEditorMapChange);
+	OnEditorMapChangeDelegateHandle = FEditorDelegates::MapChange.AddRaw(this, &FFindReferencedAssets::OnEditorMapChange);
 
 	// Set up our ignore lists
 	IgnoreClasses.Add(ULevel::StaticClass());
@@ -49,7 +49,7 @@ FFindReferencedAssets::FFindReferencedAssets(void)
  */
 FFindReferencedAssets::~FFindReferencedAssets(void)
 {
-	FEditorDelegates::MapChange.RemoveRaw(this, &FFindReferencedAssets::OnEditorMapChange);
+	FEditorDelegates::MapChange.Remove(OnEditorMapChangeDelegateHandle);
 }
 
 void FFindReferencedAssets::OnEditorMapChange( uint32 Flag )
diff --git a/Engine/Source/Editor/UnrealEd/Private/UnrealEdEngine.cpp b/Engine/Source/Editor/UnrealEd/Private/UnrealEdEngine.cpp
--- a/Engine/Source/Editor/UnrealEd/Private/UnrealEdEngine.cpp
+++ b/Engine/Source/Editor/UnrealEd/Private/UnrealEdEngine.cpp
@@ -510,7 +510,7 @@ void UUnrealEdEngine::OnPostWindowsMessage(FViewport* Viewport, uint32 Message)
 void UUnrealEdEngine::OnOpenMatinee()
 {
 	// Register a delegate to pickup when Matinee is closed.
-	GLevelEditorModeTools().OnEditorModeChanged().AddUObject( this, &UUnrealEdEngine::OnMatineeEditorClosed );
+	OnMatineeEditorClosedDelegateHandle = GLevelEditorModeTools().OnEditorModeChanged().AddUObject( this, &UUnrealEdEngine::OnMatineeEditorClosed );
 }
 
 
@@ -1242,6 +1242,6 @@ void UUnrealEdEngine::OnMatineeEditorClosed( FEdMode* Mode, bool IsEntering )
 		}
 
 		// Remove this delegate. 
-		GLevelEditorModeTools().OnEditorModeChanged().RemoveUObject( this, &UUnrealEdEngine::OnMatineeEditorClosed );
+		GLevelEditorModeTools().OnEditorModeChanged().Remove( OnMatineeEditorClosedDelegateHandle );
 	}	
 }
diff --git a/Engine/Source/Editor/UnrealEd/Private/UnrealEdMisc.cpp b/Engine/Source/Editor/UnrealEd/Private/UnrealEdMisc.cpp
--- a/Engine/Source/Editor/UnrealEd/Private/UnrealEdMisc.cpp
+++ b/Engine/Source/Editor/UnrealEd/Private/UnrealEdMisc.cpp
@@ -422,7 +422,7 @@ void FUnrealEdMisc::OnInit()
 	FAssetNameToken::OnGotoAsset().BindRaw(this, &FUnrealEdMisc::OnGotoAsset);
 
 	// Register to receive notification of new key bindings
-	FInputBindingManager::Get().RegisterUserDefinedGestureChanged(FOnUserDefinedGestureChanged::FDelegate::CreateRaw( this, &FUnrealEdMisc::OnUserDefinedGestureChanged ));
+	OnUserDefinedGestureChangedDelegateHandle = FInputBindingManager::Get().RegisterUserDefinedGestureChanged(FOnUserDefinedGestureChanged::FDelegate::CreateRaw( this, &FUnrealEdMisc::OnUserDefinedGestureChanged ));
 
 	SlowTask.EnterProgressFrame(10);
 
@@ -433,7 +433,8 @@ void FUnrealEdMisc::OnInit()
 	float Seconds = 60.0f;
 	FTimerDelegate Delegate;
 	Delegate.BindRaw( this, &FUnrealEdMisc::EditorAnalyticsHeartbeat );
-	GEditor->GetTimerManager()->SetTimer( Delegate, Seconds, true );
+
+	GEditor->GetTimerManager()->SetTimer( EditorAnalyticsHeartbeatTimerHandle, Delegate, Seconds, true );
 
 	// add handler to notify about navmesh building process
 	NavigationBuildingNotificationHandler = MakeShareable(new FNavigationBuildingNotificationImpl());
@@ -677,7 +678,7 @@ void FUnrealEdMisc::OnExit()
 		FEditorViewportStats::SendUsageData();
 	}
 
-	FInputBindingManager::Get().UnregisterUserDefinedGestureChanged(FOnUserDefinedGestureChanged::FDelegate::CreateRaw( this, &FUnrealEdMisc::OnUserDefinedGestureChanged ));
+	FInputBindingManager::Get().UnregisterUserDefinedGestureChanged(OnUserDefinedGestureChangedDelegateHandle);
 	FMessageLog::OnMessageSelectionChanged().Unbind();
 	FUObjectToken::DefaultOnMessageTokenActivated().Unbind();
 	FUObjectToken::DefaultOnGetObjectDisplayName().Unbind();
@@ -1389,7 +1390,7 @@ void FUnrealEdMisc::BeginPerformanceSurvey()
 
 	// Tell the level editor we want to be notified when selection changes
 	FLevelEditorModule& LevelEditor = FModuleManager::LoadModuleChecked<FLevelEditorModule>( LevelEditorName );
-	LevelEditor.OnMapChanged().AddRaw( this, &FUnrealEdMisc::OnMapChanged );
+	OnMapChangedDelegateHandle = LevelEditor.OnMapChanged().AddRaw( this, &FUnrealEdMisc::OnMapChanged );
 
 	// Initialize survey variables
 	bIsSurveyingPerformance = true;
@@ -1472,7 +1473,7 @@ void FUnrealEdMisc::CancelPerformanceSurvey()
 	FrameRateSamples.Empty();
 
 	FLevelEditorModule& LevelEditor = FModuleManager::LoadModuleChecked<FLevelEditorModule>( LevelEditorName );
-	LevelEditor.OnMapChanged().RemoveRaw( this, &FUnrealEdMisc::OnMapChanged );
+	LevelEditor.OnMapChanged().Remove( OnMapChangedDelegateHandle );
 }
 
 void FUnrealEdMisc::OnMapChanged( UWorld* World, EMapChangeType::Type MapChangeType )
diff --git a/Engine/Source/Editor/UnrealEd/Public/Kismet2/CompilerResultsLog.h b/Engine/Source/Editor/UnrealEd/Public/Kismet2/CompilerResultsLog.h
--- a/Engine/Source/Editor/UnrealEd/Public/Kismet2/CompilerResultsLog.h
+++ b/Engine/Source/Editor/UnrealEd/Public/Kismet2/CompilerResultsLog.h
@@ -194,6 +194,9 @@ private:
 	
 	/** The log target for compile events */
 	static FCompilerResultsLog* CurrentEventTarget;
+
+	/** Handle to the registered GetGlobalModuleCompilerDump delegate. */
+	static FDelegateHandle GetGlobalModuleCompilerDumpDelegateHandle;
 };
 
 /** This class will begin a new compile event on construction, and automatically end it when the instance goes out of scope */
diff --git a/Engine/Source/Editor/UnrealEd/Public/ReferencedAssetsUtils.h b/Engine/Source/Editor/UnrealEd/Public/ReferencedAssetsUtils.h
--- a/Engine/Source/Editor/UnrealEd/Public/ReferencedAssetsUtils.h
+++ b/Engine/Source/Editor/UnrealEd/Public/ReferencedAssetsUtils.h
@@ -183,6 +183,9 @@ protected:
 	 * The object graph for the assets referenced by the currently selected actors.
 	 */
 	ObjectReferenceGraph ReferenceGraph;
+
+	/** Handle to the registered OnEditorMapChange delegate. */
+	FDelegateHandle OnEditorMapChangeDelegateHandle;
 };
 
 #endif // __REFERENCEDASSETSBROWSER_H__
diff --git a/Engine/Source/Editor/UnrealEd/Public/UnrealEdMisc.h b/Engine/Source/Editor/UnrealEd/Public/UnrealEdMisc.h
--- a/Engine/Source/Editor/UnrealEd/Public/UnrealEdMisc.h
+++ b/Engine/Source/Editor/UnrealEd/Public/UnrealEdMisc.h
@@ -319,4 +319,12 @@ private:
 
 	/** Package names and the number of times they have been updated */
 	TMap<FName, uint32> NumUpdatesByAssetName;
+
+	/** Handle to the registered OnUserDefinedGestureChanged delegate. */
+	FDelegateHandle OnUserDefinedGestureChangedDelegateHandle;
+
+	/** Handle to the registered OnMapChanged delegate. */
+	FDelegateHandle OnMapChangedDelegateHandle;
+
+	FTimerHandle EditorAnalyticsHeartbeatTimerHandle;
 };
diff --git a/Engine/Source/Runtime/AIModule/Classes/Actions/PawnAction_Move.h b/Engine/Source/Runtime/AIModule/Classes/Actions/PawnAction_Move.h
--- a/Engine/Source/Runtime/AIModule/Classes/Actions/PawnAction_Move.h
+++ b/Engine/Source/Runtime/AIModule/Classes/Actions/PawnAction_Move.h
@@ -87,6 +87,7 @@ protected:
 	FNavPathSharedPtr Path;
 
 	FNavigationPath::FPathObserverDelegate::FDelegate PathObserver;
+	FDelegateHandle PathObserverDelegateHandle;
 	
 	/** Handle for efficient management of DeferredPerformMoveAction timer */
 	FTimerHandle TimerHandle_DeferredPerformMoveAction;
diff --git a/Engine/Source/Runtime/AIModule/Classes/BehaviorTree/BlackboardComponent.h b/Engine/Source/Runtime/AIModule/Classes/BehaviorTree/BlackboardComponent.h
--- a/Engine/Source/Runtime/AIModule/Classes/BehaviorTree/BlackboardComponent.h
+++ b/Engine/Source/Runtime/AIModule/Classes/BehaviorTree/BlackboardComponent.h
@@ -65,11 +65,15 @@ public:
 	bool HasValidAsset() const;
 
 	/** register observer for blackboard key */
-	void RegisterObserver(FBlackboard::FKey KeyID, FOnBlackboardChange ObserverDelegate);
+	FDelegateHandle RegisterObserver(FBlackboard::FKey KeyID, FOnBlackboardChange ObserverDelegate);
 
 	/** unregister observer from blackboard key */
+	DELEGATE_DEPRECATED("This overload of UnregisterObserver is deprecated, instead pass the result of RegisterObserver.")
 	void UnregisterObserver(FBlackboard::FKey KeyID, FOnBlackboardChange ObserverDelegate);
 
+	/** unregister observer from blackboard key */
+	void UnregisterObserver(FBlackboard::FKey KeyID, FDelegateHandle ObserverHandle);
+
 	/** pause change notifies and add them to queue */
 	void PauseUpdates();
 
diff --git a/Engine/Source/Runtime/AIModule/Classes/BehaviorTree/Decorators/BTDecorator_BlackboardBase.h b/Engine/Source/Runtime/AIModule/Classes/BehaviorTree/Decorators/BTDecorator_BlackboardBase.h
--- a/Engine/Source/Runtime/AIModule/Classes/BehaviorTree/Decorators/BTDecorator_BlackboardBase.h
+++ b/Engine/Source/Runtime/AIModule/Classes/BehaviorTree/Decorators/BTDecorator_BlackboardBase.h
@@ -31,6 +31,8 @@ protected:
 
 	FOnBlackboardChange BBKeyObserver;
 
+	TMap<FBlackboard::FKey, FDelegateHandle> BBKeyObserverDelegateHandles;
+
 	virtual void PostInitProperties() override;
 
 	/** called when execution flow controller becomes active */
diff --git a/Engine/Source/Runtime/AIModule/Classes/BehaviorTree/Decorators/BTDecorator_BlueprintBase.h b/Engine/Source/Runtime/AIModule/Classes/BehaviorTree/Decorators/BTDecorator_BlueprintBase.h
--- a/Engine/Source/Runtime/AIModule/Classes/BehaviorTree/Decorators/BTDecorator_BlueprintBase.h
+++ b/Engine/Source/Runtime/AIModule/Classes/BehaviorTree/Decorators/BTDecorator_BlueprintBase.h
@@ -200,6 +200,8 @@ protected:
 	UFUNCTION(BlueprintCallable, Category="AI|BehaviorTree")
 	bool IsDecoratorObserverActive() const;
 
+	TMap<FBlackboard::FKey, FDelegateHandle> BBKeyObserverDelegateHandles;
+
 	friend FBehaviorBlueprintDetails;
 
 	FORCEINLINE bool GetNeedsTickForConditionChecking() const { return ReceiveConditionCheckImplementations != 0 && (bIsObservingBB == false || bCheckConditionOnlyBlackBoardChanges == false); }
diff --git a/Engine/Source/Runtime/AIModule/Classes/Blueprint/AIAsyncTaskBlueprintProxy.h b/Engine/Source/Runtime/AIModule/Classes/Blueprint/AIAsyncTaskBlueprintProxy.h
--- a/Engine/Source/Runtime/AIModule/Classes/Blueprint/AIAsyncTaskBlueprintProxy.h
+++ b/Engine/Source/Runtime/AIModule/Classes/Blueprint/AIAsyncTaskBlueprintProxy.h
@@ -35,4 +35,6 @@ public:
 	TWeakObjectPtr<AAIController> AIController;
 	FAIRequestID MoveRequestId;
 	TWeakObjectPtr<UWorld> MyWorld;
+
+	FTimerHandle TimerHandle_OnNoPath;
 };
diff --git a/Engine/Source/Runtime/AIModule/Classes/Perception/AIPerceptionSystem.h b/Engine/Source/Runtime/AIModule/Classes/Perception/AIPerceptionSystem.h
--- a/Engine/Source/Runtime/AIModule/Classes/Perception/AIPerceptionSystem.h
+++ b/Engine/Source/Runtime/AIModule/Classes/Perception/AIPerceptionSystem.h
@@ -162,6 +162,9 @@ public:
 	FString GetSenseName(FAISenseID SenseID) const;
 	FString GetPerceptionDebugLegend() const { return PerceptionDebugLegend; }
 #endif 
+
+private:
+	FTimerHandle AgeStimuliTimerHandle;
 };
 
 //////////////////////////////////////////////////////////////////////////
diff --git a/Engine/Source/Runtime/AIModule/Classes/Perception/AISenseConfig_Damage.h b/Engine/Source/Runtime/AIModule/Classes/Perception/AISenseConfig_Damage.h
--- a/Engine/Source/Runtime/AIModule/Classes/Perception/AISenseConfig_Damage.h
+++ b/Engine/Source/Runtime/AIModule/Classes/Perception/AISenseConfig_Damage.h
@@ -16,5 +16,5 @@ public:
 	UPROPERTY(EditDefaultsOnly, BlueprintReadOnly, Category = "Sense", NoClear, config)
 	TSubclassOf<UAISense_Damage> Implementation;
 
-	virtual TSubclassOf<UAISense> GetSenseImplementation() const { return Implementation; }
+	virtual TSubclassOf<UAISense> GetSenseImplementation() const;
 };
\ No newline at end of file
diff --git a/Engine/Source/Runtime/AIModule/Private/Actions/PawnAction_Move.cpp b/Engine/Source/Runtime/AIModule/Private/Actions/PawnAction_Move.cpp
--- a/Engine/Source/Runtime/AIModule/Private/Actions/PawnAction_Move.cpp
+++ b/Engine/Source/Runtime/AIModule/Private/Actions/PawnAction_Move.cpp
@@ -228,7 +228,7 @@ void UPawnAction_Move::ClearPath()
 	ClearPendingRepath();
 	if (Path.IsValid())
 	{
-		Path->RemoveObserver(PathObserver);
+		Path->RemoveObserver(PathObserverDelegateHandle);
 		Path = NULL;
 	}
 }
@@ -239,7 +239,7 @@ void UPawnAction_Move::SetPath(FNavPathSharedRef InPath)
 	{
 		ClearPath();
 		Path = InPath;
-		Path->AddObserver(PathObserver);
+		PathObserverDelegateHandle = Path->AddObserver(PathObserver);
 
 		// skip auto updates, it will be handled manually to include controller's ShouldPostponePathUpdates()
 		Path->EnableRecalculationOnInvalidation(false);
diff --git a/Engine/Source/Runtime/AIModule/Private/BehaviorTree/BlackboardComponent.cpp b/Engine/Source/Runtime/AIModule/Private/BehaviorTree/BlackboardComponent.cpp
--- a/Engine/Source/Runtime/AIModule/Private/BehaviorTree/BlackboardComponent.cpp
+++ b/Engine/Source/Runtime/AIModule/Private/BehaviorTree/BlackboardComponent.cpp
@@ -244,14 +244,42 @@ int32 UBlackboardComponent::GetNumKeys() const
 	return BlackboardAsset ? BlackboardAsset->GetNumKeys() : 0;
 }
 
-void UBlackboardComponent::RegisterObserver(FBlackboard::FKey KeyID, FOnBlackboardChange ObserverDelegate)
+FDelegateHandle UBlackboardComponent::RegisterObserver(FBlackboard::FKey KeyID, FOnBlackboardChange ObserverDelegate)
 {
-	Observers.AddUnique(KeyID, ObserverDelegate);
+	for (auto It = Observers.CreateConstKeyIterator(KeyID); It; ++It)
+	{
+		// If the pair's value matches, return a pointer to it.
+		if (It.Value().GetHandle() == ObserverDelegate.GetHandle())
+		{
+			return It.Value().GetHandle();
+		}
+	}
+
+	return Observers.Add(KeyID, ObserverDelegate).GetHandle();
 }
 
 void UBlackboardComponent::UnregisterObserver(FBlackboard::FKey KeyID, FOnBlackboardChange ObserverDelegate)
 {
-	Observers.RemoveSingle(KeyID, ObserverDelegate);
+	for (auto It = Observers.CreateKeyIterator(KeyID); It; ++It)
+	{
+		if (It.Value().DEPRECATED_Compare(ObserverDelegate))
+		{
+			It.RemoveCurrent();
+			return;
+		}
+	}
+}
+
+void UBlackboardComponent::UnregisterObserver(FBlackboard::FKey KeyID, FDelegateHandle ObserverHandle)
+{
+	for (auto It = Observers.CreateKeyIterator(KeyID); It; ++It)
+	{
+		if (It.Value().GetHandle() == ObserverHandle)
+		{
+			It.RemoveCurrent();
+			return;
+		}
+	}
 }
 
 void UBlackboardComponent::PauseUpdates()
diff --git a/Engine/Source/Runtime/AIModule/Private/BehaviorTree/Decorators/BTDecorator_BlackboardBase.cpp b/Engine/Source/Runtime/AIModule/Private/BehaviorTree/Decorators/BTDecorator_BlackboardBase.cpp
--- a/Engine/Source/Runtime/AIModule/Private/BehaviorTree/Decorators/BTDecorator_BlackboardBase.cpp
+++ b/Engine/Source/Runtime/AIModule/Private/BehaviorTree/Decorators/BTDecorator_BlackboardBase.cpp
@@ -31,7 +31,8 @@ void UBTDecorator_BlackboardBase::OnBecomeRelevant(UBehaviorTreeComponent& Owner
 	UBlackboardComponent* BlackboardComp = OwnerComp.GetBlackboardComponent();
 	if (BlackboardComp)
 	{
-		BlackboardComp->RegisterObserver(BlackboardKey.GetSelectedKeyID(), BBKeyObserver);
+		auto KeyID = BlackboardKey.GetSelectedKeyID();
+		BBKeyObserverDelegateHandles.Add(KeyID, BlackboardComp->RegisterObserver(KeyID, BBKeyObserver));
 	}
 }
 
@@ -40,7 +41,9 @@ void UBTDecorator_BlackboardBase::OnCeaseRelevant(UBehaviorTreeComponent& OwnerC
 	UBlackboardComponent* BlackboardComp = OwnerComp.GetBlackboardComponent();
 	if (BlackboardComp)
 	{
-		BlackboardComp->UnregisterObserver(BlackboardKey.GetSelectedKeyID(), BBKeyObserver);
+		auto KeyID = BlackboardKey.GetSelectedKeyID();
+		BlackboardComp->UnregisterObserver(KeyID, BBKeyObserverDelegateHandles.FindRef(KeyID));
+		BBKeyObserverDelegateHandles.Remove(KeyID);
 	}
 }
 
diff --git a/Engine/Source/Runtime/AIModule/Private/BehaviorTree/Decorators/BTDecorator_BlueprintBase.cpp b/Engine/Source/Runtime/AIModule/Private/BehaviorTree/Decorators/BTDecorator_BlueprintBase.cpp
--- a/Engine/Source/Runtime/AIModule/Private/BehaviorTree/Decorators/BTDecorator_BlueprintBase.cpp
+++ b/Engine/Source/Runtime/AIModule/Private/BehaviorTree/Decorators/BTDecorator_BlueprintBase.cpp
@@ -106,7 +106,7 @@ void UBTDecorator_BlueprintBase::OnBecomeRelevant(UBehaviorTreeComponent& OwnerC
 			const FBlackboard::FKey KeyID = BlackboardComp->GetKeyID(ObservedKeyNames[NameIndex]);
 			if (KeyID != FBlackboard::InvalidKey)
 			{
-				BlackboardComp->RegisterObserver(KeyID, BBKeyObserver);
+				BBKeyObserverDelegateHandles.Add(KeyID, BlackboardComp->RegisterObserver(KeyID, BBKeyObserver));
 			}
 		}
 	}
@@ -122,7 +122,8 @@ void UBTDecorator_BlueprintBase::OnCeaseRelevant(UBehaviorTreeComponent& OwnerCo
 			const FBlackboard::FKey KeyID = BlackboardComp->GetKeyID(ObservedKeyNames[NameIndex]);
 			if (KeyID != FBlackboard::InvalidKey)
 			{
-				BlackboardComp->UnregisterObserver(KeyID, BBKeyObserver);
+				BlackboardComp->UnregisterObserver(KeyID, BBKeyObserverDelegateHandles.FindRef(KeyID));
+				BBKeyObserverDelegateHandles.Remove(KeyID);
 			}
 		}
 	}
diff --git a/Engine/Source/Runtime/AIModule/Private/BehaviorTree/Decorators/BTDecorator_CompareBBEntries.cpp b/Engine/Source/Runtime/AIModule/Private/BehaviorTree/Decorators/BTDecorator_CompareBBEntries.cpp
--- a/Engine/Source/Runtime/AIModule/Private/BehaviorTree/Decorators/BTDecorator_CompareBBEntries.cpp
+++ b/Engine/Source/Runtime/AIModule/Private/BehaviorTree/Decorators/BTDecorator_CompareBBEntries.cpp
@@ -70,8 +70,8 @@ void UBTDecorator_CompareBBEntries::OnBecomeRelevant(UBehaviorTreeComponent& Own
 	UBlackboardComponent* BlackboardComp = OwnerComp.GetBlackboardComponent();
 	if (BlackboardComp)
 	{
-		BlackboardComp->RegisterObserver(BlackboardKeyA.GetSelectedKeyID(), BBKeyObserver);
-		BlackboardComp->RegisterObserver(BlackboardKeyB.GetSelectedKeyID(), BBKeyObserver);
+		BBKeyObserverDelegateHandles.Add(BlackboardKeyA.GetSelectedKeyID(), BlackboardComp->RegisterObserver(BlackboardKeyA.GetSelectedKeyID(), BBKeyObserver));
+		BBKeyObserverDelegateHandles.Add(BlackboardKeyA.GetSelectedKeyID(), BlackboardComp->RegisterObserver(BlackboardKeyB.GetSelectedKeyID(), BBKeyObserver));
 	}
 }
 
@@ -80,8 +80,10 @@ void UBTDecorator_CompareBBEntries::OnCeaseRelevant(UBehaviorTreeComponent& Owne
 	UBlackboardComponent* BlackboardComp = OwnerComp.GetBlackboardComponent();
 	if (BlackboardComp)
 	{
-		BlackboardComp->UnregisterObserver(BlackboardKeyA.GetSelectedKeyID(), BBKeyObserver);
-		BlackboardComp->UnregisterObserver(BlackboardKeyB.GetSelectedKeyID(), BBKeyObserver);
+		BlackboardComp->UnregisterObserver(BlackboardKeyA.GetSelectedKeyID(), BBKeyObserverDelegateHandles.FindRef(BlackboardKeyA.GetSelectedKeyID()));
+		BlackboardComp->UnregisterObserver(BlackboardKeyB.GetSelectedKeyID(), BBKeyObserverDelegateHandles.FindRef(BlackboardKeyB.GetSelectedKeyID()));
+		BBKeyObserverDelegateHandles.Remove(BlackboardKeyA.GetSelectedKeyID());
+		BBKeyObserverDelegateHandles.Remove(BlackboardKeyB.GetSelectedKeyID());
 	}
 }
 
diff --git a/Engine/Source/Runtime/AIModule/Private/Blueprint/AIBlueprintHelperLibrary.cpp b/Engine/Source/Runtime/AIModule/Private/Blueprint/AIBlueprintHelperLibrary.cpp
--- a/Engine/Source/Runtime/AIModule/Private/Blueprint/AIBlueprintHelperLibrary.cpp
+++ b/Engine/Source/Runtime/AIModule/Private/Blueprint/AIBlueprintHelperLibrary.cpp
@@ -106,7 +106,7 @@ UAIAsyncTaskBlueprintProxy* UAIBlueprintHelperLibrary::CreateMoveToProxyObject(U
 			}
 			else
 			{
-				World->GetTimerManager().SetTimer(MyObj, &UAIAsyncTaskBlueprintProxy::OnNoPath, 0.1f, false);
+				World->GetTimerManager().SetTimer(MyObj->TimerHandle_OnNoPath, MyObj, &UAIAsyncTaskBlueprintProxy::OnNoPath, 0.1f, false);
 			}
 		}
 	}
diff --git a/Engine/Source/Runtime/AIModule/Private/Perception/AIPerceptionSystem.cpp b/Engine/Source/Runtime/AIModule/Private/Perception/AIPerceptionSystem.cpp
--- a/Engine/Source/Runtime/AIModule/Private/Perception/AIPerceptionSystem.cpp
+++ b/Engine/Source/Runtime/AIModule/Private/Perception/AIPerceptionSystem.cpp
@@ -87,7 +87,7 @@ void UAIPerceptionSystem::PostInitProperties()
 		UWorld* World = GEngine->GetWorldFromContextObject(GetOuter());
 		if (World)
 		{
-			World->GetTimerManager().SetTimer(this, &UAIPerceptionSystem::AgeStimuli, PerceptionAgingRate, /*inbLoop=*/true);
+			World->GetTimerManager().SetTimer(AgeStimuliTimerHandle, this, &UAIPerceptionSystem::AgeStimuli, PerceptionAgingRate, /*inbLoop=*/true);
 		}
 	}
 }
@@ -464,4 +464,4 @@ FString UAIPerceptionSystem::GetSenseName(FAISenseID SenseID) const
 {
 	return SenseID.IsValid() && Senses.IsValidIndex(SenseID) ? *Senses[SenseID]->GetDebugName() : TEXT("Invalid");
 }
-#endif 
\ No newline at end of file
+#endif 
diff --git a/Engine/Source/Runtime/AIModule/Private/Perception/AISenseConfig_Damage.cpp b/Engine/Source/Runtime/AIModule/Private/Perception/AISenseConfig_Damage.cpp
new file mode 100644
--- /dev/null
+++ b/Engine/Source/Runtime/AIModule/Private/Perception/AISenseConfig_Damage.cpp
@@ -0,0 +1,10 @@
+// Copyright 1998-2015 Epic Games, Inc. All Rights Reserved.
+
+#include "AIModulePrivate.h"
+#include "Perception/AISenseConfig_Damage.h"
+#include "Perception/AISense_Damage.h"
+
+TSubclassOf<UAISense> UAISenseConfig_Damage::GetSenseImplementation() const
+{
+	return Implementation;
+}
diff --git a/Engine/Source/Runtime/AppFramework/Private/Widgets/Colors/SColorThemes.cpp b/Engine/Source/Runtime/AppFramework/Private/Widgets/Colors/SColorThemes.cpp
--- a/Engine/Source/Runtime/AppFramework/Private/Widgets/Colors/SColorThemes.cpp
+++ b/Engine/Source/Runtime/AppFramework/Private/Widgets/Colors/SColorThemes.cpp
@@ -627,12 +627,12 @@ void SThemeColorBlocksBar::DestroyPlaceholders()
 
 void SThemeColorBlocksBar::RemoveRefreshCallback()
 {
-	ColorTheme.Get()->OnRefresh().Remove(RefreshCallback);
+	ColorTheme.Get()->OnRefresh().Remove(RefreshCallbackHandle);
 }
 
 void SThemeColorBlocksBar::AddRefreshCallback()
 {
-	ColorTheme.Get()->OnRefresh().Add(RefreshCallback);
+	RefreshCallbackHandle = ColorTheme.Get()->OnRefresh().Add(RefreshCallback);
 }
 
 void SThemeColorBlocksBar::Refresh()
diff --git a/Engine/Source/Runtime/AssetRegistry/Private/AssetRegistry.cpp b/Engine/Source/Runtime/AssetRegistry/Private/AssetRegistry.cpp
--- a/Engine/Source/Runtime/AssetRegistry/Private/AssetRegistry.cpp
+++ b/Engine/Source/Runtime/AssetRegistry/Private/AssetRegistry.cpp
@@ -84,7 +84,9 @@ FAssetRegistry::FAssetRegistry()
 			{
 				const FString& RootPath = *RootPathIt;
 				const FString& ContentFolder = FPackageName::LongPackageNameToFilename( RootPath );
-				DirectoryWatcher->RegisterDirectoryChangedCallback( ContentFolder, IDirectoryWatcher::FDirectoryChanged::CreateRaw(this, &FAssetRegistry::OnDirectoryChanged));
+				FDelegateHandle NewHandle;
+				DirectoryWatcher->RegisterDirectoryChangedCallback_Handle( ContentFolder, IDirectoryWatcher::FDirectoryChanged::CreateRaw(this, &FAssetRegistry::OnDirectoryChanged), NewHandle);
+				OnDirectoryChangedDelegateHandles.Add(ContentFolder, NewHandle);
 			}
 		}
 	}
@@ -190,7 +192,8 @@ FAssetRegistry::~FAssetRegistry()
 				{
 					const FString& RootPath = *RootPathIt;
 					const FString& ContentFolder = FPackageName::LongPackageNameToFilename( RootPath );
-					DirectoryWatcher->UnregisterDirectoryChangedCallback( ContentFolder, IDirectoryWatcher::FDirectoryChanged::CreateRaw(this, &FAssetRegistry::OnDirectoryChanged));
+					DirectoryWatcher->UnregisterDirectoryChangedCallback_Handle( ContentFolder, OnDirectoryChangedDelegateHandles.FindRef(ContentFolder));
+					OnDirectoryChangedDelegateHandles.Remove(ContentFolder);
 				}
 			}
 		}
@@ -1922,7 +1925,7 @@ void FAssetRegistry::OnContentPathMounted( const FString& InAssetPath, const FSt
 		{
 			// If the path doesn't exist on disk, make it so the watcher will work.
 			IFileManager::Get().MakeDirectory(*FileSystemPath);
-			DirectoryWatcher->RegisterDirectoryChangedCallback( FileSystemPath, IDirectoryWatcher::FDirectoryChanged::CreateRaw(this, &FAssetRegistry::OnDirectoryChanged));
+			DirectoryWatcher->RegisterDirectoryChangedCallback_Handle( FileSystemPath, IDirectoryWatcher::FDirectoryChanged::CreateRaw(this, &FAssetRegistry::OnDirectoryChanged), OnContentPathMountedOnDirectoryChangedDelegateHandle);
 		}
 	}
 #endif // WITH_EDITOR
@@ -1976,7 +1979,7 @@ void FAssetRegistry::OnContentPathDismounted(const FString& InAssetPath, const F
 		IDirectoryWatcher* DirectoryWatcher = DirectoryWatcherModule.Get();
 		if (DirectoryWatcher)
 		{
-			DirectoryWatcher->UnregisterDirectoryChangedCallback(FileSystemPath, IDirectoryWatcher::FDirectoryChanged::CreateRaw(this, &FAssetRegistry::OnDirectoryChanged));
+			DirectoryWatcher->UnregisterDirectoryChangedCallback_Handle(FileSystemPath, OnContentPathMountedOnDirectoryChangedDelegateHandle);
 		}
 	}
 #endif // WITH_EDITOR
diff --git a/Engine/Source/Runtime/AssetRegistry/Private/AssetRegistry.h b/Engine/Source/Runtime/AssetRegistry/Private/AssetRegistry.h
--- a/Engine/Source/Runtime/AssetRegistry/Private/AssetRegistry.h
+++ b/Engine/Source/Runtime/AssetRegistry/Private/AssetRegistry.h
@@ -254,4 +254,10 @@ private:
 
 	/** List of all class names derived from Blueprint (including Blueprint itself) */
 	TSet<FName> ClassGeneratorNames;
+
+	/** Handles to all registered OnDirectoryChanged delegates */
+	TMap<FString, FDelegateHandle> OnDirectoryChangedDelegateHandles;
+
+	/** Handle to the registered OnDirectoryChanged delegate for the OnContentPathMounted handler */
+	FDelegateHandle OnContentPathMountedOnDirectoryChangedDelegateHandle;
 };
diff --git a/Engine/Source/Runtime/Core/Private/Delegates/DelegateHandle.cpp b/Engine/Source/Runtime/Core/Private/Delegates/DelegateHandle.cpp
new file mode 100644
--- /dev/null
+++ b/Engine/Source/Runtime/Core/Private/Delegates/DelegateHandle.cpp
@@ -0,0 +1,21 @@
+// Copyright 1998-2014 Epic Games, Inc. All Rights Reserved.
+
+#include "CorePrivatePCH.h"
+
+namespace UE4Delegates_Private
+{
+	uint64 GNextID = 1;
+}
+
+uint64 FDelegateHandle::GenerateNewID()
+{
+	// Just increment a counter to generate an ID.
+	// On the next-to-impossible event that we wrap round to 0, reset back to 1, because we reserve 0 for null delegates.
+	uint64 Result = UE4Delegates_Private::GNextID++;
+	if (UE4Delegates_Private::GNextID == 0)
+	{
+		UE4Delegates_Private::GNextID = 1;
+	}
+
+	return Result;
+}
diff --git a/Engine/Source/Runtime/Core/Private/GenericPlatform/GenericPlatformTime.cpp b/Engine/Source/Runtime/Core/Private/GenericPlatform/GenericPlatformTime.cpp
--- a/Engine/Source/Runtime/Core/Private/GenericPlatform/GenericPlatformTime.cpp
+++ b/Engine/Source/Runtime/Core/Private/GenericPlatform/GenericPlatformTime.cpp
@@ -171,20 +171,22 @@ struct FCPUTimeDump
 			FParse::Value(*Args[0], TEXT("delay="), Delay);
 		}
 
-		FTicker::GetCoreTicker().RemoveTicker( GetCPUTimeDelegate );
-		FTicker::GetCoreTicker().RemoveTicker( CPUTimeDumpDelegate );
+		FTicker::GetCoreTicker().RemoveTicker( GetCPUTimeDelegateHandle );
+		FTicker::GetCoreTicker().RemoveTicker( CPUTimeDumpDelegateHandle );
+		GetCPUTimeDelegateHandle   = FDelegateHandle();
+		CPUTimeDumpDelegateHandle  = FDelegateHandle();
 
 		if( Delay == 0 )
 		{
 			UE_LOG(LogGenericPlatformTime, Log, TEXT("Disabling printing the CPU usage"));
-		}	
+		}
 		else
 		{
 			Delay = FMath::Clamp( Delay, 10, 300 );
 			UE_LOG(LogGenericPlatformTime, Log, TEXT("Delay set to %i second(s), started printing the CPU usage"), Delay);
 
-			FTicker::GetCoreTicker().AddTicker( GetCPUTimeDelegate );
-			FTicker::GetCoreTicker().AddTicker( CPUTimeDumpDelegate, Delay );
+			GetCPUTimeDelegateHandle   = FTicker::GetCoreTicker().AddTicker( GetCPUTimeDelegate );
+			CPUTimeDumpDelegateHandle  = FTicker::GetCoreTicker().AddTicker( CPUTimeDumpDelegate, Delay );
 		}
 	}
 
@@ -221,6 +223,12 @@ protected:
 
 	/** Ticker delegate for printing the average CPU usage, called every nth frame. */
 	FTickerDelegate CPUTimeDumpDelegate;
+
+	/** Handle for the added GetCPUTimeDelegate delegate. */
+	FDelegateHandle GetCPUTimeDelegateHandle;
+
+	/** Handle for the added CPUTimeDumpDelegate delegate. */
+	FDelegateHandle CPUTimeDumpDelegateHandle;
 };
 
 static FAutoConsoleCommand CPUTimeDumpCommand
diff --git a/Engine/Source/Runtime/Core/Private/HAL/ConsoleManager.cpp b/Engine/Source/Runtime/Core/Private/HAL/ConsoleManager.cpp
--- a/Engine/Source/Runtime/Core/Private/HAL/ConsoleManager.cpp
+++ b/Engine/Source/Runtime/Core/Private/HAL/ConsoleManager.cpp
@@ -543,7 +543,18 @@ void FConsoleManager::RegisterConsoleVariableSink(const FConsoleCommandDelegate&
 
 void FConsoleManager::UnregisterConsoleVariableSink(const FConsoleCommandDelegate& Command)
 {
-	ConsoleVariableChangeSinks.Remove(Command);
+	ConsoleVariableChangeSinks.RemoveAll([&](const FConsoleCommandDelegate& Element){ return Command.DEPRECATED_Compare(Element); });
+}
+
+FConsoleVariableSinkHandle FConsoleManager::RegisterConsoleVariableSink_Handle(const FConsoleCommandDelegate& Command)
+{
+	ConsoleVariableChangeSinks.Add(Command);
+	return FConsoleVariableSinkHandle(Command.GetHandle());
+}
+
+void FConsoleManager::UnregisterConsoleVariableSink_Handle(FConsoleVariableSinkHandle Handle)
+{
+	ConsoleVariableChangeSinks.RemoveAll([=](const FConsoleCommandDelegate& Delegate){ return Handle.HasSameHandle(Delegate); });
 }
 
 class FConsoleCommand : public FConsoleCommandBase
@@ -1305,7 +1316,7 @@ void FConsoleManager::Test()
 
 	// setup ---------
 
-	RegisterConsoleVariableSink(FConsoleCommandDelegate::CreateStatic(&TestSinkCallback));
+	auto TestSinkCallbackHandle = RegisterConsoleVariableSink_Handle(FConsoleCommandDelegate::CreateStatic(&TestSinkCallback));
 
 	// start tests ---------
 
@@ -1461,7 +1472,7 @@ void FConsoleManager::Test()
 	TestSinkCallback();
 	check(GConsoleManagerSinkTestCounter == 2);
 
-	UnregisterConsoleVariableSink(FConsoleCommandDelegate::CreateStatic(&TestSinkCallback));
+	UnregisterConsoleVariableSink_Handle(TestSinkCallbackHandle);
 
 #endif // !(UE_BUILD_SHIPPING || UE_BUILD_TEST)
 }
diff --git a/Engine/Source/Runtime/Core/Private/Stats/StatsCommand.cpp b/Engine/Source/Runtime/Core/Private/Stats/StatsCommand.cpp
--- a/Engine/Source/Runtime/Core/Private/Stats/StatsCommand.cpp
+++ b/Engine/Source/Runtime/Core/Private/Stats/StatsCommand.cpp
@@ -516,6 +516,9 @@ struct FHUDGroupManager
 	/** Whether it's enabled or not. */
 	bool bEnabled;
 
+	/** NewFrame delegate handle */
+	FDelegateHandle NewFrameDelegateHandle;
+
 	/** Default constructor. */
 	FHUDGroupManager(FStatsThreadState const& InStats)
 		: NumTotalStackFrames(0)
@@ -590,12 +593,12 @@ struct FHUDGroupManager
 		if( EnabledGroups.Num() && !bEnabled )
 		{
 			bEnabled = true;
-			Stats.NewFrameDelegate.AddRaw( this, &FHUDGroupManager::NewFrame );
+			NewFrameDelegateHandle = Stats.NewFrameDelegate.AddRaw( this, &FHUDGroupManager::NewFrame );
 			StatsMasterEnableAdd();
 		}
 		else if( !EnabledGroups.Num() && bEnabled )
 		{
-			Stats.NewFrameDelegate.RemoveRaw( this, &FHUDGroupManager::NewFrame );
+			Stats.NewFrameDelegate.Remove( NewFrameDelegateHandle );
 			StatsMasterEnableSubtract();
 			bEnabled = false;
 
@@ -938,6 +941,9 @@ struct FHUDGroupManager
 static float DumpCull = 5.0f;
 static int32 MaxDepth = MAX_int32;
 static FString NameFilter;
+static FDelegateHandle DumpFrameDelegateHandle;
+static FDelegateHandle DumpCPUDelegateHandle;
+static FDelegateHandle DumpMemoryDelegateHandle;
 
 static void DumpFrame(int64 Frame)
 {
@@ -945,7 +951,7 @@ static void DumpFrame(int64 Frame)
 	int64 Latest = Stats.GetLatestValidFrame();
 	check(Latest > 0);
 	DumpHistoryFrame(Stats, Latest, DumpCull, MaxDepth, *NameFilter);
-	Stats.NewFrameDelegate.RemoveStatic(&DumpFrame);
+	Stats.NewFrameDelegate.Remove(DumpFrameDelegateHandle);
 	StatsMasterEnableSubtract();
 }
 
@@ -955,7 +961,7 @@ static void DumpCPU(int64 Frame)
 	int64 Latest = Stats.GetLatestValidFrame();
 	check(Latest > 0);
 	DumpCPUSummary(Stats, Latest);
-	Stats.NewFrameDelegate.RemoveStatic(&DumpCPU);
+	Stats.NewFrameDelegate.Remove(DumpCPUDelegateHandle);
 	StatsMasterEnableSubtract();
 }
 
@@ -1046,7 +1052,7 @@ static void DumpMemory(int64 /*Frame*/)
 		UE_LOG(LogStats, Warning, TEXT("Allocated memory: %llu bytes (%.2f MB)"), TotalAllocatedMemory, TotalAllocatedMemory/1024.0f/1024.0f );
 	}
 
-	Stats.NewFrameDelegate.RemoveStatic(&DumpMemory);
+	Stats.NewFrameDelegate.Remove(DumpMemoryDelegateHandle);
 }
 
 static struct FDumpMultiple* DumpMultiple = NULL;
@@ -1059,6 +1065,7 @@ struct FDumpMultiple
 	int32 NumFrames;
 	int32 NumFramesToGo;
 	FRawStatStackNode* Stack;
+	FDelegateHandle NewFrameDelegateHandle;
 
 	FDumpMultiple()
 		: Stats(FStatsThreadState::GetLocalState())
@@ -1069,7 +1076,7 @@ struct FDumpMultiple
 		, Stack(NULL)
 	{
 		StatsMasterEnableAdd();
-		Stats.NewFrameDelegate.AddRaw(this, &FDumpMultiple::NewFrame);
+		NewFrameDelegateHandle = Stats.NewFrameDelegate.AddRaw(this, &FDumpMultiple::NewFrame);
 	}
 
 	~FDumpMultiple()
@@ -1102,7 +1109,7 @@ struct FDumpMultiple
 			delete Stack;
 			Stack = NULL;
 		}
-		Stats.NewFrameDelegate.RemoveRaw(this, &FDumpMultiple::NewFrame);
+		Stats.NewFrameDelegate.Remove(NewFrameDelegateHandle);
 		StatsMasterEnableSubtract();
 		DumpMultiple = NULL;
 	}
@@ -1226,11 +1233,11 @@ static void StatCmd(FString InCmd)
 	if( FParse::Command(&Cmd,TEXT("DUMPFRAME")) )
 	{
 		StatsMasterEnableAdd();
-		Stats.NewFrameDelegate.AddStatic(&DumpFrame);
+		DumpFrameDelegateHandle = Stats.NewFrameDelegate.AddStatic(&DumpFrame);
 	} 
 	else if( FParse::Command(&Cmd,TEXT("DumpMemory")) )
 	{
-		Stats.NewFrameDelegate.AddStatic(&DumpMemory);
+		DumpMemoryDelegateHandle = Stats.NewFrameDelegate.AddStatic(&DumpMemory);
 	}
 	else if ( FParse::Command(&Cmd,TEXT("DUMPNONFRAME")) )
 	{
@@ -1239,7 +1246,7 @@ static void StatCmd(FString InCmd)
 	else if ( FParse::Command(&Cmd,TEXT("DUMPCPU")) )
 	{
 		StatsMasterEnableAdd();
-		Stats.NewFrameDelegate.AddStatic(&DumpCPU);
+		DumpCPUDelegateHandle = Stats.NewFrameDelegate.AddStatic(&DumpCPU);
 	}
 	else if( FParse::Command(&Cmd,TEXT("STOP")) )
 	{
@@ -1290,16 +1297,17 @@ static void StatCmd(FString InCmd)
 	else if( FParse::Command(&Cmd,TEXT("DUMPHITCHES")) )
 	{
 		static bool bToggle = false;
+		static FDelegateHandle DumpHitchDelegateHandle;
 		bToggle = !bToggle;
 		if (bToggle)
 		{
 			StatsMasterEnableAdd();
-			Stats.NewFrameDelegate.AddStatic(&DumpHitch);
+			DumpHitchDelegateHandle = Stats.NewFrameDelegate.AddStatic(&DumpHitch);
 		}
 		else
 		{
 			StatsMasterEnableSubtract();
-			Stats.NewFrameDelegate.RemoveStatic(&DumpHitch);
+			Stats.NewFrameDelegate.Remove(DumpHitchDelegateHandle);
 		}
 	}
 	else if( FParse::Command( &Cmd, TEXT( "STARTFILE" ) ) )
diff --git a/Engine/Source/Runtime/Core/Private/Stats/StatsFile.cpp b/Engine/Source/Runtime/Core/Private/Stats/StatsFile.cpp
--- a/Engine/Source/Runtime/Core/Private/Stats/StatsFile.cpp
+++ b/Engine/Source/Runtime/Core/Private/Stats/StatsFile.cpp
@@ -488,13 +488,13 @@ bool FStatsWriteFile::IsValid() const
 void FStatsWriteFile::AddNewFrameDelegate()
 {
 	FStatsThreadState const& Stats = FStatsThreadState::GetLocalState();
-	Stats.NewFrameDelegate.AddThreadSafeSP( this->AsShared(), &FStatsWriteFile::NewFrame );
+	NewFrameDelegateHandle = Stats.NewFrameDelegate.AddThreadSafeSP( this->AsShared(), &FStatsWriteFile::NewFrame );
 }
 
 void FStatsWriteFile::RemoveNewFrameDelegate()
 {
 	FStatsThreadState const& Stats = FStatsThreadState::GetLocalState();
-	Stats.NewFrameDelegate.RemoveThreadSafeSP( this->AsShared(), &FStatsWriteFile::NewFrame );
+	Stats.NewFrameDelegate.Remove( NewFrameDelegateHandle );
 }
 
 /*-----------------------------------------------------------------------------
