Delegate comparisons deprecated, lots of other associated code deprecated, and lots of warning fixups:

diff --git a/Engine/Source/Runtime/Core/Public/Containers/Array.h b/Engine/Source/Runtime/Core/Public/Containers/Array.h
--- a/Engine/Source/Runtime/Core/Public/Containers/Array.h
+++ b/Engine/Source/Runtime/Core/Public/Containers/Array.h
@@ -2267,18 +2267,22 @@ public:
 	 *
 	 * @param InIntem Item to be added.
 	 * @param Predicate Predicate class instance.
+	 *
+	 * @return The index of the new element.
 	 */
 	template <class PREDICATE_CLASS>
-	void HeapPush(const ElementType& InItem, const PREDICATE_CLASS& Predicate)
+	int32 HeapPush(const ElementType& InItem, const PREDICATE_CLASS& Predicate)
 	{
 		// Add at the end, then sift up
 		Add(InItem);
 		TDereferenceWrapper<ElementType, PREDICATE_CLASS> PredicateWrapper(Predicate);
-		SiftUp(0, Num() - 1, PredicateWrapper);
+		int32 Result = SiftUp(0, Num() - 1, PredicateWrapper);
 
 #if DEBUG_HEAP
 		VerifyHeap(PredicateWrapper);
 #endif
+
+		return Result;
 	}
 
 	/** 
@@ -2286,10 +2290,12 @@ public:
 	 * template type.
 	 *
 	 * @param InIntem Item to be added.
+	 *
+	 * @return The index of the new element.
 	 */
-	void HeapPush(const ElementType& InItem)
+	int32 HeapPush(const ElementType& InItem)
 	{
-		HeapPush(InItem, TLess<ElementType>());
+		return HeapPush(InItem, TLess<ElementType>());
 	}
 
 	/** 
@@ -2404,7 +2410,7 @@ public:
 
 		TDereferenceWrapper< ElementType, PREDICATE_CLASS> PredicateWrapper(Predicate);
 		SiftDown(Index, Num(), PredicateWrapper);
-		SiftUp(0, Index, PredicateWrapper);
+		SiftUp(0, FMath::Min(Index, Num() - 1), PredicateWrapper);
 
 #if DEBUG_HEAP
 		VerifyHeap(PredicateWrapper);
@@ -2504,6 +2510,7 @@ private:
 		return (Index - 1) / 2;
 	}
 
+private:
 	/**
 	 * Fixes a possible violation of order property between node at Index and a child.
 	 *
@@ -2526,15 +2533,13 @@ private:
 				MinChildIndex = Predicate(Heap[LeftChildIndex], Heap[RightChildIndex]) ? LeftChildIndex : RightChildIndex;
 			}
 
-			if (Predicate(Heap[MinChildIndex], Heap[Index]))
-			{
-				Exchange(Heap[Index], Heap[MinChildIndex]);
-				Index = MinChildIndex;
-			}
-			else
+			if (!Predicate(Heap[MinChildIndex], Heap[Index]))
 			{
 				break;
 			}
+
+			Exchange(Heap[Index], Heap[MinChildIndex]);
+			Index = MinChildIndex;
 		}
 	}
 
@@ -2544,25 +2549,26 @@ private:
 	 * @param RootIndex How far to go up?
 	 * @param NodeIndex Node index.
 	 * @param Predicate Predicate class instance.
+	 *
+	 * @return The new index of the node that was at NodeIndex
 	 */
 	template <class PREDICATE_CLASS>
-	FORCEINLINE void SiftUp(int32 RootIndex, const int32 NodeIndex, const PREDICATE_CLASS& Predicate)
+	FORCEINLINE int32 SiftUp(int32 RootIndex, int32 NodeIndex, const PREDICATE_CLASS& Predicate)
 	{
 		ElementType* Heap = GetData();
-		int32 Index = FMath::Min(NodeIndex, Num() - 1);
-		while (Index > RootIndex)
+		while (NodeIndex > RootIndex)
 		{
-			int32 ParentIndex = HeapGetParentIndex(Index);
-			if (Predicate(Heap[Index], Heap[ParentIndex]))
-			{
-				Exchange(Heap[Index], Heap[ParentIndex]);
-				Index = ParentIndex;
-			}
-			else
+			int32 ParentIndex = HeapGetParentIndex(NodeIndex);
+			if (!Predicate(Heap[NodeIndex], Heap[ParentIndex]))
 			{
 				break;
 			}
+
+			Exchange(Heap[NodeIndex], Heap[ParentIndex]);
+			NodeIndex = ParentIndex;
 		}
+
+		return NodeIndex;
 	}
 };
 
diff --git a/Engine/Source/Runtime/Core/Public/Containers/Ticker.h b/Engine/Source/Runtime/Core/Public/Containers/Ticker.h
--- a/Engine/Source/Runtime/Core/Public/Containers/Ticker.h
+++ b/Engine/Source/Runtime/Core/Public/Containers/Ticker.h
@@ -32,10 +32,21 @@ public:
 	 * @param InDelegate	delegate to fire after the delay
 	 * @param InDelay		Delay until next fire; 0 means "next frame"
 	 */
-	void AddTicker(const FTickerDelegate& InDelegate, float InDelay = 0.0f)
+	FDelegateHandle AddTicker(const FTickerDelegate& InDelegate, float InDelay = 0.0f)
 	{
 		FElement Ticker(CurrentTime + InDelay, InDelay, InDelegate);
 		PriorityQueue.HeapPush(Ticker);
+		return InDelegate.GetHandle();
+	}
+
+	/**
+	 * Removes a previously added ticker delegate.
+	 *
+	 * @param Handle - The handle of the ticker to remove.
+	 */
+	void RemoveTicker(FDelegateHandle Handle)
+	{
+		PriorityQueue.RemoveAll([=](const FElement& Element){ return Element.Delegate.GetHandle() == Handle; });
 	}
 
 	/**
@@ -43,15 +54,16 @@ public:
 	 *
 	 * @param Delegate - The delegate to remove.
 	 */
+	DELEGATE_DEPRECATED("This RemoveTicker overload is deprecated - please remove tickers using the FDelegateHandle returned by the AddTicker function.")
 	void RemoveTicker(const FTickerDelegate& Delegate)
 	{
 		for (int32 Index = 0; Index < PriorityQueue.Num(); ++Index)
 		{
-			if (PriorityQueue[Index].Delegate == Delegate)
+			if (PriorityQueue[Index].Delegate.DEPRECATED_Compare(Delegate))
 			{
 				PriorityQueue.RemoveAt(Index);
 			}
-		}		
+		}
 	}
 
 	/**
@@ -140,8 +152,8 @@ public:
 	FTickerObjectBase(float InDelay = 0.0f)
 	{
 		// Register delegate for ticker callback
-		TickDelegate = FTickerDelegate::CreateRaw(this, &FTickerObjectBase::Tick);
-		FTicker::GetCoreTicker().AddTicker(TickDelegate, InDelay);
+		FTickerDelegate TickDelegate = FTickerDelegate::CreateRaw(this, &FTickerObjectBase::Tick);
+		TickHandle = FTicker::GetCoreTicker().AddTicker(TickDelegate, InDelay);
 	}
 
 	/**
@@ -150,7 +162,7 @@ public:
 	virtual ~FTickerObjectBase()
 	{
 		// Unregister ticker delegate
-		FTicker::GetCoreTicker().RemoveTicker(TickDelegate);
+		FTicker::GetCoreTicker().RemoveTicker(TickHandle);
 	}
 
 	/**
@@ -164,5 +176,5 @@ public:
 private:	
 
 	/** Delegate for callbacks to Tick */
-	FTickerDelegate TickDelegate;
+	FDelegateHandle TickHandle;
 };
diff --git a/Engine/Source/Runtime/Core/Public/Delegates/Delegate.h b/Engine/Source/Runtime/Core/Public/Delegates/Delegate.h
--- a/Engine/Source/Runtime/Core/Public/Delegates/Delegate.h
+++ b/Engine/Source/Runtime/Core/Public/Delegates/Delegate.h
@@ -297,6 +297,8 @@ class DynamicMulticastDelegateName : public FUNC_DECLARE_DYNAMIC_DELEGATE_BASE(T
 
 /*********************************************************************************************************************/
 
+#define DELEGATE_DEPRECATED(message) DEPRECATED(4.7, message)
+
 // We define this as a guard to prevent DelegateSignatureImpl.inl being included outside of this file
 #define __Delegate_h__
 #define FUNC_INCLUDING_INLINE_IMPL
diff --git a/Engine/Source/Runtime/Core/Public/Delegates/DelegateBase.h b/Engine/Source/Runtime/Core/Public/Delegates/DelegateBase.h
--- a/Engine/Source/Runtime/Core/Public/Delegates/DelegateBase.h
+++ b/Engine/Source/Runtime/Core/Public/Delegates/DelegateBase.h
@@ -94,6 +94,22 @@ public:
 		DelegateInstance = InDelegateInstance;
 	}
 
+	/**
+	 * Gets a handle to the delegate.
+	 *
+	 * @return The delegate instance.
+	 */
+	inline FDelegateHandle GetHandle() const
+	{
+		FDelegateHandle Result;
+		if (DelegateInstance)
+		{
+			Result = DelegateInstance->GetHandle();
+		}
+
+		return Result;
+	}
+
 private:
 
 	// Holds the delegate instance
diff --git a/Engine/Source/Runtime/Core/Public/Delegates/DelegateInstancesImpl.inl b/Engine/Source/Runtime/Core/Public/Delegates/DelegateInstancesImpl.inl
--- a/Engine/Source/Runtime/Core/Public/Delegates/DelegateInstancesImpl.inl
+++ b/Engine/Source/Runtime/Core/Public/Delegates/DelegateInstancesImpl.inl
@@ -94,6 +94,7 @@ public:
 		: UserObject(InUserObject)
 		, MethodPtr(InMethodPtr)
 		DELEGATE_COMMA_PAYLOAD_INITIALIZER_LIST
+		, Handle(FDelegateHandle::GenerateNewHandle)
 	{
 		// NOTE: Shared pointer delegates are allowed to have a null incoming object pointer.  Weak pointers can expire,
 		//       an it is possible for a copy of a delegate instance to end up with a null pointer.
@@ -140,7 +141,7 @@ public:
 
 	virtual DELEGATE_INSTANCE_INTERFACE_CLASS<FUNC_TEMPLATE_ARGS>* CreateCopy() override
 	{
-		return Create(UserObject.Pin(), MethodPtr DELEGATE_COMMA_PAYLOAD_PASSIN);
+		return new SP_METHOD_DELEGATE_INSTANCE_CLASS(*this);
 	}
 
 	virtual RetValType Execute( FUNC_PARAM_LIST ) const override
@@ -162,6 +163,11 @@ public:
 		return (MutableUserObject->*MethodPtr)(DELEGATE_PARAM_PASSTHRU_COMMA_PAYLOAD_PASSIN);
 	}
 
+	virtual FDelegateHandle GetHandle() const override
+	{
+		return Handle;
+	}
+
 #if FUNC_IS_VOID
 	virtual bool ExecuteIfSafe( FUNC_PARAM_LIST ) const override
 	{
@@ -257,6 +263,9 @@ private:
 
 	// Payload member variables, if any.
 	FUNC_PAYLOAD_MEMBERS
+
+	// The handle of this delegate
+	FDelegateHandle Handle;
 };
 
 
@@ -284,6 +293,7 @@ public:
 		: UserObject(InUserObject)
 		, MethodPtr(InMethodPtr)
 		DELEGATE_COMMA_PAYLOAD_INITIALIZER_LIST
+		, Handle(FDelegateHandle::GenerateNewHandle)
 	{
 		// Non-expirable delegates must always have a non-null object pointer on creation (otherwise they could never execute.)
 		check(InUserObject != nullptr && MethodPtr != nullptr);
@@ -331,7 +341,7 @@ public:
 
 	virtual DELEGATE_INSTANCE_INTERFACE_CLASS<FUNC_TEMPLATE_ARGS>* CreateCopy( ) override
 	{
-		return Create(UserObject, MethodPtr DELEGATE_COMMA_PAYLOAD_PASSIN);
+		return new RAW_METHOD_DELEGATE_INSTANCE_CLASS(*this);
 	}
 
 	virtual RetValType Execute( FUNC_PARAM_LIST ) const override
@@ -349,6 +359,11 @@ public:
 		return (MutableUserObject->*MethodPtr)(DELEGATE_PARAM_PASSTHRU_COMMA_PAYLOAD_PASSIN);
 	}
 
+	virtual FDelegateHandle GetHandle() const override
+	{
+		return Handle;
+	}
+
 #if FUNC_IS_VOID
 	virtual bool ExecuteIfSafe( FUNC_PARAM_LIST ) const override
 	{
@@ -421,6 +436,9 @@ private:
 
 	// Payload member variables (if any).
 	FUNC_PAYLOAD_MEMBERS
+
+	// The handle of this delegate
+	FDelegateHandle Handle;
 };
 
 
@@ -448,6 +466,7 @@ public:
 		: UserObject(InUserObject)
 		, MethodPtr(InMethodPtr)
 		DELEGATE_COMMA_PAYLOAD_INITIALIZER_LIST
+		, Handle(FDelegateHandle::GenerateNewHandle)
 	{
 		// NOTE: UObject delegates are allowed to have a null incoming object pointer.  UObject weak pointers can expire,
 		//       an it is possible for a copy of a delegate instance to end up with a null pointer.
@@ -499,7 +518,7 @@ public:
 
 	virtual DELEGATE_INSTANCE_INTERFACE_CLASS<FUNC_TEMPLATE_ARGS>* CreateCopy( ) override
 	{
-		return Create(UserObject.Get(), MethodPtr DELEGATE_COMMA_PAYLOAD_PASSIN);
+		return new UOBJECT_METHOD_DELEGATE_INSTANCE_CLASS(*this);
 	}
 
 	virtual RetValType Execute( FUNC_PARAM_LIST ) const override
@@ -520,6 +539,11 @@ public:
 		return (MutableUserObject->*MethodPtr)(DELEGATE_PARAM_PASSTHRU_COMMA_PAYLOAD_PASSIN);
 	}
 
+	virtual FDelegateHandle GetHandle() const override
+	{
+		return Handle;
+	}
+
 #if FUNC_IS_VOID
 	virtual bool ExecuteIfSafe( FUNC_PARAM_LIST ) const override
 	{
@@ -594,6 +618,9 @@ private:
 
 	// Payload member variables (if any).
 	FUNC_PAYLOAD_MEMBERS
+
+	// The handle of this delegate
+	FDelegateHandle Handle;
 };
 
 #if !FUNC_IS_CONST
@@ -623,6 +650,7 @@ public:
 	STATIC_DELEGATE_INSTANCE_CLASS( FFuncPtr InStaticFuncPtr DELEGATE_COMMA_PAYLOAD_LIST )
 		: StaticFuncPtr(InStaticFuncPtr)
 		DELEGATE_COMMA_PAYLOAD_INITIALIZER_LIST
+		, Handle(FDelegateHandle::GenerateNewHandle)
 	{
 		check(StaticFuncPtr != nullptr);
 	}
@@ -669,7 +697,7 @@ public:
 
 	virtual DELEGATE_INSTANCE_INTERFACE_CLASS<FUNC_TEMPLATE_ARGS>* CreateCopy( ) override
 	{
-		return Create(StaticFuncPtr DELEGATE_COMMA_PAYLOAD_PASSIN);
+		return new STATIC_DELEGATE_INSTANCE_CLASS(*this);
 	}
 
 	virtual RetValType Execute( FUNC_PARAM_LIST ) const override
@@ -680,6 +708,11 @@ public:
 		return StaticFuncPtr(DELEGATE_PARAM_PASSTHRU_COMMA_PAYLOAD_PASSIN);
 	}
 
+	virtual FDelegateHandle GetHandle() const override
+	{
+		return Handle;
+	}
+
 #if FUNC_IS_VOID
 	virtual bool ExecuteIfSafe( FUNC_PARAM_LIST ) const override
 	{
@@ -725,6 +758,9 @@ private:
 
 	// Payload member variables, if any.
 	FUNC_PAYLOAD_MEMBERS
+
+	// The handle of this delegate
+	FDelegateHandle Handle;
 };
 
 
@@ -746,12 +782,14 @@ public:
 	FUNCTOR_DELEGATE_INSTANCE_CLASS(const FunctorType& InFunctor DELEGATE_COMMA_PAYLOAD_LIST)
 		: Functor(InFunctor)
 		DELEGATE_COMMA_PAYLOAD_INITIALIZER_LIST
+		, Handle(FDelegateHandle::GenerateNewHandle)
 	{
 	}
 
 	FUNCTOR_DELEGATE_INSTANCE_CLASS(FunctorType&& InFunctor DELEGATE_COMMA_PAYLOAD_LIST)
-	: Functor(MoveTemp(InFunctor))
-	DELEGATE_COMMA_PAYLOAD_INITIALIZER_LIST
+		: Functor(MoveTemp(InFunctor))
+		DELEGATE_COMMA_PAYLOAD_INITIALIZER_LIST
+		, Handle(FDelegateHandle::GenerateNewHandle)
 	{
 	}
 
@@ -802,7 +840,7 @@ public:
 
 	virtual DELEGATE_INSTANCE_INTERFACE_CLASS<FUNC_TEMPLATE_ARGS>* CreateCopy() override
 	{
-		return Create(Functor DELEGATE_COMMA_PAYLOAD_PASSIN);
+		return new FUNCTOR_DELEGATE_INSTANCE_CLASS(*this);
 	}
 
 	virtual RetValType Execute(FUNC_PARAM_LIST) const override
@@ -810,6 +848,11 @@ public:
 		return Functor(DELEGATE_PARAM_PASSTHRU_COMMA_PAYLOAD_PASSIN);
 	}
 
+	virtual FDelegateHandle GetHandle() const override
+	{
+		return Handle;
+	}
+
 #if FUNC_IS_VOID
 	virtual bool ExecuteIfSafe(FUNC_PARAM_LIST) const override
 	{
@@ -847,6 +890,9 @@ private:
 
 	// Payload member variables, if any.
 	FUNC_PAYLOAD_MEMBERS
+
+	// The handle of this delegate
+	FDelegateHandle Handle;
 };
 
 /**
@@ -883,6 +929,7 @@ public:
 		: FunctionName(InFunctionName)
 		, UserObjectPtr(InUserObject)
 		DELEGATE_COMMA_PAYLOAD_INITIALIZER_LIST
+		, Handle(FDelegateHandle::GenerateNewHandle)
 	{
 		check(InFunctionName != NAME_None);
 		
@@ -937,7 +984,7 @@ public:
 
 	virtual DELEGATE_INSTANCE_INTERFACE_CLASS* CreateCopy( ) override
 	{
-		return Create(UserObjectPtr.Get(), GetFunctionName() DELEGATE_COMMA_PAYLOAD_PASSIN);
+		return new UFUNCTION_DELEGATE_INSTANCE_CLASS(*this);
 	}
 
 	virtual RetValType Execute( FUNC_PARAM_LIST ) const override
@@ -952,6 +999,11 @@ public:
 #endif
 	}
 
+	virtual FDelegateHandle GetHandle() const override
+	{
+		return Handle;
+	}
+
 #if FUNC_IS_VOID
 	virtual bool ExecuteIfSafe( FUNC_PARAM_LIST ) const override
 	{
@@ -1004,6 +1056,9 @@ public:
 
 	// Payload member variables, if any.
 	FUNC_PAYLOAD_MEMBERS
+
+	// The handle of this delegate
+	FDelegateHandle Handle;
 };
 
 #endif // FUNC_IS_CONST
diff --git a/Engine/Source/Runtime/Core/Public/Delegates/DelegateInstancesImpl_Variadics.inl b/Engine/Source/Runtime/Core/Public/Delegates/DelegateInstancesImpl_Variadics.inl
--- a/Engine/Source/Runtime/Core/Public/Delegates/DelegateInstancesImpl_Variadics.inl
+++ b/Engine/Source/Runtime/Core/Public/Delegates/DelegateInstancesImpl_Variadics.inl
@@ -30,17 +30,21 @@ class TBaseUFunctionDelegateInstance;
 template <class UserClass, typename WrappedRetValType, typename... ParamTypes, typename... VarTypes>
 class TBaseUFunctionDelegateInstance<UserClass, WrappedRetValType (ParamTypes...), VarTypes...> : public IBaseDelegateInstance<typename TUnwrapType<WrappedRetValType>::Type (ParamTypes...)>
 {
-	typedef IBaseDelegateInstance<typename TUnwrapType<WrappedRetValType>::Type (ParamTypes...)> Super;
+public:
+	typedef typename TUnwrapType<WrappedRetValType>::Type RetValType;
+
+private:
+	typedef IBaseDelegateInstance<RetValType (ParamTypes...)>                                  Super;
+	typedef TBaseUFunctionDelegateInstance<UserClass, RetValType (ParamTypes...), VarTypes...> UnwrappedThisType;
 
 	static_assert((CanConvertPointerFromTo<UserClass, const UObjectBase>::Result), "You cannot use UFunction delegates with non UObject classes.");
 
 public:
-	typedef typename TUnwrapType<WrappedRetValType>::Type RetValType;
-
 	TBaseUFunctionDelegateInstance(UserClass* InUserObject, const FName& InFunctionName, VarTypes... Vars)
 		: FunctionName (InFunctionName)
 		, UserObjectPtr(InUserObject)
 		, Payload      (Vars...)
+		, Handle       (FDelegateHandle::GenerateNewHandle)
 	{
 		check(InFunctionName != NAME_None);
 		
@@ -93,7 +97,7 @@ public:
 
 	virtual Super* CreateCopy( ) override
 	{
-		return Payload.ApplyAfter(&Create, UserObjectPtr.Get(), GetFunctionName());
+		return new UnwrappedThisType(*(UnwrappedThisType*)this);
 	}
 
 	virtual RetValType Execute(ParamTypes... Params) const override
@@ -108,6 +112,11 @@ public:
 		return PayloadAndParams->GetResult();
 	}
 
+	virtual FDelegateHandle GetHandle() const override
+	{
+		return Handle;
+	}
+
 	virtual bool IsSameFunction( const Super& Other ) const override
 	{
 		// NOTE: Payload data is not currently considered when comparing delegate instances.
@@ -128,25 +137,7 @@ public:
 	 */
 	FORCEINLINE static Super* Create(UserClass* InUserObject, const FName& InFunctionName, VarTypes... Vars)
 	{
-		return new TBaseUFunctionDelegateInstance<UserClass, RetValType (ParamTypes...), VarTypes...>(InUserObject, InFunctionName, Vars...);
-	}
-
-	template <typename TupleType, typename IntSeq>
-	struct TCreateFromTuple;
-
-	template <typename... TupleArgs, uint32... Indices>
-	struct TCreateFromTuple<TTuple<TupleArgs...>, TIntegerSequence<Indices...>>
-	{
-		static Super* CreateImpl(UserClass* InUserObject, const FName& InFunctionName, const TTuple<TupleArgs...>& Tuple)
-		{
-			return Create(InUserObject, InFunctionName, Tuple.template Get<Indices>()...);
-		}
-	};
-
-	template <typename... TupleArgs>
-	FORCEINLINE static Super* CreateFromTuple(UserClass* InUserObject, const FName& InFunctionName, const TTuple<TupleArgs...>& Tuple)
-	{
-		return TCreateFromTuple<TTuple<TupleArgs...>, TMakeIntegerSequence<sizeof...(TupleArgs)>>::CreateImpl(InUserObject, InFunctionName, Tuple);
+		return new UnwrappedThisType(InUserObject, InFunctionName, Vars...);
 	}
 
 public:
@@ -161,6 +152,9 @@ public:
 	TWeakObjectPtr<UserClass> UserObjectPtr;
 
 	TTuple<VarTypes...> Payload;
+
+	// The handle of this delegate
+	FDelegateHandle Handle;
 };
 
 template <class UserClass, typename... ParamTypes, typename... VarTypes>
@@ -206,16 +200,21 @@ class TBaseSPMethodDelegateInstance;
 template <bool bConst, class UserClass, ESPMode SPMode, typename WrappedRetValType, typename... ParamTypes, typename... VarTypes>
 class TBaseSPMethodDelegateInstance<bConst, UserClass, SPMode, WrappedRetValType (ParamTypes...), VarTypes...> : public IBaseDelegateInstance<typename TUnwrapType<WrappedRetValType>::Type (ParamTypes...)>
 {
-	typedef IBaseDelegateInstance<typename TUnwrapType<WrappedRetValType>::Type (ParamTypes...)> Super;
-
 public:
 	typedef typename TUnwrapType<WrappedRetValType>::Type RetValType;
+
+private:
+	typedef IBaseDelegateInstance<RetValType (ParamTypes...)>                                                 Super;
+	typedef TBaseSPMethodDelegateInstance<bConst, UserClass, SPMode, RetValType (ParamTypes...), VarTypes...> UnwrappedThisType;
+
+public:
 	typedef typename TMemFunPtrType<bConst, UserClass, RetValType (ParamTypes..., VarTypes...)>::Type FMethodPtr;
 
 	TBaseSPMethodDelegateInstance(const TSharedPtr<UserClass, SPMode>& InUserObject, FMethodPtr InMethodPtr, VarTypes... Vars)
 		: UserObject(InUserObject)
 		, MethodPtr (InMethodPtr)
 		, Payload   (Vars...)
+		, Handle    (FDelegateHandle::GenerateNewHandle)
 	{
 		// NOTE: Shared pointer delegates are allowed to have a null incoming object pointer.  Weak pointers can expire,
 		//       an it is possible for a copy of a delegate instance to end up with a null pointer.
@@ -260,7 +259,7 @@ public:
 
 	virtual Super* CreateCopy() override
 	{
-		return Payload.template ApplyAfter_ExplicitReturnType<Super*>(static_cast<Super*(*)(const TSharedPtr<UserClass, SPMode>&, FMethodPtr, VarTypes...)>(&Create), UserObject.Pin(), MethodPtr);
+		return new UnwrappedThisType(*(UnwrappedThisType*)this);
 	}
 
 	virtual RetValType Execute(ParamTypes... Params) const override
@@ -283,6 +282,11 @@ public:
 		return Payload.template ApplyAfter_ExplicitReturnType<RetValType>(TMemberFunctionCaller<MutableUserClass, FMethodPtr>(MutableUserObject, MethodPtr), Params...);
 	}
 
+	virtual FDelegateHandle GetHandle() const override
+	{
+		return Handle;
+	}
+
 	virtual bool IsSameFunction( const Super& InOtherDelegate ) const override
 	{
 		// NOTE: Payload data is not currently considered when comparing delegate instances.
@@ -308,7 +312,7 @@ public:
 	 */
 	FORCEINLINE static Super* Create(const TSharedPtr<UserClass, SPMode>& InUserObjectRef, FMethodPtr InFunc, VarTypes... Vars)
 	{
-		return new TBaseSPMethodDelegateInstance<bConst, UserClass, SPMode, RetValType (ParamTypes...), VarTypes...>(InUserObjectRef, InFunc, Vars...);
+		return new UnwrappedThisType(InUserObjectRef, InFunc, Vars...);
 	}
 
 	/**
@@ -353,6 +357,9 @@ protected:
 
 	// Payload member variables, if any.
 	TTuple<VarTypes...> Payload;
+
+	// The handle of this delegate
+	FDelegateHandle Handle;
 };
 
 template <bool bConst, class UserClass, ESPMode SPMode, typename... ParamTypes, typename... VarTypes>
@@ -397,12 +404,16 @@ class TBaseRawMethodDelegateInstance;
 template <bool bConst, class UserClass, typename WrappedRetValType, typename... ParamTypes, typename... VarTypes>
 class TBaseRawMethodDelegateInstance<bConst, UserClass, WrappedRetValType (ParamTypes...), VarTypes...> : public IBaseDelegateInstance<typename TUnwrapType<WrappedRetValType>::Type (ParamTypes...)>
 {
+public:
+	typedef typename TUnwrapType<WrappedRetValType>::Type RetValType;
+
+private:
 	static_assert((!CanConvertPointerFromTo<UserClass, const UObjectBase>::Result), "You cannot use raw method delegates with UObjects.");
 
-	typedef IBaseDelegateInstance<typename TUnwrapType<WrappedRetValType>::Type (ParamTypes...)> Super;
+	typedef IBaseDelegateInstance<RetValType (ParamTypes...)>                                          Super;
+	typedef TBaseRawMethodDelegateInstance<bConst, UserClass, RetValType (ParamTypes...), VarTypes...> UnwrappedThisType;
 
 public:
-	typedef typename TUnwrapType<WrappedRetValType>::Type RetValType;
 	typedef typename TMemFunPtrType<bConst, UserClass, RetValType (ParamTypes..., VarTypes...)>::Type FMethodPtr;
 
 	/**
@@ -415,6 +426,7 @@ public:
 		: UserObject(InUserObject)
 		, MethodPtr (InMethodPtr)
 		, Payload   (Vars...)
+		, Handle    (FDelegateHandle::GenerateNewHandle)
 	{
 		// Non-expirable delegates must always have a non-null object pointer on creation (otherwise they could never execute.)
 		check(InUserObject != nullptr && MethodPtr != nullptr);
@@ -460,7 +472,7 @@ public:
 
 	virtual Super* CreateCopy( ) override
 	{
-		return Payload.template ApplyAfter_ExplicitReturnType<Super*>(&Create, UserObject, MethodPtr);
+		return new UnwrappedThisType(*(UnwrappedThisType*)this);
 	}
 
 	virtual RetValType Execute(ParamTypes... Params) const override
@@ -479,6 +491,11 @@ public:
 		return Payload.template ApplyAfter_ExplicitReturnType<RetValType>(TMemberFunctionCaller<MutableUserClass, FMethodPtr>(MutableUserObject, MethodPtr), Params...);
 	}
 
+	virtual FDelegateHandle GetHandle() const override
+	{
+		return Handle;
+	}
+
 	virtual bool IsSameFunction( const Super& InOtherDelegate ) const override
 	{
 		// NOTE: Payload data is not currently considered when comparing delegate instances.
@@ -505,7 +522,7 @@ public:
 	 */
 	FORCEINLINE static Super* Create(UserClass* InUserObject, FMethodPtr InFunc, VarTypes... Vars)
 	{
-		return new TBaseRawMethodDelegateInstance<bConst, UserClass, RetValType (ParamTypes...), VarTypes...>(InUserObject, InFunc, Vars...);
+		return new UnwrappedThisType(InUserObject, InFunc, Vars...);
 	}
 
 protected:
@@ -534,6 +551,9 @@ protected:
 
 	// Payload member variables (if any).
 	TTuple<VarTypes...> Payload;
+
+	// The handle of this delegate
+	FDelegateHandle Handle;
 };
 
 template <bool bConst, class UserClass, typename... ParamTypes, typename... VarTypes>
@@ -573,18 +593,23 @@ class TBaseUObjectMethodDelegateInstance;
 template <bool bConst, class UserClass, typename WrappedRetValType, typename... ParamTypes, typename... VarTypes>
 class TBaseUObjectMethodDelegateInstance<bConst, UserClass, WrappedRetValType (ParamTypes...), VarTypes...> : public IBaseDelegateInstance<typename TUnwrapType<WrappedRetValType>::Type (ParamTypes...)>
 {
-	typedef IBaseDelegateInstance<typename TUnwrapType<WrappedRetValType>::Type (ParamTypes...)> Super;
+public:
+	typedef typename TUnwrapType<WrappedRetValType>::Type RetValType;
+
+private:
+	typedef IBaseDelegateInstance<RetValType (ParamTypes...)>                                              Super;
+	typedef TBaseUObjectMethodDelegateInstance<bConst, UserClass, RetValType (ParamTypes...), VarTypes...> UnwrappedThisType;
 
 	static_assert((CanConvertPointerFromTo<UserClass, const UObjectBase>::Result), "You cannot use UObject method delegates with raw pointers.");
 
 public:
-	typedef typename TUnwrapType<WrappedRetValType>::Type RetValType;
 	typedef typename TMemFunPtrType<bConst, UserClass, RetValType (ParamTypes..., VarTypes...)>::Type FMethodPtr;
 
 	TBaseUObjectMethodDelegateInstance(UserClass* InUserObject, FMethodPtr InMethodPtr, VarTypes... Vars)
 		: UserObject(InUserObject)
 		, MethodPtr (InMethodPtr)
 		, Payload   (Vars...)
+		, Handle    (FDelegateHandle::GenerateNewHandle)
 	{
 		// NOTE: UObject delegates are allowed to have a null incoming object pointer.  UObject weak pointers can expire,
 		//       an it is possible for a copy of a delegate instance to end up with a null pointer.
@@ -634,7 +659,7 @@ public:
 
 	virtual Super* CreateCopy( ) override
 	{
-		return Payload.template ApplyAfter_ExplicitReturnType<Super*>(&Create, UserObject.Get(), MethodPtr);
+		return new UnwrappedThisType(*(UnwrappedThisType*)this);
 	}
 
 	virtual RetValType Execute(ParamTypes... Params) const override
@@ -656,6 +681,11 @@ public:
 		return Payload.template ApplyAfter_ExplicitReturnType<RetValType>(TMemberFunctionCaller<MutableUserClass, FMethodPtr>(MutableUserObject, MethodPtr), Params...);
 	}
 
+	virtual FDelegateHandle GetHandle() const override
+	{
+		return Handle;
+	}
+
 	virtual bool IsSameFunction( const Super& InOtherDelegate ) const override
 	{
 		// NOTE: Payload data is not currently considered when comparing delegate instances.
@@ -680,7 +710,7 @@ public:
 	 */
 	FORCEINLINE static Super* Create(UserClass* InUserObject, FMethodPtr InFunc, VarTypes... Vars)
 	{
-		return new TBaseUObjectMethodDelegateInstance<bConst, UserClass, RetValType (ParamTypes...), VarTypes...>(InUserObject, InFunc, Vars...);
+		return new UnwrappedThisType(InUserObject, InFunc, Vars...);
 	}
 
 protected:
@@ -709,6 +739,9 @@ protected:
 
 	// Payload member variables (if any).
 	TTuple<VarTypes...> Payload;
+
+	// The handle of this delegate
+	FDelegateHandle Handle;
 };
 
 template <bool bConst, class UserClass, typename... ParamTypes, typename... VarTypes>
@@ -752,15 +785,19 @@ class TBaseStaticDelegateInstance;
 template <typename WrappedRetValType, typename... ParamTypes, typename... VarTypes>
 class TBaseStaticDelegateInstance<WrappedRetValType (ParamTypes...), VarTypes...> : public IBaseDelegateInstance<typename TUnwrapType<WrappedRetValType>::Type (ParamTypes...)>
 {
-	typedef IBaseDelegateInstance<typename TUnwrapType<WrappedRetValType>::Type (ParamTypes...)> Super;
-
 public:
 	typedef typename TUnwrapType<WrappedRetValType>::Type RetValType;
+
+	typedef IBaseDelegateInstance<RetValType (ParamTypes...)>                    Super;
+	typedef TBaseStaticDelegateInstance<RetValType (ParamTypes...), VarTypes...> UnwrappedThisType;
+
+public:
 	typedef RetValType (*FFuncPtr)(ParamTypes..., VarTypes...);
 
 	TBaseStaticDelegateInstance(FFuncPtr InStaticFuncPtr, VarTypes... Vars)
 		: StaticFuncPtr(InStaticFuncPtr)
 		, Payload      (Vars...)
+		, Handle       (FDelegateHandle::GenerateNewHandle)
 	{
 		check(StaticFuncPtr != nullptr);
 	}
@@ -805,7 +842,7 @@ public:
 
 	virtual Super* CreateCopy( ) override
 	{
-		return Payload.template ApplyAfter_ExplicitReturnType<Super*>(&Create, StaticFuncPtr);
+		return new UnwrappedThisType(*(UnwrappedThisType*)this);
 	}
 
 	virtual RetValType Execute(ParamTypes... Params) const override
@@ -816,6 +853,11 @@ public:
 		return Payload.template ApplyAfter_ExplicitReturnType<RetValType>(StaticFuncPtr, Params...);
 	}
 
+	virtual FDelegateHandle GetHandle() const override
+	{
+		return Handle;
+	}
+
 	virtual bool IsSameFunction( const Super& InOtherDelegate ) const override
 	{
 		// NOTE: Payload data is not currently considered when comparing delegate instances.
@@ -823,7 +865,7 @@ public:
 		if (InOtherDelegate.GetType() == EDelegateInstanceType::Raw)
 		{
 			// Downcast to our delegate type and compare
-			const auto& OtherStaticDelegate = static_cast<const TBaseStaticDelegateInstance<RetValType (ParamTypes...), VarTypes...>&>(InOtherDelegate);
+			const auto& OtherStaticDelegate = static_cast<const UnwrappedThisType&>(InOtherDelegate);
 
 			return (StaticFuncPtr == OtherStaticDelegate.StaticFuncPtr);
 		}
@@ -841,7 +883,7 @@ public:
 	 */
 	FORCEINLINE static Super* Create(FFuncPtr InFunc, VarTypes... Vars)
 	{
-		return new TBaseStaticDelegateInstance<RetValType (ParamTypes...), VarTypes...>(InFunc, Vars...);
+		return new UnwrappedThisType(InFunc, Vars...);
 	}
 
 private:
@@ -851,6 +893,9 @@ private:
 
 	// Payload member variables, if any.
 	TTuple<VarTypes...> Payload;
+
+	// The handle of this delegate
+	FDelegateHandle Handle;
 };
 
 template <typename... ParamTypes, typename... VarTypes>
@@ -886,21 +931,27 @@ class TBaseFunctorDelegateInstance;
 template <typename WrappedRetValType, typename... ParamTypes, typename FunctorType, typename... VarTypes>
 class TBaseFunctorDelegateInstance<WrappedRetValType(ParamTypes...), FunctorType, VarTypes...> : public IBaseDelegateInstance<typename TUnwrapType<WrappedRetValType>::Type(ParamTypes...)>
 {
+public:
+	typedef typename TUnwrapType<WrappedRetValType>::Type RetValType;
+
+private:
 	static_assert(TAreTypesEqual<FunctorType, typename TRemoveReference<FunctorType>::Type>::Value, "FunctorType cannot be a reference");
 
 	typedef IBaseDelegateInstance<typename TUnwrapType<WrappedRetValType>::Type(ParamTypes...)> Super;
+	typedef TBaseFunctorDelegateInstance<RetValType(ParamTypes...), FunctorType, VarTypes...>   UnwrappedThisType;
 
 public:
-	typedef typename TUnwrapType<WrappedRetValType>::Type RetValType;
-
 	TBaseFunctorDelegateInstance(const FunctorType& InFunctor, VarTypes... Vars)
 		: Functor(InFunctor)
 		, Payload(Vars...)
+		, Handle (FDelegateHandle::GenerateNewHandle)
 	{
 	}
+
 	TBaseFunctorDelegateInstance(FunctorType&& InFunctor, VarTypes... Vars)
 		: Functor(MoveTemp(InFunctor))
 		, Payload(Vars...)
+		, Handle (FDelegateHandle::GenerateNewHandle)
 	{
 	}
 
@@ -948,7 +999,7 @@ public:
 	// IBaseDelegateInstance interface
 	virtual Super* CreateCopy() override
 	{
-		return Payload.template ApplyAfter_ExplicitReturnType<Super*>(static_cast<Super* (*)(const FunctorType&)>(&Create), Functor);
+		return new UnwrappedThisType(*(UnwrappedThisType*)this);
 	}
 
 	virtual RetValType Execute(ParamTypes... Params) const override
@@ -956,6 +1007,11 @@ public:
 		return Payload.template ApplyAfter_ExplicitReturnType<RetValType>(Functor, Params...);
 	}
 
+	virtual FDelegateHandle GetHandle() const override
+	{
+		return Handle;
+	}
+
 	virtual bool IsSameFunction(const Super& InOtherDelegate) const override
 	{
 		// There's no nice way to implement this (we don't have the type info necessary to compare against OtherDelegate's Functor)
@@ -971,11 +1027,11 @@ public:
 	 */
 	FORCEINLINE static Super* Create(const FunctorType& InFunctor, VarTypes... Vars)
 	{
-		return new TBaseFunctorDelegateInstance<RetValType(ParamTypes...), FunctorType, VarTypes...>(InFunctor, Vars...);
+		return new UnwrappedThisType(InFunctor, Vars...);
 	}
 	FORCEINLINE static Super* Create(FunctorType&& InFunctor, VarTypes... Vars)
 	{
-		return new TBaseFunctorDelegateInstance<RetValType(ParamTypes...), FunctorType, VarTypes...>(MoveTemp(InFunctor), Vars...);
+		return new UnwrappedThisType(MoveTemp(InFunctor), Vars...);
 	}
 
 private:
@@ -984,6 +1040,9 @@ private:
 
 	// Payload member variables, if any.
 	TTuple<VarTypes...> Payload;
+
+	// The handle of this delegate
+	FDelegateHandle Handle;
 };
 
 template <typename FunctorType, typename... ParamTypes, typename... VarTypes>
diff --git a/Engine/Source/Runtime/Core/Public/Delegates/DelegateSignatureImpl.inl b/Engine/Source/Runtime/Core/Public/Delegates/DelegateSignatureImpl.inl
--- a/Engine/Source/Runtime/Core/Public/Delegates/DelegateSignatureImpl.inl
+++ b/Engine/Source/Runtime/Core/Public/Delegates/DelegateSignatureImpl.inl
@@ -1149,23 +1149,50 @@ public:
 	 *
 	 * @return true if this delegate equals the other, false otherwise.
 	 */
+	DELEGATE_DEPRECATED("Delegate comparison is deprecated - please replace any usage with comparison of FDelegateHandles.")
 	bool operator==( const DELEGATE_CLASS& Other ) const
 	{
 		TDelegateInstanceInterface* DelegateInstance = (TDelegateInstanceInterface*)GetDelegateInstance();
 		TDelegateInstanceInterface* OtherInstance = (TDelegateInstanceInterface*)Other.GetDelegateInstance();
-
+	
 		// The function these delegates point to must be the same
 		if ((DelegateInstance != nullptr) && (OtherInstance != nullptr))
 		{
 			return DelegateInstance->IsSameFunction(*OtherInstance);
 		}
-
+	
 		// If neither delegate is initialized to anything yet, then we treat them as equal
 		if ((DelegateInstance == nullptr) && (OtherInstance == nullptr))
 		{
 			return true;
 		}
+	
+		// No match!
+		return false;
+	}
 
+	/**
+	 * Delegate comparison operator.
+	 *
+	 * @return true if this delegate equals the other, false otherwise.
+	 */
+	bool DEPRECATED_Compare( const DELEGATE_CLASS& Other ) const
+	{
+		TDelegateInstanceInterface* DelegateInstance = (TDelegateInstanceInterface*)GetDelegateInstance();
+		TDelegateInstanceInterface* OtherInstance = (TDelegateInstanceInterface*)Other.GetDelegateInstance();
+	
+		// The function these delegates point to must be the same
+		if ((DelegateInstance != nullptr) && (OtherInstance != nullptr))
+		{
+			return DelegateInstance->IsSameFunction(*OtherInstance);
+		}
+	
+		// If neither delegate is initialized to anything yet, then we treat them as equal
+		if ((DelegateInstance == nullptr) && (OtherInstance == nullptr))
+		{
+			return true;
+		}
+	
 		// No match!
 		return false;
 	}
@@ -1210,12 +1237,16 @@ public:
 	 *
 	 * @param DelegateInstance The delegate instance to add.
 	 */
-	void Add( TDelegateInstanceInterface* DelegateInstance )
+	FDelegateHandle Add( TDelegateInstanceInterface* DelegateInstance )
 	{
+		FDelegateHandle Result;
+
 		if (DelegateInstance != nullptr)
 		{
-			AddDelegateInstance(DelegateInstance);
+			Result = AddDelegateInstance(DelegateInstance);
 		}
+
+		return Result;
 	}
 
 	/**
@@ -1225,14 +1256,17 @@ public:
 	 *
 	 * @param Delegate The delegate to add.
 	 */
-	void Add( const FDelegate& Delegate )
+	FDelegateHandle Add( const FDelegate& Delegate )
 	{
+		FDelegateHandle             Result;
 		TDelegateInstanceInterface* DelegateInstance = (TDelegateInstanceInterface*)Delegate.GetDelegateInstance();
 
 		if (DelegateInstance != nullptr)
 		{
-			AddDelegateInstance(DelegateInstance->CreateCopy());
+			Result = AddDelegateInstance(DelegateInstance->CreateCopy());
 		}
+
+		return Result;
 	}
 
 	/**
@@ -1240,29 +1274,29 @@ public:
 	 *
 	 * @param	InFunc	Function pointer
 	 */
-	inline void AddStatic( typename FDelegate::FStaticDelegate::FFuncPtr InFunc )
+	inline FDelegateHandle AddStatic( typename FDelegate::FStaticDelegate::FFuncPtr InFunc )
 	{
-		Add( FDelegate::CreateStatic( InFunc ) );
+		return Add( FDelegate::CreateStatic( InFunc ) );
 	}
 	template< PAYLOAD_TEMPLATE_DECL_OneVar >
-	inline void AddStatic( typename FDelegate::template TStaticDelegate_OneVar< PAYLOAD_TEMPLATE_LIST_OneVar >::FFuncPtr InFunc, PAYLOAD_TEMPLATE_ARGS_OneVar )
+	inline FDelegateHandle AddStatic( typename FDelegate::template TStaticDelegate_OneVar< PAYLOAD_TEMPLATE_LIST_OneVar >::FFuncPtr InFunc, PAYLOAD_TEMPLATE_ARGS_OneVar )
 	{
-		Add( FDelegate::CreateStatic( InFunc, PAYLOAD_TEMPLATE_PASSIN_OneVar ) );
+		return Add( FDelegate::CreateStatic( InFunc, PAYLOAD_TEMPLATE_PASSIN_OneVar ) );
 	}
 	template< PAYLOAD_TEMPLATE_DECL_TwoVars >
-	inline void AddStatic( typename FDelegate::template TStaticDelegate_TwoVars< PAYLOAD_TEMPLATE_LIST_TwoVars >::FFuncPtr InFunc, PAYLOAD_TEMPLATE_ARGS_TwoVars )
+	inline FDelegateHandle AddStatic( typename FDelegate::template TStaticDelegate_TwoVars< PAYLOAD_TEMPLATE_LIST_TwoVars >::FFuncPtr InFunc, PAYLOAD_TEMPLATE_ARGS_TwoVars )
 	{
-		Add( FDelegate::CreateStatic( InFunc, PAYLOAD_TEMPLATE_PASSIN_TwoVars ) );
+		return Add( FDelegate::CreateStatic( InFunc, PAYLOAD_TEMPLATE_PASSIN_TwoVars ) );
 	}
 	template< PAYLOAD_TEMPLATE_DECL_ThreeVars >
-	inline void AddStatic( typename FDelegate::template TStaticDelegate_ThreeVars< PAYLOAD_TEMPLATE_LIST_ThreeVars >::FFuncPtr InFunc, PAYLOAD_TEMPLATE_ARGS_ThreeVars )
+	inline FDelegateHandle AddStatic( typename FDelegate::template TStaticDelegate_ThreeVars< PAYLOAD_TEMPLATE_LIST_ThreeVars >::FFuncPtr InFunc, PAYLOAD_TEMPLATE_ARGS_ThreeVars )
 	{
-		Add( FDelegate::CreateStatic( InFunc, PAYLOAD_TEMPLATE_PASSIN_ThreeVars ) );
+		return Add( FDelegate::CreateStatic( InFunc, PAYLOAD_TEMPLATE_PASSIN_ThreeVars ) );
 	}
 	template< PAYLOAD_TEMPLATE_DECL_FourVars >
-	inline void AddStatic( typename FDelegate::template TStaticDelegate_FourVars< PAYLOAD_TEMPLATE_LIST_FourVars >::FFuncPtr InFunc, PAYLOAD_TEMPLATE_ARGS_FourVars )
+	inline FDelegateHandle AddStatic( typename FDelegate::template TStaticDelegate_FourVars< PAYLOAD_TEMPLATE_LIST_FourVars >::FFuncPtr InFunc, PAYLOAD_TEMPLATE_ARGS_FourVars )
 	{
-		Add( FDelegate::CreateStatic( InFunc, PAYLOAD_TEMPLATE_PASSIN_FourVars ) );
+		return Add( FDelegate::CreateStatic( InFunc, PAYLOAD_TEMPLATE_PASSIN_FourVars ) );
 	}
 
 	/**
@@ -1275,54 +1309,54 @@ public:
 	 * @param	InFunc			Class method function address
 	 */
 	template< class UserClass >
-	inline void AddRaw( UserClass* InUserObject, typename FDelegate::template TRawMethodDelegate< UserClass >::FMethodPtr InFunc )
+	inline FDelegateHandle AddRaw( UserClass* InUserObject, typename FDelegate::template TRawMethodDelegate< UserClass >::FMethodPtr InFunc )
 	{
-		Add( FDelegate::CreateRaw( InUserObject, InFunc ) );
+		return Add( FDelegate::CreateRaw( InUserObject, InFunc ) );
 	}
 	template< class UserClass >
-	inline void AddRaw( UserClass* InUserObject, typename FDelegate::template TRawMethodDelegate_Const< UserClass >::FMethodPtr InFunc )
+	inline FDelegateHandle AddRaw( UserClass* InUserObject, typename FDelegate::template TRawMethodDelegate_Const< UserClass >::FMethodPtr InFunc )
 	{
-		Add( FDelegate::CreateRaw( InUserObject, InFunc ) );
+		return Add( FDelegate::CreateRaw( InUserObject, InFunc ) );
 	}
 	template< class UserClass, PAYLOAD_TEMPLATE_DECL_OneVar >
-	inline void AddRaw( UserClass* InUserObject, typename FDelegate::template TRawMethodDelegate_OneVar< UserClass, PAYLOAD_TEMPLATE_LIST_OneVar >::FMethodPtr InFunc, PAYLOAD_TEMPLATE_ARGS_OneVar )
+	inline FDelegateHandle AddRaw( UserClass* InUserObject, typename FDelegate::template TRawMethodDelegate_OneVar< UserClass, PAYLOAD_TEMPLATE_LIST_OneVar >::FMethodPtr InFunc, PAYLOAD_TEMPLATE_ARGS_OneVar )
 	{
-		Add( FDelegate::CreateRaw( InUserObject, InFunc, PAYLOAD_TEMPLATE_PASSIN_OneVar ) );
+		return Add( FDelegate::CreateRaw( InUserObject, InFunc, PAYLOAD_TEMPLATE_PASSIN_OneVar ) );
 	}
 	template< class UserClass, PAYLOAD_TEMPLATE_DECL_OneVar >
-	inline void AddRaw( UserClass* InUserObject, typename FDelegate::template TRawMethodDelegate_OneVar_Const< UserClass, PAYLOAD_TEMPLATE_LIST_OneVar >::FMethodPtr InFunc, PAYLOAD_TEMPLATE_ARGS_OneVar )
+	inline FDelegateHandle AddRaw( UserClass* InUserObject, typename FDelegate::template TRawMethodDelegate_OneVar_Const< UserClass, PAYLOAD_TEMPLATE_LIST_OneVar >::FMethodPtr InFunc, PAYLOAD_TEMPLATE_ARGS_OneVar )
 	{
-		Add( FDelegate::CreateRaw( InUserObject, InFunc, PAYLOAD_TEMPLATE_PASSIN_OneVar ) );
+		return Add( FDelegate::CreateRaw( InUserObject, InFunc, PAYLOAD_TEMPLATE_PASSIN_OneVar ) );
 	}
 	template< class UserClass, PAYLOAD_TEMPLATE_DECL_TwoVars >
-	inline void AddRaw( UserClass* InUserObject, typename FDelegate::template TRawMethodDelegate_TwoVars< UserClass, PAYLOAD_TEMPLATE_LIST_TwoVars >::FMethodPtr InFunc, PAYLOAD_TEMPLATE_ARGS_TwoVars )
+	inline FDelegateHandle AddRaw( UserClass* InUserObject, typename FDelegate::template TRawMethodDelegate_TwoVars< UserClass, PAYLOAD_TEMPLATE_LIST_TwoVars >::FMethodPtr InFunc, PAYLOAD_TEMPLATE_ARGS_TwoVars )
 	{
-		Add( FDelegate::CreateRaw( InUserObject, InFunc, PAYLOAD_TEMPLATE_PASSIN_TwoVars ) );
+		return Add( FDelegate::CreateRaw( InUserObject, InFunc, PAYLOAD_TEMPLATE_PASSIN_TwoVars ) );
 	}
 	template< class UserClass, PAYLOAD_TEMPLATE_DECL_TwoVars >
-	inline void AddRaw( UserClass* InUserObject, typename FDelegate::template TRawMethodDelegate_TwoVars_Const< UserClass, PAYLOAD_TEMPLATE_LIST_TwoVars >::FMethodPtr InFunc, PAYLOAD_TEMPLATE_ARGS_TwoVars )
+	inline FDelegateHandle AddRaw( UserClass* InUserObject, typename FDelegate::template TRawMethodDelegate_TwoVars_Const< UserClass, PAYLOAD_TEMPLATE_LIST_TwoVars >::FMethodPtr InFunc, PAYLOAD_TEMPLATE_ARGS_TwoVars )
 	{
-		Add( FDelegate::CreateRaw( InUserObject, InFunc, PAYLOAD_TEMPLATE_PASSIN_TwoVars ) );
+		return Add( FDelegate::CreateRaw( InUserObject, InFunc, PAYLOAD_TEMPLATE_PASSIN_TwoVars ) );
 	}
 	template< class UserClass, PAYLOAD_TEMPLATE_DECL_ThreeVars >
-	inline void AddRaw( UserClass* InUserObject, typename FDelegate::template TRawMethodDelegate_ThreeVars< UserClass, PAYLOAD_TEMPLATE_LIST_ThreeVars >::FMethodPtr InFunc, PAYLOAD_TEMPLATE_ARGS_ThreeVars )
+	inline FDelegateHandle AddRaw( UserClass* InUserObject, typename FDelegate::template TRawMethodDelegate_ThreeVars< UserClass, PAYLOAD_TEMPLATE_LIST_ThreeVars >::FMethodPtr InFunc, PAYLOAD_TEMPLATE_ARGS_ThreeVars )
 	{
-		Add( FDelegate::CreateRaw( InUserObject, InFunc, PAYLOAD_TEMPLATE_PASSIN_ThreeVars ) );
+		return Add( FDelegate::CreateRaw( InUserObject, InFunc, PAYLOAD_TEMPLATE_PASSIN_ThreeVars ) );
 	}
 	template< class UserClass, PAYLOAD_TEMPLATE_DECL_ThreeVars >
-	inline void AddRaw( UserClass* InUserObject, typename FDelegate::template TRawMethodDelegate_ThreeVars_Const< UserClass, PAYLOAD_TEMPLATE_LIST_ThreeVars >::FMethodPtr InFunc, PAYLOAD_TEMPLATE_ARGS_ThreeVars )
+	inline FDelegateHandle AddRaw( UserClass* InUserObject, typename FDelegate::template TRawMethodDelegate_ThreeVars_Const< UserClass, PAYLOAD_TEMPLATE_LIST_ThreeVars >::FMethodPtr InFunc, PAYLOAD_TEMPLATE_ARGS_ThreeVars )
 	{
-		Add( FDelegate::CreateRaw( InUserObject, InFunc, PAYLOAD_TEMPLATE_PASSIN_ThreeVars ) );
+		return Add( FDelegate::CreateRaw( InUserObject, InFunc, PAYLOAD_TEMPLATE_PASSIN_ThreeVars ) );
 	}
 	template< class UserClass, PAYLOAD_TEMPLATE_DECL_FourVars >
-	inline void AddRaw( UserClass* InUserObject, typename FDelegate::template TRawMethodDelegate_FourVars< UserClass, PAYLOAD_TEMPLATE_LIST_FourVars >::FMethodPtr InFunc, PAYLOAD_TEMPLATE_ARGS_FourVars )
+	inline FDelegateHandle AddRaw( UserClass* InUserObject, typename FDelegate::template TRawMethodDelegate_FourVars< UserClass, PAYLOAD_TEMPLATE_LIST_FourVars >::FMethodPtr InFunc, PAYLOAD_TEMPLATE_ARGS_FourVars )
 	{
-		Add( FDelegate::CreateRaw( InUserObject, InFunc, PAYLOAD_TEMPLATE_PASSIN_FourVars ) );
+		return Add( FDelegate::CreateRaw( InUserObject, InFunc, PAYLOAD_TEMPLATE_PASSIN_FourVars ) );
 	}
 	template< class UserClass, PAYLOAD_TEMPLATE_DECL_FourVars >
-	inline void AddRaw( UserClass* InUserObject, typename FDelegate::template TRawMethodDelegate_FourVars_Const< UserClass, PAYLOAD_TEMPLATE_LIST_FourVars >::FMethodPtr InFunc, PAYLOAD_TEMPLATE_ARGS_FourVars )
+	inline FDelegateHandle AddRaw( UserClass* InUserObject, typename FDelegate::template TRawMethodDelegate_FourVars_Const< UserClass, PAYLOAD_TEMPLATE_LIST_FourVars >::FMethodPtr InFunc, PAYLOAD_TEMPLATE_ARGS_FourVars )
 	{
-		Add( FDelegate::CreateRaw( InUserObject, InFunc, PAYLOAD_TEMPLATE_PASSIN_FourVars ) );
+		return Add( FDelegate::CreateRaw( InUserObject, InFunc, PAYLOAD_TEMPLATE_PASSIN_FourVars ) );
 	}
 
 	/**
@@ -1334,54 +1368,54 @@ public:
 	 * @param	InFunc			Class method function address
 	 */
 	template< class UserClass >
-	inline void AddSP( const TSharedRef< UserClass, ESPMode::Fast >& InUserObjectRef, typename FDelegate::template TSPMethodDelegate< UserClass >::FMethodPtr InFunc )
+	inline FDelegateHandle AddSP( const TSharedRef< UserClass, ESPMode::Fast >& InUserObjectRef, typename FDelegate::template TSPMethodDelegate< UserClass >::FMethodPtr InFunc )
 	{
-		Add( FDelegate::CreateSP( InUserObjectRef, InFunc ) );
+		return Add( FDelegate::CreateSP( InUserObjectRef, InFunc ) );
 	}
 	template< class UserClass >
-	inline void AddSP( const TSharedRef< UserClass, ESPMode::Fast >& InUserObjectRef, typename FDelegate::template TSPMethodDelegate_Const< UserClass >::FMethodPtr InFunc )
+	inline FDelegateHandle AddSP( const TSharedRef< UserClass, ESPMode::Fast >& InUserObjectRef, typename FDelegate::template TSPMethodDelegate_Const< UserClass >::FMethodPtr InFunc )
 	{
-		Add( FDelegate::CreateSP( InUserObjectRef, InFunc ) );
+		return Add( FDelegate::CreateSP( InUserObjectRef, InFunc ) );
 	}
 	template< class UserClass, PAYLOAD_TEMPLATE_DECL_OneVar >
-	inline void AddSP( const TSharedRef< UserClass, ESPMode::Fast >& InUserObjectRef, typename FDelegate::template TSPMethodDelegate_OneVar< UserClass, PAYLOAD_TEMPLATE_LIST_OneVar >::FMethodPtr InFunc, PAYLOAD_TEMPLATE_ARGS_OneVar )
+	inline FDelegateHandle AddSP( const TSharedRef< UserClass, ESPMode::Fast >& InUserObjectRef, typename FDelegate::template TSPMethodDelegate_OneVar< UserClass, PAYLOAD_TEMPLATE_LIST_OneVar >::FMethodPtr InFunc, PAYLOAD_TEMPLATE_ARGS_OneVar )
 	{
-		Add( FDelegate::CreateSP( InUserObjectRef, InFunc, PAYLOAD_TEMPLATE_PASSIN_OneVar ) );
+		return Add( FDelegate::CreateSP( InUserObjectRef, InFunc, PAYLOAD_TEMPLATE_PASSIN_OneVar ) );
 	}
 	template< class UserClass, PAYLOAD_TEMPLATE_DECL_OneVar >
-	inline void AddSP( const TSharedRef< UserClass, ESPMode::Fast >& InUserObjectRef, typename FDelegate::template TSPMethodDelegate_OneVar_Const< UserClass, PAYLOAD_TEMPLATE_LIST_OneVar >::FMethodPtr InFunc, PAYLOAD_TEMPLATE_ARGS_OneVar )
+	inline FDelegateHandle AddSP( const TSharedRef< UserClass, ESPMode::Fast >& InUserObjectRef, typename FDelegate::template TSPMethodDelegate_OneVar_Const< UserClass, PAYLOAD_TEMPLATE_LIST_OneVar >::FMethodPtr InFunc, PAYLOAD_TEMPLATE_ARGS_OneVar )
 	{
-		Add( FDelegate::CreateSP( InUserObjectRef, InFunc, PAYLOAD_TEMPLATE_PASSIN_OneVar ) );
+		return Add( FDelegate::CreateSP( InUserObjectRef, InFunc, PAYLOAD_TEMPLATE_PASSIN_OneVar ) );
 	}
 	template< class UserClass, PAYLOAD_TEMPLATE_DECL_TwoVars >
-	inline void AddSP( const TSharedRef< UserClass, ESPMode::Fast >& InUserObjectRef, typename FDelegate::template TSPMethodDelegate_TwoVars< UserClass, PAYLOAD_TEMPLATE_LIST_TwoVars >::FMethodPtr InFunc, PAYLOAD_TEMPLATE_ARGS_TwoVars )
+	inline FDelegateHandle AddSP( const TSharedRef< UserClass, ESPMode::Fast >& InUserObjectRef, typename FDelegate::template TSPMethodDelegate_TwoVars< UserClass, PAYLOAD_TEMPLATE_LIST_TwoVars >::FMethodPtr InFunc, PAYLOAD_TEMPLATE_ARGS_TwoVars )
 	{
-		Add( FDelegate::CreateSP( InUserObjectRef, InFunc, PAYLOAD_TEMPLATE_PASSIN_TwoVars ) );
+		return Add( FDelegate::CreateSP( InUserObjectRef, InFunc, PAYLOAD_TEMPLATE_PASSIN_TwoVars ) );
 	}
 	template< class UserClass, PAYLOAD_TEMPLATE_DECL_TwoVars >
-	inline void AddSP( const TSharedRef< UserClass, ESPMode::Fast >& InUserObjectRef, typename FDelegate::template TSPMethodDelegate_TwoVars_Const< UserClass, PAYLOAD_TEMPLATE_LIST_TwoVars >::FMethodPtr InFunc, PAYLOAD_TEMPLATE_ARGS_TwoVars )
+	inline FDelegateHandle AddSP( const TSharedRef< UserClass, ESPMode::Fast >& InUserObjectRef, typename FDelegate::template TSPMethodDelegate_TwoVars_Const< UserClass, PAYLOAD_TEMPLATE_LIST_TwoVars >::FMethodPtr InFunc, PAYLOAD_TEMPLATE_ARGS_TwoVars )
 	{
-		Add( FDelegate::CreateSP( InUserObjectRef, InFunc, PAYLOAD_TEMPLATE_PASSIN_TwoVars ) );
+		return Add( FDelegate::CreateSP( InUserObjectRef, InFunc, PAYLOAD_TEMPLATE_PASSIN_TwoVars ) );
 	}
 	template< class UserClass, PAYLOAD_TEMPLATE_DECL_ThreeVars >
-	inline void AddSP( const TSharedRef< UserClass, ESPMode::Fast >& InUserObjectRef, typename FDelegate::template TSPMethodDelegate_ThreeVars< UserClass, PAYLOAD_TEMPLATE_LIST_ThreeVars >::FMethodPtr InFunc, PAYLOAD_TEMPLATE_ARGS_ThreeVars )
+	inline FDelegateHandle AddSP( const TSharedRef< UserClass, ESPMode::Fast >& InUserObjectRef, typename FDelegate::template TSPMethodDelegate_ThreeVars< UserClass, PAYLOAD_TEMPLATE_LIST_ThreeVars >::FMethodPtr InFunc, PAYLOAD_TEMPLATE_ARGS_ThreeVars )
 	{
-		Add( FDelegate::CreateSP( InUserObjectRef, InFunc, PAYLOAD_TEMPLATE_PASSIN_ThreeVars ) );
+		return Add( FDelegate::CreateSP( InUserObjectRef, InFunc, PAYLOAD_TEMPLATE_PASSIN_ThreeVars ) );
 	}
 	template< class UserClass, PAYLOAD_TEMPLATE_DECL_ThreeVars >
-	inline void AddSP( const TSharedRef< UserClass, ESPMode::Fast >& InUserObjectRef, typename FDelegate::template TSPMethodDelegate_ThreeVars_Const< UserClass, PAYLOAD_TEMPLATE_LIST_ThreeVars >::FMethodPtr InFunc, PAYLOAD_TEMPLATE_ARGS_ThreeVars )
+	inline FDelegateHandle AddSP( const TSharedRef< UserClass, ESPMode::Fast >& InUserObjectRef, typename FDelegate::template TSPMethodDelegate_ThreeVars_Const< UserClass, PAYLOAD_TEMPLATE_LIST_ThreeVars >::FMethodPtr InFunc, PAYLOAD_TEMPLATE_ARGS_ThreeVars )
 	{
-		Add( FDelegate::CreateSP( InUserObjectRef, InFunc, PAYLOAD_TEMPLATE_PASSIN_ThreeVars ) );
+		return Add( FDelegate::CreateSP( InUserObjectRef, InFunc, PAYLOAD_TEMPLATE_PASSIN_ThreeVars ) );
 	}
 	template< class UserClass, PAYLOAD_TEMPLATE_DECL_FourVars >
-	inline void AddSP( const TSharedRef< UserClass, ESPMode::Fast >& InUserObjectRef, typename FDelegate::template TSPMethodDelegate_FourVars< UserClass, PAYLOAD_TEMPLATE_LIST_FourVars >::FMethodPtr InFunc, PAYLOAD_TEMPLATE_ARGS_FourVars )
+	inline FDelegateHandle AddSP( const TSharedRef< UserClass, ESPMode::Fast >& InUserObjectRef, typename FDelegate::template TSPMethodDelegate_FourVars< UserClass, PAYLOAD_TEMPLATE_LIST_FourVars >::FMethodPtr InFunc, PAYLOAD_TEMPLATE_ARGS_FourVars )
 	{
-		Add( FDelegate::CreateSP( InUserObjectRef, InFunc, PAYLOAD_TEMPLATE_PASSIN_FourVars ) );
+		return Add( FDelegate::CreateSP( InUserObjectRef, InFunc, PAYLOAD_TEMPLATE_PASSIN_FourVars ) );
 	}
 	template< class UserClass, PAYLOAD_TEMPLATE_DECL_FourVars >
-	inline void AddSP( const TSharedRef< UserClass, ESPMode::Fast >& InUserObjectRef, typename FDelegate::template TSPMethodDelegate_FourVars_Const< UserClass, PAYLOAD_TEMPLATE_LIST_FourVars >::FMethodPtr InFunc, PAYLOAD_TEMPLATE_ARGS_FourVars )
+	inline FDelegateHandle AddSP( const TSharedRef< UserClass, ESPMode::Fast >& InUserObjectRef, typename FDelegate::template TSPMethodDelegate_FourVars_Const< UserClass, PAYLOAD_TEMPLATE_LIST_FourVars >::FMethodPtr InFunc, PAYLOAD_TEMPLATE_ARGS_FourVars )
 	{
-		Add( FDelegate::CreateSP( InUserObjectRef, InFunc, PAYLOAD_TEMPLATE_PASSIN_FourVars ) );
+		return Add( FDelegate::CreateSP( InUserObjectRef, InFunc, PAYLOAD_TEMPLATE_PASSIN_FourVars ) );
 	}
 
 	/**
@@ -1393,54 +1427,54 @@ public:
 	 * @param	InFunc			Class method function address
 	 */
 	template< class UserClass >
-	inline void AddSP( UserClass* InUserObject, typename FDelegate::template TSPMethodDelegate< UserClass >::FMethodPtr InFunc )
+	inline FDelegateHandle AddSP( UserClass* InUserObject, typename FDelegate::template TSPMethodDelegate< UserClass >::FMethodPtr InFunc )
 	{
-		Add( FDelegate::CreateSP( InUserObject, InFunc ) );
+		return Add( FDelegate::CreateSP( InUserObject, InFunc ) );
 	}
 	template< class UserClass >
-	inline void AddSP( UserClass* InUserObject, typename FDelegate::template TSPMethodDelegate_Const< UserClass >::FMethodPtr InFunc )
+	inline FDelegateHandle AddSP( UserClass* InUserObject, typename FDelegate::template TSPMethodDelegate_Const< UserClass >::FMethodPtr InFunc )
 	{
-		Add( FDelegate::CreateSP( InUserObject, InFunc ) );
+		return Add( FDelegate::CreateSP( InUserObject, InFunc ) );
 	}
 	template< class UserClass, PAYLOAD_TEMPLATE_DECL_OneVar >
-	inline void AddSP( UserClass* InUserObject, typename FDelegate::template TSPMethodDelegate_OneVar< UserClass, PAYLOAD_TEMPLATE_LIST_OneVar >::FMethodPtr InFunc, PAYLOAD_TEMPLATE_ARGS_OneVar )
+	inline FDelegateHandle AddSP( UserClass* InUserObject, typename FDelegate::template TSPMethodDelegate_OneVar< UserClass, PAYLOAD_TEMPLATE_LIST_OneVar >::FMethodPtr InFunc, PAYLOAD_TEMPLATE_ARGS_OneVar )
 	{
-		Add( FDelegate::CreateSP( InUserObject, InFunc, PAYLOAD_TEMPLATE_PASSIN_OneVar ) );
+		return Add( FDelegate::CreateSP( InUserObject, InFunc, PAYLOAD_TEMPLATE_PASSIN_OneVar ) );
 	}
 	template< class UserClass, PAYLOAD_TEMPLATE_DECL_OneVar >
-	inline void AddSP( UserClass* InUserObject, typename FDelegate::template TSPMethodDelegate_OneVar_Const< UserClass, PAYLOAD_TEMPLATE_LIST_OneVar >::FMethodPtr InFunc, PAYLOAD_TEMPLATE_ARGS_OneVar )
+	inline FDelegateHandle AddSP( UserClass* InUserObject, typename FDelegate::template TSPMethodDelegate_OneVar_Const< UserClass, PAYLOAD_TEMPLATE_LIST_OneVar >::FMethodPtr InFunc, PAYLOAD_TEMPLATE_ARGS_OneVar )
 	{
-		Add( FDelegate::CreateSP( InUserObject, InFunc, PAYLOAD_TEMPLATE_PASSIN_OneVar ) );
+		return Add( FDelegate::CreateSP( InUserObject, InFunc, PAYLOAD_TEMPLATE_PASSIN_OneVar ) );
 	}
 	template< class UserClass, PAYLOAD_TEMPLATE_DECL_TwoVars >
-	inline void AddSP( UserClass* InUserObject, typename FDelegate::template TSPMethodDelegate_TwoVars< UserClass, PAYLOAD_TEMPLATE_LIST_TwoVars >::FMethodPtr InFunc, PAYLOAD_TEMPLATE_ARGS_TwoVars )
+	inline FDelegateHandle AddSP( UserClass* InUserObject, typename FDelegate::template TSPMethodDelegate_TwoVars< UserClass, PAYLOAD_TEMPLATE_LIST_TwoVars >::FMethodPtr InFunc, PAYLOAD_TEMPLATE_ARGS_TwoVars )
 	{
-		Add( FDelegate::CreateSP( InUserObject, InFunc, PAYLOAD_TEMPLATE_PASSIN_TwoVars ) );
+		return Add( FDelegate::CreateSP( InUserObject, InFunc, PAYLOAD_TEMPLATE_PASSIN_TwoVars ) );
 	}
 	template< class UserClass, PAYLOAD_TEMPLATE_DECL_TwoVars >
-	inline void AddSP( UserClass* InUserObject, typename FDelegate::template TSPMethodDelegate_TwoVars_Const< UserClass, PAYLOAD_TEMPLATE_LIST_TwoVars >::FMethodPtr InFunc, PAYLOAD_TEMPLATE_ARGS_TwoVars )
+	inline FDelegateHandle AddSP( UserClass* InUserObject, typename FDelegate::template TSPMethodDelegate_TwoVars_Const< UserClass, PAYLOAD_TEMPLATE_LIST_TwoVars >::FMethodPtr InFunc, PAYLOAD_TEMPLATE_ARGS_TwoVars )
 	{
-		Add( FDelegate::CreateSP( InUserObject, InFunc, PAYLOAD_TEMPLATE_PASSIN_TwoVars ) );
+		return Add( FDelegate::CreateSP( InUserObject, InFunc, PAYLOAD_TEMPLATE_PASSIN_TwoVars ) );
 	}
 	template< class UserClass, PAYLOAD_TEMPLATE_DECL_ThreeVars >
-	inline void AddSP( UserClass* InUserObject, typename FDelegate::template TSPMethodDelegate_ThreeVars< UserClass, PAYLOAD_TEMPLATE_LIST_ThreeVars >::FMethodPtr InFunc, PAYLOAD_TEMPLATE_ARGS_ThreeVars )
+	inline FDelegateHandle AddSP( UserClass* InUserObject, typename FDelegate::template TSPMethodDelegate_ThreeVars< UserClass, PAYLOAD_TEMPLATE_LIST_ThreeVars >::FMethodPtr InFunc, PAYLOAD_TEMPLATE_ARGS_ThreeVars )
 	{
-		Add( FDelegate::CreateSP( InUserObject, InFunc, PAYLOAD_TEMPLATE_PASSIN_ThreeVars ) );
+		return Add( FDelegate::CreateSP( InUserObject, InFunc, PAYLOAD_TEMPLATE_PASSIN_ThreeVars ) );
 	}
 	template< class UserClass, PAYLOAD_TEMPLATE_DECL_ThreeVars >
-	inline void AddSP( UserClass* InUserObject, typename FDelegate::template TSPMethodDelegate_ThreeVars_Const< UserClass, PAYLOAD_TEMPLATE_LIST_ThreeVars >::FMethodPtr InFunc, PAYLOAD_TEMPLATE_ARGS_ThreeVars )
+	inline FDelegateHandle AddSP( UserClass* InUserObject, typename FDelegate::template TSPMethodDelegate_ThreeVars_Const< UserClass, PAYLOAD_TEMPLATE_LIST_ThreeVars >::FMethodPtr InFunc, PAYLOAD_TEMPLATE_ARGS_ThreeVars )
 	{
-		Add( FDelegate::CreateSP( InUserObject, InFunc, PAYLOAD_TEMPLATE_PASSIN_ThreeVars ) );
+		return Add( FDelegate::CreateSP( InUserObject, InFunc, PAYLOAD_TEMPLATE_PASSIN_ThreeVars ) );
 	}
 	template< class UserClass, PAYLOAD_TEMPLATE_DECL_FourVars >
-	inline void AddSP( UserClass* InUserObject, typename FDelegate::template TSPMethodDelegate_FourVars< UserClass, PAYLOAD_TEMPLATE_LIST_FourVars >::FMethodPtr InFunc, PAYLOAD_TEMPLATE_ARGS_FourVars )
+	inline FDelegateHandle AddSP( UserClass* InUserObject, typename FDelegate::template TSPMethodDelegate_FourVars< UserClass, PAYLOAD_TEMPLATE_LIST_FourVars >::FMethodPtr InFunc, PAYLOAD_TEMPLATE_ARGS_FourVars )
 	{
-		Add( FDelegate::CreateSP( InUserObject, InFunc, PAYLOAD_TEMPLATE_PASSIN_FourVars ) );
+		return Add( FDelegate::CreateSP( InUserObject, InFunc, PAYLOAD_TEMPLATE_PASSIN_FourVars ) );
 	}
 	template< class UserClass, PAYLOAD_TEMPLATE_DECL_FourVars >
-	inline void AddSP( UserClass* InUserObject, typename FDelegate::template TSPMethodDelegate_FourVars_Const< UserClass, PAYLOAD_TEMPLATE_LIST_FourVars >::FMethodPtr InFunc, PAYLOAD_TEMPLATE_ARGS_FourVars )
+	inline FDelegateHandle AddSP( UserClass* InUserObject, typename FDelegate::template TSPMethodDelegate_FourVars_Const< UserClass, PAYLOAD_TEMPLATE_LIST_FourVars >::FMethodPtr InFunc, PAYLOAD_TEMPLATE_ARGS_FourVars )
 	{
-		Add( FDelegate::CreateSP( InUserObject, InFunc, PAYLOAD_TEMPLATE_PASSIN_FourVars ) );
+		return Add( FDelegate::CreateSP( InUserObject, InFunc, PAYLOAD_TEMPLATE_PASSIN_FourVars ) );
 	}
 
 	/**
@@ -1450,54 +1484,54 @@ public:
 	 * @param	InFunc			Class method function address
 	 */
 	template< class UserClass >
-	inline void AddThreadSafeSP( const TSharedRef< UserClass, ESPMode::ThreadSafe >& InUserObjectRef, typename FDelegate::template TThreadSafeSPMethodDelegate< UserClass >::FMethodPtr InFunc )
+	inline FDelegateHandle AddThreadSafeSP( const TSharedRef< UserClass, ESPMode::ThreadSafe >& InUserObjectRef, typename FDelegate::template TThreadSafeSPMethodDelegate< UserClass >::FMethodPtr InFunc )
 	{
-		Add( FDelegate::CreateThreadSafeSP( InUserObjectRef, InFunc ) );
+		return Add( FDelegate::CreateThreadSafeSP( InUserObjectRef, InFunc ) );
 	}
 	template< class UserClass >
-	inline void AddThreadSafeSP( const TSharedRef< UserClass, ESPMode::ThreadSafe >& InUserObjectRef, typename FDelegate::template TThreadSafeSPMethodDelegate_Const< UserClass >::FMethodPtr InFunc )
+	inline FDelegateHandle AddThreadSafeSP( const TSharedRef< UserClass, ESPMode::ThreadSafe >& InUserObjectRef, typename FDelegate::template TThreadSafeSPMethodDelegate_Const< UserClass >::FMethodPtr InFunc )
 	{
-		Add( FDelegate::CreateThreadSafeSP( InUserObjectRef, InFunc ) );
+		return Add( FDelegate::CreateThreadSafeSP( InUserObjectRef, InFunc ) );
 	}
 	template< class UserClass, PAYLOAD_TEMPLATE_DECL_OneVar >
-	inline void AddThreadSafeSP( const TSharedRef< UserClass, ESPMode::ThreadSafe >& InUserObjectRef, typename FDelegate::template TThreadSafeSPMethodDelegate_OneVar< UserClass, PAYLOAD_TEMPLATE_LIST_OneVar >::FMethodPtr InFunc, PAYLOAD_TEMPLATE_ARGS_OneVar )
+	inline FDelegateHandle AddThreadSafeSP( const TSharedRef< UserClass, ESPMode::ThreadSafe >& InUserObjectRef, typename FDelegate::template TThreadSafeSPMethodDelegate_OneVar< UserClass, PAYLOAD_TEMPLATE_LIST_OneVar >::FMethodPtr InFunc, PAYLOAD_TEMPLATE_ARGS_OneVar )
 	{
-		Add( FDelegate::CreateThreadSafeSP( InUserObjectRef, InFunc, PAYLOAD_TEMPLATE_PASSIN_OneVar ) );
+		return Add( FDelegate::CreateThreadSafeSP( InUserObjectRef, InFunc, PAYLOAD_TEMPLATE_PASSIN_OneVar ) );
 	}
 	template< class UserClass, PAYLOAD_TEMPLATE_DECL_OneVar >
-	inline void AddThreadSafeSP( const TSharedRef< UserClass, ESPMode::ThreadSafe >& InUserObjectRef, typename FDelegate::template TThreadSafeSPMethodDelegate_OneVar_Const< UserClass, PAYLOAD_TEMPLATE_LIST_OneVar >::FMethodPtr InFunc, PAYLOAD_TEMPLATE_ARGS_OneVar )
+	inline FDelegateHandle AddThreadSafeSP( const TSharedRef< UserClass, ESPMode::ThreadSafe >& InUserObjectRef, typename FDelegate::template TThreadSafeSPMethodDelegate_OneVar_Const< UserClass, PAYLOAD_TEMPLATE_LIST_OneVar >::FMethodPtr InFunc, PAYLOAD_TEMPLATE_ARGS_OneVar )
 	{
-		Add( FDelegate::CreateThreadSafeSP( InUserObjectRef, InFunc, PAYLOAD_TEMPLATE_PASSIN_OneVar ) );
+		return Add( FDelegate::CreateThreadSafeSP( InUserObjectRef, InFunc, PAYLOAD_TEMPLATE_PASSIN_OneVar ) );
 	}
 	template< class UserClass, PAYLOAD_TEMPLATE_DECL_TwoVars >
-	inline void AddThreadSafeSP( const TSharedRef< UserClass, ESPMode::ThreadSafe >& InUserObjectRef, typename FDelegate::template TThreadSafeSPMethodDelegate_TwoVars< UserClass, PAYLOAD_TEMPLATE_LIST_TwoVars >::FMethodPtr InFunc, PAYLOAD_TEMPLATE_ARGS_TwoVars )
+	inline FDelegateHandle AddThreadSafeSP( const TSharedRef< UserClass, ESPMode::ThreadSafe >& InUserObjectRef, typename FDelegate::template TThreadSafeSPMethodDelegate_TwoVars< UserClass, PAYLOAD_TEMPLATE_LIST_TwoVars >::FMethodPtr InFunc, PAYLOAD_TEMPLATE_ARGS_TwoVars )
 	{
-		Add( FDelegate::CreateThreadSafeSP( InUserObjectRef, InFunc, PAYLOAD_TEMPLATE_PASSIN_TwoVars ) );
+		return Add( FDelegate::CreateThreadSafeSP( InUserObjectRef, InFunc, PAYLOAD_TEMPLATE_PASSIN_TwoVars ) );
 	}
 	template< class UserClass, PAYLOAD_TEMPLATE_DECL_TwoVars >
-	inline void AddThreadSafeSP( const TSharedRef< UserClass, ESPMode::ThreadSafe >& InUserObjectRef, typename FDelegate::template TThreadSafeSPMethodDelegate_TwoVars_Const< UserClass, PAYLOAD_TEMPLATE_LIST_TwoVars >::FMethodPtr InFunc, PAYLOAD_TEMPLATE_ARGS_TwoVars )
+	inline FDelegateHandle AddThreadSafeSP( const TSharedRef< UserClass, ESPMode::ThreadSafe >& InUserObjectRef, typename FDelegate::template TThreadSafeSPMethodDelegate_TwoVars_Const< UserClass, PAYLOAD_TEMPLATE_LIST_TwoVars >::FMethodPtr InFunc, PAYLOAD_TEMPLATE_ARGS_TwoVars )
 	{
-		Add( FDelegate::CreateThreadSafeSP( InUserObjectRef, InFunc, PAYLOAD_TEMPLATE_PASSIN_TwoVars ) );
+		return Add( FDelegate::CreateThreadSafeSP( InUserObjectRef, InFunc, PAYLOAD_TEMPLATE_PASSIN_TwoVars ) );
 	}
 	template< class UserClass, PAYLOAD_TEMPLATE_DECL_ThreeVars >
-	inline void AddThreadSafeSP( const TSharedRef< UserClass, ESPMode::ThreadSafe >& InUserObjectRef, typename FDelegate::template TThreadSafeSPMethodDelegate_ThreeVars< UserClass, PAYLOAD_TEMPLATE_LIST_ThreeVars >::FMethodPtr InFunc, PAYLOAD_TEMPLATE_ARGS_ThreeVars )
+	inline FDelegateHandle AddThreadSafeSP( const TSharedRef< UserClass, ESPMode::ThreadSafe >& InUserObjectRef, typename FDelegate::template TThreadSafeSPMethodDelegate_ThreeVars< UserClass, PAYLOAD_TEMPLATE_LIST_ThreeVars >::FMethodPtr InFunc, PAYLOAD_TEMPLATE_ARGS_ThreeVars )
 	{
-		Add( FDelegate::CreateThreadSafeSP( InUserObjectRef, InFunc, PAYLOAD_TEMPLATE_PASSIN_ThreeVars ) );
+		return Add( FDelegate::CreateThreadSafeSP( InUserObjectRef, InFunc, PAYLOAD_TEMPLATE_PASSIN_ThreeVars ) );
 	}
 	template< class UserClass, PAYLOAD_TEMPLATE_DECL_ThreeVars >
-	inline void AddThreadSafeSP( const TSharedRef< UserClass, ESPMode::ThreadSafe >& InUserObjectRef, typename FDelegate::template TThreadSafeSPMethodDelegate_ThreeVars_Const< UserClass, PAYLOAD_TEMPLATE_LIST_ThreeVars >::FMethodPtr InFunc, PAYLOAD_TEMPLATE_ARGS_ThreeVars )
+	inline FDelegateHandle AddThreadSafeSP( const TSharedRef< UserClass, ESPMode::ThreadSafe >& InUserObjectRef, typename FDelegate::template TThreadSafeSPMethodDelegate_ThreeVars_Const< UserClass, PAYLOAD_TEMPLATE_LIST_ThreeVars >::FMethodPtr InFunc, PAYLOAD_TEMPLATE_ARGS_ThreeVars )
 	{
-		Add( FDelegate::CreateThreadSafeSP( InUserObjectRef, InFunc, PAYLOAD_TEMPLATE_PASSIN_ThreeVars ) );
+		return Add( FDelegate::CreateThreadSafeSP( InUserObjectRef, InFunc, PAYLOAD_TEMPLATE_PASSIN_ThreeVars ) );
 	}
 	template< class UserClass, PAYLOAD_TEMPLATE_DECL_FourVars >
-	inline void AddThreadSafeSP( const TSharedRef< UserClass, ESPMode::ThreadSafe >& InUserObjectRef, typename FDelegate::template TThreadSafeSPMethodDelegate_FourVars< UserClass, PAYLOAD_TEMPLATE_LIST_FourVars >::FMethodPtr InFunc, PAYLOAD_TEMPLATE_ARGS_FourVars )
+	inline FDelegateHandle AddThreadSafeSP( const TSharedRef< UserClass, ESPMode::ThreadSafe >& InUserObjectRef, typename FDelegate::template TThreadSafeSPMethodDelegate_FourVars< UserClass, PAYLOAD_TEMPLATE_LIST_FourVars >::FMethodPtr InFunc, PAYLOAD_TEMPLATE_ARGS_FourVars )
 	{
-		Add( FDelegate::CreateThreadSafeSP( InUserObjectRef, InFunc, PAYLOAD_TEMPLATE_PASSIN_FourVars ) );
+		return Add( FDelegate::CreateThreadSafeSP( InUserObjectRef, InFunc, PAYLOAD_TEMPLATE_PASSIN_FourVars ) );
 	}
 	template< class UserClass, PAYLOAD_TEMPLATE_DECL_FourVars >
-	inline void AddThreadSafeSP( const TSharedRef< UserClass, ESPMode::ThreadSafe >& InUserObjectRef, typename FDelegate::template TThreadSafeSPMethodDelegate_FourVars_Const< UserClass, PAYLOAD_TEMPLATE_LIST_FourVars >::FMethodPtr InFunc, PAYLOAD_TEMPLATE_ARGS_FourVars )
+	inline FDelegateHandle AddThreadSafeSP( const TSharedRef< UserClass, ESPMode::ThreadSafe >& InUserObjectRef, typename FDelegate::template TThreadSafeSPMethodDelegate_FourVars_Const< UserClass, PAYLOAD_TEMPLATE_LIST_FourVars >::FMethodPtr InFunc, PAYLOAD_TEMPLATE_ARGS_FourVars )
 	{
-		Add( FDelegate::CreateThreadSafeSP( InUserObjectRef, InFunc, PAYLOAD_TEMPLATE_PASSIN_FourVars ) );
+		return Add( FDelegate::CreateThreadSafeSP( InUserObjectRef, InFunc, PAYLOAD_TEMPLATE_PASSIN_FourVars ) );
 	}
 
 	/**
@@ -1509,54 +1543,54 @@ public:
 	 * @param	InFunc			Class method function address
 	 */
 	template< class UserClass >
-	inline void AddThreadSafeSP( UserClass* InUserObject, typename FDelegate::template TThreadSafeSPMethodDelegate< UserClass >::FMethodPtr InFunc )
+	inline FDelegateHandle AddThreadSafeSP( UserClass* InUserObject, typename FDelegate::template TThreadSafeSPMethodDelegate< UserClass >::FMethodPtr InFunc )
 	{
-		Add( FDelegate::CreateThreadSafeSP( InUserObject, InFunc ) );
+		return Add( FDelegate::CreateThreadSafeSP( InUserObject, InFunc ) );
 	}
 	template< class UserClass >
-	inline void AddThreadSafeSP( UserClass* InUserObject, typename FDelegate::template TThreadSafeSPMethodDelegate_Const< UserClass >::FMethodPtr InFunc )
+	inline FDelegateHandle AddThreadSafeSP( UserClass* InUserObject, typename FDelegate::template TThreadSafeSPMethodDelegate_Const< UserClass >::FMethodPtr InFunc )
 	{
-		Add( FDelegate::CreateThreadSafeSP( InUserObject, InFunc ) );
+		return Add( FDelegate::CreateThreadSafeSP( InUserObject, InFunc ) );
 	}
 	template< class UserClass, PAYLOAD_TEMPLATE_DECL_OneVar >
-	inline void AddThreadSafeSP( UserClass* InUserObject, typename FDelegate::template TThreadSafeSPMethodDelegate_OneVar< UserClass, PAYLOAD_TEMPLATE_LIST_OneVar >::FMethodPtr InFunc, PAYLOAD_TEMPLATE_ARGS_OneVar )
+	inline FDelegateHandle AddThreadSafeSP( UserClass* InUserObject, typename FDelegate::template TThreadSafeSPMethodDelegate_OneVar< UserClass, PAYLOAD_TEMPLATE_LIST_OneVar >::FMethodPtr InFunc, PAYLOAD_TEMPLATE_ARGS_OneVar )
 	{
-		Add( FDelegate::CreateThreadSafeSP( InUserObject, InFunc, PAYLOAD_TEMPLATE_PASSIN_OneVar ) );
+		return Add( FDelegate::CreateThreadSafeSP( InUserObject, InFunc, PAYLOAD_TEMPLATE_PASSIN_OneVar ) );
 	}
 	template< class UserClass, PAYLOAD_TEMPLATE_DECL_OneVar >
-	inline void AddThreadSafeSP( UserClass* InUserObject, typename FDelegate::template TThreadSafeSPMethodDelegate_OneVar_Const< UserClass, PAYLOAD_TEMPLATE_LIST_OneVar >::FMethodPtr InFunc, PAYLOAD_TEMPLATE_ARGS_OneVar )
+	inline FDelegateHandle AddThreadSafeSP( UserClass* InUserObject, typename FDelegate::template TThreadSafeSPMethodDelegate_OneVar_Const< UserClass, PAYLOAD_TEMPLATE_LIST_OneVar >::FMethodPtr InFunc, PAYLOAD_TEMPLATE_ARGS_OneVar )
 	{
-		Add( FDelegate::CreateThreadSafeSP( InUserObject, InFunc, PAYLOAD_TEMPLATE_PASSIN_OneVar ) );
+		return Add( FDelegate::CreateThreadSafeSP( InUserObject, InFunc, PAYLOAD_TEMPLATE_PASSIN_OneVar ) );
 	}
 	template< class UserClass, PAYLOAD_TEMPLATE_DECL_TwoVars >
-	inline void AddThreadSafeSP( UserClass* InUserObject, typename FDelegate::template TThreadSafeSPMethodDelegate_TwoVars< UserClass, PAYLOAD_TEMPLATE_LIST_TwoVars >::FMethodPtr InFunc, PAYLOAD_TEMPLATE_ARGS_TwoVars )
+	inline FDelegateHandle AddThreadSafeSP( UserClass* InUserObject, typename FDelegate::template TThreadSafeSPMethodDelegate_TwoVars< UserClass, PAYLOAD_TEMPLATE_LIST_TwoVars >::FMethodPtr InFunc, PAYLOAD_TEMPLATE_ARGS_TwoVars )
 	{
-		Add( FDelegate::CreateThreadSafeSP( InUserObject, InFunc, PAYLOAD_TEMPLATE_PASSIN_TwoVars ) );
+		return Add( FDelegate::CreateThreadSafeSP( InUserObject, InFunc, PAYLOAD_TEMPLATE_PASSIN_TwoVars ) );
 	}
 	template< class UserClass, PAYLOAD_TEMPLATE_DECL_TwoVars >
-	inline void AddThreadSafeSP( UserClass* InUserObject, typename FDelegate::template TThreadSafeSPMethodDelegate_TwoVars_Const< UserClass, PAYLOAD_TEMPLATE_LIST_TwoVars >::FMethodPtr InFunc, PAYLOAD_TEMPLATE_ARGS_TwoVars )
+	inline FDelegateHandle AddThreadSafeSP( UserClass* InUserObject, typename FDelegate::template TThreadSafeSPMethodDelegate_TwoVars_Const< UserClass, PAYLOAD_TEMPLATE_LIST_TwoVars >::FMethodPtr InFunc, PAYLOAD_TEMPLATE_ARGS_TwoVars )
 	{
-		Add( FDelegate::CreateThreadSafeSP( InUserObject, InFunc, PAYLOAD_TEMPLATE_PASSIN_TwoVars ) );
+		return Add( FDelegate::CreateThreadSafeSP( InUserObject, InFunc, PAYLOAD_TEMPLATE_PASSIN_TwoVars ) );
 	}
 	template< class UserClass, PAYLOAD_TEMPLATE_DECL_ThreeVars >
-	inline void AddThreadSafeSP( UserClass* InUserObject, typename FDelegate::template TThreadSafeSPMethodDelegate_ThreeVars< UserClass, PAYLOAD_TEMPLATE_LIST_ThreeVars >::FMethodPtr InFunc, PAYLOAD_TEMPLATE_ARGS_ThreeVars )
+	inline FDelegateHandle AddThreadSafeSP( UserClass* InUserObject, typename FDelegate::template TThreadSafeSPMethodDelegate_ThreeVars< UserClass, PAYLOAD_TEMPLATE_LIST_ThreeVars >::FMethodPtr InFunc, PAYLOAD_TEMPLATE_ARGS_ThreeVars )
 	{
-		Add( FDelegate::CreateThreadSafeSP( InUserObject, InFunc, PAYLOAD_TEMPLATE_PASSIN_ThreeVars ) );
+		return Add( FDelegate::CreateThreadSafeSP( InUserObject, InFunc, PAYLOAD_TEMPLATE_PASSIN_ThreeVars ) );
 	}
 	template< class UserClass, PAYLOAD_TEMPLATE_DECL_ThreeVars >
-	inline void AddThreadSafeSP( UserClass* InUserObject, typename FDelegate::template TThreadSafeSPMethodDelegate_ThreeVars_Const< UserClass, PAYLOAD_TEMPLATE_LIST_ThreeVars >::FMethodPtr InFunc, PAYLOAD_TEMPLATE_ARGS_ThreeVars )
+	inline FDelegateHandle AddThreadSafeSP( UserClass* InUserObject, typename FDelegate::template TThreadSafeSPMethodDelegate_ThreeVars_Const< UserClass, PAYLOAD_TEMPLATE_LIST_ThreeVars >::FMethodPtr InFunc, PAYLOAD_TEMPLATE_ARGS_ThreeVars )
 	{
-		Add( FDelegate::CreateThreadSafeSP( InUserObject, InFunc, PAYLOAD_TEMPLATE_PASSIN_ThreeVars ) );
+		return Add( FDelegate::CreateThreadSafeSP( InUserObject, InFunc, PAYLOAD_TEMPLATE_PASSIN_ThreeVars ) );
 	}
 	template< class UserClass, PAYLOAD_TEMPLATE_DECL_FourVars >
-	inline void AddThreadSafeSP( UserClass* InUserObject, typename FDelegate::template TThreadSafeSPMethodDelegate_FourVars< UserClass, PAYLOAD_TEMPLATE_LIST_FourVars >::FMethodPtr InFunc, PAYLOAD_TEMPLATE_ARGS_FourVars )
+	inline FDelegateHandle AddThreadSafeSP( UserClass* InUserObject, typename FDelegate::template TThreadSafeSPMethodDelegate_FourVars< UserClass, PAYLOAD_TEMPLATE_LIST_FourVars >::FMethodPtr InFunc, PAYLOAD_TEMPLATE_ARGS_FourVars )
 	{
-		Add( FDelegate::CreateThreadSafeSP( InUserObject, InFunc, PAYLOAD_TEMPLATE_PASSIN_FourVars ) );
+		return Add( FDelegate::CreateThreadSafeSP( InUserObject, InFunc, PAYLOAD_TEMPLATE_PASSIN_FourVars ) );
 	}
 	template< class UserClass, PAYLOAD_TEMPLATE_DECL_FourVars >
-	inline void AddThreadSafeSP( UserClass* InUserObject, typename FDelegate::template TThreadSafeSPMethodDelegate_FourVars_Const< UserClass, PAYLOAD_TEMPLATE_LIST_FourVars >::FMethodPtr InFunc, PAYLOAD_TEMPLATE_ARGS_FourVars )
+	inline FDelegateHandle AddThreadSafeSP( UserClass* InUserObject, typename FDelegate::template TThreadSafeSPMethodDelegate_FourVars_Const< UserClass, PAYLOAD_TEMPLATE_LIST_FourVars >::FMethodPtr InFunc, PAYLOAD_TEMPLATE_ARGS_FourVars )
 	{
-		Add( FDelegate::CreateThreadSafeSP( InUserObject, InFunc, PAYLOAD_TEMPLATE_PASSIN_FourVars ) );
+		return Add( FDelegate::CreateThreadSafeSP( InUserObject, InFunc, PAYLOAD_TEMPLATE_PASSIN_FourVars ) );
 	}
 
 	/**
@@ -1568,29 +1602,29 @@ public:
 	 * @param	InFunctionName			Class method function address
 	 */
 	template< class UObjectTemplate >
-	inline void AddUFunction( UObjectTemplate* InUserObject, const FName& InFunctionName )
+	inline FDelegateHandle AddUFunction( UObjectTemplate* InUserObject, const FName& InFunctionName )
 	{
-		Add( FDelegate::CreateUFunction( InUserObject, InFunctionName ) );
+		return Add( FDelegate::CreateUFunction( InUserObject, InFunctionName ) );
 	}
 	template< class UObjectTemplate, PAYLOAD_TEMPLATE_DECL_OneVar >
-	inline void AddUFunction( UObjectTemplate* InUserObject, const FName& InFunctionName, PAYLOAD_TEMPLATE_ARGS_OneVar )
+	inline FDelegateHandle AddUFunction( UObjectTemplate* InUserObject, const FName& InFunctionName, PAYLOAD_TEMPLATE_ARGS_OneVar )
 	{
-		Add( FDelegate::CreateUFunction( InUserObject, InFunctionName, PAYLOAD_TEMPLATE_PASSIN_OneVar ) );
+		return Add( FDelegate::CreateUFunction( InUserObject, InFunctionName, PAYLOAD_TEMPLATE_PASSIN_OneVar ) );
 	}
 	template< class UObjectTemplate, PAYLOAD_TEMPLATE_DECL_TwoVars >
-	inline void AddUFunction( UObjectTemplate* InUserObject, const FName& InFunctionName, PAYLOAD_TEMPLATE_ARGS_TwoVars )
+	inline FDelegateHandle AddUFunction( UObjectTemplate* InUserObject, const FName& InFunctionName, PAYLOAD_TEMPLATE_ARGS_TwoVars )
 	{
-		Add( FDelegate::CreateUFunction( InUserObject, InFunctionName, PAYLOAD_TEMPLATE_PASSIN_TwoVars ) );
+		return Add( FDelegate::CreateUFunction( InUserObject, InFunctionName, PAYLOAD_TEMPLATE_PASSIN_TwoVars ) );
 	}
 	template< class UObjectTemplate, PAYLOAD_TEMPLATE_DECL_ThreeVars >
-	inline void AddUFunction( UObjectTemplate* InUserObject, const FName& InFunctionName, PAYLOAD_TEMPLATE_ARGS_ThreeVars )
+	inline FDelegateHandle AddUFunction( UObjectTemplate* InUserObject, const FName& InFunctionName, PAYLOAD_TEMPLATE_ARGS_ThreeVars )
 	{
-		Add( FDelegate::CreateUFunction( InUserObject, InFunctionName, PAYLOAD_TEMPLATE_PASSIN_ThreeVars ) );
+		return Add( FDelegate::CreateUFunction( InUserObject, InFunctionName, PAYLOAD_TEMPLATE_PASSIN_ThreeVars ) );
 	}
 	template< class UObjectTemplate, PAYLOAD_TEMPLATE_DECL_FourVars >
-	inline void AddUFunction( UObjectTemplate* InUserObject, const FName& InFunctionName, PAYLOAD_TEMPLATE_ARGS_FourVars )
+	inline FDelegateHandle AddUFunction( UObjectTemplate* InUserObject, const FName& InFunctionName, PAYLOAD_TEMPLATE_ARGS_FourVars )
 	{
-		Add( FDelegate::CreateUFunction( InUserObject, InFunctionName, PAYLOAD_TEMPLATE_PASSIN_FourVars ) );
+		return Add( FDelegate::CreateUFunction( InUserObject, InFunctionName, PAYLOAD_TEMPLATE_PASSIN_FourVars ) );
 	}
 
 	/**
@@ -1602,58 +1636,70 @@ public:
 	 * @param	InFunc			Class method function address
 	 */
 	template< class UserClass >
-	inline void AddUObject( UserClass* InUserObject, typename FDelegate::template TUObjectMethodDelegate< UserClass >::FMethodPtr InFunc )
+	inline FDelegateHandle AddUObject( UserClass* InUserObject, typename FDelegate::template TUObjectMethodDelegate< UserClass >::FMethodPtr InFunc )
 	{
-		Add( FDelegate::CreateUObject( InUserObject, InFunc ) );
+		return Add( FDelegate::CreateUObject( InUserObject, InFunc ) );
 	}
 	template< class UserClass >
-	inline void AddUObject( UserClass* InUserObject, typename FDelegate::template TUObjectMethodDelegate_Const< UserClass >::FMethodPtr InFunc )
+	inline FDelegateHandle AddUObject( UserClass* InUserObject, typename FDelegate::template TUObjectMethodDelegate_Const< UserClass >::FMethodPtr InFunc )
 	{
-		Add( FDelegate::CreateUObject( InUserObject, InFunc ) );
+		return Add( FDelegate::CreateUObject( InUserObject, InFunc ) );
 	}
 	template< class UserClass, PAYLOAD_TEMPLATE_DECL_OneVar >
-	inline void AddUObject( UserClass* InUserObject, typename FDelegate::template TUObjectMethodDelegate_OneVar< UserClass, PAYLOAD_TEMPLATE_LIST_OneVar >::FMethodPtr InFunc, PAYLOAD_TEMPLATE_ARGS_OneVar )
+	inline FDelegateHandle AddUObject( UserClass* InUserObject, typename FDelegate::template TUObjectMethodDelegate_OneVar< UserClass, PAYLOAD_TEMPLATE_LIST_OneVar >::FMethodPtr InFunc, PAYLOAD_TEMPLATE_ARGS_OneVar )
 	{
-		Add( FDelegate::CreateUObject( InUserObject, InFunc, PAYLOAD_TEMPLATE_PASSIN_OneVar ) );
+		return Add( FDelegate::CreateUObject( InUserObject, InFunc, PAYLOAD_TEMPLATE_PASSIN_OneVar ) );
 	}
 	template< class UserClass, PAYLOAD_TEMPLATE_DECL_OneVar >
-	inline void AddUObject( UserClass* InUserObject, typename FDelegate::template TUObjectMethodDelegate_OneVar_Const< UserClass, PAYLOAD_TEMPLATE_LIST_OneVar >::FMethodPtr InFunc, PAYLOAD_TEMPLATE_ARGS_OneVar )
+	inline FDelegateHandle AddUObject( UserClass* InUserObject, typename FDelegate::template TUObjectMethodDelegate_OneVar_Const< UserClass, PAYLOAD_TEMPLATE_LIST_OneVar >::FMethodPtr InFunc, PAYLOAD_TEMPLATE_ARGS_OneVar )
 	{
-		Add( FDelegate::CreateUObject( InUserObject, InFunc, PAYLOAD_TEMPLATE_PASSIN_OneVar ) );
+		return Add( FDelegate::CreateUObject( InUserObject, InFunc, PAYLOAD_TEMPLATE_PASSIN_OneVar ) );
 	}
 	template< class UserClass, PAYLOAD_TEMPLATE_DECL_TwoVars >
-	inline void AddUObject( UserClass* InUserObject, typename FDelegate::template TUObjectMethodDelegate_TwoVars< UserClass, PAYLOAD_TEMPLATE_LIST_TwoVars >::FMethodPtr InFunc, PAYLOAD_TEMPLATE_ARGS_TwoVars )
+	inline FDelegateHandle AddUObject( UserClass* InUserObject, typename FDelegate::template TUObjectMethodDelegate_TwoVars< UserClass, PAYLOAD_TEMPLATE_LIST_TwoVars >::FMethodPtr InFunc, PAYLOAD_TEMPLATE_ARGS_TwoVars )
 	{
-		Add( FDelegate::CreateUObject( InUserObject, InFunc, PAYLOAD_TEMPLATE_PASSIN_TwoVars ) );
+		return Add( FDelegate::CreateUObject( InUserObject, InFunc, PAYLOAD_TEMPLATE_PASSIN_TwoVars ) );
 	}
 	template< class UserClass, PAYLOAD_TEMPLATE_DECL_TwoVars >
-	inline void AddUObject( UserClass* InUserObject, typename FDelegate::template TUObjectMethodDelegate_TwoVars_Const< UserClass, PAYLOAD_TEMPLATE_LIST_TwoVars >::FMethodPtr InFunc, PAYLOAD_TEMPLATE_ARGS_TwoVars )
+	inline FDelegateHandle AddUObject( UserClass* InUserObject, typename FDelegate::template TUObjectMethodDelegate_TwoVars_Const< UserClass, PAYLOAD_TEMPLATE_LIST_TwoVars >::FMethodPtr InFunc, PAYLOAD_TEMPLATE_ARGS_TwoVars )
 	{
-		Add( FDelegate::CreateUObject( InUserObject, InFunc, PAYLOAD_TEMPLATE_PASSIN_TwoVars ) );
+		return Add( FDelegate::CreateUObject( InUserObject, InFunc, PAYLOAD_TEMPLATE_PASSIN_TwoVars ) );
 	}
 	template< class UserClass, PAYLOAD_TEMPLATE_DECL_ThreeVars >
-	inline void AddUObject( UserClass* InUserObject, typename FDelegate::template TUObjectMethodDelegate_ThreeVars< UserClass, PAYLOAD_TEMPLATE_LIST_ThreeVars >::FMethodPtr InFunc, PAYLOAD_TEMPLATE_ARGS_ThreeVars )
+	inline FDelegateHandle AddUObject( UserClass* InUserObject, typename FDelegate::template TUObjectMethodDelegate_ThreeVars< UserClass, PAYLOAD_TEMPLATE_LIST_ThreeVars >::FMethodPtr InFunc, PAYLOAD_TEMPLATE_ARGS_ThreeVars )
 	{
-		Add( FDelegate::CreateUObject( InUserObject, InFunc, PAYLOAD_TEMPLATE_PASSIN_ThreeVars ) );
+		return Add( FDelegate::CreateUObject( InUserObject, InFunc, PAYLOAD_TEMPLATE_PASSIN_ThreeVars ) );
 	}
 	template< class UserClass, PAYLOAD_TEMPLATE_DECL_ThreeVars >
-	inline void AddUObject( UserClass* InUserObject, typename FDelegate::template TUObjectMethodDelegate_ThreeVars_Const< UserClass, PAYLOAD_TEMPLATE_LIST_ThreeVars >::FMethodPtr InFunc, PAYLOAD_TEMPLATE_ARGS_ThreeVars )
+	inline FDelegateHandle AddUObject( UserClass* InUserObject, typename FDelegate::template TUObjectMethodDelegate_ThreeVars_Const< UserClass, PAYLOAD_TEMPLATE_LIST_ThreeVars >::FMethodPtr InFunc, PAYLOAD_TEMPLATE_ARGS_ThreeVars )
 	{
-		Add( FDelegate::CreateUObject( InUserObject, InFunc, PAYLOAD_TEMPLATE_PASSIN_ThreeVars ) );
+		return Add( FDelegate::CreateUObject( InUserObject, InFunc, PAYLOAD_TEMPLATE_PASSIN_ThreeVars ) );
 	}
 	template< class UserClass, PAYLOAD_TEMPLATE_DECL_FourVars >
-	inline void AddUObject( UserClass* InUserObject, typename FDelegate::template TUObjectMethodDelegate_FourVars< UserClass, PAYLOAD_TEMPLATE_LIST_FourVars >::FMethodPtr InFunc, PAYLOAD_TEMPLATE_ARGS_FourVars )
+	inline FDelegateHandle AddUObject( UserClass* InUserObject, typename FDelegate::template TUObjectMethodDelegate_FourVars< UserClass, PAYLOAD_TEMPLATE_LIST_FourVars >::FMethodPtr InFunc, PAYLOAD_TEMPLATE_ARGS_FourVars )
 	{
-		Add( FDelegate::CreateUObject( InUserObject, InFunc, PAYLOAD_TEMPLATE_PASSIN_FourVars ) );
+		return Add( FDelegate::CreateUObject( InUserObject, InFunc, PAYLOAD_TEMPLATE_PASSIN_FourVars ) );
 	}
 	template< class UserClass, PAYLOAD_TEMPLATE_DECL_FourVars >
-	inline void AddUObject( UserClass* InUserObject, typename FDelegate::template TUObjectMethodDelegate_FourVars_Const< UserClass, PAYLOAD_TEMPLATE_LIST_FourVars >::FMethodPtr InFunc, PAYLOAD_TEMPLATE_ARGS_FourVars )
+	inline FDelegateHandle AddUObject( UserClass* InUserObject, typename FDelegate::template TUObjectMethodDelegate_FourVars_Const< UserClass, PAYLOAD_TEMPLATE_LIST_FourVars >::FMethodPtr InFunc, PAYLOAD_TEMPLATE_ARGS_FourVars )
 	{
-		Add( FDelegate::CreateUObject( InUserObject, InFunc, PAYLOAD_TEMPLATE_PASSIN_FourVars ) );
+		return Add( FDelegate::CreateUObject( InUserObject, InFunc, PAYLOAD_TEMPLATE_PASSIN_FourVars ) );
 	}
 
 public:
 
+	/**
+	 * Removes a delegate instance from this multi-cast delegate's invocation list (performance is O(N)).
+	 *
+	 * Note that the order of the delegate instances may not be preserved!
+	 *
+	 * @param Handle The handle of the delegate instance to remove.
+	 */
+	void Remove( FDelegateHandle Handle )
+	{
+		RemoveDelegateInstance(Handle);
+	}
+
 	/**
 	 * Removes a delegate instance from this multi-cast delegate's invocation list (performance is O(N)).
 	 *
@@ -1661,14 +1707,18 @@ public:
 	 *
 	 * @param DelegateInstance The delegate instance to remove.
 	 */
+	DELEGATE_DEPRECATED("This Remove overload is deprecated - please remove delegates using the FDelegateHandle returned by the Add function.")
 	void Remove( const TDelegateInstanceInterface& DelegateInstance )
 	{
-		// NOTE: For delegates with PAYLOAD, the payload variables are not considered when comparing delegate objects!
-		// This means that two delegate instances bound to the same object/method, but with different payload data,
-		// cannot be bound to the same single multi-cast delegate.  We could support this by changing IsSameFunction()
-		// to do comparison of payload data, but there are some trade-offs if we do that.
+		DEPRECATED_RemoveDelegateInstance(&DelegateInstance);
+	}
 
-		RemoveDelegateInstance(&DelegateInstance);
+	void DEPRECATED_Remove( const TDelegateInstanceInterface& DelegateInstance )
+	{
+		// Provided to help implement other deprecated functions without giving multiple warnings.
+		// Should not be called directly.
+
+		DEPRECATED_RemoveDelegateInstance(&DelegateInstance);
 	}
 
 	/**
@@ -1679,18 +1729,22 @@ public:
 	 *
 	 * @param Delegate The delegate to remove.
 	 */
+	DELEGATE_DEPRECATED("This Remove overload is deprecated - please remove delegates using the FDelegateHandle returned by the Add function.")
 	void Remove( const FDelegate& Delegate )
 	{
-		// NOTE: For delegates with PAYLOAD, the payload variables are not considered when comparing delegate objects!
-		// This means that two delegate instances bound to the same object/method, but with different payload data,
-		// cannot be bound to the same single multi-cast delegate.  We could support this by changing IsSameFunction()
-		// to do comparison of payload data, but there are some trade-offs if we do that.
+		DEPRECATED_Remove(Delegate);
+	}
+
+	void DEPRECATED_Remove( const FDelegate& Delegate )
+	{
+		// Provided to help implement other deprecated functions without giving multiple warnings.
+		// Should not be called directly.
 
 		TDelegateInstanceInterface* DelegateInstance = (TDelegateInstanceInterface*)Delegate.GetDelegateInstance();
 
 		if (DelegateInstance != nullptr)
 		{
-			RemoveDelegateInstance(DelegateInstance);
+			DEPRECATED_RemoveDelegateInstance(DelegateInstance);
 		}		
 	}
 
@@ -1703,9 +1757,10 @@ public:
 	 *
 	 * @param	InFunc	Function pointer
 	 */
+	DELEGATE_DEPRECATED("RemoveStatic is deprecated - please remove delegates using the FDelegateHandle returned by the Add function.")
 	inline void RemoveStatic( typename FDelegate::FStaticDelegate::FFuncPtr InFunc )
 	{
-		Remove( FDelegate::CreateStatic( InFunc ) );
+		DEPRECATED_RemoveDelegateInstance( (TDelegateInstanceInterface*)FDelegate::CreateStatic( InFunc ).GetDelegateInstance() );
 	}
 
 	/**
@@ -1716,14 +1771,16 @@ public:
 	 * @param	InFunc			Class method function address
 	 */
 	template< class UserClass >
+	DELEGATE_DEPRECATED("RemoveRaw is deprecated - please remove delegates using the FDelegateHandle returned by the Add function.")
 	inline void RemoveRaw( UserClass* InUserObject, typename FDelegate::template TRawMethodDelegate< UserClass >::FMethodPtr InFunc )
 	{
-		Remove( FDelegate::CreateRaw( InUserObject, InFunc ) );
+		DEPRECATED_RemoveDelegateInstance( (TDelegateInstanceInterface*)FDelegate::CreateRaw( InUserObject, InFunc ).GetDelegateInstance() );
 	}
 	template< class UserClass >
+	DELEGATE_DEPRECATED("RemoveRaw is deprecated - please remove delegates using the FDelegateHandle returned by the Add function.")
 	inline void RemoveRaw( UserClass* InUserObject, typename FDelegate::template TRawMethodDelegate_Const< UserClass >::FMethodPtr InFunc )
 	{
-		Remove( FDelegate::CreateRaw( InUserObject, InFunc ) );
+		DEPRECATED_RemoveDelegateInstance( (TDelegateInstanceInterface*)FDelegate::CreateRaw( InUserObject, InFunc ).GetDelegateInstance() );
 	}
 
 	/**
@@ -1734,14 +1791,16 @@ public:
 	 * @param	InFunc			Class method function address
 	 */
 	template< class UserClass >
+	DELEGATE_DEPRECATED("RemoveSP is deprecated - please remove delegates using the FDelegateHandle returned by the Add function.")
 	inline void RemoveSP( const TSharedRef< UserClass, ESPMode::Fast >& InUserObjectRef, typename FDelegate::template TSPMethodDelegate< UserClass >::FMethodPtr InFunc )
 	{
-		Remove( FDelegate::CreateSP( InUserObjectRef, InFunc ) );
+		DEPRECATED_RemoveDelegateInstance( (TDelegateInstanceInterface*)FDelegate::CreateSP( InUserObjectRef, InFunc ).GetDelegateInstance() );
 	}
 	template< class UserClass >
+	DELEGATE_DEPRECATED("RemoveSP is deprecated - please remove delegates using the FDelegateHandle returned by the Add function.")
 	inline void RemoveSP( const TSharedRef< UserClass, ESPMode::Fast >& InUserObjectRef, typename FDelegate::template TSPMethodDelegate_Const< UserClass >::FMethodPtr InFunc )
 	{
-		Remove( FDelegate::CreateSP( InUserObjectRef, InFunc ) );
+		DEPRECATED_RemoveDelegateInstance( (TDelegateInstanceInterface*)FDelegate::CreateSP( InUserObjectRef, InFunc ).GetDelegateInstance() );
 	}
 
 	/**
@@ -1752,14 +1811,16 @@ public:
 	 * @param	InFunc			Class method function address
 	 */
 	template< class UserClass >
+	DELEGATE_DEPRECATED("RemoveSP is deprecated - please remove delegates using the FDelegateHandle returned by the Add function.")
 	inline void RemoveSP( UserClass* InUserObject, typename FDelegate::template TSPMethodDelegate< UserClass >::FMethodPtr InFunc )
 	{
-		Remove( FDelegate::CreateSP( InUserObject, InFunc ) );
+		DEPRECATED_RemoveDelegateInstance( (TDelegateInstanceInterface*)FDelegate::CreateSP( InUserObject, InFunc ).GetDelegateInstance() );
 	}
 	template< class UserClass >
+	DELEGATE_DEPRECATED("RemoveSP is deprecated - please remove delegates using the FDelegateHandle returned by the Add function.")
 	inline void RemoveSP( UserClass* InUserObject, typename FDelegate::template TSPMethodDelegate_Const< UserClass >::FMethodPtr InFunc )
 	{
-		Remove( FDelegate::CreateSP( InUserObject, InFunc ) );
+		DEPRECATED_RemoveDelegateInstance( (TDelegateInstanceInterface*)FDelegate::CreateSP( InUserObject, InFunc ).GetDelegateInstance() );
 	}
 
 	/**
@@ -1770,14 +1831,16 @@ public:
 	 * @param	InFunc			Class method function address
 	 */
 	template< class UserClass >
+	DELEGATE_DEPRECATED("RemoveThreadSafeSP is deprecated - please remove delegates using the FDelegateHandle returned by the Add function.")
 	inline void RemoveThreadSafeSP( const TSharedRef< UserClass, ESPMode::ThreadSafe >& InUserObjectRef, typename FDelegate::template TThreadSafeSPMethodDelegate< UserClass >::FMethodPtr InFunc )
 	{
-		Remove( FDelegate::CreateThreadSafeSP( InUserObjectRef, InFunc ) );
+		DEPRECATED_RemoveDelegateInstance( (TDelegateInstanceInterface*)FDelegate::CreateThreadSafeSP( InUserObjectRef, InFunc ).GetDelegateInstance() );
 	}
 	template< class UserClass >
+	DELEGATE_DEPRECATED("RemoveThreadSafeSP is deprecated - please remove delegates using the FDelegateHandle returned by the Add function.")
 	inline void RemoveThreadSafeSP( const TSharedRef< UserClass, ESPMode::ThreadSafe >& InUserObjectRef, typename FDelegate::template TThreadSafeSPMethodDelegate_Const< UserClass >::FMethodPtr InFunc )
 	{
-		Remove( FDelegate::CreateThreadSafeSP( InUserObjectRef, InFunc ) );
+		DEPRECATED_RemoveDelegateInstance( (TDelegateInstanceInterface*)FDelegate::CreateThreadSafeSP( InUserObjectRef, InFunc ).GetDelegateInstance() );
 	}
 
 	/**
@@ -1788,14 +1851,16 @@ public:
 	 * @param	InFunc			Class method function address
 	 */
 	template< class UserClass >
+	DELEGATE_DEPRECATED("RemoveThreadSafeSP is deprecated - please remove delegates using the FDelegateHandle returned by the Add function.")
 	inline void RemoveThreadSafeSP( UserClass* InUserObject, typename FDelegate::template TThreadSafeSPMethodDelegate< UserClass >::FMethodPtr InFunc )
 	{
-		Remove( FDelegate::CreateThreadSafeSP( InUserObject, InFunc ) );
+		DEPRECATED_RemoveDelegateInstance( (TDelegateInstanceInterface*)FDelegate::CreateThreadSafeSP( InUserObject, InFunc ).GetDelegateInstance() );
 	}
 	template< class UserClass >
+	DELEGATE_DEPRECATED("RemoveThreadSafeSP is deprecated - please remove delegates using the FDelegateHandle returned by the Add function.")
 	inline void RemoveThreadSafeSP( UserClass* InUserObject, typename FDelegate::template TThreadSafeSPMethodDelegate_Const< UserClass >::FMethodPtr InFunc )
 	{
-		Remove( FDelegate::CreateThreadSafeSP( InUserObject, InFunc ) );
+		DEPRECATED_RemoveDelegateInstance( (TDelegateInstanceInterface*)FDelegate::CreateThreadSafeSP( InUserObject, InFunc ).GetDelegateInstance() );
 	}
 
 	/**
@@ -1806,9 +1871,10 @@ public:
 	 * @param	InFunc			Class method function address
 	 */
 	template< class UObjectTemplate >
+	DELEGATE_DEPRECATED("RemoveUFunction is deprecated - please remove delegates using the FDelegateHandle returned by the Add function.")
 	inline void RemoveUFunction( UObjectTemplate* InUserObject, const FName& InFunctionName )
 	{
-		Remove( FDelegate::CreateUFunction( InUserObject, InFunctionName ) );
+		DEPRECATED_RemoveDelegateInstance( (TDelegateInstanceInterface*)FDelegate::CreateUFunction( InUserObject, InFunctionName ).GetDelegateInstance() );
 	}
 
 	/**
@@ -1819,14 +1885,16 @@ public:
 	 * @param	InFunc			Class method function address
 	 */
 	template< class UserClass >
+	DELEGATE_DEPRECATED("RemoveUObject is deprecated - please remove delegates using the FDelegateHandle returned by the Add function.")
 	inline void RemoveUObject( UserClass* InUserObject, typename FDelegate::template TUObjectMethodDelegate< UserClass >::FMethodPtr InFunc )
 	{
-		Remove( FDelegate::CreateUObject( InUserObject, InFunc ) );
+		DEPRECATED_RemoveDelegateInstance( (TDelegateInstanceInterface*)FDelegate::CreateUObject( InUserObject, InFunc ).GetDelegateInstance() );
 	}
 	template< class UserClass >
+	DELEGATE_DEPRECATED("RemoveUObject is deprecated - please remove delegates using the FDelegateHandle returned by the Add function.")
 	inline void RemoveUObject( UserClass* InUserObject, typename FDelegate::template TUObjectMethodDelegate_Const< UserClass >::FMethodPtr InFunc )
 	{
-		Remove( FDelegate::CreateUObject( InUserObject, InFunc ) );
+		DEPRECATED_RemoveDelegateInstance( (TDelegateInstanceInterface*)FDelegate::CreateUObject( InUserObject, InFunc ).GetDelegateInstance() );
 	}
 
 protected:
@@ -1882,7 +1950,7 @@ protected:
 	 *
 	 * @param	InDelegate	Delegate to add
 	 */
-	void AddDelegateInstance( TDelegateInstanceInterface* InDelegateInstance )
+	FDelegateHandle AddDelegateInstance( TDelegateInstanceInterface* InDelegateInstance )
 	{
 #if !UE_BUILD_SHIPPING && !UE_BUILD_TEST
 		// verify that the same function isn't already bound
@@ -1897,7 +1965,7 @@ protected:
 		}
 #endif
 
-		AddInternal(InDelegateInstance);
+		return AddInternal(InDelegateInstance);
 	}
 
 	/**
@@ -1933,6 +2001,41 @@ protected:
 		}
 	}
 
+	/**
+	 * Removes a function from this multi-cast delegate's invocation list (performance is O(N)).
+	 *
+	 * The function is not actually removed, but deleted and marked as removed.
+	 * It will be removed next time the invocation list is compacted within Broadcast().
+	 *
+	 * @param Handle The handle of the delegate instance to remove.
+	 */
+	void RemoveDelegateInstance( FDelegateHandle Handle )
+	{
+		const TArray<IDelegateInstance*>& InvocationList = GetInvocationList();
+
+		// NOTE: We assume that this method is never called with a nullptr object, in which case the
+		//       the following algorithm would break down (it would remove the first found instance
+		//       of a matching function binding, which is not necessarily the instance we wish to remove).
+
+		for (int32 InvocationListIndex = 0; InvocationListIndex < InvocationList.Num(); ++InvocationListIndex)
+		{
+			// this down-cast is OK! allows for managing invocation list in the base class without requiring virtual functions
+			TDelegateInstanceInterface*& DelegateInstanceRef = (TDelegateInstanceInterface*&)InvocationList[InvocationListIndex];
+
+			// NOTE: We must do a deep compare here, not just compare delegate pointers, because multiple
+			//       delegate pointers can refer to the exact same object and method
+			if ((DelegateInstanceRef != nullptr) && DelegateInstanceRef->GetHandle() == Handle)
+			{
+				delete DelegateInstanceRef;
+				DelegateInstanceRef = nullptr;
+
+				break;	// no need to continue, as we never allow the same delegate to be bound twice
+			}
+		}
+
+		const_cast<BASE_MULTICAST_DELEGATE_CLASS*>(this)->CompactInvocationList();
+	}
+
 	/**
 	 * Removes a function from this multi-cast delegate's invocation list (performance is O(N)).
 	 *
@@ -1941,7 +2044,7 @@ protected:
 	 *
 	 * @param InDelegateInstance The delegate instance to remove.
 	 */
-	void RemoveDelegateInstance( const TDelegateInstanceInterface* InDelegateInstance )
+	void DEPRECATED_RemoveDelegateInstance( const TDelegateInstanceInterface* InDelegateInstance )
 	{
 		const TArray<IDelegateInstance*>& InvocationList = GetInvocationList();
 
