Delegate comparisons deprecated, lots of other associated code deprecated, and lots of warning fixups:

diff --git a/Engine/Source/Runtime/Core/Public/Delegates/DelegateSignatureImpl_Variadics.inl b/Engine/Source/Runtime/Core/Public/Delegates/DelegateSignatureImpl_Variadics.inl
--- a/Engine/Source/Runtime/Core/Public/Delegates/DelegateSignatureImpl_Variadics.inl
+++ b/Engine/Source/Runtime/Core/Public/Delegates/DelegateSignatureImpl_Variadics.inl
@@ -444,23 +444,34 @@ public:
 	 *
 	 * @return true if this delegate equals the other, false otherwise.
 	 */
+	DELEGATE_DEPRECATED("Delegate comparison is deprecated - please replace any usage with comparison of FDelegateHandles.")
 	bool operator==( const TBaseDelegate& Other ) const
+	{
+		return DEPRECATED_Compare(Other);
+	}
+
+	/**
+	 * Delegate comparison operator.
+	 *
+	 * @return true if this delegate equals the other, false otherwise.
+	 */
+	bool DEPRECATED_Compare( const TBaseDelegate& Other ) const
 	{
 		TDelegateInstanceInterface* DelegateInstance = (TDelegateInstanceInterface*)GetDelegateInstance();
 		TDelegateInstanceInterface* OtherInstance = (TDelegateInstanceInterface*)Other.GetDelegateInstance();
-
+		
 		// The function these delegates point to must be the same
 		if ((DelegateInstance != nullptr) && (OtherInstance != nullptr))
 		{
 			return DelegateInstance->IsSameFunction(*OtherInstance);
 		}
-
+		
 		// If neither delegate is initialized to anything yet, then we treat them as equal
 		if ((DelegateInstance == nullptr) && (OtherInstance == nullptr))
 		{
 			return true;
 		}
-
+		
 		// No match!
 		return false;
 	}
@@ -544,12 +555,16 @@ public:
 	 *
 	 * @param DelegateInstance The delegate instance to add.
 	 */
-	void Add( TDelegateInstanceInterface* DelegateInstance )
+	FDelegateHandle Add( TDelegateInstanceInterface* DelegateInstance )
 	{
+		FDelegateHandle Result;
+
 		if (DelegateInstance != nullptr)
 		{
-			AddDelegateInstance(DelegateInstance);
+			Result = AddDelegateInstance(DelegateInstance);
 		}
+
+		return Result;
 	}
 
 	/**
@@ -559,14 +574,17 @@ public:
 	 *
 	 * @param Delegate The delegate to add.
 	 */
-	void Add( const FDelegate& Delegate )
+	FDelegateHandle Add( const FDelegate& Delegate )
 	{
+		FDelegateHandle             Result;
 		TDelegateInstanceInterface* DelegateInstance = (TDelegateInstanceInterface*)Delegate.GetDelegateInstance();
 
 		if (DelegateInstance != nullptr)
 		{
-			AddDelegateInstance(DelegateInstance->CreateCopy());
+			Result = AddDelegateInstance(DelegateInstance->CreateCopy());
 		}
+
+		return Result;
 	}
 
 	/**
@@ -575,9 +593,9 @@ public:
 	 * @param	InFunc	Function pointer
 	 */
 	template <typename... VarTypes>
-	inline void AddStatic(typename TBaseStaticDelegateInstance<void (ParamTypes...), VarTypes...>::FFuncPtr InFunc, VarTypes... Vars)
+	inline FDelegateHandle AddStatic(typename TBaseStaticDelegateInstance<void (ParamTypes...), VarTypes...>::FFuncPtr InFunc, VarTypes... Vars)
 	{
-		Add(FDelegate::CreateStatic(InFunc, Vars...));
+		return Add(FDelegate::CreateStatic(InFunc, Vars...));
 	}
 
 	/**
@@ -590,14 +608,14 @@ public:
 	 * @param	InFunc			Class method function address
 	 */
 	template <typename UserClass, typename... VarTypes>
-	inline void AddRaw(UserClass* InUserObject, typename TMemFunPtrType<false, UserClass, void (ParamTypes..., VarTypes...)>::Type InFunc, VarTypes... Vars)
+	inline FDelegateHandle AddRaw(UserClass* InUserObject, typename TMemFunPtrType<false, UserClass, void (ParamTypes..., VarTypes...)>::Type InFunc, VarTypes... Vars)
 	{
-		Add(FDelegate::CreateRaw(InUserObject, InFunc, Vars...));
+		return Add(FDelegate::CreateRaw(InUserObject, InFunc, Vars...));
 	}
 	template <typename UserClass, typename... VarTypes>
-	inline void AddRaw(UserClass* InUserObject, typename TMemFunPtrType<true, UserClass, void (ParamTypes..., VarTypes...)>::Type InFunc, VarTypes... Vars)
+	inline FDelegateHandle AddRaw(UserClass* InUserObject, typename TMemFunPtrType<true, UserClass, void (ParamTypes..., VarTypes...)>::Type InFunc, VarTypes... Vars)
 	{
-		Add(FDelegate::CreateRaw(InUserObject, InFunc, Vars...));
+		return Add(FDelegate::CreateRaw(InUserObject, InFunc, Vars...));
 	}
 
 	/**
@@ -609,14 +627,14 @@ public:
 	 * @param	InFunc			Class method function address
 	 */
 	template <typename UserClass, typename... VarTypes>
-	inline void AddSP(const TSharedRef<UserClass, ESPMode::Fast>& InUserObjectRef, typename TMemFunPtrType<false, UserClass, void (ParamTypes..., VarTypes...)>::Type InFunc, VarTypes... Vars)
+	inline FDelegateHandle AddSP(const TSharedRef<UserClass, ESPMode::Fast>& InUserObjectRef, typename TMemFunPtrType<false, UserClass, void (ParamTypes..., VarTypes...)>::Type InFunc, VarTypes... Vars)
 	{
-		Add(FDelegate::CreateSP(InUserObjectRef, InFunc, Vars...));
+		return Add(FDelegate::CreateSP(InUserObjectRef, InFunc, Vars...));
 	}
 	template <typename UserClass, typename... VarTypes>
-	inline void AddSP(const TSharedRef<UserClass, ESPMode::Fast>& InUserObjectRef, typename TMemFunPtrType<true, UserClass, void (ParamTypes..., VarTypes...)>::Type InFunc, VarTypes... Vars)
+	inline FDelegateHandle AddSP(const TSharedRef<UserClass, ESPMode::Fast>& InUserObjectRef, typename TMemFunPtrType<true, UserClass, void (ParamTypes..., VarTypes...)>::Type InFunc, VarTypes... Vars)
 	{
-		Add(FDelegate::CreateSP(InUserObjectRef, InFunc, Vars...));
+		return Add(FDelegate::CreateSP(InUserObjectRef, InFunc, Vars...));
 	}
 
 	/**
@@ -628,14 +646,14 @@ public:
 	 * @param	InFunc			Class method function address
 	 */
 	template <typename UserClass, typename... VarTypes>
-	inline void AddSP(UserClass* InUserObject, typename TMemFunPtrType<false, UserClass, void (ParamTypes..., VarTypes...)>::Type InFunc, VarTypes... Vars)
+	inline FDelegateHandle AddSP(UserClass* InUserObject, typename TMemFunPtrType<false, UserClass, void (ParamTypes..., VarTypes...)>::Type InFunc, VarTypes... Vars)
 	{
-		Add(FDelegate::CreateSP(InUserObject, InFunc, Vars...));
+		return Add(FDelegate::CreateSP(InUserObject, InFunc, Vars...));
 	}
 	template <typename UserClass, typename... VarTypes>
-	inline void AddSP(UserClass* InUserObject, typename TMemFunPtrType<true, UserClass, void (ParamTypes..., VarTypes...)>::Type InFunc, VarTypes... Vars)
+	inline FDelegateHandle AddSP(UserClass* InUserObject, typename TMemFunPtrType<true, UserClass, void (ParamTypes..., VarTypes...)>::Type InFunc, VarTypes... Vars)
 	{
-		Add(FDelegate::CreateSP(InUserObject, InFunc, Vars...));
+		return Add(FDelegate::CreateSP(InUserObject, InFunc, Vars...));
 	}
 
 	/**
@@ -645,14 +663,14 @@ public:
 	 * @param	InFunc			Class method function address
 	 */
 	template <typename UserClass, typename... VarTypes>
-	inline void AddThreadSafeSP(const TSharedRef<UserClass, ESPMode::ThreadSafe>& InUserObjectRef, typename TMemFunPtrType<false, UserClass, void (ParamTypes..., VarTypes...)>::Type InFunc, VarTypes... Vars)
+	inline FDelegateHandle AddThreadSafeSP(const TSharedRef<UserClass, ESPMode::ThreadSafe>& InUserObjectRef, typename TMemFunPtrType<false, UserClass, void (ParamTypes..., VarTypes...)>::Type InFunc, VarTypes... Vars)
 	{
-		Add(FDelegate::CreateThreadSafeSP(InUserObjectRef, InFunc, Vars...));
+		return Add(FDelegate::CreateThreadSafeSP(InUserObjectRef, InFunc, Vars...));
 	}
 	template <typename UserClass, typename... VarTypes>
-	inline void AddThreadSafeSP(const TSharedRef<UserClass, ESPMode::ThreadSafe>& InUserObjectRef, typename TMemFunPtrType<true, UserClass, void (ParamTypes..., VarTypes...)>::Type InFunc, VarTypes... Vars)
+	inline FDelegateHandle AddThreadSafeSP(const TSharedRef<UserClass, ESPMode::ThreadSafe>& InUserObjectRef, typename TMemFunPtrType<true, UserClass, void (ParamTypes..., VarTypes...)>::Type InFunc, VarTypes... Vars)
 	{
-		Add(FDelegate::CreateThreadSafeSP(InUserObjectRef, InFunc, Vars...));
+		return Add(FDelegate::CreateThreadSafeSP(InUserObjectRef, InFunc, Vars...));
 	}
 
 	/**
@@ -664,14 +682,14 @@ public:
 	 * @param	InFunc			Class method function address
 	 */
 	template <typename UserClass, typename... VarTypes>
-	inline void AddThreadSafeSP(UserClass* InUserObject, typename TMemFunPtrType<false, UserClass, void (ParamTypes..., VarTypes...)>::Type InFunc, VarTypes... Vars)
+	inline FDelegateHandle AddThreadSafeSP(UserClass* InUserObject, typename TMemFunPtrType<false, UserClass, void (ParamTypes..., VarTypes...)>::Type InFunc, VarTypes... Vars)
 	{
-		Add(FDelegate::CreateThreadSafeSP(InUserObject, InFunc, Vars...));
+		return Add(FDelegate::CreateThreadSafeSP(InUserObject, InFunc, Vars...));
 	}
 	template <typename UserClass, typename... VarTypes>
-	inline void AddThreadSafeSP(UserClass* InUserObject, typename TMemFunPtrType<true, UserClass, void (ParamTypes..., VarTypes...)>::Type InFunc, VarTypes... Vars)
+	inline FDelegateHandle AddThreadSafeSP(UserClass* InUserObject, typename TMemFunPtrType<true, UserClass, void (ParamTypes..., VarTypes...)>::Type InFunc, VarTypes... Vars)
 	{
-		Add(FDelegate::CreateThreadSafeSP(InUserObject, InFunc, Vars...));
+		return Add(FDelegate::CreateThreadSafeSP(InUserObject, InFunc, Vars...));
 	}
 
 	/**
@@ -683,9 +701,9 @@ public:
 	 * @param	InFunctionName			Class method function address
 	 */
 	template <typename UObjectTemplate, typename... VarTypes>
-	inline void AddUFunction(UObjectTemplate* InUserObject, const FName& InFunctionName, VarTypes... Vars)
+	inline FDelegateHandle AddUFunction(UObjectTemplate* InUserObject, const FName& InFunctionName, VarTypes... Vars)
 	{
-		Add(FDelegate::CreateUFunction(InUserObject, InFunctionName, Vars...));
+		return Add(FDelegate::CreateUFunction(InUserObject, InFunctionName, Vars...));
 	}
 
 	/**
@@ -697,18 +715,30 @@ public:
 	 * @param	InFunc			Class method function address
 	 */
 	template <typename UserClass, typename... VarTypes>
-	inline void AddUObject(UserClass* InUserObject, typename TMemFunPtrType<false, UserClass, void (ParamTypes..., VarTypes...)>::Type InFunc, VarTypes... Vars)
+	inline FDelegateHandle AddUObject(UserClass* InUserObject, typename TMemFunPtrType<false, UserClass, void (ParamTypes..., VarTypes...)>::Type InFunc, VarTypes... Vars)
 	{
-		Add(FDelegate::CreateUObject(InUserObject, InFunc, Vars...));
+		return Add(FDelegate::CreateUObject(InUserObject, InFunc, Vars...));
 	}
 	template <typename UserClass, typename... VarTypes>
-	inline void AddUObject(UserClass* InUserObject, typename TMemFunPtrType<true, UserClass, void (ParamTypes..., VarTypes...)>::Type InFunc, VarTypes... Vars)
+	inline FDelegateHandle AddUObject(UserClass* InUserObject, typename TMemFunPtrType<true, UserClass, void (ParamTypes..., VarTypes...)>::Type InFunc, VarTypes... Vars)
 	{
-		Add(FDelegate::CreateUObject(InUserObject, InFunc, Vars...));
+		return Add(FDelegate::CreateUObject(InUserObject, InFunc, Vars...));
 	}
 
 public:
 
+	/**
+	 * Removes a delegate instance from this multi-cast delegate's invocation list (performance is O(N)).
+	 *
+	 * Note that the order of the delegate instances may not be preserved!
+	 *
+	 * @param Handle The handle of the delegate instance to remove.
+	 */
+	void Remove( FDelegateHandle Handle )
+	{
+		RemoveDelegateInstance(Handle);
+	}
+
 	/**
 	 * Removes a delegate instance from this multi-cast delegate's invocation list (performance is O(N)).
 	 *
@@ -716,14 +746,18 @@ public:
 	 *
 	 * @param DelegateInstance The delegate instance to remove.
 	 */
+	DELEGATE_DEPRECATED("This Remove overload is deprecated - please remove delegates using the FDelegateHandle returned by the Add function.")
 	void Remove( const TDelegateInstanceInterface& DelegateInstance )
 	{
-		// NOTE: For delegates with PAYLOAD, the payload variables are not considered when comparing delegate objects!
-		// This means that two delegate instances bound to the same object/method, but with different payload data,
-		// cannot be bound to the same single multi-cast delegate.  We could support this by changing IsSameFunction()
-		// to do comparison of payload data, but there are some trade-offs if we do that.
+		DEPRECATED_Remove(DelegateInstance);
+	}
 
-		RemoveDelegateInstance(&DelegateInstance);
+	void DEPRECATED_Remove( const TDelegateInstanceInterface& DelegateInstance )
+	{
+		// Provided to help implement other deprecated functions without giving multiple warnings.
+		// Should not be called directly.
+
+		DEPRECATED_RemoveDelegateInstance(&DelegateInstance);
 	}
 
 	/**
@@ -734,18 +768,22 @@ public:
 	 *
 	 * @param Delegate The delegate to remove.
 	 */
+	DELEGATE_DEPRECATED("This Remove overload is deprecated - please remove delegates using the FDelegateHandle returned by the Add function.")
 	void Remove( const FDelegate& Delegate )
 	{
-		// NOTE: For delegates with PAYLOAD, the payload variables are not considered when comparing delegate objects!
-		// This means that two delegate instances bound to the same object/method, but with different payload data,
-		// cannot be bound to the same single multi-cast delegate.  We could support this by changing IsSameFunction()
-		// to do comparison of payload data, but there are some trade-offs if we do that.
+		DEPRECATED_Remove(Delegate);
+	}
+
+	void DEPRECATED_Remove( const FDelegate& Delegate )
+	{
+		// Provided to help implement other deprecated functions without giving multiple warnings.
+		// Should not be called directly.
 
 		TDelegateInstanceInterface* DelegateInstance = (TDelegateInstanceInterface*)Delegate.GetDelegateInstance();
 
 		if (DelegateInstance != nullptr)
 		{
-			RemoveDelegateInstance(DelegateInstance);
+			DEPRECATED_RemoveDelegateInstance(DelegateInstance);
 		}		
 	}
 
@@ -758,9 +796,10 @@ public:
 	 *
 	 * @param	InFunc	Function pointer
 	 */
+	DELEGATE_DEPRECATED("RemoveStatic is deprecated - please remove delegates using the FDelegateHandle returned by the Add function.")
 	inline void RemoveStatic( typename TBaseStaticDelegateInstance< void (ParamTypes...) >::FFuncPtr InFunc )
 	{
-		Remove( FDelegate::CreateStatic( InFunc ) );
+		DEPRECATED_RemoveDelegateInstance( (TDelegateInstanceInterface*)FDelegate::CreateStatic( InFunc ).GetDelegateInstance() );
 	}
 
 	/**
@@ -771,14 +810,16 @@ public:
 	 * @param	InFunc			Class method function address
 	 */
 	template< class UserClass >
+	DELEGATE_DEPRECATED("RemoveRaw is deprecated - please remove delegates using the FDelegateHandle returned by the Add function.")
 	inline void RemoveRaw( UserClass* InUserObject, typename TBaseRawMethodDelegateInstance<false, UserClass, void (ParamTypes...)>::FMethodPtr InFunc )
 	{
-		Remove( FDelegate::CreateRaw( InUserObject, InFunc ) );
+		DEPRECATED_RemoveDelegateInstance( (TDelegateInstanceInterface*)FDelegate::CreateRaw( InUserObject, InFunc ).GetDelegateInstance() );
 	}
 	template< class UserClass >
+	DELEGATE_DEPRECATED("RemoveRaw is deprecated - please remove delegates using the FDelegateHandle returned by the Add function.")
 	inline void RemoveRaw( UserClass* InUserObject, typename TBaseRawMethodDelegateInstance<true , UserClass, void (ParamTypes...)>::FMethodPtr InFunc )
 	{
-		Remove( FDelegate::CreateRaw( InUserObject, InFunc ) );
+		DEPRECATED_RemoveDelegateInstance( (TDelegateInstanceInterface*)FDelegate::CreateRaw( InUserObject, InFunc ).GetDelegateInstance() );
 	}
 
 	/**
@@ -789,14 +830,16 @@ public:
 	 * @param	InFunc			Class method function address
 	 */
 	template< class UserClass >
+	DELEGATE_DEPRECATED("RemoveSP is deprecated - please remove delegates using the FDelegateHandle returned by the Add function.")
 	inline void RemoveSP( const TSharedRef< UserClass, ESPMode::Fast >& InUserObjectRef, typename TBaseSPMethodDelegateInstance<false, UserClass, ESPMode::Fast, void (ParamTypes...)>::FMethodPtr InFunc )
 	{
-		Remove( FDelegate::CreateSP( InUserObjectRef, InFunc ) );
+		DEPRECATED_RemoveDelegateInstance( (TDelegateInstanceInterface*)FDelegate::CreateSP( InUserObjectRef, InFunc ).GetDelegateInstance() );
 	}
 	template< class UserClass >
+	DELEGATE_DEPRECATED("RemoveSP is deprecated - please remove delegates using the FDelegateHandle returned by the Add function.")
 	inline void RemoveSP( const TSharedRef< UserClass, ESPMode::Fast >& InUserObjectRef, typename TBaseSPMethodDelegateInstance<true , UserClass, ESPMode::Fast, void (ParamTypes...)>::FMethodPtr InFunc )
 	{
-		Remove( FDelegate::CreateSP( InUserObjectRef, InFunc ) );
+		DEPRECATED_RemoveDelegateInstance( (TDelegateInstanceInterface*)FDelegate::CreateSP( InUserObjectRef, InFunc ).GetDelegateInstance() );
 	}
 
 	/**
@@ -807,14 +850,16 @@ public:
 	 * @param	InFunc			Class method function address
 	 */
 	template< class UserClass >
+	DELEGATE_DEPRECATED("RemoveSP is deprecated - please remove delegates using the FDelegateHandle returned by the Add function.")
 	inline void RemoveSP( UserClass* InUserObject, typename TBaseSPMethodDelegateInstance<false, UserClass, ESPMode::Fast, void (ParamTypes...)>::FMethodPtr InFunc )
 	{
-		Remove( FDelegate::CreateSP( InUserObject, InFunc ) );
+		DEPRECATED_RemoveDelegateInstance( (TDelegateInstanceInterface*)FDelegate::CreateSP( InUserObject, InFunc ).GetDelegateInstance() );
 	}
 	template< class UserClass >
+	DELEGATE_DEPRECATED("RemoveSP is deprecated - please remove delegates using the FDelegateHandle returned by the Add function.")
 	inline void RemoveSP( UserClass* InUserObject, typename TBaseSPMethodDelegateInstance<true , UserClass, ESPMode::Fast, void (ParamTypes...)>::FMethodPtr InFunc )
 	{
-		Remove( FDelegate::CreateSP( InUserObject, InFunc ) );
+		DEPRECATED_RemoveDelegateInstance( (TDelegateInstanceInterface*)FDelegate::CreateSP( InUserObject, InFunc ).GetDelegateInstance() );
 	}
 
 	/**
@@ -825,14 +870,16 @@ public:
 	 * @param	InFunc			Class method function address
 	 */
 	template< class UserClass >
+	DELEGATE_DEPRECATED("RemoveThreadSafeSP is deprecated - please remove delegates using the FDelegateHandle returned by the Add function.")
 	inline void RemoveThreadSafeSP( const TSharedRef< UserClass, ESPMode::ThreadSafe >& InUserObjectRef, typename TBaseSPMethodDelegateInstance<false, UserClass, ESPMode::ThreadSafe, void (ParamTypes...)>::FMethodPtr InFunc )
 	{
-		Remove( FDelegate::CreateThreadSafeSP( InUserObjectRef, InFunc ) );
+		DEPRECATED_RemoveDelegateInstance( (TDelegateInstanceInterface*)FDelegate::CreateThreadSafeSP( InUserObjectRef, InFunc ).GetDelegateInstance() );
 	}
 	template< class UserClass >
+	DELEGATE_DEPRECATED("RemoveThreadSafeSP is deprecated - please remove delegates using the FDelegateHandle returned by the Add function.")
 	inline void RemoveThreadSafeSP( const TSharedRef< UserClass, ESPMode::ThreadSafe >& InUserObjectRef, typename TBaseSPMethodDelegateInstance<true , UserClass, ESPMode::ThreadSafe, void (ParamTypes...)>::FMethodPtr InFunc )
 	{
-		Remove( FDelegate::CreateThreadSafeSP( InUserObjectRef, InFunc ) );
+		DEPRECATED_RemoveDelegateInstance( (TDelegateInstanceInterface*)FDelegate::CreateThreadSafeSP( InUserObjectRef, InFunc ).GetDelegateInstance() );
 	}
 
 	/**
@@ -843,14 +890,16 @@ public:
 	 * @param	InFunc			Class method function address
 	 */
 	template< class UserClass >
+	DELEGATE_DEPRECATED("RemoveThreadSafeSP is deprecated - please remove delegates using the FDelegateHandle returned by the Add function.")
 	inline void RemoveThreadSafeSP( UserClass* InUserObject, typename TBaseSPMethodDelegateInstance<false, UserClass, ESPMode::ThreadSafe, void (ParamTypes...)>::FMethodPtr InFunc )
 	{
-		Remove( FDelegate::CreateThreadSafeSP( InUserObject, InFunc ) );
+		DEPRECATED_RemoveDelegateInstance( (TDelegateInstanceInterface*)FDelegate::CreateThreadSafeSP( InUserObject, InFunc ).GetDelegateInstance() );
 	}
 	template< class UserClass >
+	DELEGATE_DEPRECATED("RemoveThreadSafeSP is deprecated - please remove delegates using the FDelegateHandle returned by the Add function.")
 	inline void RemoveThreadSafeSP( UserClass* InUserObject, typename TBaseSPMethodDelegateInstance<true , UserClass, ESPMode::ThreadSafe, void (ParamTypes...)>::FMethodPtr InFunc )
 	{
-		Remove( FDelegate::CreateThreadSafeSP( InUserObject, InFunc ) );
+		DEPRECATED_RemoveDelegateInstance( (TDelegateInstanceInterface*)FDelegate::CreateThreadSafeSP( InUserObject, InFunc ).GetDelegateInstance() );
 	}
 
 	/**
@@ -861,9 +910,10 @@ public:
 	 * @param	InFunc			Class method function address
 	 */
 	template< class UObjectTemplate >
+	DELEGATE_DEPRECATED("RemoveUFunction is deprecated - please remove delegates using the FDelegateHandle returned by the Add function.")
 	inline void RemoveUFunction( UObjectTemplate* InUserObject, const FName& InFunctionName )
 	{
-		Remove( FDelegate::CreateUFunction( InUserObject, InFunctionName ) );
+		DEPRECATED_RemoveDelegateInstance( (TDelegateInstanceInterface*)FDelegate::CreateUFunction( InUserObject, InFunctionName ).GetDelegateInstance() );
 	}
 
 	/**
@@ -874,14 +924,16 @@ public:
 	 * @param	InFunc			Class method function address
 	 */
 	template< class UserClass >
+	DELEGATE_DEPRECATED("RemoveUObject is deprecated - please remove delegates using the FDelegateHandle returned by the Add function.")
 	inline void RemoveUObject( UserClass* InUserObject, typename TBaseUObjectMethodDelegateInstance<false, UserClass, void (ParamTypes...)>::FMethodPtr InFunc )
 	{
-		Remove( FDelegate::CreateUObject( InUserObject, InFunc ) );
+		DEPRECATED_RemoveDelegateInstance( (TDelegateInstanceInterface*)FDelegate::CreateUObject( InUserObject, InFunc ).GetDelegateInstance() );
 	}
 	template< class UserClass >
+	DELEGATE_DEPRECATED("RemoveUObject is deprecated - please remove delegates using the FDelegateHandle returned by the Add function.")
 	inline void RemoveUObject( UserClass* InUserObject, typename TBaseUObjectMethodDelegateInstance<true , UserClass, void (ParamTypes...)>::FMethodPtr InFunc )
 	{
-		Remove( FDelegate::CreateUObject( InUserObject, InFunc ) );
+		DEPRECATED_RemoveDelegateInstance( (TDelegateInstanceInterface*)FDelegate::CreateUObject( InUserObject, InFunc ).GetDelegateInstance() );
 	}
 
 protected:
@@ -937,7 +989,7 @@ protected:
 	 *
 	 * @param	InDelegate	Delegate to add
 	 */
-	void AddDelegateInstance( TDelegateInstanceInterface* InDelegateInstance )
+	FDelegateHandle AddDelegateInstance( TDelegateInstanceInterface* InDelegateInstance )
 	{
 #if !UE_BUILD_SHIPPING && !UE_BUILD_TEST
 		// verify that the same function isn't already bound
@@ -952,7 +1004,7 @@ protected:
 		}
 #endif
 
-		AddInternal(InDelegateInstance);
+		return AddInternal(InDelegateInstance);
 	}
 
 public:
@@ -990,6 +1042,41 @@ public:
 	}
 
 protected:
+	/**
+	 * Removes a function from this multi-cast delegate's invocation list (performance is O(N)).
+	 *
+	 * The function is not actually removed, but deleted and marked as removed.
+	 * It will be removed next time the invocation list is compacted within Broadcast().
+	 *
+	 * @param Handle The handle of the delegate instance to remove.
+	 */
+	void RemoveDelegateInstance( FDelegateHandle Handle )
+	{
+		const TArray<IDelegateInstance*>& InvocationList = GetInvocationList();
+
+		// NOTE: We assume that this method is never called with a nullptr object, in which case the
+		//       the following algorithm would break down (it would remove the first found instance
+		//       of a matching function binding, which is not necessarily the instance we wish to remove).
+
+		for (int32 InvocationListIndex = 0; InvocationListIndex < InvocationList.Num(); ++InvocationListIndex)
+		{
+			// this down-cast is OK! allows for managing invocation list in the base class without requiring virtual functions
+			TDelegateInstanceInterface*& DelegateInstanceRef = (TDelegateInstanceInterface*&)InvocationList[InvocationListIndex];
+
+			// NOTE: We must do a deep compare here, not just compare delegate pointers, because multiple
+			//       delegate pointers can refer to the exact same object and method
+			if ((DelegateInstanceRef != nullptr) && DelegateInstanceRef->GetHandle() == Handle)
+			{
+				delete DelegateInstanceRef;
+				DelegateInstanceRef = nullptr;
+
+				break;	// no need to continue, as we never allow the same delegate to be bound twice
+			}
+		}
+
+		const_cast<TBaseMulticastDelegate*>(this)->CompactInvocationList();
+	}
+
 	/**
 	 * Removes a function from this multi-cast delegate's invocation list (performance is O(N)).
 	 *
@@ -998,7 +1085,7 @@ protected:
 	 *
 	 * @param InDelegateInstance The delegate instance to remove.
 	 */
-	void RemoveDelegateInstance( const TDelegateInstanceInterface* InDelegateInstance )
+	void DEPRECATED_RemoveDelegateInstance( const TDelegateInstanceInterface* InDelegateInstance )
 	{
 		const TArray<IDelegateInstance*>& InvocationList = GetInvocationList();
 
diff --git a/Engine/Source/Runtime/Core/Public/Delegates/IDelegateInstance.h b/Engine/Source/Runtime/Core/Public/Delegates/IDelegateInstance.h
--- a/Engine/Source/Runtime/Core/Public/Delegates/IDelegateInstance.h
+++ b/Engine/Source/Runtime/Core/Public/Delegates/IDelegateInstance.h
@@ -34,6 +34,49 @@ namespace EDelegateInstanceType
 }
 
 
+/**
+ * Class representing an handle to a delegate.
+ */
+class FDelegateHandle
+{
+public:
+	enum EGenerateNewHandleType
+	{
+		GenerateNewHandle
+	};
+
+	FDelegateHandle()
+		: ID(0)
+	{
+	}
+
+	explicit FDelegateHandle(EGenerateNewHandleType)
+		: ID(GenerateNewID())
+	{
+	}
+
+private:
+	friend bool operator==(const FDelegateHandle& Lhs, const FDelegateHandle& Rhs)
+	{
+		return Lhs.ID == Rhs.ID;
+	}
+
+	friend bool operator!=(const FDelegateHandle& Lhs, const FDelegateHandle& Rhs)
+	{
+		return Lhs.ID != Rhs.ID;
+	}
+
+	/**
+	 * Generates a new ID for use the delegate handle.
+	 *
+	 * @return A unique ID for the delegate.
+	 */
+	static CORE_API uint64 GenerateNewID();
+
+	uint64 ID;
+};
+
+
 /**
  * Interface for delegate instances.
  */
@@ -96,6 +139,11 @@ public:
 	 */
 	virtual bool IsSafeToExecute( ) const = 0;
 
+	/**
+	 * Returns a handle for the delegate.
+	 */
+	virtual FDelegateHandle GetHandle() const = 0;
+
 public:
 
 	/**
diff --git a/Engine/Source/Runtime/Core/Public/Delegates/MulticastDelegateBase.h b/Engine/Source/Runtime/Core/Public/Delegates/MulticastDelegateBase.h
--- a/Engine/Source/Runtime/Core/Public/Delegates/MulticastDelegateBase.h
+++ b/Engine/Source/Runtime/Core/Public/Delegates/MulticastDelegateBase.h
@@ -95,9 +95,10 @@ protected:
 	 *
 	 * @param DelegateInstance The delegate instance to add.
 	 */
-	inline void AddInternal( IDelegateInstance* DelegateInstance )
+	inline FDelegateHandle AddInternal( IDelegateInstance* DelegateInstance )
 	{
 		InvocationList.Add(DelegateInstance);
+		return DelegateInstance->GetHandle();
 	}
 
 	/**
diff --git a/Engine/Source/Runtime/Core/Public/HAL/ConsoleManager.h b/Engine/Source/Runtime/Core/Public/HAL/ConsoleManager.h
--- a/Engine/Source/Runtime/Core/Public/HAL/ConsoleManager.h
+++ b/Engine/Source/Runtime/Core/Public/HAL/ConsoleManager.h
@@ -7,6 +7,13 @@
 
 class FConsoleVariableBase;
 
+// We disable the deprecation warnings here because otherwise it'll complain about us
+// implementing RegisterConsoleVariableSink and UnregisterConsoleVariableSink.  We know
+// that, but we only want it to complain if *others* implement or call these functions.
+//
+// These macros should be removed when those functions are removed.
+PRAGMA_DISABLE_DEPRECATION_WARNINGS
+
 class CORE_API FConsoleManager :public IConsoleManager
 {
 public:
@@ -53,9 +60,16 @@ public:
 	virtual IConsoleVariable* RegisterConsoleVariableBitRef(const TCHAR* CVarName, const TCHAR* FlagName, uint32 BitNumber, uint8* Force0MaskPtr, uint8* Force1MaskPtr, const TCHAR* Help, uint32 Flags) override;
 	
 	virtual void CallAllConsoleVariableSinks();
+
+	DELEGATE_DEPRECATED("This function is deprecated - please replace any usage with RegisterConsoleVariableSink_Handle.")
 	virtual void RegisterConsoleVariableSink(const FConsoleCommandDelegate& Command);
+
+	DELEGATE_DEPRECATED("Delegate comparison is deprecated - please replace any usage with UnregisterConsoleVariableSink_Handle, passing the result of RegisterConsoleVariableSink_Handle.")
 	virtual void UnregisterConsoleVariableSink(const FConsoleCommandDelegate& Command);
 
+	virtual FConsoleVariableSinkHandle RegisterConsoleVariableSink_Handle(const FConsoleCommandDelegate& Command);
+	virtual void UnregisterConsoleVariableSink_Handle(FConsoleVariableSinkHandle Handle);
+
 	virtual IConsoleCommand* RegisterConsoleCommand(const TCHAR* Name, const TCHAR* Help, const FConsoleCommandDelegate& Command, uint32 Flags) override;
 	virtual IConsoleCommand* RegisterConsoleCommand(const TCHAR* Name, const TCHAR* Help, const FConsoleCommandWithArgsDelegate& Command, uint32 Flags) override;
 	virtual IConsoleCommand* RegisterConsoleCommand(const TCHAR* Name, const TCHAR* Help, const FConsoleCommandWithWorldDelegate& Command, uint32 Flags) override;
@@ -138,3 +152,4 @@ private: // ----------------------------------------------------
 	void SaveHistory();
 };
 
+PRAGMA_ENABLE_DEPRECATION_WARNINGS
diff --git a/Engine/Source/Runtime/Core/Public/HAL/IConsoleManager.h b/Engine/Source/Runtime/Core/Public/HAL/IConsoleManager.h
--- a/Engine/Source/Runtime/Core/Public/HAL/IConsoleManager.h
+++ b/Engine/Source/Runtime/Core/Public/HAL/IConsoleManager.h
@@ -294,6 +294,39 @@ struct IConsoleThreadPropagation
  */
 DECLARE_DELEGATE_TwoParams( FConsoleObjectVisitor, const TCHAR*, IConsoleObject* );
 
+
+/**
+ * Class representing an handle to an online delegate.
+ */
+class FConsoleVariableSinkHandle
+{
+public:
+	FConsoleVariableSinkHandle()
+	{
+	}
+
+	explicit FConsoleVariableSinkHandle(FDelegateHandle InHandle)
+		: Handle(InHandle)
+	{
+	}
+
+	template <typename MulticastDelegateType>
+	void RemoveFromDelegate(MulticastDelegateType& MulticastDelegate)
+	{
+		MulticastDelegate.Remove(Handle);
+	}
+
+	template <typename DelegateType>
+	bool HasSameHandle(const DelegateType& Delegate) const
+	{
+		return Delegate.GetHandle() == Handle;
+	}
+
+private:
+	FDelegateHandle Handle;
+};
+
+
 /**
  * handles console commands and variables, registered console variables are released on destruction
  */
@@ -362,14 +395,28 @@ struct CORE_API IConsoleManager
 	 * The registered command is executed at few defined points (see CallAllConsoleVariableSinks)
 	 * @param Command
 	 */
+	DELEGATE_DEPRECATED("This function is deprecated - please replace any usage with RegisterConsoleVariableSink_Handle.")
 	virtual void RegisterConsoleVariableSink(const FConsoleCommandDelegate& Command) = 0;
 
 	/**
 	 * The registered command is executed at few defined points (see CallAllConsoleVariableSinks)
 	 * @param Command
 	 */
+	DELEGATE_DEPRECATED("Delegate comparison is deprecated - please replace any usage with UnregisterConsoleVariableSink_Handle, passing the result of RegisterConsoleVariableSink_Handle.")
 	virtual void UnregisterConsoleVariableSink(const FConsoleCommandDelegate& Command) = 0;
 
+	/**
+	 * The registered command is executed at few defined points (see CallAllConsoleVariableSinks)
+	 * @param Command
+	 */
+	virtual FConsoleVariableSinkHandle RegisterConsoleVariableSink_Handle(const FConsoleCommandDelegate& Command) = 0;
+
+	/**
+	 * The registered command is executed at few defined points (see CallAllConsoleVariableSinks)
+	 * @param Command
+	 */
+	virtual void UnregisterConsoleVariableSink_Handle(FConsoleVariableSinkHandle Handle) = 0;
+
 	// ----------
 
 	/**
@@ -539,15 +586,16 @@ public:
 	FAutoConsoleVariableSink(const FConsoleCommandDelegate& InCommand)
 		: Command(InCommand)
 	{
-		IConsoleManager::Get().RegisterConsoleVariableSink(Command);
+		Handle = IConsoleManager::Get().RegisterConsoleVariableSink_Handle(Command);
 	}
 	/** Destructor, removes the console variable sink **/
 	virtual ~FAutoConsoleVariableSink()
 	{
-//disabled for now, destruction order makes this not always working		IConsoleManager::Get().UnregisterConsoleVariableSink(Command);
+//disabled for now, destruction order makes this not always working		IConsoleManager::Get().UnregisterConsoleVariableSink_Handle(Handle);
 	}
 
 	const FConsoleCommandDelegate& Command;
+	FConsoleVariableSinkHandle Handle;
 };
 
 
diff --git a/Engine/Source/Runtime/Core/Public/Stats/StatsFile.h b/Engine/Source/Runtime/Core/Public/Stats/StatsFile.h
--- a/Engine/Source/Runtime/Core/Public/Stats/StatsFile.h
+++ b/Engine/Source/Runtime/Core/Public/Stats/StatsFile.h
@@ -334,6 +334,9 @@ protected:
 	/** Buffer to store the compressed data, used by the FAsyncWriteWorker. */
 	TArray<uint8> CompressedData;
 
+	/** NewFrame delegate handle  */
+	FDelegateHandle NewFrameDelegateHandle;
+
 public:
 	/** Constructor. **/
 	FStatsWriteFile();
diff --git a/Engine/Source/Runtime/CoreUObject/Public/Serialization/AsyncLoading.h b/Engine/Source/Runtime/CoreUObject/Public/Serialization/AsyncLoading.h
--- a/Engine/Source/Runtime/CoreUObject/Public/Serialization/AsyncLoading.h
+++ b/Engine/Source/Runtime/CoreUObject/Public/Serialization/AsyncLoading.h
@@ -64,7 +64,7 @@ struct FAsyncPackage : public FGCObject
 		// This is to ensure that there is no one trying to subscribe to a already loaded package
 		check(!bLoadHasFinished && !bLoadHasFailed);
 		
-		CompletionCallbacks.AddUnique(Callback);
+		CompletionCallbacks.Add(Callback);
 	}
 
 	/**
diff --git a/Engine/Source/Runtime/Engine/Classes/AI/Navigation/NavigationData.h b/Engine/Source/Runtime/Engine/Classes/AI/Navigation/NavigationData.h
--- a/Engine/Source/Runtime/Engine/Classes/AI/Navigation/NavigationData.h
+++ b/Engine/Source/Runtime/Engine/Classes/AI/Navigation/NavigationData.h
@@ -97,13 +97,18 @@ struct ENGINE_API FNavigationPath : public TSharedFromThis<FNavigationPath, ESPM
 	{
 		return ObserverDelegate;
 	}
-	FORCEINLINE void AddObserver(FPathObserverDelegate::FDelegate& NewObserver)
+	FORCEINLINE FDelegateHandle AddObserver(FPathObserverDelegate::FDelegate& NewObserver)
 	{
-		ObserverDelegate.Add(NewObserver);
+		return ObserverDelegate.Add(NewObserver);
 	}
+	DELEGATE_DEPRECATED("This RemoveObserver overload has been deprecated - please pass the handle returned from AddObserver instead.")
 	FORCEINLINE void RemoveObserver(FPathObserverDelegate::FDelegate& ObserverToRemove)
 	{
-		ObserverDelegate.Remove(ObserverToRemove);
+		ObserverDelegate.DEPRECATED_Remove(ObserverToRemove);
+	}
+	FORCEINLINE void RemoveObserver(FDelegateHandle HandleOfObserverToRemove)
+	{
+		ObserverDelegate.Remove(HandleOfObserverToRemove);
 	}
 
 	FORCEINLINE void MarkReady()
diff --git a/Engine/Source/Runtime/Engine/Classes/AI/Navigation/NavigationPath.h b/Engine/Source/Runtime/Engine/Classes/AI/Navigation/NavigationPath.h
--- a/Engine/Source/Runtime/Engine/Classes/AI/Navigation/NavigationPath.h
+++ b/Engine/Source/Runtime/Engine/Classes/AI/Navigation/NavigationPath.h
@@ -32,10 +32,13 @@ private:
 	uint32 bDebugDrawingEnabled : 1;
 	FColor DebugDrawingColor;
 
+	FDelegateHandle DrawDebugDelegateHandle;
+
 protected:
 	FNavPathSharedPtr SharedPath;
 
 	FNavigationPath::FPathObserverDelegate::FDelegate PathObserver;
+	FDelegateHandle PathObserverDelegateHandle;
 
 public:
 
diff --git a/Engine/Source/Runtime/Engine/Classes/Animation/Skeleton.h b/Engine/Source/Runtime/Engine/Classes/Animation/Skeleton.h
--- a/Engine/Source/Runtime/Engine/Classes/Animation/Skeleton.h
+++ b/Engine/Source/Runtime/Engine/Classes/Animation/Skeleton.h
@@ -345,19 +345,27 @@ public:
 	typedef FOnRetargetSourceChangedMulticaster::FDelegate FOnRetargetSourceChanged;
 
 	/** Registers a delegate to be called after the preview animation has been changed */
-	void RegisterOnRetargetSourceChanged(const FOnRetargetSourceChanged& Delegate)
+	FDelegateHandle RegisterOnRetargetSourceChanged(const FOnRetargetSourceChanged& Delegate)
 	{
-		OnRetargetSourceChanged.Add(Delegate);
+		return OnRetargetSourceChanged.Add(Delegate);
 	}
 
 	const FGuid GetGuid() const
 	{
 		return Guid;
 	}
+
 	/** Unregisters a delegate to be called after the preview animation has been changed */
+	DELEGATE_DEPRECATED("This UnregisterOnRetargetSourceChanged overload has been deprecated - please pass the handle returned from RegisterOnRetargetSourceChanged instead.")
 	void UnregisterOnRetargetSourceChanged(const FOnRetargetSourceChanged& Delegate)
 	{
-		OnRetargetSourceChanged.Remove(Delegate);
+		OnRetargetSourceChanged.DEPRECATED_Remove(Delegate);
+	}
+
+	/** Unregisters a delegate to be called after the preview animation has been changed */
+	void UnregisterOnRetargetSourceChanged(FDelegateHandle Handle)
+	{
+		OnRetargetSourceChanged.Remove(Handle);
 	}
 
 	void CallbackRetargetSourceChanged()
diff --git a/Engine/Source/Runtime/Engine/Classes/Components/SkeletalMeshComponent.h b/Engine/Source/Runtime/Engine/Classes/Components/SkeletalMeshComponent.h
--- a/Engine/Source/Runtime/Engine/Classes/Components/SkeletalMeshComponent.h
+++ b/Engine/Source/Runtime/Engine/Classes/Components/SkeletalMeshComponent.h
@@ -772,8 +772,10 @@ public:
 	typedef FOnSkeletalMeshPropertyChangedMulticaster::FDelegate FOnSkeletalMeshPropertyChanged;
 
 	/** Register / Unregister delegates called when the skeletal mesh property is changed */
-	void RegisterOnSkeletalMeshPropertyChanged(const FOnSkeletalMeshPropertyChanged& Delegate);
+	FDelegateHandle RegisterOnSkeletalMeshPropertyChanged(const FOnSkeletalMeshPropertyChanged& Delegate);
+	DELEGATE_DEPRECATED("This overload of UnregisterOnSkeletalMeshPropertyChanged is deprecated, instead pass the result of RegisterOnSkeletalMeshPropertyChanged.")
 	void UnregisterOnSkeletalMeshPropertyChanged(const FOnSkeletalMeshPropertyChanged& Delegate);
+	void UnregisterOnSkeletalMeshPropertyChanged(FDelegateHandle Handle);
 
 	virtual void PostEditChangeProperty(FPropertyChangedEvent& PropertyChangedEvent) override;
 
diff --git a/Engine/Source/Runtime/Engine/Classes/Debug/DebugDrawService.h b/Engine/Source/Runtime/Engine/Classes/Debug/DebugDrawService.h
--- a/Engine/Source/Runtime/Engine/Classes/Debug/DebugDrawService.h
+++ b/Engine/Source/Runtime/Engine/Classes/Debug/DebugDrawService.h
@@ -19,8 +19,10 @@ class ENGINE_API UDebugDrawService : public UBlueprintFunctionLibrary
 
 	//void Register
 
-	static void Register(const TCHAR* Name, const FDebugDrawDelegate& NewDelegate);
-	static void Unregister(const FDebugDrawDelegate& NewDelegate);
+	static FDelegateHandle Register(const TCHAR* Name, const FDebugDrawDelegate& NewDelegate);
+	DELEGATE_DEPRECATED("This overload of Unregister is deprecated, instead pass the result of Register.")
+	static void Unregister(const FDebugDrawDelegate& DelegateToRemove);
+	static void Unregister(FDelegateHandle HandleToRemove);
 	static void Draw(const FEngineShowFlags Flags, class UCanvas* Canvas);
 	static void Draw(const FEngineShowFlags Flags, class FViewport* Viewport, FSceneView* View, FCanvas* Canvas);
 
diff --git a/Engine/Source/Runtime/Engine/Classes/EdGraph/EdGraph.h b/Engine/Source/Runtime/Engine/Classes/EdGraph/EdGraph.h
--- a/Engine/Source/Runtime/Engine/Classes/EdGraph/EdGraph.h
+++ b/Engine/Source/Runtime/Engine/Classes/EdGraph/EdGraph.h
@@ -98,11 +98,15 @@ public:
 	const class UEdGraphSchema* GetSchema() const;
 
 	/** Add a listener for OnGraphChanged events */
-	void AddOnGraphChangedHandler( const FOnGraphChanged::FDelegate& InHandler );
+	FDelegateHandle AddOnGraphChangedHandler( const FOnGraphChanged::FDelegate& InHandler );
 
 	/** Remove a listener for OnGraphChanged events */
+	DELEGATE_DEPRECATED("This overload of RemoveOnGraphChangedHandler is deprecated, instead pass the result of AddOnGraphChangedHandler.")
 	void RemoveOnGraphChangedHandler( const FOnGraphChanged::FDelegate& InHandler );
 
+	/** Remove a listener for OnGraphChanged events */
+	void RemoveOnGraphChangedHandler( FDelegateHandle Handle );
+
 #if WITH_EDITORONLY_DATA
 	// Begin UObject interface
 	virtual void PostInitProperties() override;
@@ -185,10 +189,14 @@ public:
 	void NotifyPostChange( const FPropertyChangedEvent& PropertyChangedEvent, const FString& PropertyName );
 
 	/** Add a delegate listening for property change notifications */
-	void AddPropertyChangedNotifier(const FOnPropertyChanged::FDelegate& InDelegate );
+	FDelegateHandle AddPropertyChangedNotifier(const FOnPropertyChanged::FDelegate& InDelegate );
 
 	/** Remove a delegate listening for property changed notifications */
+	DELEGATE_DEPRECATED("This overload of RemovePropertyChangedNotifier is deprecated, instead pass the result of AddPropertyChangedNotifier.")
 	void RemovePropertyChangedNotifier(const FOnPropertyChanged::FDelegate& InDelegate );
+
+	/** Remove a delegate listening for property changed notifications */
+	void RemovePropertyChangedNotifier(FDelegateHandle InHandle );
 #endif
 
 protected:
diff --git a/Engine/Source/Runtime/Engine/Classes/Engine/Engine.h b/Engine/Source/Runtime/Engine/Classes/Engine/Engine.h
--- a/Engine/Source/Runtime/Engine/Classes/Engine/Engine.h
+++ b/Engine/Source/Runtime/Engine/Classes/Engine/Engine.h
@@ -2737,4 +2737,6 @@ private:
 #if STATS
 	int32 RenderStatSlateBatches(UWorld* World, FViewport* Viewport, FCanvas* Canvas, int32 X, int32 Y, const FVector* ViewLocation = nullptr, const FRotator* ViewRotation = nullptr);
 #endif
+
+	FDelegateHandle HandleScreenshotCapturedDelegateHandle;
 };
diff --git a/Engine/Source/Runtime/Engine/Classes/Engine/GameViewportClient.h b/Engine/Source/Runtime/Engine/Classes/Engine/GameViewportClient.h
--- a/Engine/Source/Runtime/Engine/Classes/Engine/GameViewportClient.h
+++ b/Engine/Source/Runtime/Engine/Classes/Engine/GameViewportClient.h
@@ -699,6 +699,9 @@ private:
 
 	/** Whether or not the cursor is hidden when the viewport captures the mouse */
 	bool bHideCursorDuringCapture;
+
+	/** Handle to the registered ShowCollisionOnSpawnedActors delegate */
+	FDelegateHandle ShowCollisionOnSpawnedActorsDelegateHandle;
 };
 
 
diff --git a/Engine/Source/Runtime/Engine/Classes/Engine/LevelBounds.h b/Engine/Source/Runtime/Engine/Classes/Engine/LevelBounds.h
--- a/Engine/Source/Runtime/Engine/Classes/Engine/LevelBounds.h
+++ b/Engine/Source/Runtime/Engine/Classes/Engine/LevelBounds.h
@@ -80,6 +80,11 @@ private:
 
 	/** True when there are no actors contributing to the bounds and we are currently using the default bounds instead */
 	bool bUsingDefaultBounds;
+
+	/** Handles to various registered delegates */
+	FDelegateHandle OnLevelActorMovedDelegateHandle;
+	FDelegateHandle OnLevelActorDeletedDelegateHandle;
+	FDelegateHandle OnLevelActorAddedDelegateHandle;
 #endif
 };
 
diff --git a/Engine/Source/Runtime/Engine/Classes/Engine/NetDriver.h b/Engine/Source/Runtime/Engine/Classes/Engine/NetDriver.h
--- a/Engine/Source/Runtime/Engine/Classes/Engine/NetDriver.h
+++ b/Engine/Source/Runtime/Engine/Classes/Engine/NetDriver.h
@@ -297,6 +297,11 @@ public:
 
 	TSet< TWeakPtr< FObjectReplicator > > UnmappedReplicators;
 
+	/** Handles to various registered delegates */
+	FDelegateHandle TickDispatchDelegateHandle;
+	FDelegateHandle TickFlushDelegateHandle;
+	FDelegateHandle PostTickFlushDelegateHandle;
+
 	/**
 	* Updates the standby cheat information and
 	 * causes the dialog to be shown/hidden as needed
@@ -532,9 +537,9 @@ protected:
 	ENGINE_API void	AddClientConnection(UNetConnection * NewConnection);
 
 	/** Register all TickDispatch, TickFlush, PostTickFlush to tick in World */
-	ENGINE_API void RegisterTickEvents(class UWorld* InWorld) const;
+	ENGINE_API void RegisterTickEvents(class UWorld* InWorld);
 	/** Unregister all TickDispatch, TickFlush, PostTickFlush to tick in World */
-	ENGINE_API void UnregisterTickEvents(class UWorld* InWorld) const;
+	ENGINE_API void UnregisterTickEvents(class UWorld* InWorld);
 	/** Returns true if this actor is considered to be in a loaded level */
 	bool IsLevelInitializedForActor(const AActor* InActor, const UNetConnection* InConnection) const;
 };
diff --git a/Engine/Source/Runtime/Engine/Classes/Engine/World.h b/Engine/Source/Runtime/Engine/Classes/Engine/World.h
--- a/Engine/Source/Runtime/Engine/Classes/Engine/World.h
+++ b/Engine/Source/Runtime/Engine/Classes/Engine/World.h
@@ -1746,11 +1746,15 @@ public:
 	void RemoveNetworkActor( AActor* Actor );
 
 	/** Add a listener for OnActorSpawned events */
-	void AddOnActorSpawnedHandler( const FOnActorSpawned::FDelegate& InHandler );
+	FDelegateHandle AddOnActorSpawnedHandler( const FOnActorSpawned::FDelegate& InHandler );
 
 	/** Remove a listener for OnActorSpawned events */
+	DELEGATE_DEPRECATED("This overload of RemoveOnActorSpawnedHandler is deprecated, instead pass the result of AddOnActorSpawnedHandler.")
 	void RemoveOnActorSpawnedHandler( const FOnActorSpawned::FDelegate& InHandler );
 
+	/** Remove a listener for OnActorSpawned events */
+	void RemoveOnActorSpawnedHandler( FDelegateHandle InHandle );
+
 	/**
 	 * Returns whether the passed in actor is part of any of the loaded levels actors array.
 	 * Warning: Will return true for pending kill actors!
diff --git a/Engine/Source/Runtime/Engine/Classes/GameFramework/GameSession.h b/Engine/Source/Runtime/Engine/Classes/GameFramework/GameSession.h
--- a/Engine/Source/Runtime/Engine/Classes/GameFramework/GameSession.h
+++ b/Engine/Source/Runtime/Engine/Classes/GameFramework/GameSession.h
@@ -190,6 +190,8 @@ class ENGINE_API AGameSession : public AInfo
 private:
 	// Hidden functions that don't make sense to use on this class.
 	HIDE_ACTOR_TRANSFORM_FUNCTIONS();
+
+	FDelegateHandle OnLoginCompleteDelegateHandle;
 };
 
 /** 
diff --git a/Engine/Source/Runtime/Engine/Private/AI/Navigation/NavTestRenderingComponent.cpp b/Engine/Source/Runtime/Engine/Private/AI/Navigation/NavTestRenderingComponent.cpp
--- a/Engine/Source/Runtime/Engine/Private/AI/Navigation/NavTestRenderingComponent.cpp
+++ b/Engine/Source/Runtime/Engine/Private/AI/Navigation/NavTestRenderingComponent.cpp
@@ -69,14 +69,14 @@ public:
 	virtual void RegisterDebugDrawDelgate() override
 	{
 		DebugTextDrawingDelegate = FDebugDrawDelegate::CreateRaw(this, &FNavTestSceneProxy::DrawDebugLabels);
-		UDebugDrawService::Register(TEXT("Navigation"), DebugTextDrawingDelegate);
+		DebugTextDrawingDelegateHandle = UDebugDrawService::Register(TEXT("Navigation"), DebugTextDrawingDelegate);
 	}
 
 	virtual void UnregisterDebugDrawDelgate() override
 	{
 		if (DebugTextDrawingDelegate.IsBound())
 		{
-			UDebugDrawService::Unregister(DebugTextDrawingDelegate);
+			UDebugDrawService::Unregister(DebugTextDrawingDelegateHandle);
 		}
 	}
 
@@ -406,6 +406,7 @@ private:
 	FVector NavMeshDrawOffset;
 	ANavigationTestingActor* NavTestActor;
 	FDebugDrawDelegate DebugTextDrawingDelegate;
+	FDelegateHandle DebugTextDrawingDelegateHandle;
 	TArray<FVector> PathPoints;
 	TArray<FString> PathPointFlags;
 
diff --git a/Engine/Source/Runtime/Engine/Private/AI/Navigation/NavigationPath.cpp b/Engine/Source/Runtime/Engine/Private/AI/Navigation/NavigationPath.cpp
--- a/Engine/Source/Runtime/Engine/Private/AI/Navigation/NavigationPath.cpp
+++ b/Engine/Source/Runtime/Engine/Private/AI/Navigation/NavigationPath.cpp
@@ -1127,7 +1127,7 @@ void UNavigationPath::BeginDestroy()
 {
 	if (SharedPath.IsValid())
 	{
-		SharedPath->RemoveObserver(PathObserver);
+		SharedPath->RemoveObserver(PathObserverDelegateHandle);
 	}
 	Super::BeginDestroy();
 }
@@ -1183,11 +1183,11 @@ void UNavigationPath::EnableDebugDrawing(bool bShouldDrawDebugData, FLinearColor
 	bDebugDrawingEnabled = bShouldDrawDebugData;
 	if (bShouldDrawDebugData)
 	{
-		UDebugDrawService::Register(TEXT("Navigation"), FDebugDrawDelegate::CreateUObject(this, &UNavigationPath::DrawDebug));
+		DrawDebugDelegateHandle = UDebugDrawService::Register(TEXT("Navigation"), FDebugDrawDelegate::CreateUObject(this, &UNavigationPath::DrawDebug));
 	}
 	else
 	{
-		UDebugDrawService::Unregister(FDebugDrawDelegate::CreateUObject(this, &UNavigationPath::DrawDebug));
+		UDebugDrawService::Unregister(DrawDebugDelegateHandle);
 	}
 }
 
@@ -1239,12 +1239,12 @@ void UNavigationPath::SetPath(FNavPathSharedPtr NewSharedPath)
 	{
 		if (SharedPath.IsValid())
 		{
-			SharedPath->RemoveObserver(PathObserver);
+			SharedPath->RemoveObserver(PathObserverDelegateHandle);
 		}
 		SharedPath = NewSharedPath;
 		if (NewPath != NULL)
 		{
-			NewPath->AddObserver(PathObserver);
+			PathObserverDelegateHandle = NewPath->AddObserver(PathObserver);
 
 			if (RecalculateOnInvalidation != ENavigationOptionFlag::Default)
 			{
diff --git a/Engine/Source/Runtime/Engine/Private/Debug/DebugDrawService.cpp b/Engine/Source/Runtime/Engine/Private/Debug/DebugDrawService.cpp
--- a/Engine/Source/Runtime/Engine/Private/Debug/DebugDrawService.cpp
+++ b/Engine/Source/Runtime/Engine/Private/Debug/DebugDrawService.cpp
@@ -13,12 +13,13 @@ UDebugDrawService::UDebugDrawService(const FObjectInitializer& ObjectInitializer
 	Delegates.Reserve(sizeof(FEngineShowFlags)*8);
 }
 
-void UDebugDrawService::Register(const TCHAR* Name, const FDebugDrawDelegate& NewDelegate)
+FDelegateHandle UDebugDrawService::Register(const TCHAR* Name, const FDebugDrawDelegate& NewDelegate)
 {
 	check(IsInGameThread());
 
 	int32 Index = FEngineShowFlags::FindIndexByName(Name);
 
+	FDelegateHandle Result;
 	if (Index != INDEX_NONE)
 	{
 		if (Index >= Delegates.Num())
@@ -26,8 +27,10 @@ void UDebugDrawService::Register(const TCHAR* Name, const FDebugDrawDelegate& Ne
 			Delegates.AddZeroed(Index - Delegates.Num() + 1);
 		}
 		Delegates[Index].Add(NewDelegate);
+		Result = Delegates[Index].Last().GetHandle();
 		ObservedFlags.SetSingleFlag(Index, true);
 	}
+	return Result;
 }
 
 void UDebugDrawService::Unregister(const FDebugDrawDelegate& DelegateToRemove)
@@ -38,7 +41,27 @@ void UDebugDrawService::Unregister(const FDebugDrawDelegate& DelegateToRemove)
 	for (int32 Flag = 0; Flag < Delegates.Num(); ++Flag, ++DelegatesArray)
 	{
 		check(DelegatesArray); //it shouldn't happen, but to be sure
-		const uint32 Index = DelegatesArray->Find(DelegateToRemove);
+		const uint32 Index = DelegatesArray->IndexOfByPredicate([&](const FDebugDrawDelegate& Delegate){ return Delegate.DEPRECATED_Compare(DelegateToRemove); });
+		if (Index != INDEX_NONE)
+		{
+			DelegatesArray->RemoveAtSwap(Index, 1, false);
+			if (DelegatesArray->Num() == 0)
+			{
+				ObservedFlags.SetSingleFlag(Flag, false);
+			}
+		}
+	}	
+}
+
+void UDebugDrawService::Unregister(FDelegateHandle HandleToRemove)
+{
+	check(IsInGameThread());
+
+	TArray<FDebugDrawDelegate>* DelegatesArray = Delegates.GetData();
+	for (int32 Flag = 0; Flag < Delegates.Num(); ++Flag, ++DelegatesArray)
+	{
+		check(DelegatesArray); //it shouldn't happen, but to be sure
+		const uint32 Index = DelegatesArray->IndexOfByPredicate([=](const FDebugDrawDelegate& Delegate){ return Delegate.GetHandle() == HandleToRemove; });
 		if (Index != INDEX_NONE)
 		{
 			DelegatesArray->RemoveAtSwap(Index, 1, false);
diff --git a/Engine/Source/Runtime/Engine/Private/DebugRenderSceneProxy.cpp b/Engine/Source/Runtime/Engine/Private/DebugRenderSceneProxy.cpp
--- a/Engine/Source/Runtime/Engine/Private/DebugRenderSceneProxy.cpp
+++ b/Engine/Source/Runtime/Engine/Private/DebugRenderSceneProxy.cpp
@@ -26,14 +26,14 @@ FDebugRenderSceneProxy::FDebugRenderSceneProxy(const UPrimitiveComponent* InComp
 void FDebugRenderSceneProxy::RegisterDebugDrawDelgate()
 {
 	DebugTextDrawingDelegate = FDebugDrawDelegate::CreateRaw(this, &FDebugRenderSceneProxy::DrawDebugLabels);
-	UDebugDrawService::Register(*ViewFlagName, DebugTextDrawingDelegate);
+	DebugTextDrawingDelegateHandle = UDebugDrawService::Register(*ViewFlagName, DebugTextDrawingDelegate);
 }
 
 void FDebugRenderSceneProxy::UnregisterDebugDrawDelgate()
 {
 	if (DebugTextDrawingDelegate.IsBound())
 	{
-		UDebugDrawService::Unregister(DebugTextDrawingDelegate);
+		UDebugDrawService::Unregister(DebugTextDrawingDelegateHandle);
 	}
 }
 
diff --git a/Engine/Source/Runtime/Engine/Private/EdGraph/EdGraph.cpp b/Engine/Source/Runtime/Engine/Private/EdGraph/EdGraph.cpp
--- a/Engine/Source/Runtime/Engine/Private/EdGraph/EdGraph.cpp
+++ b/Engine/Source/Runtime/Engine/Private/EdGraph/EdGraph.cpp
@@ -106,14 +106,19 @@ const UEdGraphSchema* UEdGraph::GetSchema() const
 	return GetDefault<UEdGraphSchema>(Schema);
 }
 
-void UEdGraph::AddOnGraphChangedHandler( const FOnGraphChanged::FDelegate& InHandler )
+FDelegateHandle UEdGraph::AddOnGraphChangedHandler( const FOnGraphChanged::FDelegate& InHandler )
 {
-	OnGraphChanged.Add( InHandler );
+	return OnGraphChanged.Add( InHandler );
 }
 
 void UEdGraph::RemoveOnGraphChangedHandler( const FOnGraphChanged::FDelegate& InHandler )
 {
-	OnGraphChanged.Remove( InHandler );
+	OnGraphChanged.DEPRECATED_Remove( InHandler );
+}
+
+void UEdGraph::RemoveOnGraphChangedHandler( FDelegateHandle Handle )
+{
+	OnGraphChanged.Remove( Handle );
 }
 
 UEdGraphNode* UEdGraph::CreateNode( TSubclassOf<UEdGraphNode> NewNodeClass, bool bSelectNewNode/* = true*/ )
@@ -257,17 +262,26 @@ void UEdGraph::NotifyPostChange( const FPropertyChangedEvent& PropertyChangedEve
 #endif
 }
 
-void UEdGraph::AddPropertyChangedNotifier(const FOnPropertyChanged::FDelegate& InDelegate )
+FDelegateHandle UEdGraph::AddPropertyChangedNotifier(const FOnPropertyChanged::FDelegate& InDelegate )
 {
 #if WITH_EDITORONLY_DATA
-	PropertyChangedNotifiers.Add(InDelegate);
+	return PropertyChangedNotifiers.Add(InDelegate);
+#else
+	return FDelegateHandle();
 #endif
 }
 
 void UEdGraph::RemovePropertyChangedNotifier(const FOnPropertyChanged::FDelegate& InDelegate )
 {
 #if WITH_EDITORONLY_DATA
-	PropertyChangedNotifiers.Remove(InDelegate);
+	PropertyChangedNotifiers.DEPRECATED_Remove(InDelegate);
+#endif
+}
+
+void UEdGraph::RemovePropertyChangedNotifier(FDelegateHandle Handle )
+{
+#if WITH_EDITORONLY_DATA
+	PropertyChangedNotifiers.Remove(Handle);
 #endif
 }
 #endif	//WITH_EDITOR
diff --git a/Engine/Source/Runtime/Engine/Private/GameMode.cpp b/Engine/Source/Runtime/Engine/Private/GameMode.cpp
--- a/Engine/Source/Runtime/Engine/Private/GameMode.cpp
+++ b/Engine/Source/Runtime/Engine/Private/GameMode.cpp
@@ -179,12 +179,9 @@ bool AGameMode::SetPause(APlayerController* PC, FCanUnpause CanUnpauseDelegate)
 {
 	if ( AllowPausing(PC) )
 	{
-		// Don't add the delegate twice (no need)
-		if (Pausers.Find(CanUnpauseDelegate) == INDEX_NONE)
-		{
-			// Not in the list so add it for querying
-			Pausers.Add(CanUnpauseDelegate);
-		}
+		// Add it for querying
+		Pausers.Add(CanUnpauseDelegate);
+
 		// Let the first one in "own" the pause state
 		AWorldSettings * WorldSettings = GetWorldSettings();
 		if (WorldSettings->Pauser == NULL)
@@ -1726,4 +1723,4 @@ void AGameMode::DefaultTimer()
 FString AGameMode::GetRedirectURL(const FString& MapName) const
 {
 	return FString();
-}
\ No newline at end of file
+}
diff --git a/Engine/Source/Runtime/Engine/Private/GameSession.cpp b/Engine/Source/Runtime/Engine/Private/GameSession.cpp
--- a/Engine/Source/Runtime/Engine/Private/GameSession.cpp
+++ b/Engine/Source/Runtime/Engine/Private/GameSession.cpp
@@ -71,7 +71,7 @@ bool AGameSession::ProcessAutoLogin()
 	IOnlineIdentityPtr IdentityInt = Online::GetIdentityInterface(World);
 	if (IdentityInt.IsValid())
 	{
-		IdentityInt->AddOnLoginCompleteDelegate(0, FOnLoginCompleteDelegate::CreateUObject(this, &AGameSession::OnLoginComplete));
+		OnLoginCompleteDelegateHandle = IdentityInt->AddOnLoginCompleteDelegate_Handle(0, FOnLoginCompleteDelegate::CreateUObject(this, &AGameSession::OnLoginComplete));
 		if (!IdentityInt->AutoLogin(0))
 		{
 			// Not waiting for async login
@@ -91,7 +91,7 @@ void AGameSession::OnLoginComplete(int32 LocalUserNum, bool bWasSuccessful, cons
 	IOnlineIdentityPtr IdentityInt = Online::GetIdentityInterface(World);
 	if (IdentityInt.IsValid())
 	{
-		IdentityInt->ClearOnLoginCompleteDelegate(0, FOnLoginCompleteDelegate::CreateUObject(this, &AGameSession::OnLoginComplete));
+		IdentityInt->ClearOnLoginCompleteDelegate_Handle(0, OnLoginCompleteDelegateHandle);
 		if (IdentityInt->GetLoginStatus(0) == ELoginStatus::LoggedIn)
 		{
 			RegisterServer();
diff --git a/Engine/Source/Runtime/Engine/Private/GameViewportClient.cpp b/Engine/Source/Runtime/Engine/Private/GameViewportClient.cpp
--- a/Engine/Source/Runtime/Engine/Private/GameViewportClient.cpp
+++ b/Engine/Source/Runtime/Engine/Private/GameViewportClient.cpp
@@ -2201,13 +2201,13 @@ void UGameViewportClient::ToggleShowCollision()
 	if (bIsShowingCollision)
 	{
 		NumViewportsShowingCollision++;
-		GetWorld()->AddOnActorSpawnedHandler(FOnActorSpawned::FDelegate::CreateUObject(this, &UGameViewportClient::ShowCollisionOnSpawnedActors));
+		ShowCollisionOnSpawnedActorsDelegateHandle = GetWorld()->AddOnActorSpawnedHandler(FOnActorSpawned::FDelegate::CreateUObject(this, &UGameViewportClient::ShowCollisionOnSpawnedActors));
 	}
 	else
 	{
 		NumViewportsShowingCollision--;
 		check(NumViewportsShowingCollision >= 0);
-		GetWorld()->RemoveOnActorSpawnedHandler(FOnActorSpawned::FDelegate::CreateUObject(this, &UGameViewportClient::ShowCollisionOnSpawnedActors));
+		GetWorld()->RemoveOnActorSpawnedHandler(ShowCollisionOnSpawnedActorsDelegateHandle);
 	}
 
 	CollisionComponentVisibilityMap& Mapping = GetCollisionComponentVisibilityMap();
diff --git a/Engine/Source/Runtime/Engine/Private/KismetSystemLibrary.cpp b/Engine/Source/Runtime/Engine/Private/KismetSystemLibrary.cpp
--- a/Engine/Source/Runtime/Engine/Private/KismetSystemLibrary.cpp
+++ b/Engine/Source/Runtime/Engine/Private/KismetSystemLibrary.cpp
@@ -291,7 +291,9 @@ void UKismetSystemLibrary::K2_SetTimerDelegate(FTimerDynamicDelegate Delegate, f
 	if (Delegate.IsBound())
 	{
 		const UWorld* const World = GEngine->GetWorldFromContextObject(Delegate.GetUObject());
-		World->GetTimerManager().SetTimer(Delegate, Time, bLooping);
+		auto& TimerManager = World->GetTimerManager();
+		auto Handle = TimerManager.K2_FindDynamicTimerHandle(Delegate);
+		TimerManager.SetTimer(Handle, Delegate, Time, bLooping);
 	}
 	else
 	{
@@ -316,7 +318,9 @@ void UKismetSystemLibrary::K2_ClearTimerDelegate(FTimerDynamicDelegate Delegate)
 		UWorld* World = GEngine->GetWorldFromContextObject(Delegate.GetUObject());
 		if (World)
 		{
-			World->GetTimerManager().ClearTimer(Delegate);
+			auto& TimerManager = World->GetTimerManager();
+			auto Handle = TimerManager.K2_FindDynamicTimerHandle(Delegate);
+			World->GetTimerManager().ClearTimer(Handle);
 		}
 	}
 	else
@@ -340,7 +344,9 @@ void UKismetSystemLibrary::K2_PauseTimerDelegate(FTimerDynamicDelegate Delegate)
 	if (Delegate.IsBound())
 	{
 		UWorld* World = GEngine->GetWorldFromContextObject(Delegate.GetUObject());
-		World->GetTimerManager().PauseTimer(Delegate);
+		auto& TimerManager = World->GetTimerManager();
+		auto Handle = TimerManager.K2_FindDynamicTimerHandle(Delegate);
+		World->GetTimerManager().PauseTimer(Handle);
 	}
 	else
 	{
@@ -363,7 +369,9 @@ void UKismetSystemLibrary::K2_UnPauseTimerDelegate(FTimerDynamicDelegate Delegat
 	if (Delegate.IsBound())
 	{
 		UWorld* World = GEngine->GetWorldFromContextObject(Delegate.GetUObject());
-		World->GetTimerManager().UnPauseTimer(Delegate);
+		auto& TimerManager = World->GetTimerManager();
+		auto Handle = TimerManager.K2_FindDynamicTimerHandle(Delegate);
+		World->GetTimerManager().UnPauseTimer(Handle);
 	}
 	else
 	{
@@ -386,7 +394,9 @@ bool UKismetSystemLibrary::K2_IsTimerActiveDelegate(FTimerDynamicDelegate Delega
 	if (Delegate.IsBound())
 	{
 		UWorld* World = GEngine->GetWorldFromContextObject(Delegate.GetUObject());
-		return World->GetTimerManager().IsTimerActive(Delegate);
+		auto& TimerManager = World->GetTimerManager();
+		auto Handle = TimerManager.K2_FindDynamicTimerHandle(Delegate);
+		return World->GetTimerManager().IsTimerActive(Handle);
 	}
 	else
 	{
@@ -410,7 +420,9 @@ bool UKismetSystemLibrary::K2_IsTimerPausedDelegate(FTimerDynamicDelegate Delega
 	if (Delegate.IsBound())
 	{
 		UWorld* World = GEngine->GetWorldFromContextObject(Delegate.GetUObject());
-		return World->GetTimerManager().IsTimerPaused(Delegate);
+		auto& TimerManager = World->GetTimerManager();
+		auto Handle = TimerManager.K2_FindDynamicTimerHandle(Delegate);
+		return World->GetTimerManager().IsTimerPaused(Handle);
 	}
 	else
 	{
@@ -434,7 +446,9 @@ bool UKismetSystemLibrary::K2_TimerExistsDelegate(FTimerDynamicDelegate Delegate
 	if (Delegate.IsBound())
 	{
 		UWorld* World = GEngine->GetWorldFromContextObject(Delegate.GetUObject());
-		return World->GetTimerManager().TimerExists(Delegate);
+		auto& TimerManager = World->GetTimerManager();
+		auto Handle = TimerManager.K2_FindDynamicTimerHandle(Delegate);
+		return World->GetTimerManager().TimerExists(Handle);
 	}
 	else
 	{
@@ -458,7 +472,9 @@ float UKismetSystemLibrary::K2_GetTimerElapsedTimeDelegate(FTimerDynamicDelegate
 	if (Delegate.IsBound())
 	{
 		UWorld* World = GEngine->GetWorldFromContextObject(Delegate.GetUObject());
-		return World->GetTimerManager().GetTimerElapsed(Delegate);
+		auto& TimerManager = World->GetTimerManager();
+		auto Handle = TimerManager.K2_FindDynamicTimerHandle(Delegate);
+		return World->GetTimerManager().GetTimerElapsed(Handle);
 	}
 	else
 	{
@@ -482,7 +498,9 @@ float UKismetSystemLibrary::K2_GetTimerRemainingTimeDelegate(FTimerDynamicDelega
 	if (Delegate.IsBound())
 	{
 		UWorld* World = GEngine->GetWorldFromContextObject(Delegate.GetUObject());
-		return World->GetTimerManager().GetTimerRemaining(Delegate);
+		auto& TimerManager = World->GetTimerManager();
+		auto Handle = TimerManager.K2_FindDynamicTimerHandle(Delegate);
+		return World->GetTimerManager().GetTimerRemaining(Handle);
 	}
 	else
 	{
diff --git a/Engine/Source/Runtime/Engine/Private/LevelBounds.cpp b/Engine/Source/Runtime/Engine/Private/LevelBounds.cpp
--- a/Engine/Source/Runtime/Engine/Private/LevelBounds.cpp
+++ b/Engine/Source/Runtime/Engine/Private/LevelBounds.cpp
@@ -195,9 +195,9 @@ void ALevelBounds::SubscribeToUpdateEvents()
 	if (bSubscribedToEvents == false && GIsEditor && !GetWorld()->IsPlayInEditor())
 	{
 		GetWorldTimerManager().SetTimer(TimerHandle_OnTimerTick, this, &ALevelBounds::OnTimerTick, 1.f, true);
-		GEngine->OnActorMoved().AddUObject(this, &ALevelBounds::OnLevelActorMoved);
-		GEngine->OnLevelActorDeleted().AddUObject(this, &ALevelBounds::OnLevelActorAddedRemoved);
-		GEngine->OnLevelActorAdded().AddUObject(this, &ALevelBounds::OnLevelActorAddedRemoved);
+		OnLevelActorMovedDelegateHandle   = GEngine->OnActorMoved       ().AddUObject(this, &ALevelBounds::OnLevelActorMoved);
+		OnLevelActorDeletedDelegateHandle = GEngine->OnLevelActorDeleted().AddUObject(this, &ALevelBounds::OnLevelActorAddedRemoved);
+		OnLevelActorAddedDelegateHandle   = GEngine->OnLevelActorAdded  ().AddUObject(this, &ALevelBounds::OnLevelActorAddedRemoved);
 		
 		bSubscribedToEvents = true;
 	}
@@ -211,9 +211,9 @@ void ALevelBounds::UnsubscribeFromUpdateEvents()
 		{
 			GetWorldTimerManager().ClearTimer(TimerHandle_OnTimerTick);
 		}
-		GEngine->OnActorMoved().RemoveUObject(this, &ALevelBounds::OnLevelActorMoved);
-		GEngine->OnLevelActorDeleted().RemoveUObject(this, &ALevelBounds::OnLevelActorAddedRemoved);
-		GEngine->OnLevelActorAdded().RemoveUObject(this, &ALevelBounds::OnLevelActorAddedRemoved);
+		GEngine->OnActorMoved       ().Remove(OnLevelActorMovedDelegateHandle);
+		GEngine->OnLevelActorDeleted().Remove(OnLevelActorDeletedDelegateHandle);
+		GEngine->OnLevelActorAdded  ().Remove(OnLevelActorAddedDelegateHandle);
 
 		bSubscribedToEvents = false;
 	}
diff --git a/Engine/Source/Runtime/Engine/Private/LevelTick.cpp b/Engine/Source/Runtime/Engine/Private/LevelTick.cpp
--- a/Engine/Source/Runtime/Engine/Private/LevelTick.cpp
+++ b/Engine/Source/Runtime/Engine/Private/LevelTick.cpp
@@ -119,7 +119,7 @@ FDetailedTickStats::FDetailedTickStats( int32 InNumObjectsToReport, float InTime
 FDetailedTickStats::~FDetailedTickStats()
 {
 	// remove callback as we are dead
-	FCoreUObjectDelegates::PreGarbageCollect.RemoveRaw(this, &FDetailedTickStats::OnPreGarbageCollect);
+	FCoreUObjectDelegates::PreGarbageCollect.Remove(OnPreGarbageCollectDelegateHandle);
 }
 
 /**
@@ -179,7 +179,7 @@ void FDetailedTickStats::EndObject( UObject* Object, float DeltaTime, bool bForS
 		{
 			GCCallBackRegistered = true;
 			// register callback so that we can avoid finding the wrong stats for new objects reusing memory that used to be associated with a different object
-			FCoreUObjectDelegates::PreGarbageCollect.AddRaw(this, &FDetailedTickStats::OnPreGarbageCollect);
+			OnPreGarbageCollectDelegateHandle = FCoreUObjectDelegates::PreGarbageCollect.AddRaw(this, &FDetailedTickStats::OnPreGarbageCollect);
 		}
 
 		FTickStats NewTickStats;
diff --git a/Engine/Source/Runtime/Engine/Private/NavMeshRenderingHelpers.h b/Engine/Source/Runtime/Engine/Private/NavMeshRenderingHelpers.h
--- a/Engine/Source/Runtime/Engine/Private/NavMeshRenderingHelpers.h
+++ b/Engine/Source/Runtime/Engine/Private/NavMeshRenderingHelpers.h
@@ -316,14 +316,14 @@ public:
 	void RegisterDebugDrawDelgate()
 	{
 		DebugTextDrawingDelegate = FDebugDrawDelegate::CreateRaw(this, &FRecastRenderingSceneProxy::DrawDebugLabels);
-		UDebugDrawService::Register(TEXT("Navigation"), DebugTextDrawingDelegate);
+		DebugTextDrawingDelegateHandle = UDebugDrawService::Register(TEXT("Navigation"), DebugTextDrawingDelegate);
 	}
 
 	void UnregisterDebugDrawDelgate()
 	{
 		if (DebugTextDrawingDelegate.IsBound())
 		{
-			UDebugDrawService::Unregister(DebugTextDrawingDelegate);
+			UDebugDrawService::Unregister(DebugTextDrawingDelegateHandle);
 		}
 	}
 
@@ -584,6 +584,7 @@ private:
 	TArray<FMeshBatchElement> MeshBatchElements;
 
 	FDebugDrawDelegate DebugTextDrawingDelegate;
+	FDelegateHandle DebugTextDrawingDelegateHandle;
 	TWeakObjectPtr<UNavMeshRenderingComponent> RenderingComponent;
 	uint32 bRequestedData : 1;
 	uint32 bForceRendering : 1;
diff --git a/Engine/Source/Runtime/Engine/Private/NetworkDriver.cpp b/Engine/Source/Runtime/Engine/Private/NetworkDriver.cpp
--- a/Engine/Source/Runtime/Engine/Private/NetworkDriver.cpp
+++ b/Engine/Source/Runtime/Engine/Private/NetworkDriver.cpp
@@ -504,23 +504,23 @@ ENetMode UNetDriver::GetNetMode() const
 	return (IsServer() ? (GIsClient ? NM_ListenServer : NM_DedicatedServer) : NM_Client);
 }
 
-void UNetDriver::RegisterTickEvents(class UWorld* InWorld) const
+void UNetDriver::RegisterTickEvents(class UWorld* InWorld)
 {
 	if (InWorld)
 	{
-		InWorld->OnTickDispatch().AddUObject(this, &UNetDriver::TickDispatch);
-		InWorld->OnTickFlush().AddUObject(this, &UNetDriver::TickFlush);
-		InWorld->OnPostTickFlush().AddUObject(this, &UNetDriver::PostTickFlush);
+		TickDispatchDelegateHandle  = InWorld->OnTickDispatch ().AddUObject(this, &UNetDriver::TickDispatch);
+		TickFlushDelegateHandle     = InWorld->OnTickFlush    ().AddUObject(this, &UNetDriver::TickFlush);
+		PostTickFlushDelegateHandle = InWorld->OnPostTickFlush().AddUObject(this, &UNetDriver::PostTickFlush);
 	}
 }
 
-void UNetDriver::UnregisterTickEvents(class UWorld* InWorld) const
+void UNetDriver::UnregisterTickEvents(class UWorld* InWorld)
 {
 	if (InWorld)
 	{
-		InWorld->OnTickDispatch().RemoveUObject(this, &UNetDriver::TickDispatch);
-		InWorld->OnTickFlush().RemoveUObject(this, &UNetDriver::TickFlush);
-		InWorld->OnPostTickFlush().RemoveUObject(this, &UNetDriver::PostTickFlush);
+		InWorld->OnTickDispatch ().Remove(TickDispatchDelegateHandle);
+		InWorld->OnTickFlush    ().Remove(TickFlushDelegateHandle);
+		InWorld->OnPostTickFlush().Remove(PostTickFlushDelegateHandle);
 	}
 }
 
diff --git a/Engine/Source/Runtime/Engine/Private/PhysicsEngine2D/Box2DIntegration.cpp b/Engine/Source/Runtime/Engine/Private/PhysicsEngine2D/Box2DIntegration.cpp
--- a/Engine/Source/Runtime/Engine/Private/PhysicsEngine2D/Box2DIntegration.cpp
+++ b/Engine/Source/Runtime/Engine/Private/PhysicsEngine2D/Box2DIntegration.cpp
@@ -278,6 +278,8 @@ FString FEndPhysics2DTickFunction::DiagnosticMessage()
 TMap< UWorld*, TSharedPtr<FPhysicsScene2D> > FPhysicsIntegration2D::WorldMappings;
 FWorldDelegates::FWorldInitializationEvent::FDelegate FPhysicsIntegration2D::OnWorldCreatedDelegate;
 FWorldDelegates::FWorldEvent::FDelegate FPhysicsIntegration2D::OnWorldDestroyedDelegate;
+FDelegateHandle FPhysicsIntegration2D::OnWorldCreatedDelegateHandle;
+FDelegateHandle FPhysicsIntegration2D::OnWorldDestroyedDelegateHandle;
 
 void FPhysicsIntegration2D::OnWorldCreated(UWorld* UnrealWorld, const UWorld::InitializationValues IVS)
 {
@@ -294,17 +296,17 @@ void FPhysicsIntegration2D::OnWorldDestroyed(UWorld* UnrealWorld)
 
 void FPhysicsIntegration2D::InitializePhysics()
 {
-	OnWorldCreatedDelegate = FWorldDelegates::FWorldInitializationEvent::FDelegate::CreateStatic(&FPhysicsIntegration2D::OnWorldCreated);
+	OnWorldCreatedDelegate   = FWorldDelegates::FWorldInitializationEvent::FDelegate::CreateStatic(&FPhysicsIntegration2D::OnWorldCreated);
 	OnWorldDestroyedDelegate = FWorldDelegates::FWorldEvent::FDelegate::CreateStatic(&FPhysicsIntegration2D::OnWorldDestroyed);
 
-	FWorldDelegates::OnPreWorldInitialization.Add(OnWorldCreatedDelegate);
-	FWorldDelegates::OnPreWorldFinishDestroy.Add(OnWorldDestroyedDelegate);
+	OnWorldCreatedDelegateHandle   = FWorldDelegates::OnPreWorldInitialization.Add(OnWorldCreatedDelegate);
+	OnWorldDestroyedDelegateHandle = FWorldDelegates::OnPreWorldFinishDestroy .Add(OnWorldDestroyedDelegate);
 }
 
 void FPhysicsIntegration2D::ShutdownPhysics()
 {
-	FWorldDelegates::OnPreWorldInitialization.Remove(OnWorldCreatedDelegate);
-	FWorldDelegates::OnPreWorldFinishDestroy.Remove(OnWorldDestroyedDelegate);
+	FWorldDelegates::OnPreWorldInitialization.Remove(OnWorldCreatedDelegateHandle);
+	FWorldDelegates::OnPreWorldFinishDestroy .Remove(OnWorldDestroyedDelegateHandle);
 
 	check(WorldMappings.Num() == 0);
 }
diff --git a/Engine/Source/Runtime/Engine/Private/PhysicsEngine2D/Box2DIntegration.h b/Engine/Source/Runtime/Engine/Private/PhysicsEngine2D/Box2DIntegration.h
--- a/Engine/Source/Runtime/Engine/Private/PhysicsEngine2D/Box2DIntegration.h
+++ b/Engine/Source/Runtime/Engine/Private/PhysicsEngine2D/Box2DIntegration.h
@@ -79,6 +79,9 @@ private:
 	static FWorldDelegates::FWorldInitializationEvent::FDelegate OnWorldCreatedDelegate;
 	static FWorldDelegates::FWorldEvent::FDelegate OnWorldDestroyedDelegate;
 
+	static FDelegateHandle OnWorldCreatedDelegateHandle;
+	static FDelegateHandle OnWorldDestroyedDelegateHandle;
+
 public:
 
 	//@TODO: BOX2D: These conversion functions ignore the Paper2D coordinate system (they are fixed on a XZ coordinate system)
diff --git a/Engine/Source/Runtime/Engine/Private/SkeletalMeshComponent.cpp b/Engine/Source/Runtime/Engine/Private/SkeletalMeshComponent.cpp
--- a/Engine/Source/Runtime/Engine/Private/SkeletalMeshComponent.cpp
+++ b/Engine/Source/Runtime/Engine/Private/SkeletalMeshComponent.cpp
@@ -1872,14 +1872,19 @@ void USkeletalMeshComponent::UpdateCollisionProfile()
 	}
 }
 
-void USkeletalMeshComponent::RegisterOnSkeletalMeshPropertyChanged( const FOnSkeletalMeshPropertyChanged& Delegate )
+FDelegateHandle USkeletalMeshComponent::RegisterOnSkeletalMeshPropertyChanged( const FOnSkeletalMeshPropertyChanged& Delegate )
 {
-	OnSkeletalMeshPropertyChanged.Add(Delegate);
+	return OnSkeletalMeshPropertyChanged.Add(Delegate);
 }
 
 void USkeletalMeshComponent::UnregisterOnSkeletalMeshPropertyChanged( const FOnSkeletalMeshPropertyChanged& Delegate )
 {
-	OnSkeletalMeshPropertyChanged.Remove(Delegate);
+	OnSkeletalMeshPropertyChanged.DEPRECATED_Remove(Delegate);
+}
+
+void USkeletalMeshComponent::UnregisterOnSkeletalMeshPropertyChanged( FDelegateHandle Handle )
+{
+	OnSkeletalMeshPropertyChanged.Remove(Handle);
 }
 
 void USkeletalMeshComponent::ValidateAnimation()
diff --git a/Engine/Source/Runtime/Engine/Private/SystemSettings.cpp b/Engine/Source/Runtime/Engine/Private/SystemSettings.cpp
--- a/Engine/Source/Runtime/Engine/Private/SystemSettings.cpp
+++ b/Engine/Source/Runtime/Engine/Private/SystemSettings.cpp
@@ -339,7 +339,7 @@ void FSystemSettings::Initialize( bool bSetupForEditor )
 
 	ApplyOverrides();
 
-	IConsoleManager::Get().RegisterConsoleVariableSink(FConsoleCommandDelegate::CreateRaw(this, &FSystemSettings::CVarSink));
+	IConsoleManager::Get().RegisterConsoleVariableSink_Handle(FConsoleCommandDelegate::CreateRaw(this, &FSystemSettings::CVarSink));
 
 	// intialize a critical texture streaming value used by texture loading, etc
 	int32 MinTextureResidentMipCount = 1;
diff --git a/Engine/Source/Runtime/Engine/Private/TimerManager.cpp b/Engine/Source/Runtime/Engine/Private/TimerManager.cpp
--- a/Engine/Source/Runtime/Engine/Private/TimerManager.cpp
+++ b/Engine/Source/Runtime/Engine/Private/TimerManager.cpp
@@ -20,14 +20,19 @@ void FTimerHandle::MakeValid()
 }
 
 
+inline bool DEPRECATED_CompareUnifiedDelegates(const FTimerUnifiedDelegate& Lhs, const FTimerUnifiedDelegate& Rhs)
+{
+	return ( Lhs.FuncDelegate.IsBound() && (Lhs.FuncDelegate.DEPRECATED_Compare(Rhs.FuncDelegate)) ) || ( Lhs.FuncDynDelegate.IsBound() && (Lhs.FuncDynDelegate == Rhs.FuncDynDelegate) );
+}
+
 // ---------------------------------
 // Private members
 // ---------------------------------
 
 /** Will find and return a timer if it exists, regardless whether it is paused. */ 
-FTimerData const* FTimerManager::FindTimer(FTimerUnifiedDelegate const& InDelegate, int32* OutTimerIndex) const
+FTimerData const* FTimerManager::DEPRECATED_FindTimer(FTimerUnifiedDelegate const& InDelegate, int32* OutTimerIndex) const
 {
-	int32 ActiveTimerIdx = FindTimerInList(ActiveTimerHeap, InDelegate);
+	int32 ActiveTimerIdx = DEPRECATED_FindTimerInList(ActiveTimerHeap, InDelegate);
 	if (ActiveTimerIdx != INDEX_NONE)
 	{
 		// found it in the active heap
@@ -38,7 +43,7 @@ FTimerData const* FTimerManager::FindTimer(FTimerUnifiedDelegate const& InDelega
 		return &ActiveTimerHeap[ActiveTimerIdx];
 	}
 
-	int32 PausedTimerIdx = FindTimerInList(PausedTimerList, InDelegate);
+	int32 PausedTimerIdx = DEPRECATED_FindTimerInList(PausedTimerList, InDelegate);
 	if (PausedTimerIdx != INDEX_NONE)
 	{
 		// found it in the paused list
@@ -49,7 +54,7 @@ FTimerData const* FTimerManager::FindTimer(FTimerUnifiedDelegate const& InDelega
 		return &PausedTimerList[PausedTimerIdx];
 	}
 
-	int32 PendingTimerIdx = FindTimerInList(PendingTimerList, InDelegate);
+	int32 PendingTimerIdx = DEPRECATED_FindTimerInList(PendingTimerList, InDelegate);
 	if (PendingTimerIdx != INDEX_NONE)
 	{
 		// found it in the pending list
@@ -60,7 +65,7 @@ FTimerData const* FTimerManager::FindTimer(FTimerUnifiedDelegate const& InDelega
 		return &PendingTimerList[PendingTimerIdx];
 	}
 
-	return NULL;
+	return nullptr;
 }
 
 FTimerData const* FTimerManager::FindTimer(FTimerHandle const& InHandle, int32* OutTimerIndex) const
@@ -98,16 +103,16 @@ FTimerData const* FTimerManager::FindTimer(FTimerHandle const& InHandle, int32*
 		return &PendingTimerList[PendingTimerIdx];
 	}
 
-	return NULL;
+	return nullptr;
 }
 
 
 /** Will find the given timer in the given TArray and return its index. */ 
-int32 FTimerManager::FindTimerInList(const TArray<FTimerData> &SearchArray, FTimerUnifiedDelegate const& InDelegate) const
+int32 FTimerManager::DEPRECATED_FindTimerInList(const TArray<FTimerData> &SearchArray, FTimerUnifiedDelegate const& InDelegate) const
 {
 	for (int32 Idx=0; Idx<SearchArray.Num(); ++Idx)
 	{
-		if (!SearchArray[Idx].TimerHandle.IsValid() && SearchArray[Idx].TimerDelegate == InDelegate)
+		if (!SearchArray[Idx].TimerHandle.IsValid() && DEPRECATED_CompareUnifiedDelegates(SearchArray[Idx].TimerDelegate, InDelegate))
 		{
 			return Idx;
 		}
@@ -133,14 +138,36 @@ int32 FTimerManager::FindTimerInList(const TArray<FTimerData> &SearchArray, FTim
 	return INDEX_NONE;
 }
 
-void FTimerManager::InternalSetTimer(FTimerUnifiedDelegate const& InDelegate, float InRate, bool InbLoop, float InFirstDelay)
+/** Finds a handle to a dynamic timer bound to a particular pointer and function name. */
+FTimerHandle FTimerManager::K2_FindDynamicTimerHandle(FTimerDynamicDelegate InDynamicDelegate) const
+{
+	if (auto* Timer = ActiveTimerHeap.FindByPredicate([=](const FTimerData& Data){ return Data.TimerDelegate.FuncDynDelegate == InDynamicDelegate; }))
+	{
+		return Timer->TimerHandle;
+	}
+
+	if (auto* Timer = PausedTimerList.FindByPredicate([=](const FTimerData& Data){ return Data.TimerDelegate.FuncDynDelegate == InDynamicDelegate; }))
+	{
+		return Timer->TimerHandle;
+	}
+
+	if (auto* Timer = PendingTimerList.FindByPredicate([=](const FTimerData& Data){ return Data.TimerDelegate.FuncDynDelegate == InDynamicDelegate; }))
+	{
+		return Timer->TimerHandle;
+	}
+
+	return FTimerHandle();
+}
+
+
+void FTimerManager::DEPRECATED_InternalSetTimer(FTimerUnifiedDelegate const& InDelegate, float InRate, bool InbLoop, float InFirstDelay)
 {
 	// not currently threadsafe
 	check(IsInGameThread());
 
 	// if the timer is already set, just clear it and we'll re-add it, since 
 	// there's no data to maintain.
-	InternalClearTimer(InDelegate);
+	DEPRECATED_InternalClearTimer(InDelegate);
 
 	if (InRate > 0.f)
 	{
@@ -216,28 +243,22 @@ void FTimerManager::InternalSetTimerForNextTick(FTimerUnifiedDelegate const& InD
 	ActiveTimerHeap.HeapPush(NewTimerData);
 }
 
-void FTimerManager::InternalClearTimer(FTimerUnifiedDelegate const& InDelegate)
+void FTimerManager::DEPRECATED_InternalClearTimer(FTimerUnifiedDelegate const& InDelegate)
 {
 	// not currently threadsafe
 	check(IsInGameThread());
 
 	int32 TimerIdx;
-	FTimerData const* const TimerData = FindTimer(InDelegate, &TimerIdx);
+	FTimerData const* const TimerData = DEPRECATED_FindTimer(InDelegate, &TimerIdx);
 	if (TimerData)
 	{
-		switch (TimerData->Status)
-		{
-		case ETimerStatus::Pending: PendingTimerList.RemoveAtSwap(TimerIdx); break;
-		case ETimerStatus::Active: ActiveTimerHeap.HeapRemoveAt(TimerIdx); break;
-		case ETimerStatus::Paused: PausedTimerList.RemoveAtSwap(TimerIdx); break;
-		default: check(false);
-		}
+		InternalClearTimer(TimerIdx, TimerData->Status);
 	}
 	else
 	{
 		// Edge case. We're currently handling this timer when it got cleared.  Unbind it to prevent it firing again
 		// in case it was scheduled to fire multiple times.
-		if (!CurrentlyExecutingTimer.TimerHandle.IsValid() && CurrentlyExecutingTimer.TimerDelegate == InDelegate)
+		if (!CurrentlyExecutingTimer.TimerHandle.IsValid() && DEPRECATED_CompareUnifiedDelegates(CurrentlyExecutingTimer.TimerDelegate, InDelegate))
 		{
 			CurrentlyExecutingTimer.TimerDelegate.Unbind();
 		}
@@ -259,13 +280,7 @@ void FTimerManager::InternalClearTimer(FTimerHandle const& InHandle)
 	FTimerData const* const TimerData = FindTimer(InHandle, &TimerIdx);
 	if (TimerData)
 	{
-		switch (TimerData->Status)
-		{
-		case ETimerStatus::Pending: PendingTimerList.RemoveAtSwap(TimerIdx); break;
-		case ETimerStatus::Active: ActiveTimerHeap.HeapRemoveAt(TimerIdx); break;
-		case ETimerStatus::Paused: PausedTimerList.RemoveAtSwap(TimerIdx); break;
-		default: check(false);
-		}
+		InternalClearTimer(TimerIdx, TimerData->Status);
 	}
 	else
 	{
@@ -279,6 +294,27 @@ void FTimerManager::InternalClearTimer(FTimerHandle const& InHandle)
 	}
 }
 
+void FTimerManager::InternalClearTimer(int32 TimerIdx, ETimerStatus::Type TimerStatus)
+{
+	switch (TimerStatus)
+	{
+		case ETimerStatus::Pending:
+			PendingTimerList.RemoveAtSwap(TimerIdx);
+			break;
+
+		case ETimerStatus::Active:
+			ActiveTimerHeap.HeapRemoveAt(TimerIdx);
+			break;
+
+		case ETimerStatus::Paused:
+			PausedTimerList.RemoveAtSwap(TimerIdx);
+			break;
+
+		default:
+			check(false);
+	}
+}
+
 
 void FTimerManager::InternalClearAllTimers(void const* Object)
 {
@@ -482,7 +518,7 @@ void FTimerManager::Tick(float DeltaTime)
 				 CurrentlyExecutingTimer.TimerDelegate.IsBound()) && 							// did not get cleared during execution
 				(CurrentlyExecutingTimer.TimerHandle.IsValid() ? 
 					(FindTimer(CurrentlyExecutingTimer.TimerHandle) == nullptr) : 
-					(FindTimer(CurrentlyExecutingTimer.TimerDelegate) == nullptr)) // did not get manually re-added during execution			  
+					(DEPRECATED_FindTimer(CurrentlyExecutingTimer.TimerDelegate) == nullptr)) // did not get manually re-added during execution			  
 				)
 			{
 				// Put this timer back on the heap
diff --git a/Engine/Source/Runtime/Engine/Private/TimerManagerTests.cpp b/Engine/Source/Runtime/Engine/Private/TimerManagerTests.cpp
--- a/Engine/Source/Runtime/Engine/Private/TimerManagerTests.cpp
+++ b/Engine/Source/Runtime/Engine/Private/TimerManagerTests.cpp
@@ -36,35 +36,18 @@ public:
 bool TimerManagerTest_InvalidTimers(UWorld* World, FAutomationTestBase* Test)
 {
 	FTimerManager& TimerManager = World->GetTimerManager();
-	FTimerDelegate Delegate;
 	FTimerHandle Handle;
 
-	Test->TestFalse(TIMER_TEST_TEXT("TimerExists called with an invalid delegate"), TimerManager.TimerExists(Delegate));
 	Test->TestFalse(TIMER_TEST_TEXT("TimerExists called with an invalid handle"), TimerManager.TimerExists(Handle));
-
-	Test->TestFalse(TIMER_TEST_TEXT("IsTimerActive called with an invalid delegate"), TimerManager.IsTimerActive(Delegate));
 	Test->TestFalse(TIMER_TEST_TEXT("IsTimerActive called with an invalid handle"), TimerManager.IsTimerActive(Handle));
-
-	Test->TestFalse(TIMER_TEST_TEXT("IsTimerPaused called with an invalid delegate"), TimerManager.IsTimerPaused(Delegate));
 	Test->TestFalse(TIMER_TEST_TEXT("IsTimerPaused called with an invalid handle"), TimerManager.IsTimerPaused(Handle));
-
-	Test->TestTrue(TIMER_TEST_TEXT("GetTimerRate called with an invalid delegate"), (TimerManager.GetTimerRate(Delegate) == -1.f));
 	Test->TestTrue(TIMER_TEST_TEXT("GetTimerRate called with an invalid handle"), (TimerManager.GetTimerRate(Handle) == -1.f));
-
-	Test->TestTrue(TIMER_TEST_TEXT("GetTimerElapsed called with an invalid delegate"), (TimerManager.GetTimerElapsed(Delegate) == -1.f));
 	Test->TestTrue(TIMER_TEST_TEXT("GetTimerElapsed called with an invalid handle"), (TimerManager.GetTimerElapsed(Handle) == -1.f));
-
-	Test->TestTrue(TIMER_TEST_TEXT("GetTimerRemaining called with an invalid delegate"), (TimerManager.GetTimerRemaining(Delegate) == -1.f));
 	Test->TestTrue(TIMER_TEST_TEXT("GetTimerRemaining called with an invalid handle"), (TimerManager.GetTimerRemaining(Handle) == -1.f));
 
 	// these don't return values but we should run them to make sure they don't do something horrible like crash
-	TimerManager.PauseTimer(Delegate);
 	TimerManager.PauseTimer(Handle);
-
-	TimerManager.UnPauseTimer(Delegate);
 	TimerManager.UnPauseTimer(Handle);
-
-	TimerManager.ClearTimer(Delegate);
 	TimerManager.ClearTimer(Handle);
 	
 	return true;
@@ -74,161 +57,39 @@ bool TimerManagerTest_InvalidTimers(UWorld* World, FAutomationTestBase* Test)
 bool TimerManagerTest_MissingTimers(UWorld* World, FAutomationTestBase* Test)
 {
 	FTimerManager& TimerManager = World->GetTimerManager();
-	FTimerDelegate Delegate;
 	FTimerHandle Handle;
 
-	FDummy Dummy;
-	Delegate.BindRaw(&Dummy, &FDummy::Callback);
 	Handle.MakeValid();
 
-	Test->TestFalse(TIMER_TEST_TEXT("TimerExists called with an invalid delegate"), TimerManager.TimerExists(Delegate));
 	Test->TestFalse(TIMER_TEST_TEXT("TimerExists called with an invalid handle"), TimerManager.TimerExists(Handle));
-
-	Test->TestFalse(TIMER_TEST_TEXT("IsTimerActive called with an invalid delegate"), TimerManager.IsTimerActive(Delegate));
 	Test->TestFalse(TIMER_TEST_TEXT("IsTimerActive called with an invalid handle"), TimerManager.IsTimerActive(Handle));
-
-	Test->TestFalse(TIMER_TEST_TEXT("IsTimerPaused called with an invalid delegate"), TimerManager.IsTimerPaused(Delegate));
 	Test->TestFalse(TIMER_TEST_TEXT("IsTimerPaused called with an invalid handle"), TimerManager.IsTimerPaused(Handle));
-
-	Test->TestTrue(TIMER_TEST_TEXT("GetTimerRate called with an invalid delegate"), (TimerManager.GetTimerRate(Delegate) == -1.f));
 	Test->TestTrue(TIMER_TEST_TEXT("GetTimerRate called with an invalid handle"), (TimerManager.GetTimerRate(Handle) == -1.f));
-
-	Test->TestTrue(TIMER_TEST_TEXT("GetTimerElapsed called with an invalid delegate"), (TimerManager.GetTimerElapsed(Delegate) == -1.f));
 	Test->TestTrue(TIMER_TEST_TEXT("GetTimerElapsed called with an invalid handle"), (TimerManager.GetTimerElapsed(Handle) == -1.f));
-
-	Test->TestTrue(TIMER_TEST_TEXT("GetTimerRemaining called with an invalid delegate"), (TimerManager.GetTimerRemaining(Delegate) == -1.f));
 	Test->TestTrue(TIMER_TEST_TEXT("GetTimerRemaining called with an invalid handle"), (TimerManager.GetTimerRemaining(Handle) == -1.f));
 
 	// these don't return values but we should run them to make sure they don't do something horrible like crash
-	TimerManager.PauseTimer(Delegate);
 	TimerManager.PauseTimer(Handle);
-
-	TimerManager.UnPauseTimer(Delegate);
 	TimerManager.UnPauseTimer(Handle);
-
-	TimerManager.ClearTimer(Delegate);
 	TimerManager.ClearTimer(Handle);
 
 	return true;
 }
 
-bool TimerManagerTest_ValidTimer_Delegate(UWorld* World, FAutomationTestBase* Test)
-{
-	FTimerManager& TimerManager = World->GetTimerManager();
-	FTimerDelegate Delegate;
-	const float Rate = 1.5f;
-
-	FDummy Dummy;
-
-	Delegate.BindRaw(&Dummy, &FDummy::Callback);
-
-	TimerManager.SetTimer(Delegate, Rate, false);
-
-	Test->TestTrue(TIMER_TEST_TEXT("TimerExists called with a pending timer"), TimerManager.TimerExists(Delegate));
-	Test->TestTrue(TIMER_TEST_TEXT("IsTimerActive called with a pending timer"), TimerManager.IsTimerActive(Delegate));
-	Test->TestFalse(TIMER_TEST_TEXT("IsTimerPaused called with a pending timer"), TimerManager.IsTimerPaused(Delegate));
-	Test->TestTrue(TIMER_TEST_TEXT("GetTimerRate called with a pending timer"), (TimerManager.GetTimerRate(Delegate) == Rate));
-	Test->TestTrue(TIMER_TEST_TEXT("GetTimerElapsed called with a pending timer"), (TimerManager.GetTimerElapsed(Delegate) == 0.f));
-	Test->TestTrue(TIMER_TEST_TEXT("GetTimerRemaining called with a pending timer"), (TimerManager.GetTimerRemaining(Delegate) == Rate));
-
-	// small tick to move the timer from the pending list to the active list, the timer will start counting time after this tick
-	TimerTest_TickWorld(World, KINDA_SMALL_NUMBER);
-
-	Test->TestTrue(TIMER_TEST_TEXT("TimerExists called with a pending timer"), TimerManager.TimerExists(Delegate));
-	Test->TestTrue(TIMER_TEST_TEXT("IsTimerActive called with an active timer"), TimerManager.IsTimerActive(Delegate));
-	Test->TestFalse(TIMER_TEST_TEXT("IsTimerPaused called with an active timer"), TimerManager.IsTimerPaused(Delegate));
-
-	TimerTest_TickWorld(World);
-
-	Test->TestTrue(TIMER_TEST_TEXT("GetTimerElapsed called with an active timer after one step"), (FMath::IsNearlyEqual(TimerManager.GetTimerElapsed(Delegate), 1.f, KINDA_SMALL_NUMBER)));
-	Test->TestTrue(TIMER_TEST_TEXT("GetTimerRemaining called with an active timer after one step"), (FMath::IsNearlyEqual(TimerManager.GetTimerRemaining(Delegate), Rate - 1.f, KINDA_SMALL_NUMBER)));
-
-	TimerManager.PauseTimer(Delegate);
-
-	Test->TestTrue(TIMER_TEST_TEXT("TimerExists called with a paused timer"), TimerManager.TimerExists(Delegate));
-	Test->TestFalse(TIMER_TEST_TEXT("IsTimerActive called with a paused timer"), TimerManager.IsTimerActive(Delegate));
-	Test->TestTrue(TIMER_TEST_TEXT("IsTimerPaused called with a paused timer"), TimerManager.IsTimerPaused(Delegate));
-
-	TimerTest_TickWorld(World);
-
-	Test->TestTrue(TIMER_TEST_TEXT("TimerExists called with a paused timer"), TimerManager.TimerExists(Delegate));
-	Test->TestFalse(TIMER_TEST_TEXT("IsTimerActive called with a paused timer"), TimerManager.IsTimerActive(Delegate));
-	Test->TestTrue(TIMER_TEST_TEXT("IsTimerPaused called with a paused timer"), TimerManager.IsTimerPaused(Delegate));
-
-	Test->TestTrue(TIMER_TEST_TEXT("GetTimerElapsed called with a paused timer after one step"), (FMath::IsNearlyEqual(TimerManager.GetTimerElapsed(Delegate), 1.f, KINDA_SMALL_NUMBER)));
-	Test->TestTrue(TIMER_TEST_TEXT("GetTimerRemaining called with a paused timer after one step"), (FMath::IsNearlyEqual(TimerManager.GetTimerRemaining(Delegate), Rate - 1.f, KINDA_SMALL_NUMBER)));
-
-	TimerManager.UnPauseTimer(Delegate);
-
-	Test->TestTrue(TIMER_TEST_TEXT("TimerExists called with a pending timer"), TimerManager.TimerExists(Delegate));
-	Test->TestTrue(TIMER_TEST_TEXT("IsTimerActive called with a pending timer"), TimerManager.IsTimerActive(Delegate));
-	Test->TestFalse(TIMER_TEST_TEXT("IsTimerPaused called with a pending timer"), TimerManager.IsTimerPaused(Delegate));
-
-	TimerTest_TickWorld(World);
-
-	Test->TestFalse(TIMER_TEST_TEXT("TimerExists called with a completed timer"), TimerManager.TimerExists(Delegate));
-	Test->TestTrue(TIMER_TEST_TEXT("Count of callback executions"), Dummy.Count == 1);
-
-
-	// Test resetting the timer
-	TimerManager.SetTimer(Delegate, Rate, false);
-	TimerManager.SetTimer(Delegate, 0.f, false);
-
-	Test->TestFalse(TIMER_TEST_TEXT("TimerExists called with a reset timer"), TimerManager.TimerExists(Delegate));
-
-	// Test looping timers
-	Dummy.Reset();
-	TimerManager.SetTimer(Delegate, Rate, true);
-	TimerTest_TickWorld(World, KINDA_SMALL_NUMBER);
-
-	TimerTest_TickWorld(World, 2.f);
-
-	Test->TestTrue(TIMER_TEST_TEXT("TimerExists called with a looping timer"), TimerManager.TimerExists(Delegate));
-	Test->TestTrue(TIMER_TEST_TEXT("IsTimerActive called with a looping timer"), TimerManager.IsTimerActive(Delegate));
-
-	Test->TestTrue(TIMER_TEST_TEXT("Count of callback executions"), Dummy.Count == 1);
-	Test->TestTrue(TIMER_TEST_TEXT("GetTimerElapsed called with a looping timer"), (FMath::IsNearlyEqual(TimerManager.GetTimerElapsed(Delegate), 2.f - (Rate * Dummy.Count), KINDA_SMALL_NUMBER)));
-	Test->TestTrue(TIMER_TEST_TEXT("GetTimerRemaining called with a looping timer"), (FMath::IsNearlyEqual(TimerManager.GetTimerRemaining(Delegate), Rate * (Dummy.Count + 1) - 2.f, KINDA_SMALL_NUMBER)));
-
-	TimerTest_TickWorld(World, 2.f);
-
-	Test->TestTrue(TIMER_TEST_TEXT("Count of callback executions"), Dummy.Count == 2);
-	Test->TestTrue(TIMER_TEST_TEXT("GetTimerElapsed called with a looping timer"), (FMath::IsNearlyEqual(TimerManager.GetTimerElapsed(Delegate), 4.f - (Rate * Dummy.Count), KINDA_SMALL_NUMBER)));
-	Test->TestTrue(TIMER_TEST_TEXT("GetTimerRemaining called with a looping timer"), (FMath::IsNearlyEqual(TimerManager.GetTimerRemaining(Delegate), Rate * (Dummy.Count + 1) - 4.f, KINDA_SMALL_NUMBER)));
-
-	TimerManager.SetTimer(Delegate, 0.f, false);
-
-	Test->TestFalse(TIMER_TEST_TEXT("TimerExists called with a reset looping timer"), TimerManager.TimerExists(Delegate));
-
-	Dummy.Reset();
-	TimerManager.SetTimerForNextTick(Delegate);
-
-	Test->TestTrue(TIMER_TEST_TEXT("TimerExists called with a timer executing on the next tick"), TimerManager.TimerExists(Delegate));
-	Test->TestTrue(TIMER_TEST_TEXT("IsTimerActive called with a timer executing on the next tick"), TimerManager.IsTimerActive(Delegate));
-
-	TimerTest_TickWorld(World, SMALL_NUMBER);
-
-	Test->TestFalse(TIMER_TEST_TEXT("TimerExists called with a timer executing on the previous tick"), TimerManager.TimerExists(Delegate));
-	Test->TestTrue(TIMER_TEST_TEXT("Count of callback executions"), Dummy.Count == 1);
-
-	return true;
-}
-
 bool TimerManagerTest_ValidTimer_HandleWithDelegate(UWorld* World, FAutomationTestBase* Test)
 {
 	FTimerManager& TimerManager = World->GetTimerManager();
 	FTimerDelegate Delegate;
-	FTimerHandle Handle;
 	const float Rate = 1.5f;
 
 	FDummy Dummy;
 	Delegate.BindRaw(&Dummy, &FDummy::Callback);
 
+	FTimerHandle Handle;
 	TimerManager.SetTimer(Handle, Delegate, Rate, false);
 
 	Test->TestTrue(TIMER_TEST_TEXT("Handle should be valid after calling SetTimer"), Handle.IsValid());
 	Test->TestTrue(TIMER_TEST_TEXT("TimerExists called with a pending timer"), TimerManager.TimerExists(Handle));
-	Test->TestFalse(TIMER_TEST_TEXT("TimerExists called by delegate with a pending timer"), TimerManager.TimerExists(Delegate));
 	Test->TestTrue(TIMER_TEST_TEXT("IsTimerActive called with a pending timer"), TimerManager.IsTimerActive(Handle));
 	Test->TestFalse(TIMER_TEST_TEXT("IsTimerPaused called with a pending timer"), TimerManager.IsTimerPaused(Handle));
 	Test->TestTrue(TIMER_TEST_TEXT("GetTimerRate called with a pending timer"), (TimerManager.GetTimerRate(Handle) == Rate));
@@ -307,89 +168,6 @@ bool TimerManagerTest_ValidTimer_HandleWithDelegate(UWorld* World, FAutomationTe
 	return true;
 }
 
-bool TimerManagerTest_ValidTimer_HandleNoDelegate(UWorld* World, FAutomationTestBase* Test)
-{
-	FTimerManager& TimerManager = World->GetTimerManager();
-	FTimerHandle Handle;
-	const float Rate = 1.5f;
-
-	TimerManager.SetTimer(Handle, Rate, false);
-
-	Test->TestTrue(TIMER_TEST_TEXT("Handle should be valid after calling SetTimer"), Handle.IsValid());
-	Test->TestTrue(TIMER_TEST_TEXT("TimerExists called with a pending timer"), TimerManager.TimerExists(Handle));
-	Test->TestTrue(TIMER_TEST_TEXT("IsTimerActive called with a pending timer"), TimerManager.IsTimerActive(Handle));
-	Test->TestFalse(TIMER_TEST_TEXT("IsTimerPaused called with a pending timer"), TimerManager.IsTimerPaused(Handle));
-	Test->TestTrue(TIMER_TEST_TEXT("GetTimerRate called with a pending timer"), (TimerManager.GetTimerRate(Handle) == Rate));
-	Test->TestTrue(TIMER_TEST_TEXT("GetTimerElapsed called with a pending timer"), (TimerManager.GetTimerElapsed(Handle) == 0.f));
-	Test->TestTrue(TIMER_TEST_TEXT("GetTimerRemaining called with a pending timer"), (TimerManager.GetTimerRemaining(Handle) == Rate));
-
-	// small tick to move the timer from the pending list to the active list, the timer will start counting time after this tick
-	TimerTest_TickWorld(World, KINDA_SMALL_NUMBER);
-
-	Test->TestTrue(TIMER_TEST_TEXT("TimerExists called with a pending timer"), TimerManager.TimerExists(Handle));
-	Test->TestTrue(TIMER_TEST_TEXT("IsTimerActive called with an active timer"), TimerManager.IsTimerActive(Handle));
-	Test->TestFalse(TIMER_TEST_TEXT("IsTimerPaused called with an active timer"), TimerManager.IsTimerPaused(Handle));
-
-	TimerTest_TickWorld(World);
-
-	Test->TestTrue(TIMER_TEST_TEXT("GetTimerElapsed called with an active timer after one step"), (FMath::IsNearlyEqual(TimerManager.GetTimerElapsed(Handle), 1.f, KINDA_SMALL_NUMBER)));
-	Test->TestTrue(TIMER_TEST_TEXT("GetTimerRemaining called with an active timer after one step"), (FMath::IsNearlyEqual(TimerManager.GetTimerRemaining(Handle), Rate - 1.f, KINDA_SMALL_NUMBER)));
-
-	TimerManager.PauseTimer(Handle);
-
-	Test->TestTrue(TIMER_TEST_TEXT("TimerExists called with a paused timer"), TimerManager.TimerExists(Handle));
-	Test->TestFalse(TIMER_TEST_TEXT("IsTimerActive called with a paused timer"), TimerManager.IsTimerActive(Handle));
-	Test->TestTrue(TIMER_TEST_TEXT("IsTimerPaused called with a paused timer"), TimerManager.IsTimerPaused(Handle));
-
-	TimerTest_TickWorld(World);
-
-	Test->TestTrue(TIMER_TEST_TEXT("TimerExists called with a paused timer"), TimerManager.TimerExists(Handle));
-	Test->TestFalse(TIMER_TEST_TEXT("IsTimerActive called with a paused timer"), TimerManager.IsTimerActive(Handle));
-	Test->TestTrue(TIMER_TEST_TEXT("IsTimerPaused called with a paused timer"), TimerManager.IsTimerPaused(Handle));
-
-	Test->TestTrue(TIMER_TEST_TEXT("GetTimerElapsed called with a paused timer after one step"), (FMath::IsNearlyEqual(TimerManager.GetTimerElapsed(Handle), 1.f, KINDA_SMALL_NUMBER)));
-	Test->TestTrue(TIMER_TEST_TEXT("GetTimerRemaining called with a paused timer after one step"), (FMath::IsNearlyEqual(TimerManager.GetTimerRemaining(Handle), Rate - 1.f, KINDA_SMALL_NUMBER)));
-
-	TimerManager.UnPauseTimer(Handle);
-
-	Test->TestTrue(TIMER_TEST_TEXT("TimerExists called with a pending timer"), TimerManager.TimerExists(Handle));
-	Test->TestTrue(TIMER_TEST_TEXT("IsTimerActive called with a pending timer"), TimerManager.IsTimerActive(Handle));
-	Test->TestFalse(TIMER_TEST_TEXT("IsTimerPaused called with a pending timer"), TimerManager.IsTimerPaused(Handle));
-
-	TimerTest_TickWorld(World);
-
-	Test->TestFalse(TIMER_TEST_TEXT("TimerExists called with a completed timer"), TimerManager.TimerExists(Handle));
-
-	// Test resetting the timer
-	TimerManager.SetTimer(Handle, Rate, false);
-	TimerManager.SetTimer(Handle, 0.f, false);
-
-	Test->TestFalse(TIMER_TEST_TEXT("TimerExists called with a reset timer"), TimerManager.TimerExists(Handle));
-
-	// Test looping timers
-	TimerManager.SetTimer(Handle, Rate, true);
-	TimerTest_TickWorld(World, KINDA_SMALL_NUMBER);
-
-	TimerTest_TickWorld(World, 2.f);
-
-	Test->TestTrue(TIMER_TEST_TEXT("TimerExists called with a looping timer"), TimerManager.TimerExists(Handle));
-	Test->TestTrue(TIMER_TEST_TEXT("IsTimerActive called with a looping timer"), TimerManager.IsTimerActive(Handle));
-
-	Test->TestTrue(TIMER_TEST_TEXT("GetTimerElapsed called with a looping timer"), (FMath::IsNearlyEqual(TimerManager.GetTimerElapsed(Handle), 2.f - Rate, KINDA_SMALL_NUMBER)));
-	Test->TestTrue(TIMER_TEST_TEXT("GetTimerRemaining called with a looping timer"), (FMath::IsNearlyEqual(TimerManager.GetTimerRemaining(Handle), Rate * 2 - 2.f, KINDA_SMALL_NUMBER)));
-
-	TimerTest_TickWorld(World, 2.f);
-
-	Test->TestTrue(TIMER_TEST_TEXT("GetTimerElapsed called with a looping timer"), (FMath::IsNearlyEqual(TimerManager.GetTimerElapsed(Handle), 4.f - (Rate * 2), KINDA_SMALL_NUMBER)));
-	Test->TestTrue(TIMER_TEST_TEXT("GetTimerRemaining called with a looping timer"), (FMath::IsNearlyEqual(TimerManager.GetTimerRemaining(Handle), Rate * 3 - 4.f, KINDA_SMALL_NUMBER)));
-
-	TimerManager.SetTimer(Handle, 0.f, false);
-
-	Test->TestFalse(TIMER_TEST_TEXT("TimerExists called with a reset looping timer"), TimerManager.TimerExists(Handle));
-
-	return true;
-}
-
 struct FLoopingTestFunc
 {
 	static FTimerManager* TimerManager;
@@ -454,17 +232,18 @@ bool TimerManagerTest_LoopingTimers_DifferentHandles(UWorld* World, FAutomationT
 
 	FTimerDelegate Delegate = FTimerDelegate::CreateStatic(Func, &CallCount);
 
-	TimerManager.SetTimer(Delegate, 1.0f, false);
+	FTimerHandle Handle;
+	TimerManager.SetTimer(Handle, Delegate, 1.0f, false);
 	TimerTest_TickWorld(World, KINDA_SMALL_NUMBER);
 
 	Test->TestTrue(TIMER_TEST_TEXT("First delegate time remaining is 1.0f"),
-		FMath::IsNearlyEqual(TimerManager.GetTimerRemaining(Delegate), 1.0f, KINDA_SMALL_NUMBER));
+		FMath::IsNearlyEqual(TimerManager.GetTimerRemaining(Handle), 1.0f, KINDA_SMALL_NUMBER));
 	
-	TimerManager.SetTimer(Delegate, 5.0f, false);
+	TimerManager.SetTimer(Handle, Delegate, 5.0f, false);
 	TimerTest_TickWorld(World, KINDA_SMALL_NUMBER);
 	Test->TestTrue(TIMER_TEST_TEXT("Reset delegate time remaining is 5.0f"),
-		FMath::IsNearlyEqual(TimerManager.GetTimerRemaining(Delegate), 5.0f, KINDA_SMALL_NUMBER));
-	
+		FMath::IsNearlyEqual(TimerManager.GetTimerRemaining(Handle), 5.0f, KINDA_SMALL_NUMBER));
+
 	TimerManager.SetTimer(HandleOne, Delegate, 1.0f, true);
 	TimerManager.SetTimer(HandleTwo, Delegate, 1.5f, true);
 	TimerTest_TickWorld(World, KINDA_SMALL_NUMBER);
@@ -497,9 +276,7 @@ bool FTimerManagerTest::RunTest(const FString& Parameters)
 
 	TimerManagerTest_InvalidTimers(World, this);
 	TimerManagerTest_MissingTimers(World, this);
-	TimerManagerTest_ValidTimer_Delegate(World, this);
 	TimerManagerTest_ValidTimer_HandleWithDelegate(World, this);
-	TimerManagerTest_ValidTimer_HandleNoDelegate(World, this);
 	TimerManagerTest_ValidTimer_HandleLoopingSetDuringExecute(World, this);
 	TimerManagerTest_LoopingTimers_DifferentHandles(World, this);
 
