Core - UE/UE-6435 Virtual Memory Limit Listed in Logs

diff --git a/Engine/Source/Runtime/Core/Private/Windows/WindowsPlatformMemory.cpp b/Engine/Source/Runtime/Core/Private/Windows/WindowsPlatformMemory.cpp
--- a/Engine/Source/Runtime/Core/Private/Windows/WindowsPlatformMemory.cpp
+++ b/Engine/Source/Runtime/Core/Private/Windows/WindowsPlatformMemory.cpp
@@ -43,12 +43,18 @@ void FWindowsPlatformMemory::Init()
 	SET_MEMORY_STAT(MCR_Physical, 2*GB); //2Gb of physical memory on win32
 #endif
 
-
 	const FPlatformMemoryConstants& MemoryConstants = FPlatformMemory::GetConstants();
+#if PLATFORM_32BITS	
 	UE_LOG(LogMemory, Log, TEXT("Memory total: Physical=%.1fGB (%dGB approx) Virtual=%.1fGB"), 
 		float(MemoryConstants.TotalPhysical/1024.0/1024.0/1024.0),
 		MemoryConstants.TotalPhysicalGB, 
 		float(MemoryConstants.TotalVirtual/1024.0/1024.0/1024.0) );
+#else
+	// Logging virtual memory size for 64bits is pointless.
+	UE_LOG(LogMemory, Log, TEXT("Memory total: Physical=%.1fGB (%dGB approx)"), 
+		float(MemoryConstants.TotalPhysical/1024.0/1024.0/1024.0),
+		MemoryConstants.TotalPhysicalGB );
+#endif //PLATFORM_32BITS
 
 	UpdateStats();
 	DumpStats( *GLog );
