PR #671: Linux: Create default project path using user localization setting (Contributed by [REDACTED_PERSON])

diff --git a/Engine/Source/Runtime/Core/Private/Linux/LinuxPlatformProcess.cpp b/Engine/Source/Runtime/Core/Private/Linux/LinuxPlatformProcess.cpp
--- a/Engine/Source/Runtime/Core/Private/Linux/LinuxPlatformProcess.cpp
+++ b/Engine/Source/Runtime/Core/Private/Linux/LinuxPlatformProcess.cpp
@@ -141,20 +141,36 @@ const TCHAR* FLinuxPlatformProcess::UserDir()
 {
 	// The UserDir is where user visible files (such as game projects) live.
 	// On Linux (just like on Mac) this corresponds to $HOME/Documents.
+	// To accomodate localization requirement we use xdg-user-dir command,
+	// and fall back to $HOME/Documents if setting not found.
 	static TCHAR Result[MAX_PATH] = TEXT("");
 
 	if (!Result[0])
 	{
-		char* Home = secure_getenv("HOME");
-		if (!Home)
+		char DocPath[MAX_PATH];
+		
+		FILE* FilePtr = popen("xdg-user-dir DOCUMENTS", "r");
+		if(fgets(DocPath, MAX_PATH, FilePtr) == NULL)
 		{
-			UE_LOG(LogHAL, Warning, TEXT("Unable to read the $HOME environment variable"));
+			char* Home = secure_getenv("HOME");
+			if (!Home)
+			{
+				UE_LOG(LogHAL, Warning, TEXT("Unable to read the $HOME environment variable"));
+			}
+			else
+			{
+				FCString::Strcpy(Result, ANSI_TO_TCHAR(Home));
+				FCString::Strcat(Result, TEXT("/Documents/"));
+			}
 		}
 		else
 		{
-			FCString::Strcpy(Result, ANSI_TO_TCHAR(Home));
-			FCString::Strcat(Result, TEXT("/Documents/"));
+				size_t DocLen = strlen(DocPath) - 1;
+				DocPath[DocLen] = '\0';
+				FCString::Strcpy(Result, ANSI_TO_TCHAR(DocPath));
+				FCString::Strcat(Result, TEXT("/"));
 		}
+		pclose(FilePtr);
 	}
 	return Result;
 }
