#ue4 - Some minor cleanup to FQuat::Rotator() to reorder Yaw/Pitch/Roll for better cache use, and clean up some constants.

diff --git a/Engine/Source/Runtime/Core/Private/Math/UnrealMath.cpp b/Engine/Source/Runtime/Core/Private/Math/UnrealMath.cpp
--- a/Engine/Source/Runtime/Core/Private/Math/UnrealMath.cpp
+++ b/Engine/Source/Runtime/Core/Private/Math/UnrealMath.cpp
@@ -397,13 +397,11 @@ FRotator FQuat::Rotator() const
 	checkSlow(IsNormalized());
 #endif
 
-	static float RAD_TO_DEG = 180.f/PI;
+	static const float RAD_TO_DEG = 180.f/PI;
 
 	FRotator RotatorFromQuat;
-	float SingularityTest = Z*X-W*Y;
-	float Pitch = FMath::Asin(2*(SingularityTest));
-
-	RotatorFromQuat.Yaw = FMath::Atan2(2.f*(W*Z+X*Y), (1-2.f*(FMath::Square(Y) + FMath::Square(Z))))*RAD_TO_DEG;
+	const float SingularityTest = Z*X-W*Y;
+	const float Yaw = FMath::Atan2(2.f*(W*Z+X*Y), (1.f-2.f*(FMath::Square(Y) + FMath::Square(Z))))*RAD_TO_DEG;
 
 	// reference 
 	// http://en.wikipedia.org/wiki/Conversion_between_quaternions_and_Euler_angles
@@ -416,18 +414,22 @@ FRotator FQuat::Rotator() const
 
 	if (SingularityTest < -SINGULARITY_THRESHOLD)
 	{
-		RotatorFromQuat.Roll = -RotatorFromQuat.Yaw - 2.f* FMath::Atan2(X, W) * RAD_TO_DEG;
 		RotatorFromQuat.Pitch = 270.f;
+		RotatorFromQuat.Yaw = Yaw;
+		RotatorFromQuat.Roll = -RotatorFromQuat.Yaw - 2.f* FMath::Atan2(X, W) * RAD_TO_DEG;
 	}
 	else if (SingularityTest > SINGULARITY_THRESHOLD)
 	{
-		RotatorFromQuat.Roll = RotatorFromQuat.Yaw - 2.f* FMath::Atan2(X, W) * RAD_TO_DEG;
 		RotatorFromQuat.Pitch = 90.f;
+		RotatorFromQuat.Yaw = Yaw;
+		RotatorFromQuat.Roll = RotatorFromQuat.Yaw - 2.f* FMath::Atan2(X, W) * RAD_TO_DEG;
 	}
 	else
 	{
-		RotatorFromQuat.Roll = FMath::Atan2(-2.f*(W*X+Y*Z), (1-2.f*(FMath::Square(X) + FMath::Square(Y))))*RAD_TO_DEG;
+		const float Pitch = FMath::Asin(2.f*(SingularityTest));
 		RotatorFromQuat.Pitch = FRotator::ClampAxis(Pitch*RAD_TO_DEG); //clamp it so within 360 - this is for if below
+		RotatorFromQuat.Yaw = Yaw;
+		RotatorFromQuat.Roll = FMath::Atan2(-2.f*(W*X+Y*Z), (1-2.f*(FMath::Square(X) + FMath::Square(Y))))*RAD_TO_DEG;
 	}
 
 	RotatorFromQuat.Normalize();
