Fixed platform windows using old transparency flag

diff --git a/Engine/Source/Runtime/Core/Private/Linux/LinuxWindow.cpp b/Engine/Source/Runtime/Core/Private/Linux/LinuxWindow.cpp
--- a/Engine/Source/Runtime/Core/Private/Linux/LinuxWindow.cpp
+++ b/Engine/Source/Runtime/Core/Private/Linux/LinuxWindow.cpp
@@ -129,7 +129,7 @@ void FLinuxWindow::Initialize( FLinuxApplication* const Application, const TShar
 		return;
 	}
 
-	if ( Definition->SupportsTransparency )
+	if ( Definition->TransparencySupport == EWindowTransparency::PerWindow )
 	{
 		SetOpacity( Definition->Opacity );
 	}
diff --git a/Engine/Source/Runtime/Core/Private/Mac/MacWindow.cpp b/Engine/Source/Runtime/Core/Private/Mac/MacWindow.cpp
--- a/Engine/Source/Runtime/Core/Private/Mac/MacWindow.cpp
+++ b/Engine/Source/Runtime/Core/Private/Mac/MacWindow.cpp
@@ -184,7 +184,7 @@ void FMacWindow::Initialize( FMacApplication* const Application, const TSharedRe
 				[WindowHandle setCollectionBehavior:NSWindowCollectionBehaviorCanJoinAllSpaces|NSWindowCollectionBehaviorTransient|NSWindowCollectionBehaviorIgnoresCycle];
 			}
 
-			if( Definition->SupportsTransparency )
+			if( Definition->TransparencySupport == EWindowTransparency::PerWindow )
 			{
 				SetOpacity( Definition->Opacity );
 			}
