Fixed compile errors cause by missing symbol where ALPHA_BLENDED_WINDOWS is false

diff --git a/Engine/Source/Runtime/Core/Private/Windows/WindowsApplication.cpp b/Engine/Source/Runtime/Core/Private/Windows/WindowsApplication.cpp
--- a/Engine/Source/Runtime/Core/Private/Windows/WindowsApplication.cpp
+++ b/Engine/Source/Runtime/Core/Private/Windows/WindowsApplication.cpp
@@ -610,10 +610,14 @@ EWindowTitleAlignment::Type FWindowsApplication::GetWindowTitleAlignment() const
 	
 EWindowTransparency FWindowsApplication::GetWindowTransparencySupport() const
 {
+#if ALPHA_BLENDED_WINDOWS
 	BOOL bIsCompositionEnabled = FALSE;
 	::DwmIsCompositionEnabled(&bIsCompositionEnabled);
 
 	return bIsCompositionEnabled ? EWindowTransparency::PerPixel : EWindowTransparency::PerWindow;
+#else
+	return EWindowTransparency::PerWindow;
+#endif
 }
 
 /** All WIN32 messages sent to our app go here; this method simply passes them on */
diff --git a/Engine/Source/Runtime/Core/Private/Windows/WindowsWindow.cpp b/Engine/Source/Runtime/Core/Private/Windows/WindowsWindow.cpp
--- a/Engine/Source/Runtime/Core/Private/Windows/WindowsWindow.cpp
+++ b/Engine/Source/Runtime/Core/Private/Windows/WindowsWindow.cpp
@@ -59,7 +59,12 @@ void FWindowsWindow::Initialize( FWindowsApplication* const Application, const T
 	int32 ClientHeight = FMath::TruncToInt( HeightInitial );
 	int32 WindowWidth = ClientWidth;
 	int32 WindowHeight = ClientHeight;
-	const bool bApplicationSupportsPerPixelBlending = Application->GetWindowTransparencySupport() == EWindowTransparency::PerPixel;
+	const bool bApplicationSupportsPerPixelBlending =
+#if ALPHA_BLENDED_WINDOWS
+		Application->GetWindowTransparencySupport() == EWindowTransparency::PerPixel;
+#else
+		false;
+#endif
 
 	if( !Definition->HasOSWindowBorder )
 	{
diff --git a/Engine/Source/Runtime/Core/Public/Linux/LinuxApplication.h b/Engine/Source/Runtime/Core/Public/Linux/LinuxApplication.h
--- a/Engine/Source/Runtime/Core/Public/Linux/LinuxApplication.h
+++ b/Engine/Source/Runtime/Core/Public/Linux/LinuxApplication.h
@@ -53,7 +53,7 @@ public:
 
 	virtual EWindowTransparency GetWindowTransparencySupport() const override
 	{
-		return EWindowTransparency::PerPixel;
+		return EWindowTransparency::PerWindow;
 	}
 
 	void AddPendingEvent( SDL_Event event );
diff --git a/Engine/Source/Runtime/Core/Public/Mac/MacApplication.h b/Engine/Source/Runtime/Core/Public/Mac/MacApplication.h
--- a/Engine/Source/Runtime/Core/Public/Mac/MacApplication.h
+++ b/Engine/Source/Runtime/Core/Public/Mac/MacApplication.h
@@ -55,7 +55,7 @@ public:
 	
 	virtual EWindowTransparency GetWindowTransparencySupport() const override
 	{
-		return EWindowTransparency::PerPixel;
+		return EWindowTransparency::PerWindow;
 	}
 
 	virtual ITextInputMethodSystem *GetTextInputMethodSystem() override
