Added documentation

diff --git a/Engine/Source/Runtime/Core/Public/Windows/WindowsApplication.h b/Engine/Source/Runtime/Core/Public/Windows/WindowsApplication.h
--- a/Engine/Source/Runtime/Core/Public/Windows/WindowsApplication.h
+++ b/Engine/Source/Runtime/Core/Public/Windows/WindowsApplication.h
@@ -13,11 +13,13 @@
 #include "IForceFeedbackSystem.h"
 #include "WindowsTextInputMethodSystem.h"
 
+
 DECLARE_LOG_CATEGORY_EXTERN(LogWindowsDesktop, Log, All);
 
 class FWindowsWindow;
 class FGenericApplicationMessageHandler;
 
+
 namespace ETaskbarProgressState
 {
 	enum Type
@@ -26,7 +28,7 @@ namespace ETaskbarProgressState
 		NoProgress = 0x0,
 
 		//The progress indicator does not grow in size, but cycles repeatedly along the 
-		//length of the taskbar button. This indicates activity without specifying what 
+		//length of the task bar button. This indicates activity without specifying what 
 		//proportion of the progress is complete. Progress is taking place, but there is 
 		//no prediction as to how long the operation will take.
 		Indeterminate = 0x1,
@@ -51,6 +53,7 @@ namespace ETaskbarProgressState
 	};
 }
 
+
 class FTaskbarList
 {
 public:
@@ -76,6 +79,7 @@ private:
 	ITaskbarList3* TaskBarList3;
 };
 
+
 struct FDeferredWindowsMessage
 {
 	FDeferredWindowsMessage( const TSharedPtr<FWindowsWindow>& InNativeWindow, HWND InHWnd, uint32 InMessage, WPARAM InWParam, LPARAM InLParam, int32 InX=0, int32 InY=0, uint32 InRawInputFlags = 0 )
@@ -87,22 +91,28 @@ struct FDeferredWindowsMessage
 		, X( InX )
 		, Y( InY )
 		, RawInputFlags( InRawInputFlags )
-	{}
+	{ }
+
 	/** Native window that received the message */
 	TWeakPtr<FWindowsWindow> NativeWindow;
+
 	/** Window handle */
 	HWND hWND;
+
 	/** Message code */
 	uint32 Message;
+
 	/** Message data */
 	WPARAM wParam;
 	LPARAM lParam;
+
 	/** Mouse coordinates */
 	int32 X;
 	int32 Y;
 	uint32 RawInputFlags;
 };
 
+
 namespace EWindowsDragDropOperationType
 {
 	enum Type
@@ -114,6 +124,7 @@ namespace EWindowsDragDropOperationType
 	};
 }
 
+
 struct FDragDropOLEData
 {
 	enum EWindowsOLEDataType
@@ -125,16 +136,18 @@ struct FDragDropOLEData
 
 	FDragDropOLEData()
 		: Type(None)
-	{}
+	{ }
 
 	FString OperationText;
 	TArray<FString> OperationFilenames;
 	EWindowsOLEDataType Type;
 };
 
+
 struct FDeferredWindowsDragDropOperation
 {
 private:
+
 	// Private constructor. Use the factory functions below.
 	FDeferredWindowsDragDropOperation()
 		: HWnd(NULL)
@@ -145,6 +158,7 @@ private:
 	}
 
 public:
+
 	static FDeferredWindowsDragDropOperation MakeDragEnter(HWND InHwnd, const FDragDropOLEData& InOLEData, ::DWORD InKeyState, POINTL InCursorPosition)
 	{
 		FDeferredWindowsDragDropOperation NewOperation;
@@ -193,16 +207,18 @@ public:
 	POINTL CursorPosition;
 };
 
-//disable warnings from overriding the deprecated forcefeedback.  
+
+//disable warnings from overriding the deprecated force feedback.  
 //calls to the deprecated function will still generate warnings.
 PRAGMA_DISABLE_DEPRECATION_WARNINGS
 
+
 /**
  * Windows-specific application implementation.
  */
-class FWindowsApplication : public GenericApplication, IForceFeedbackSystem
+class FWindowsApplication
+	: public GenericApplication, IForceFeedbackSystem
 {
-
 public:
 
 	/**
@@ -215,53 +231,34 @@ public:
 	 */
 	static FWindowsApplication* CreateWindowsApplication( const HINSTANCE InstanceHandle, const HICON IconHandle );
 
-
 public:	
 
 	virtual ~FWindowsApplication();
-	
 	virtual void DestroyApplication() override;
 
-
 public:
 
-	virtual void SetMessageHandler( const TSharedRef< class FGenericApplicationMessageHandler >& InMessageHandler ) override;
+	// GenericApplication overrides
 
+	virtual void SetMessageHandler( const TSharedRef< class FGenericApplicationMessageHandler >& InMessageHandler ) override;
 	virtual void PollGameDeviceState( const float TimeDelta ) override;
-
 	virtual void PumpMessages( const float TimeDelta ) override;
-
 	virtual void ProcessDeferredEvents( const float TimeDelta ) override;
-
 	virtual TSharedRef< FGenericWindow > MakeWindow() override;
-
 	virtual void InitializeWindow( const TSharedRef< FGenericWindow >& Window, const TSharedRef< FGenericWindowDefinition >& InDefinition, const TSharedPtr< FGenericWindow >& InParent, const bool bShowImmediately ) override;
-
 	virtual void SetCapture( const TSharedPtr< FGenericWindow >& InWindow ) override;
-
 	virtual void* GetCapture( void ) const override;
-
 	virtual void SetHighPrecisionMouseMode( const bool Enable, const TSharedPtr< FGenericWindow >& InWindow ) override;
-
 	virtual bool IsUsingHighPrecisionMouseMode() const override { return bUsingHighPrecisionMouseInput; }
-
 	virtual bool IsMouseAttached() const override { return bIsMouseAttached; }
-
 	virtual FModifierKeysState GetModifierKeys() const override;
-
 	virtual bool IsCursorDirectlyOverSlateWindow() const override;
-
 	virtual FPlatformRect GetWorkArea( const FPlatformRect& CurrentWindow ) const override;
-
 	virtual bool TryCalculatePopupWindowPosition( const FPlatformRect& InAnchor, const FVector2D& InSize, const EPopUpOrientation::Type Orientation, /*OUT*/ FVector2D* const CalculatedPopUpPosition ) const override;
-
 	virtual void GetInitialDisplayMetrics( FDisplayMetrics& OutDisplayMetrics ) const override;
-
 	virtual EWindowTitleAlignment::Type GetWindowTitleAlignment() const override;
-	
 	virtual EWindowTransparency GetWindowTransparencySupport() const override;
 
-	/** Function to return the current implementation of the ForceFeedback system */
 	DEPRECATED(4.7, "Please use GetInputInterface()")
 	virtual IForceFeedbackSystem* GetForceFeedbackSystem() override
 	{
@@ -279,9 +276,7 @@ public:
 	}
 
 	virtual void SetForceFeedbackChannelValue (int32 ControllerId, FForceFeedbackChannelType ChannelType, float Value) override;
-
 	virtual void SetForceFeedbackChannelValues(int32 ControllerId, const FForceFeedbackValues &Values) override;
-
 	virtual void SetLightColor(int32 ControllerId, FColor Color) override {}
 
 	virtual ITextInputMethodSystem *GetTextInputMethodSystem() override
@@ -293,15 +288,18 @@ public:
 
 	TSharedPtr<FTaskbarList> GetTaskbarList();
 
-
 protected:
 
+	/** Windows callback for message processing (forwards messages to the FWindowsApplication instance). */
 	static LRESULT CALLBACK AppWndProc(HWND hwnd, uint32 msg, WPARAM wParam, LPARAM lParam);
 
+	/** Processes a single Windows message. */
 	int32 ProcessMessage( HWND hwnd, uint32 msg, WPARAM wParam, LPARAM lParam );
 
+	/** Processes a deferred Windows message. */
 	int32 ProcessDeferredMessage( const FDeferredWindowsMessage& DeferredMessage );
 
+	/** Processes deferred drag and drop operations. */
 	void ProcessDeferredDragDropOperation(const FDeferredWindowsDragDropOperation& Op);
 
 public:
@@ -321,9 +319,9 @@ public:
 	/** Invoked by a window when an OLE Drag and Drop is dropped onto the window. */
 	HRESULT OnOLEDrop( const HWND HWnd, const FDragDropOLEData& OLEData, ::DWORD KeyState, POINTL CursorPosition, ::DWORD *CursorEffect);
 
-
 private:
 
+	/** Hidden constructor. */
 	FWindowsApplication( const HINSTANCE HInstance, const HICON IconHandle );
 
 	/** Registers the Windows class for windows and assigns the application instance and icon */
@@ -332,14 +330,19 @@ private:
 	/**  @return  True if a windows message is related to user input (mouse, keyboard) */
 	static bool IsInputMessage( uint32 msg );
 
+	/** Defers a Windows message for later processing. */
 	void DeferMessage( TSharedPtr<FWindowsWindow>& NativeWindow, HWND InHWnd, uint32 InMessage, WPARAM InWParam, LPARAM InLParam, int32 MouseX = 0, int32 MouseY = 0, uint32 RawInputFlags = 0 );
 
+	/** Checks a key code for release of the Shift key. */
 	void CheckForShiftUpEvents(const int32 KeyCode);
 
+	/** Shuts down the application (called after an unrecoverable error occurred). */
 	void ShutDownAfterError();
 
+	/** Enables or disables Windows accessibility features, such as sticky keys. */
 	void AllowAccessibilityShortcutKeys(const bool bAllowKeys);
 
+	/** Queries and caches the number of connected mouse devices. */
 	void QueryConnectedMice();
 
 private:
@@ -352,19 +355,19 @@ private:
 
 	bool bIsMouseAttached;
 
-	TArray< FDeferredWindowsMessage > DeferredMessages;
+	TArray<FDeferredWindowsMessage> DeferredMessages;
 
-	TArray< FDeferredWindowsDragDropOperation > DeferredDragDropOperations;
+	TArray<FDeferredWindowsDragDropOperation> DeferredDragDropOperations;
 
-	TArray< TSharedRef< FWindowsWindow > > Windows;
+	TArray<TSharedRef<FWindowsWindow>> Windows;
 
-	TSharedRef< class XInputInterface > XInput;
+	TSharedRef<class XInputInterface> XInput;
 
 	/** List of input devices implemented in external modules. */
-	TArray< TSharedPtr<class IInputDevice> > ExternalInputDevices;
+	TArray<TSharedPtr<class IInputDevice>> ExternalInputDevices;
 	bool bHasLoadedInputPlugins;
 
-	TArray< int32 > PressedModifierKeys;
+	TArray<int32> PressedModifierKeys;
 
 	FAutoConsoleVariableRef CVarDeferMessageProcessing;
 
@@ -385,4 +388,5 @@ private:
 	FILTERKEYS							StartupFilterKeys;
 };
 
-PRAGMA_ENABLE_DEPRECATION_WARNINGS
\ No newline at end of file
+
+PRAGMA_ENABLE_DEPRECATION_WARNINGS
