Core: ThreadingBase documentation clenaup pass; NULL -> nullptr

diff --git a/Engine/Source/Runtime/Core/Private/HAL/ThreadingBase.cpp b/Engine/Source/Runtime/Core/Private/HAL/ThreadingBase.cpp
--- a/Engine/Source/Runtime/Core/Private/HAL/ThreadingBase.cpp
+++ b/Engine/Source/Runtime/Core/Private/HAL/ThreadingBase.cpp
@@ -5,8 +5,10 @@
 #include "LockFreeList.h"
 #include "StatsData.h"
 
+
 /** The global thread pool */
-FQueuedThreadPool*		GThreadPool						= NULL;
+FQueuedThreadPool* GThreadPool = nullptr;
+
 
 CORE_API bool IsInGameThread()
 {
@@ -23,7 +25,7 @@ CORE_API bool IsInSlateThread()
 
 CORE_API int32 GIsRenderingThreadSuspended = 0;
 
-CORE_API FRunnableThread* GRenderingThread = NULL;
+CORE_API FRunnableThread* GRenderingThread = nullptr;
 
 CORE_API bool IsInActualRenderingThread()
 {
@@ -44,11 +46,12 @@ CORE_API bool IsInRHIThread()
 {
 	return GRHIThread && FPlatformTLS::GetCurrentThreadId() == GRHIThread->GetThreadID();
 }
-CORE_API FRunnableThread* GRHIThread = NULL;
+CORE_API FRunnableThread* GRHIThread = nullptr;
 // Fake threads
 
+
 /**
- * Fake thread created when multithreading is disabled.
+ * Fake thread created when multi-threading is disabled.
  */
 class FFakeThread : public FRunnableThread
 {
@@ -57,39 +60,36 @@ class FFakeThread : public FRunnableThread
 
 	/** Thread is suspended. */
 	bool bIsSuspended;
+
 	/** Name of this thread. */
 	FString Name;
+
 	/** Runnable object associated with this thread. */
 	FSingleThreadRunnable* Runnable;
+
 	/** Thread Id. */
 	uint32 ThreadId;
 
 public:
 
-	/**
-	 * Constructor.
-	 */
+	/** Constructor. */
 	FFakeThread()
 		: bIsSuspended(false)
-		, Runnable(NULL)
+		, Runnable(nullptr)
 		, ThreadId(ThreadIdCounter++)
 	{
 		// Auto register with single thread manager.
 		FSingleThreadManager::Get().AddThread(this);
 	}
 
-	/**
-	 * Destructor
-	 */
+	/** Virtual destructor. */
 	virtual ~FFakeThread()
 	{
 		// Remove from the manager.
 		FSingleThreadManager::Get().RemoveThread(this);
 	}
 
-	/**
-	 * Tick one time per frame
-	 */
+	/** Tick one time per frame. */
 	void Tick()
 	{
 		if (Runnable && !bIsSuspended)
@@ -98,7 +98,10 @@ public:
 		}
 	}
 
+public:
+
 	// FRunnableThread interface
+
 	virtual void SetThreadPriority(EThreadPriority NewPriority) override
 	{
 		// Not relevant.
@@ -140,11 +143,12 @@ public:
 		{
 			InRunnable->Init();
 		}		
-		return Runnable != NULL;
+		return Runnable != nullptr;
 	}
 };
 uint32 FFakeThread::ThreadIdCounter = 0xffff;
 
+
 void FSingleThreadManager::AddThread(FFakeThread* Thread)
 {
 	ThreadList.Add(Thread);
@@ -230,7 +234,7 @@ FRunnableThread* FRunnableThread::Create(
 	EThreadPriority InThreadPri, 
 	uint64 InThreadAffinityMask)
 {
-	FRunnableThread* NewThread = NULL;
+	FRunnableThread* NewThread = nullptr;
 	if (FPlatformProcess::SupportsMultithreading())
 	{
 		check(InRunnable);
@@ -243,7 +247,7 @@ FRunnableThread* FRunnableThread::Create(
 			{
 				// We failed to start the thread correctly so clean up
 				delete NewThread;
-				NewThread = NULL;
+				NewThread = nullptr;
 			}
 		}
 	}
@@ -255,7 +259,7 @@ FRunnableThread* FRunnableThread::Create(
 		{
 			// We failed to start the thread correctly so clean up
 			delete NewThread;
-			NewThread = NULL;
+			NewThread = nullptr;
 		}
 	}
 
@@ -268,9 +272,9 @@ FRunnableThread* FRunnableThread::Create(
 	}
 
 	return NewThread;
-
 }
 
+
 /**
  * This is the interface used for all poolable threads. The usage pattern for
  * a poolable thread is different from a regular thread and this interface
@@ -278,27 +282,21 @@ FRunnableThread* FRunnableThread::Create(
  * for work to do. When signaled they perform a job and then return themselves
  * to their owning pool via a callback and go back to an idle state.
  */
-class FQueuedThread : public FRunnable
+class FQueuedThread
+	: public FRunnable
 {
 protected:
-	/**
-	 * The event that tells the thread there is work to do
-	 */
+
+	/** The event that tells the thread there is work to do. */
 	FEvent* DoWorkEvent;
 
-	/**
-	 * If true, the thread should exit
-	 */
+	/** If true, the thread should exit. */
 	volatile int32 TimeToDie;
 
-	/**
-	 * The work this thread is doing
-	 */
+	/** The work this thread is doing. */
 	FQueuedWork* volatile QueuedWork;
 
-	/**
-	 * The pool this thread belongs to
-	 */
+	/** The pool this thread belongs to. */
 	class FQueuedThreadPool* OwningThreadPool;
 
 	/** My Thread  */
@@ -315,7 +313,7 @@ protected:
 			// Wait for some work to do
 			DoWorkEvent->Wait();
 			FQueuedWork* LocalQueuedWork = QueuedWork;
-			QueuedWork = NULL;
+			QueuedWork = nullptr;
 			FPlatformMisc::MemoryBarrier();
 			check(LocalQueuedWork || TimeToDie); // well you woke me up, where is the job or termination request?
 			while (LocalQueuedWork)
@@ -331,26 +329,22 @@ protected:
 
 public:
 
-	/** Constructor **/
+	/** Default constructor **/
 	FQueuedThread()
-		: DoWorkEvent(NULL)
+		: DoWorkEvent(nullptr)
 		, TimeToDie(0)
-		, QueuedWork(NULL)
-		, OwningThreadPool(NULL)
-		, Thread(NULL)
-	{
-	}
+		, QueuedWork(nullptr)
+		, OwningThreadPool(nullptr)
+		, Thread(nullptr)
+	{ }
 
 	/**
 	 * Creates the thread with the specified stack size and creates the various
 	 * events to be able to communicate with it.
 	 *
-	 * @param InPool The thread pool interface used to place this thread
-	 * back into the pool of available threads when its work is done
-	 * @param InStackSize The size of the stack to create. 0 means use the
-	 * current thread's stack size
+	 * @param InPool The thread pool interface used to place this thread back into the pool of available threads when its work is done
+	 * @param InStackSize The size of the stack to create. 0 means use the current thread's stack size
 	 * @param ThreadPriority priority of new thread
-	 *
 	 * @return True if the thread and all of its initialization was successful, false otherwise
 	 */
 	virtual bool Create(class FQueuedThreadPool* InPool,uint32 InStackSize = 0,EThreadPriority ThreadPriority=TPri_Normal)
@@ -387,7 +381,7 @@ public:
 		Thread->WaitForCompletion();
 		// Clean up the event
 		delete DoWorkEvent;
-		DoWorkEvent = NULL;
+		DoWorkEvent = nullptr;
 		delete DoWorkEvent;
 		delete Thread;
 		return bDidExitOK;
@@ -401,7 +395,7 @@ public:
 	 */
 	void DoWork(FQueuedWork* InQueuedWork)
 	{
-		check(QueuedWork == NULL && "Can't do more than one task at a time");
+		check(QueuedWork == nullptr && "Can't do more than one task at a time");
 		// Tell the thread the work to be done
 		QueuedWork = InQueuedWork;
 		FPlatformMisc::MemoryBarrier();
@@ -411,48 +405,38 @@ public:
 
 };
 
+
 /**
  * Implementation of a queued thread pool.
  */
 class FQueuedThreadPoolBase : public FQueuedThreadPool
 {
 protected:
-	/**
-	 * The work queue to pull from
-	 */
+
+	/** The work queue to pull from. */
 	TArray<FQueuedWork*> QueuedWork;
 	
-	/**
-	 * The thread pool to dole work out to
-	 */
+	/** The thread pool to dole work out to. */
 	TArray<FQueuedThread*> QueuedThreads;
 
-	/**
-	 * All threads in the pool
-	 */
+	/** All threads in the pool. */
 	TArray<FQueuedThread*> AllThreads;
 
-	/**
-	 * The synchronization object used to protect access to the queued work
-	 */
+	/** The synchronization object used to protect access to the queued work. */
 	FCriticalSection* SynchQueue;
 
-	/**
-	 * If true, indicates the destruction process has taken place
-	 */
+	/** If true, indicates the destruction process has taken place. */
 	bool TimeToDie;
 
 public:
 
+	/** Default constructor. */
 	FQueuedThreadPoolBase()
-		: SynchQueue(NULL)
+		: SynchQueue(nullptr)
 		, TimeToDie(0)
-	{
-	}
+	{ }
 
-	/**
-	 * Clean up the synch objects
-	 */
+	/** Virtual destructor (cleans up the synchronization objects). */
 	virtual ~FQueuedThreadPoolBase()
 	{
 		Destroy();
@@ -462,7 +446,7 @@ public:
 	{
 		// Make sure we have synch objects
 		bool bWasSuccessful = true;
-		check(SynchQueue == NULL);
+		check(SynchQueue == nullptr);
 		SynchQueue = new FCriticalSection();
 		FScopeLock Lock(SynchQueue);
 		// Presize the array so there is no extra memory allocated
@@ -534,7 +518,7 @@ public:
 				AllThreads.Empty();
 			}
 			delete SynchQueue;
-			SynchQueue = NULL;
+			SynchQueue = nullptr;
 		}
 	}
 
@@ -545,8 +529,8 @@ public:
 			InQueuedWork->Abandon();
 			return;
 		}
-		check(InQueuedWork != NULL);
-		FQueuedThread* Thread = NULL;
+		check(InQueuedWork != nullptr);
+		FQueuedThread* Thread = nullptr;
 		// Check to see if a thread is available. Make sure no other threads
 		// can manipulate the thread pool while we do this.
 		check(SynchQueue);
@@ -561,7 +545,7 @@ public:
 			QueuedThreads.RemoveAt(Index);
 		}
 		// Was there a thread ready?
-		if (Thread != NULL)
+		if (Thread != nullptr)
 		{
 			// We have a thread, so tell it to do the work
 			Thread->DoWork(InQueuedWork);
@@ -580,7 +564,7 @@ public:
 		{
 			return false; // no special consideration for this, refuse the retraction and let shutdown proceed
 		}
-		check(InQueuedWork != NULL);
+		check(InQueuedWork != nullptr);
 		check(SynchQueue);
 		FScopeLock sl(SynchQueue);
 		return !!QueuedWork.RemoveSingle(InQueuedWork);
@@ -588,8 +572,8 @@ public:
 
 	virtual FQueuedWork* ReturnToPoolOrGetNextJob(FQueuedThread* InQueuedThread) override
 	{
-		check(InQueuedThread != NULL);
-		FQueuedWork* Work = NULL;
+		check(InQueuedThread != nullptr);
+		FQueuedWork* Work = nullptr;
 		// Check to see if there is any work to be done
 		FScopeLock sl(SynchQueue);
 		if (TimeToDie)
@@ -614,11 +598,13 @@ public:
 	}
 };
 
+
 FQueuedThreadPool* FQueuedThreadPool::Allocate()
 {
 	return new FQueuedThreadPoolBase;
 }
 
+
 /*-----------------------------------------------------------------------------
 	FThreadSingletonInitializer
 -----------------------------------------------------------------------------*/
diff --git a/Engine/Source/Runtime/Core/Public/HAL/ThreadingBase.h b/Engine/Source/Runtime/Core/Public/HAL/ThreadingBase.h
--- a/Engine/Source/Runtime/Core/Public/HAL/ThreadingBase.h
+++ b/Engine/Source/Runtime/Core/Public/HAL/ThreadingBase.h
@@ -44,14 +44,10 @@ public:
 	 */
 	virtual bool Create( bool bIsManualReset = false ) = 0;
 
-	/**
-	 * Triggers the event so any waiting threads are activated
-	 */
+	/** Triggers the event so any waiting threads are activated */
 	virtual void Trigger() = 0;
 
-	/**
-	 * Resets the event to an untriggered (waitable) state
-	 */
+	/** Resets the event to an untriggered (waitable) state */
 	virtual void Reset() = 0;
 
 	/**
@@ -89,9 +85,7 @@ public:
 
 public:
 
-	/**
-	 * Virtual destructor.
-	 */
+	/** Virtual destructor. */
 	virtual ~FEvent() { }
 };
 
@@ -110,17 +104,12 @@ class CORE_API FScopedEvent
 {
 public:
 
-	/**
-	 * Default constructor.
-	 */
+	/** Default constructor. */
 	FScopedEvent()
 		: Event(GetEventFromPool())
-	{
-	}
+	{ }
 
-	/**
-	 * Destructor.
-	 */
+	/** Destructor. */
 	~FScopedEvent()
 	{
 		Event->Wait();
@@ -128,16 +117,16 @@ public:
 		Event = nullptr;
 	}
 
-	/**
-	 * Triggers the event.
-	 */
+	/** Triggers the event. */
 	void Trigger()
 	{
 		Event->Trigger();
 	}
 
 	/**
-	 * Retrieve the event, usually for passing around
+	 * Retrieve the event, usually for passing around.
+	 *
+	 * @return The event.
 	 */
 	FEvent* Get()
 	{
@@ -150,6 +139,7 @@ protected:
 	 * Returns an event object from the pool.
 	 *
 	 * @return An event object.
+	 * @see ReturnToPool
 	 */
 	static FEvent* GetEventFromPool();
 
@@ -157,6 +147,7 @@ protected:
 	 * Returns an event object to the pool.
 	 *
 	 * @param Event The event object to return.
+	 * @see GetEventFromPool
 	 */
 	static void ReturnToPool( FEvent* Event );
 
@@ -190,6 +181,7 @@ public:
 	 * thread that passes this runnable to a new thread.
 	 *
 	 * @return True if initialization was successful, false otherwise
+	 * @see Run, Stop, Exit
 	 */
 	virtual bool Init()
 	{
@@ -202,6 +194,7 @@ public:
 	 * This is where all per object thread work is done. This is only called if the initialization was successful.
 	 *
 	 * @return The exit code of the runnable object
+	 * @see Init, Stop, Exit
 	 */
 	virtual uint32 Run() = 0;
 
@@ -209,6 +202,7 @@ public:
 	 * Stops the runnable object.
 	 *
 	 * This is called if a thread is requested to terminate early.
+	 * @see Init, Run, Exit
 	 */
 	virtual void Stop() { }
 
@@ -216,11 +210,12 @@ public:
 	 * Exits the runnable object.
 	 *
 	 * Called in the context of the aggregating thread to perform any cleanup.
+	 * @see Init, Run, Stop
 	 */
 	virtual void Exit() { }
 
 	/**
-	 * Gets single thread interface pointer used for ticking this runnable when multithreading is disabled.
+	 * Gets single thread interface pointer used for ticking this runnable when multi-threading is disabled.
 	 * If the interface is not implemented, this runnable will not be ticked when FPlatformProcess::SupportsMultithreading() is false.
 	 *
 	 * @return Pointer to the single thread interface or nullptr if not implemented.
@@ -232,9 +227,7 @@ public:
 
 public:
 
-	/**
-	 * Virtual destructor
-	 */
+	/** Virtual destructor */
 	virtual ~FRunnable() { }
 };
 
@@ -312,15 +305,14 @@ public:
 	 */
 	virtual bool Kill( bool bShouldWait = false ) = 0;
 
-	/**
-	 * Halts the caller until this thread is has completed its work.
-	 */
+	/** Halts the caller until this thread is has completed its work. */
 	virtual void WaitForCompletion() = 0;
 
 	/**
 	 * Thread ID for this thread 
 	 *
 	 * @return ID that was set by CreateThread
+	 * @see GetThreadName
 	 */
 	virtual uint32 GetThreadID() = 0;
 
@@ -328,6 +320,7 @@ public:
 	 * Retrieves the given name of the thread
 	 *
 	 * @return Name that was set by CreateThread
+	 * @see GetThreadID
 	 */
 	virtual FString GetThreadName() = 0;
 
@@ -406,10 +399,9 @@ public:
 		return ThreadDestroyedDelegate;
 	}
 
-	/**
-	 * Virtual destructor
-	 */
+	/** Virtual destructor */
 	virtual ~FRunnableThread();
+
 protected:
 
 	/**
@@ -444,9 +436,7 @@ class FSingleThreadEvent
 
 public:
 
-	/**
-	 * Default constructor.
-	 */
+	/** Default constructor. */
 	FSingleThreadEvent()
 		: bTriggered(false)
 		, bManualReset(false)
@@ -512,6 +502,7 @@ public:
 	 * Used internally to add a new thread object when multithreading is disabled.
 	 *
 	 * @param Thread Fake thread object.
+	 * @see RemoveThread
 	 */
 	void AddThread( class FFakeThread* Thread );
 
@@ -519,16 +510,17 @@ public:
 	 * Used internally to remove fake thread object.
 	 *
 	 * @param Thread Fake thread object to be removed.
+	 * @see AddThread
 	 */
 	void RemoveThread( class FFakeThread* Thread );
 
-	/**
-	 * Ticks all fake threads and their runnable objects.
-	 */
+	/** Ticks all fake threads and their runnable objects. */
 	void Tick();
 
 	/**
 	 * Access to the singleton object.
+	 *
+	 * @return Thread manager object.
 	 */
 	static FSingleThreadManager& Get();
 };
@@ -591,9 +583,7 @@ public:
 	 */
 	virtual bool Create( uint32 InNumQueuedThreads, uint32 StackSize = (32 * 1024), EThreadPriority ThreadPriority=TPri_Normal ) = 0;
 
-	/**
-	 * Tells the pool to clean up all background threads
-	 */
+	/** Tells the pool to clean up all background threads */
 	virtual void Destroy() = 0;
 
 	/**
@@ -601,6 +591,7 @@ public:
 	 * it queues the work for later. Otherwise it is immediately dispatched.
 	 *
 	 * @param InQueuedWork The work that needs to be done asynchronously
+	 * @see RetractQueuedWork
 	 */
 	virtual void AddQueuedWork( FQueuedWork* InQueuedWork ) = 0;
 
@@ -609,6 +600,7 @@ public:
 	 *
 	 * @param InQueuedWork The work to try to retract
 	 * @return true if the work was retracted
+	 * @see AddQueuedWork
 	 */
 	virtual bool RetractQueuedWork( FQueuedWork* InQueuedWork ) = 0;
 
@@ -622,15 +614,15 @@ public:
 
 public:
 
-	/**
-	 * Virtual destructor.
-	 */
+	/** Virtual destructor. */
 	virtual ~FQueuedThreadPool() { }
 
 public:
 
 	/**
 	 * Allocates a thread pool
+	 *
+	 * @return A new thread pool.
 	 */
 	static FQueuedThreadPool* Allocate();
 };
@@ -685,6 +677,7 @@ public:
 	 * Increment and return new value.	
 	 *
 	 * @return the new, incremented value
+	 * @see Add, Decrement, Reset, Set, Subtract
 	 */
 	int32 Increment()
 	{
@@ -696,6 +689,7 @@ public:
 	 *
 	 * @param Amount Amount to increase the counter by
 	 * @return the old value
+	 * @see Decrement, Increment, Reset, Set, Subtract
 	 */
 	int32 Add( int32 Amount )
 	{
@@ -706,6 +700,7 @@ public:
 	 * Decrement and return new value.
 	 *
 	 * @return the new, decremented value
+	 * @see Add, Increment, Reset, Set, Subtract
 	 */
 	int32 Decrement()
 	{
@@ -717,6 +712,7 @@ public:
 	 *
 	 * @param Amount Amount to decrease the counter by
 	 * @return the old value
+	 * @see Add, Decrement, Increment, Reset, Set
 	 */
 	int32 Subtract( int32 Amount )
 	{
@@ -728,6 +724,7 @@ public:
 	 *
 	 * @param Value	Value to set the counter to
 	 * @return The old value
+	 * @see Add, Decrement, Increment, Reset, Subtract
 	 */
 	int32 Set( int32 Value )
 	{
@@ -738,6 +735,7 @@ public:
 	 * Resets the counter's value to zero.
 	 *
 	 * @return the old value.
+	 * @see Add, Decrement, Increment, Set, Subtract
 	 */
 	int32 Reset()
 	{
@@ -888,9 +886,7 @@ public:
 		SynchObject->Lock();
 	}
 
-	/**
-	 * Destructor that performs a release on the synchronization object
-	 */
+	/** Destructor that performs a release on the synchronization object. */
 	~FScopeLock()
 	{
 		check(SynchObject);
@@ -898,13 +894,13 @@ public:
 	}
 private:
 
-	// Default constructor (hidden on purpose).
+	/** Default constructor (hidden on purpose). */
 	FScopeLock();
 
-	// Copy constructor( hidden on purpose).
+	/** Copy constructor( hidden on purpose). */
 	FScopeLock( FScopeLock* InScopeLock);
 
-	// Assignment operator (hidden on purpose).
+	/** Assignment operator (hidden on purpose). */
 	FScopeLock& operator=( FScopeLock& InScopeLock )
 	{
 		return *this;
@@ -941,11 +937,14 @@ extern CORE_API int32 GIsRenderingThreadSuspended;
 
 /** @return True if called from the RHI thread, or if called from ANY thread during single threaded rendering */
 extern CORE_API bool IsInRHIThread();
+
 /** Thread used for RHI */
 extern CORE_API FRunnableThread* GRHIThread;
 
 
-/** Minimal base class for the thread singleton. */
+/**
+ * Minimal base class for the thread singleton.
+ */
 struct FThreadSingleton
 {
 	typedef FThreadSingleton* (*TCreateSingletonFuncPtr)();
@@ -957,14 +956,16 @@ struct FThreadSingleton
 };
 
 
-/** Thread singleton initializer. */
+/**
+ * Thread singleton initializer.
+ */
 class FThreadSingletonInitializer
 {
 public:
 
 	/**
-	* @return an instance of a singleton for the current thread.
-	*/
+	 * @return an instance of a singleton for the current thread.
+	 */
 	static CORE_API FThreadSingleton* Get( const FThreadSingleton::TCreateSingletonFuncPtr CreateFunc, const FThreadSingleton::TDestroySingletonFuncPtr DestroyFunc, uint32& TlsSlot );
 };
 
@@ -1001,9 +1002,7 @@ protected:
 		return new T();
 	}
 
-	/**
-	 *	Deletes this instance of the thread singleton. 
-	 */
+	/**	Deletes this instance of the thread singleton.  */
 	virtual void DeleteThis()
 	{
 		T* This = (T*)this;
