Async: Added support for scheduling tasks on the queued thread pool

diff --git a/Engine/Source/Runtime/Core/Private/Tests/Async/AsyncTest.cpp b/Engine/Source/Runtime/Core/Private/Tests/Async/AsyncTest.cpp
--- a/Engine/Source/Runtime/Core/Private/Tests/Async/AsyncTest.cpp
+++ b/Engine/Source/Runtime/Core/Private/Tests/Async/AsyncTest.cpp
@@ -7,7 +7,7 @@
 #include "Future.h"
 
 
-IMPLEMENT_SIMPLE_AUTOMATION_TEST(FAsyncTest, "Core.Async.Async", EAutomationTestFlags::ATF_Editor)
+IMPLEMENT_SIMPLE_AUTOMATION_TEST(FAsyncTest, "Core.Async.Async", EAutomationTestFlags::ATF_SmokeTest)// ATF_Editor)
 
 
 bool FAsyncTest::RunTest(const FString& Parameters)
@@ -30,7 +30,15 @@ bool FAsyncTest::RunTest(const FString& Parameters)
 		auto Future = Async<int>(EAsyncExecution::Thread, Task);
 		int Result = Future.Get();
 
-		TestEqual(TEXT("Task graph task must return expected value"), Result, 123);
+		TestEqual(TEXT("Threaded task must return expected value"), Result, 123);
+	}
+
+	// thread pool task
+	{
+		auto Future = Async<int>(EAsyncExecution::ThreadPool, Task);
+		int Result = Future.Get();
+
+		TestEqual(TEXT("Thread pool task must return expected value"), Result, 123);
 	}
 
 	return true;
diff --git a/Engine/Source/Runtime/Core/Public/Async/Async.h b/Engine/Source/Runtime/Core/Public/Async/Async.h
--- a/Engine/Source/Runtime/Core/Public/Async/Async.h
+++ b/Engine/Source/Runtime/Core/Public/Async/Async.h
@@ -155,6 +155,51 @@ private:
 };
 
 
+/**
+ * Template for asynchronous functions that are executed in the queued thread pool.
+ */
+template<typename ResultType>
+class TAsyncQueuedWork
+	: public FQueuedWork
+{
+public:
+
+	/**
+	 * Creates and initializes a new instance.
+	 *
+	 * @param InFunction The function to execute asynchronously.
+	 * @param InPromise The promise object used to return the function's result.
+	 */
+	TAsyncQueuedWork(TFunction<ResultType()>&& InFunction, TPromise<ResultType>&& InPromise)
+		: Function(InFunction)
+		, Promise(MoveTemp(InPromise))
+	{ }
+
+public:
+
+	// FQueuedWork interface
+
+	virtual void DoThreadedWork() override
+	{
+		Promise.SetValue(Function());
+		delete this;
+	}
+
+	virtual void Abandon() override
+	{
+		// not supported
+	}
+
+private:
+
+	/** The function to execute on the Task Graph. */
+	TFunction<ResultType()> Function;
+
+	/** The promise to assign the result to. */
+	TPromise<ResultType> Promise;
+};
+
+
 /**
  * Executes a given function asynchronously.
  *
@@ -188,6 +233,12 @@ TFuture<ResultType> Async(EAsyncExecution Execution, TFunction<ResultType()> Fun
 		}
 		break;
 
+	case EAsyncExecution::ThreadPool:
+		{
+			GThreadPool->AddQueuedWork(new TAsyncQueuedWork<ResultType>(MoveTemp(Function), MoveTemp(Promise)));
+		}
+		break;
+
 	default:
 		check(false); // not implemented yet!
 	}
