#perf - FQuat, FRotator, FRotationMatrix optimizations cross platform.

Much of this is integrated from Mike.Rayner's changes in /Partners/Microsoft/UE4-MS/
XBoxOne results (from Mike):

FRotator::Quaternion                             3479ns ==> 337ns
FRotator::Vector                                     7860ns ==> 1637ns
FRotationMatrix::FRotationMatrix          4227ns ==> 540ns
FQuat::Rotator                                       5101ns ==> 2160ns

Optimizations for FQuat::RotateVector and  FQuat::RotateVector as well.
Optimized AActor::GetActor<Forward|Up|Right>Vector

Includes UnitTests.

diff --git a/Engine/Source/Runtime/Core/Private/Math/UnrealMath.cpp b/Engine/Source/Runtime/Core/Private/Math/UnrealMath.cpp
--- a/Engine/Source/Runtime/Core/Private/Math/UnrealMath.cpp
+++ b/Engine/Source/Runtime/Core/Private/Math/UnrealMath.cpp
@@ -18,6 +18,7 @@ DEFINE_STAT(STAT_MathConvertRotatorToQuat);
 CORE_API const FVector FVector::ZeroVector(0.0f, 0.0f, 0.0f);
 CORE_API const FVector FVector::UpVector(0.0f, 0.0f, 1.0f);
 CORE_API const FVector FVector::ForwardVector(1.0f, 0.0f, 0.0f);
+CORE_API const FVector FVector::RightVector(0.0f, 1.0f, 0.0f);
 CORE_API const FVector2D FVector2D::ZeroVector(0.0f, 0.0f);
 CORE_API const FVector2D FVector2D::UnitVector(1.0f, 1.0f);
 CORE_API const FRotator FRotator::ZeroRotator(0.f,0.f,0.f);
@@ -258,7 +259,12 @@ FRotator::FRotator(const FQuat& Quat)
 
 CORE_API FVector FRotator::Vector() const
 {
-	return FRotationMatrix( *this ).GetScaledAxis( EAxis::X );
+	float CP, SP, CY, SY;
+	FMath::SinCos( &SP, &CP, FMath::DegreesToRadians(Pitch) );
+	FMath::SinCos( &SY, &CY, FMath::DegreesToRadians(Yaw) );
+	FVector V = FVector( CP*CY, CP*SY, SP );
+
+	return V;
 }
 
 
@@ -269,15 +275,32 @@ FQuat FRotator::Quaternion() const
 #if USE_MATRIX_ROTATOR 
 	FQuat RotationMatrix = FQuat( FRotationMatrix( *this ) );
 #endif
-	static float DEG_TO_RAD = PI/(180.f);
-	static float DIVIDE_BY_2 = DEG_TO_RAD/2.f;
 
-	float CR = FMath::Cos(Roll*DIVIDE_BY_2);
-	float CP = FMath::Cos(Pitch*DIVIDE_BY_2);
-	float CY = FMath::Cos(Yaw*DIVIDE_BY_2);
-	float SR = FMath::Sin(Roll*DIVIDE_BY_2);
-	float SP = FMath::Sin(Pitch*DIVIDE_BY_2);
-	float SY = FMath::Sin(Yaw*DIVIDE_BY_2);
+#if WITH_DIRECTXMATH	// Currently VectorSinCos() is only vectorized in DirectX implementation. Other platforms become slower if they use this.
+	VectorRegister Angles = MakeVectorRegister(Roll, Pitch, Yaw, 0.0f);
+	VectorRegister HalfAngles = VectorMultiply(Angles, GlobalVectorConstants::DEG_TO_RAD_HALF);
+
+	union { VectorRegister v; float f[4]; } SinAngles, CosAngles;	
+	VectorSinCos(&SinAngles.v, &CosAngles.v, &HalfAngles);
+
+	const float	SR	= SinAngles.f[0];
+	const float	SP	= SinAngles.f[1];
+	const float	SY	= SinAngles.f[2];
+	const float	CR	= CosAngles.f[0];
+	const float	CP	= CosAngles.f[1];
+	const float	CY	= CosAngles.f[2];
+
+#else
+	static const float DEG_TO_RAD = PI/(180.f);
+	static const float DIVIDE_BY_2 = DEG_TO_RAD/2.f;
+
+	float SP, SY, SR;
+	float CP, CY, CR;
+
+	FMath::SinCos(&SP, &CP, Pitch*DIVIDE_BY_2);
+	FMath::SinCos(&SY, &CY, Yaw*DIVIDE_BY_2);
+	FMath::SinCos(&SR, &CR, Roll*DIVIDE_BY_2);
+#endif
 
 	FQuat RotationQuat;
 	RotationQuat.W = CR*CP*CY + SR*SP*SY;
@@ -288,7 +311,7 @@ FQuat FRotator::Quaternion() const
 #if USE_MATRIX_ROTATOR 
 	if (!RotationMatrix.Equals(RotationQuat, KINDA_SMALL_NUMBER) && !RotationMatrix.Equals(RotationQuat*-1.f, KINDA_SMALL_NUMBER))
 	{
-		UE_LOG(LogUnrealMath, Log, TEXT("RotationMatrix (%s), RoptationQuat (%s)"), *RotationMatrix.ToString(), *RotationQuat.ToString());
+		UE_LOG(LogUnrealMath, Log, TEXT("RotationMatrix (%s), RotationQuat (%s)"), *RotationMatrix.ToString(), *RotationQuat.ToString());
 	}
 	return RotationMatrix;
 
@@ -397,11 +420,9 @@ FRotator FQuat::Rotator() const
 	checkSlow(IsNormalized());
 #endif
 
-	static const float RAD_TO_DEG = 180.f/PI;
-
-	FRotator RotatorFromQuat;
 	const float SingularityTest = Z*X-W*Y;
-	const float Yaw = FMath::Atan2(2.f*(W*Z+X*Y), (1.f-2.f*(FMath::Square(Y) + FMath::Square(Z))))*RAD_TO_DEG;
+	const float YawY = 2.f*(W*Z+X*Y);
+	const float YawX = (1.f-2.f*(FMath::Square(Y) + FMath::Square(Z)));
 
 	// reference 
 	// http://en.wikipedia.org/wiki/Conversion_between_quaternions_and_Euler_angles
@@ -410,30 +431,75 @@ FRotator FQuat::Rotator() const
 	// this value was found from experience, the above websites recommend different values
 	// but that isn't the case for us, so I went through different testing, and finally found the case 
 	// where both of world lives happily. 
-#define SINGULARITY_THRESHOLD	0.4999995f
+	const float SINGULARITY_THRESHOLD = 0.4999995f;
+
+#if PLATFORM_ENABLE_VECTORINTRINSICS
+
+	FRotator RotatorFromQuat;
+	union { VectorRegister v; float f[4]; } VRotatorFromQuat;
+
+	if (SingularityTest < -SINGULARITY_THRESHOLD)
+	{
+		// Pitch
+		VRotatorFromQuat.f[0] = 3.0f*HALF_PI;	// 270 deg
+		// Yaw
+		VRotatorFromQuat.f[1] = FMath::Atan2(YawY, YawX);
+		// Roll
+		VRotatorFromQuat.f[2] = -VRotatorFromQuat.f[1] - (2.f * FMath::Atan2(X, W));
+	}
+	else if (SingularityTest > SINGULARITY_THRESHOLD)
+	{
+		// Pitch
+		VRotatorFromQuat.f[0] = HALF_PI;	// 90 deg
+		// Yaw
+		VRotatorFromQuat.f[1] = FMath::Atan2(YawY, YawX);
+		//Roll
+		VRotatorFromQuat.f[2] = VRotatorFromQuat.f[1] - (2.f * FMath::Atan2(X, W));
+	}
+	else
+	{
+		//Pitch
+		VRotatorFromQuat.f[0] = FMath::FastAsin(2.f*(SingularityTest));
+		// Yaw
+		VRotatorFromQuat.f[1] = FMath::Atan2(YawY, YawX);
+		// Roll
+		VRotatorFromQuat.f[2] = FMath::Atan2(-2.f*(W*X+Y*Z), (1.f-2.f*(FMath::Square(X) + FMath::Square(Y))));
+	}
+
+	VRotatorFromQuat.f[3] = 0.f; // We should initialize this, otherwise the value can be denormalized which is bad for floating point perf.
+	VRotatorFromQuat.v = VectorMultiply(VRotatorFromQuat.v, GlobalVectorConstants::RAD_TO_DEG);
+	VRotatorFromQuat.v = VectorNormalizeRotator(VRotatorFromQuat.v);
+	VectorStoreFloat3(VRotatorFromQuat.v, &RotatorFromQuat);
+
+#else // PLATFORM_ENABLE_VECTORINTRINSICS
+
+	static const float RAD_TO_DEG = (180.f)/PI;
+	FRotator RotatorFromQuat;
 
 	if (SingularityTest < -SINGULARITY_THRESHOLD)
 	{
 		RotatorFromQuat.Pitch = 270.f;
-		RotatorFromQuat.Yaw = Yaw;
-		RotatorFromQuat.Roll = -RotatorFromQuat.Yaw - 2.f* FMath::Atan2(X, W) * RAD_TO_DEG;
+		RotatorFromQuat.Yaw = FMath::Atan2(YawY, YawX) * RAD_TO_DEG;
+		RotatorFromQuat.Roll = -RotatorFromQuat.Yaw - (2.f * FMath::Atan2(X, W) * RAD_TO_DEG);
 	}
 	else if (SingularityTest > SINGULARITY_THRESHOLD)
 	{
 		RotatorFromQuat.Pitch = 90.f;
-		RotatorFromQuat.Yaw = Yaw;
-		RotatorFromQuat.Roll = RotatorFromQuat.Yaw - 2.f* FMath::Atan2(X, W) * RAD_TO_DEG;
+		RotatorFromQuat.Yaw = FMath::Atan2(YawY, YawX) * RAD_TO_DEG;
+		RotatorFromQuat.Roll = RotatorFromQuat.Yaw - (2.f * FMath::Atan2(X, W) * RAD_TO_DEG);
 	}
 	else
 	{
-		const float Pitch = FMath::Asin(2.f*(SingularityTest));
-		RotatorFromQuat.Pitch = FRotator::ClampAxis(Pitch*RAD_TO_DEG); //clamp it so within 360 - this is for if below
-		RotatorFromQuat.Yaw = Yaw;
-		RotatorFromQuat.Roll = FMath::Atan2(-2.f*(W*X+Y*Z), (1-2.f*(FMath::Square(X) + FMath::Square(Y))))*RAD_TO_DEG;
+		RotatorFromQuat.Pitch = FMath::FastAsin(2.f*(SingularityTest)) * RAD_TO_DEG;
+		RotatorFromQuat.Yaw = FMath::Atan2(YawY, YawX) * RAD_TO_DEG;
+		RotatorFromQuat.Roll = FMath::Atan2(-2.f*(W*X+Y*Z), (1.f-2.f*(FMath::Square(X) + FMath::Square(Y)))) * RAD_TO_DEG;
 	}
 
 	RotatorFromQuat.Normalize();
 
+#endif // PLATFORM_ENABLE_VECTORINTRINSICS
+
+
 #if USE_MATRIX_ROTATOR
 	RotatorFromMatrix = RotatorFromMatrix.Clamp();
 
diff --git a/Engine/Source/Runtime/Core/Private/Tests/Math/UnrealMathTest.cpp b/Engine/Source/Runtime/Core/Private/Tests/Math/UnrealMathTest.cpp
--- a/Engine/Source/Runtime/Core/Private/Tests/Math/UnrealMathTest.cpp
+++ b/Engine/Source/Runtime/Core/Private/Tests/Math/UnrealMathTest.cpp
@@ -77,6 +77,34 @@ bool TestVectorsEqual3( VectorRegister Vec0, VectorRegister Vec1, float Toleranc
 	return GSum <= Tolerance;
 }
 
+/**
+ * Tests if two vectors (xyz) are equal within an optional tolerance
+ *
+ * @param Vec0 First vector
+ * @param Vec1 Second vector
+ * @param Tolerance Error allowed for the comparison
+ *
+ * @return true if equal(ish)
+ */
+bool TestVectorsEqual3_NoVec( const FVector& Vec0, const FVector& Vec1, float Tolerance = 0.0f)
+{
+	GScratch[0] = Vec0.X;
+	GScratch[1] = Vec0.Y;
+	GScratch[2] = Vec0.Z;
+	GScratch[3] = 0.0f;
+	GScratch[4] = Vec1.X;
+	GScratch[5] = Vec1.Y;
+	GScratch[6] = Vec1.Z;
+	GScratch[7] = 0.0f;
+	GSum = 0.f;
+
+	for ( int32 Component = 0; Component < 3; Component++ ) 
+	{
+		GSum += FMath::Abs<float>( GScratch[ Component + 0 ] - GScratch[ Component + 4 ] );
+	}
+	GPassing = GPassing && GSum <= Tolerance;
+	return GSum <= Tolerance;
+}
 
 /**
  * Tests if two matrices (4x4 xyzw) are equal within an optional tolerance
@@ -256,6 +284,49 @@ VectorRegister TestVectorTransformVector(const VectorRegister&  VecP,  const voi
 	return Result.v;
 }
 
+/**
+* Get Rotation as a quaternion.
+* @param Rotator FRotator 
+* @return Rotation as a quaternion.
+*/
+FORCENOINLINE FQuat TestRotatorToQuaternion( const FRotator& Rotator)
+{
+	const float CR = FMath::Cos(FMath::DegreesToRadians(Rotator.Roll  * 0.5f));
+	const float CP = FMath::Cos(FMath::DegreesToRadians(Rotator.Pitch * 0.5f));
+	const float CY = FMath::Cos(FMath::DegreesToRadians(Rotator.Yaw   * 0.5f));
+	const float SR = FMath::Sin(FMath::DegreesToRadians(Rotator.Roll  * 0.5f));
+	const float SP = FMath::Sin(FMath::DegreesToRadians(Rotator.Pitch * 0.5f));
+	const float SY = FMath::Sin(FMath::DegreesToRadians(Rotator.Yaw   * 0.5f));
+
+	FQuat RotationQuat;
+	RotationQuat.W = CR*CP*CY + SR*SP*SY;
+	RotationQuat.X = CR*SP*SY - SR*CP*CY;
+	RotationQuat.Y = -CR*SP*CY - SR*CP*SY;
+	RotationQuat.Z = CR*CP*SY - SR*SP*CY;
+	return RotationQuat;
+}
+
+FORCENOINLINE FVector TestQuaternionRotateVector(const FQuat& Quat, const FVector& Vector)
+{
+	// (q.W*q.W-qv.qv)v + 2(qv.v)qv + 2 q.W (qv x v)
+	const FVector qv(Quat.X, Quat.Y, Quat.Z);
+	FVector vOut = (2.f * Quat.W) * (qv ^ Vector);
+	vOut += ((Quat.W * Quat.W) - (qv | qv)) * Vector;
+	vOut += (2.f * (qv | Vector)) * qv;
+	
+	return vOut;
+}
+
+FORCENOINLINE FVector TestQuaternionMultiplyVector(const FQuat& Quat, const FVector& Vector)
+{
+	FQuat VQ(Vector.X, Vector.Y, Vector.Z, 0.f);
+	FQuat VT, VR;
+	FQuat I = Quat.Inverse();
+	VectorQuaternionMultiply(&VT, &Quat, &VQ);
+	VectorQuaternionMultiply(&VR, &VT, &I);
+
+	return FVector(VR.X, VR.Y, VR.Z);
+}
 
 /**
 * Multiplies two quaternions: The order matters.
@@ -284,6 +355,50 @@ void TestVectorQuaternionMultiply( void *Result, const void* Quat1, const void*
 	R[3] = TW;
 }
 
+/**
+ * Converts a Quaternion to a Rotator.
+ */
+FORCENOINLINE FRotator TestQuaternionToRotator(const FQuat& Quat)
+{
+	const float X = Quat.X;
+	const float Y = Quat.Y;
+	const float Z = Quat.Z;
+	const float W = Quat.W;
+
+	const float SingularityTest = Z*X-W*Y;
+	const float YawY = 2.f*(W*Z+X*Y);
+	const float YawX = (1.f-2.f*(FMath::Square(Y) + FMath::Square(Z)));
+	const float SINGULARITY_THRESHOLD = 0.4999995f;
+
+	static const float RAD_TO_DEG = (180.f)/PI;
+	FRotator RotatorFromQuat;
+
+	if (SingularityTest < -SINGULARITY_THRESHOLD)
+	{
+		RotatorFromQuat.Pitch = 270.f;
+		RotatorFromQuat.Yaw = FMath::Atan2(YawY, YawX) * RAD_TO_DEG;
+		RotatorFromQuat.Roll = -RotatorFromQuat.Yaw - (2.f * FMath::Atan2(X, W) * RAD_TO_DEG);
+	}
+	else if (SingularityTest > SINGULARITY_THRESHOLD)
+	{
+		RotatorFromQuat.Pitch = 90.f;
+		RotatorFromQuat.Yaw = FMath::Atan2(YawY, YawX) * RAD_TO_DEG;
+		RotatorFromQuat.Roll = RotatorFromQuat.Yaw - (2.f * FMath::Atan2(X, W) * RAD_TO_DEG);
+	}
+	else
+	{
+		RotatorFromQuat.Pitch = FMath::Asin(2.f*(SingularityTest)) * RAD_TO_DEG;
+		RotatorFromQuat.Yaw = FMath::Atan2(YawY, YawX) * RAD_TO_DEG;
+		RotatorFromQuat.Roll = FMath::Atan2(-2.f*(W*X+Y*Z), (1.f-2.f*(FMath::Square(X) + FMath::Square(Y)))) * RAD_TO_DEG;
+	}
+
+	RotatorFromQuat.Pitch = FRotator::NormalizeAxis(RotatorFromQuat.Pitch);
+	RotatorFromQuat.Yaw = FRotator::NormalizeAxis(RotatorFromQuat.Yaw);
+	RotatorFromQuat.Roll = FRotator::NormalizeAxis(RotatorFromQuat.Roll);
+
+	return RotatorFromQuat;
+}
+
 
 /**
  * Helper debugf function to print out success or failure information for a test
@@ -715,6 +830,70 @@ bool FVectorRegisterAbstractionTest::RunTest(const FString& Parameters)
 	LogTest( TEXT("VectorTransformVector"), TestVectorsEqual( V1, V2 ) );
 
 	FQuat Q0, Q1, Q2, Q3;
+
+	{
+		FRotator Rotator0;
+		Rotator0 = FRotator(30.0f, -45.0f, 90.0f);
+		Q0 = Rotator0.Quaternion();
+		Q1 = TestRotatorToQuaternion(Rotator0);
+		LogTest( TEXT("TestRotatorToQuaternion"), Q0.Equals( Q1, 1e-6f ) );
+
+		FVector FV0, FV1;
+		FV0 = Rotator0.Vector();
+		FV1 = FRotationMatrix( Rotator0 ).GetScaledAxis( EAxis::X );
+		LogTest( TEXT("Test0 Rotator::Vector()"), FV0.Equals( FV1, 1e-6f ) );
+
+
+		Rotator0 = FRotator(45.0f,  60.0f, 120.0f);
+		Q0 = Rotator0.Quaternion();
+		Q1 = TestRotatorToQuaternion(Rotator0);
+		LogTest( TEXT("TestRotatorToQuaternion"), Q0.Equals( Q1, 1e-6f ) );
+
+		FV0 = Rotator0.Vector();
+		FV1 = FRotationMatrix( Rotator0 ).GetScaledAxis( EAxis::X );
+		LogTest( TEXT("Test1 Rotator::Vector()"), FV0.Equals( FV1, 1e-6f ) );
+	}
+
+	{
+		const FRotator Rotator0(30.0f, -45.0f, 90.0f);
+		const FRotator Rotator1(45.0f,  60.0f, 120.0f);
+		const FQuat Quat0(Rotator0);
+		const FQuat Quat1(Rotator1);
+
+		LogTest( TEXT("Test0 FQuat::RotateVector(const FVector& Vec)"),		TestVectorsEqual3_NoVec(Quat0.RotateVector(FVector::ForwardVector), TestQuaternionRotateVector(Quat0, FVector::ForwardVector),	1e-4f ));
+		LogTest( TEXT("Test0 FQuat::operator*(const FVector& Vec)"),		TestVectorsEqual3_NoVec(Quat0 * FVector::ForwardVector, TestQuaternionMultiplyVector(Quat0, FVector::ForwardVector), 1e-4f ));
+		LogTest( TEXT("Test1 FQuat::RotateVector(const FVector& Vec)"),		TestVectorsEqual3_NoVec(Quat0.RotateVector(FVector::RightVector), TestQuaternionRotateVector(Quat0, FVector::RightVector),	1e-4f ));
+		LogTest( TEXT("Test1 FQuat::operator*(const FVector& Vec)"),		TestVectorsEqual3_NoVec(Quat0 * FVector::RightVector, TestQuaternionMultiplyVector(Quat0, FVector::RightVector), 1e-4f ));
+		LogTest( TEXT("Test2 FQuat::RotateVector(const FVector& Vec)"),		TestVectorsEqual3_NoVec(Quat0.RotateVector(FVector::UpVector), TestQuaternionRotateVector(Quat0, FVector::UpVector),	1e-4f ));
+		LogTest( TEXT("Test2 FQuat::operator*(const FVector& Vec)"),		TestVectorsEqual3_NoVec(Quat0 * FVector::UpVector, TestQuaternionMultiplyVector(Quat0, FVector::UpVector), 1e-4f ));
+		LogTest( TEXT("Test3 FQuat::RotateVector(const FVector& Vec)"),		TestVectorsEqual3_NoVec(Quat0.RotateVector(FVector(45.0f, -60.0f, 120.0f)), TestQuaternionRotateVector(Quat0, FVector(45.0f, -60.0f, 120.0f)),	1e-4f ));
+		LogTest( TEXT("Test3 FQuat::operator*(const FVector& Vec)"),		TestVectorsEqual3_NoVec(Quat0 * FVector(45.0f, -60.0f, 120.0f), TestQuaternionMultiplyVector(Quat0, FVector(45.0f, -60.0f, 120.0f)), 1e-4f ));
+		LogTest( TEXT("Test4 FQuat::RotateVector(const FVector& Vec)"),		TestVectorsEqual3_NoVec(Quat1.RotateVector(FVector::ForwardVector), TestQuaternionRotateVector(Quat1, FVector::ForwardVector),	1e-4f ));
+		LogTest( TEXT("Test4 FQuat::operator*(const FVector& Vec)"),		TestVectorsEqual3_NoVec(Quat1 * FVector::ForwardVector, TestQuaternionMultiplyVector(Quat1, FVector::ForwardVector), 1e-4f ));
+		LogTest( TEXT("Test5 FQuat::RotateVector(const FVector& Vec)"),		TestVectorsEqual3_NoVec(Quat1.RotateVector(FVector::RightVector), TestQuaternionRotateVector(Quat1, FVector::RightVector),	1e-4f ));
+		LogTest( TEXT("Test5 FQuat::operator*(const FVector& Vec)"),		TestVectorsEqual3_NoVec(Quat1 * FVector::RightVector, TestQuaternionMultiplyVector(Quat1, FVector::RightVector), 1e-4f ));
+		LogTest( TEXT("Test6 FQuat::RotateVector(const FVector& Vec)"),		TestVectorsEqual3_NoVec(Quat1.RotateVector(FVector::UpVector), TestQuaternionRotateVector(Quat1, FVector::UpVector),	1e-4f ));
+		LogTest( TEXT("Test6 FQuat::operator*(const FVector& Vec)"),		TestVectorsEqual3_NoVec(Quat1 * FVector::UpVector, TestQuaternionMultiplyVector(Quat1, FVector::UpVector), 1e-4f ));
+		LogTest( TEXT("Test7 FQuat::RotateVector(const FVector& Vec)"),		TestVectorsEqual3_NoVec(Quat1.RotateVector(FVector(45.0f, -60.0f, 120.0f)), TestQuaternionRotateVector(Quat1, FVector(45.0f, -60.0f, 120.0f)),	1e-4f ));
+		LogTest( TEXT("Test7 FQuat::operator*(const FVector& Vec)"),		TestVectorsEqual3_NoVec(Quat1 * FVector(45.0f, -60.0f, 120.0f), TestQuaternionMultiplyVector(Quat1, FVector(45.0f, -60.0f, 120.0f)), 1e-4f ));
+
+	}
+
+	{
+		FRotator Rotator0, Rotator1, Rotator2;
+		Rotator0 = FRotator(30.0f, -45.0f, 90.0f);
+		Q0 = TestRotatorToQuaternion(Rotator0);
+		Rotator1 = Q0.Rotator();
+		Rotator2 = TestQuaternionToRotator(Q0);
+		LogTest(TEXT("TestQuaternionToRotator"), Rotator1.Equals(Rotator2, 1e-5f));
+
+		Rotator0 = FRotator(45.0f, 60.0f, 120.0f);
+		Q0 = TestRotatorToQuaternion(Rotator0);
+		Rotator1 = Q0.Rotator();
+		Rotator2 = TestQuaternionToRotator(Q0);
+		LogTest(TEXT("TestQuaternionToRotator"), Rotator1.Equals(Rotator2, 1e-5f));
+	}
+
 	Q0 = FQuat(FRotator(30.0f, -45.0f, 90.0f));
 	Q1 = FQuat(FRotator(45.0f,  60.0f, 120.0f));
 	VectorQuaternionMultiply(&Q2, &Q0, &Q1);
diff --git a/Engine/Source/Runtime/Core/Public/Math/Quat.h b/Engine/Source/Runtime/Core/Public/Math/Quat.h
--- a/Engine/Source/Runtime/Core/Public/Math/Quat.h
+++ b/Engine/Source/Runtime/Core/Public/Math/Quat.h
@@ -290,6 +290,14 @@ public:
 	 * @return vector after rotation
 	 */
 	FVector RotateVector( FVector V ) const;
+	
+	/**
+	 * Rotate a vector by the inverse of this quaternion.
+	 *
+	 * @param V the vector to be rotated
+	 * @return vector after rotation by the inverse of this quaternion.
+	 */
+	FVector UnrotateVector( FVector V ) const;
 
 	/**
 	 * @return quaternion with W=0 and V=theta*v.
@@ -515,15 +523,9 @@ FORCEINLINE FQuat::FQuat( const FRotator& R )
 }
 
 
-inline FVector FQuat::operator*( const FVector& V ) const
+FORCEINLINE FVector FQuat::operator*( const FVector& V ) const
 {
-	FQuat VQ(V.X, V.Y, V.Z, 0.f);
-	FQuat VT, VR;
-	FQuat I = Inverse();
-	VectorQuaternionMultiply(&VT, this, &VQ);
-	VectorQuaternionMultiply(&VR, &VT, &I);
-
-	return FVector(VR.X, VR.Y, VR.Z);
+	return RotateVector(V);
 }
 
 
@@ -809,14 +811,44 @@ FORCEINLINE FVector FQuat::GetRotationAxis() const
 
 FORCEINLINE FVector FQuat::RotateVector( FVector V ) const
 {	
+#if WITH_DIRECTXMATH
+	FVector Result;
+	VectorQuaternionVector3Rotate(&Result, &V, this);
+	return Result;
+
+	/*
+	// In unit testing this appears to be slower than the non-vectorized version.
+#elif PLATFORM_ENABLE_VECTORINTRINSICS
+	FQuat VQ(V.X, V.Y, V.Z, 0.f);
+	FQuat VT, VR;
+	FQuat I = Inverse();
+	VectorQuaternionMultiply(&VT, this, &VQ);
+	VectorQuaternionMultiply(&VR, &VT, &I);
+
+	return FVector(VR.X, VR.Y, VR.Z);
+	*/
+
+#else
 	// (q.W*q.W-qv.qv)v + 2(qv.v)qv + 2 q.W (qv x v)
 
 	const FVector qv(X, Y, Z);
-	FVector vOut = 2.f * W * (qv ^ V);
+	FVector vOut = (2.f * W) * (qv ^ V);
 	vOut += ((W * W) - (qv | qv)) * V;
 	vOut += (2.f * (qv | V)) * qv;
 
 	return vOut;
+#endif
+}
+
+FORCEINLINE FVector FQuat::UnrotateVector( FVector V ) const
+{	
+#if WITH_DIRECTXMATH
+	FVector Result;
+	VectorQuaternionVector3InverseRotate(&Result, &V, this);
+	return Result;
+#else
+	return Inverse().RotateVector(V);
+#endif
 }
 
 
diff --git a/Engine/Source/Runtime/Core/Public/Math/RotationTranslationMatrix.h b/Engine/Source/Runtime/Core/Public/Math/RotationTranslationMatrix.h
--- a/Engine/Source/Runtime/Core/Public/Math/RotationTranslationMatrix.h
+++ b/Engine/Source/Runtime/Core/Public/Math/RotationTranslationMatrix.h
@@ -27,12 +27,32 @@ public:
 
 FORCEINLINE FRotationTranslationMatrix::FRotationTranslationMatrix(const FRotator& Rot, const FVector& Origin)
 {
-	const float	SR	= FMath::Sin(Rot.Roll * PI / 180.f);
-	const float	SP	= FMath::Sin(Rot.Pitch * PI / 180.f);
-	const float	SY	= FMath::Sin(Rot.Yaw * PI / 180.f);
-	const float	CR	= FMath::Cos(Rot.Roll * PI / 180.f);
-	const float	CP	= FMath::Cos(Rot.Pitch * PI / 180.f);
-	const float	CY	= FMath::Cos(Rot.Yaw * PI / 180.f);
+#if WITH_DIRECTXMATH	// Currently VectorSinCos() is only vectorized in DirectX implementation. Other platforms become slower if they use this.
+
+	VectorRegister Angles = MakeVectorRegister(Rot.Roll, Rot.Pitch, Rot.Yaw, 0.0f);
+	VectorRegister HalfAngles = VectorMultiply(Angles, GlobalVectorConstants::DEG_TO_RAD);
+
+	union { VectorRegister v; float f[4]; } SinAngles, CosAngles;	
+	VectorSinCos(&SinAngles.v, &CosAngles.v, &HalfAngles);
+
+	const float	SR	= SinAngles.f[0];
+	const float	SP	= SinAngles.f[1];
+	const float	SY	= SinAngles.f[2];
+	const float	CR	= CosAngles.f[0];
+	const float	CP	= CosAngles.f[1];
+	const float	CY	= CosAngles.f[2];
+
+#else
+	
+	static const float DEG_TO_RAD = PI/(180.f);
+
+	float SP, SY, SR;
+	float CP, CY, CR;
+	FMath::SinCos(&SP, &CP, Rot.Pitch * DEG_TO_RAD);
+	FMath::SinCos(&SY, &CY, Rot.Yaw * DEG_TO_RAD);
+	FMath::SinCos(&SR, &CR, Rot.Roll * DEG_TO_RAD);
+
+#endif
 
 	M[0][0]	= CP * CY;
 	M[0][1]	= CP * SY;
diff --git a/Engine/Source/Runtime/Core/Public/Math/Rotator.h b/Engine/Source/Runtime/Core/Public/Math/Rotator.h
--- a/Engine/Source/Runtime/Core/Public/Math/Rotator.h
+++ b/Engine/Source/Runtime/Core/Public/Math/Rotator.h
@@ -605,9 +605,15 @@ FORCEINLINE FRotator FRotator::GetDenormalized() const
 
 FORCEINLINE void FRotator::Normalize()
 {
-	Yaw = NormalizeAxis(Yaw);
+#if PLATFORM_ENABLE_VECTORINTRINSICS
+	VectorRegister VRotator = VectorLoadFloat3_W0(this);
+	VRotator = VectorNormalizeRotator(VRotator);
+	VectorStoreFloat3(VRotator, this);
+#else
 	Pitch = NormalizeAxis(Pitch);
+	Yaw = NormalizeAxis(Yaw);
 	Roll = NormalizeAxis(Roll);
+#endif
 }
 
 FORCEINLINE FString FRotator::ToString() const
diff --git a/Engine/Source/Runtime/Core/Public/Math/UnrealMathDirectX.h b/Engine/Source/Runtime/Core/Public/Math/UnrealMathDirectX.h
--- a/Engine/Source/Runtime/Core/Public/Math/UnrealMathDirectX.h
+++ b/Engine/Source/Runtime/Core/Public/Math/UnrealMathDirectX.h
@@ -370,6 +370,16 @@ FORCEINLINE float VectorGetComponent( VectorRegister Vec, uint32 ComponentIndex
  */
 #define VectorBitwiseXor( Vec1, Vec2 )	DirectX::XMVectorXorInt( Vec1, Vec2 )
 
+ /**
+ * Returns an integer bit-mask (0x00 - 0x0f) based on the sign-bit for each component in a vector.
+ *
+ * @param VecMask		Vector
+ * @return				Bit 0 = sign(VecMask.x), Bit 1 = sign(VecMask.y), Bit 2 = sign(VecMask.z), Bit 3 = sign(VecMask.w)
+ */
+#define VectorMaskBits( VecMask )			_mm_movemask_ps( VecMask )
+
+
+
 /**
  * Calculates the cross product of two vectors (XYZ components). W is set to 0.
  *
@@ -684,6 +694,102 @@ FORCEINLINE void VectorQuaternionMultiply( VectorRegister *VResult, const Vector
 	*VResult = VectorQuaternionMultiply2(*VQuat1, *VQuat2);
 }
 
+FORCEINLINE void VectorQuaternionVector3Rotate( FVector *Result, const FVector* Vec, const FQuat* Quat)
+{	
+	VectorRegister XMVec = VectorLoad(Vec);
+	VectorRegister XMQuat = VectorLoadAligned(Quat);
+	VectorRegister XMResult = DirectX::XMVector3Rotate(XMVec, XMQuat);
+	VectorStoreFloat3(XMResult, Result);
+}
+
+FORCEINLINE void VectorQuaternionVector3InverseRotate( FVector *Result, const FVector* Vec, const FQuat* Quat)
+{	
+	VectorRegister XMVec = VectorLoad(Vec);
+	VectorRegister XMQuat = VectorLoadAligned(Quat);
+	VectorRegister XMResult = DirectX::XMVector3InverseRotate(XMVec, XMQuat);
+	VectorStoreFloat3(XMResult, Result);
+}
+
+
+/**
+* Computes the sine and cosine of each component of a Vector.
+*
+* @param VSinAngles	VectorRegister Pointer to where the Sin result should be stored
+* @param VCosAngles	VectorRegister Pointer to where the Cos result should be stored
+* @param VAngles VectorRegister Pointer to the input angles 
+*/
+FORCEINLINE void VectorSinCos(  VectorRegister* VSinAngles, VectorRegister* VCosAngles, const VectorRegister* VAngles )
+{
+	using namespace DirectX;
+	// Force the value within the bounds of pi	
+	XMVECTOR x = XMVectorModAngles(*VAngles);
+
+	// Map in [-pi/2,pi/2] with sin(y) = sin(x), cos(y) = sign*cos(x).
+	XMVECTOR sign = _mm_and_ps(x, g_XMNegativeZero);
+	__m128 c = _mm_or_ps(g_XMPi, sign);  // pi when x >= 0, -pi when x < 0
+	__m128 absx = _mm_andnot_ps(sign, x);  // |x|
+	__m128 rflx = _mm_sub_ps(c, x);
+	__m128 comp = _mm_cmple_ps(absx, g_XMHalfPi);
+	__m128 select0 = _mm_and_ps(comp, x);
+	__m128 select1 = _mm_andnot_ps(comp, rflx);
+	x = _mm_or_ps(select0, select1);
+	select0 = _mm_and_ps(comp, g_XMOne);
+	select1 = _mm_andnot_ps(comp, g_XMNegativeOne);
+	sign = _mm_or_ps(select0, select1);
+
+	__m128 x2 = _mm_mul_ps(x, x);
+
+	// Compute polynomial approximation of sine
+	const XMVECTOR SC1 = g_XMSinCoefficients1;
+	XMVECTOR vConstants = XM_PERMUTE_PS( SC1, _MM_SHUFFLE(0, 0, 0, 0) );
+	__m128 Result = _mm_mul_ps(vConstants, x2);
+
+	const XMVECTOR SC0 = g_XMSinCoefficients0;
+	vConstants = XM_PERMUTE_PS( SC0, _MM_SHUFFLE(3, 3, 3, 3) );
+	Result = _mm_add_ps(Result, vConstants);
+	Result = _mm_mul_ps(Result, x2);
+
+	vConstants = XM_PERMUTE_PS( SC0, _MM_SHUFFLE(2, 2, 2, 2) );
+	Result = _mm_add_ps(Result, vConstants);
+	Result = _mm_mul_ps(Result, x2);
+
+	vConstants = XM_PERMUTE_PS( SC0, _MM_SHUFFLE(1, 1, 1, 1) );
+	Result = _mm_add_ps(Result, vConstants);
+	Result = _mm_mul_ps(Result, x2);
+
+	vConstants = XM_PERMUTE_PS( SC0, _MM_SHUFFLE(0, 0, 0, 0) );
+	Result = _mm_add_ps(Result, vConstants);
+	Result = _mm_mul_ps(Result, x2);
+	Result = _mm_add_ps(Result, g_XMOne);
+	Result = _mm_mul_ps(Result, x);
+	*VSinAngles = Result;
+
+	// Compute polynomial approximation of cosine
+	const XMVECTOR CC1 = g_XMCosCoefficients1;
+	vConstants = XM_PERMUTE_PS( CC1, _MM_SHUFFLE(0, 0, 0, 0) );
+	Result = _mm_mul_ps(vConstants, x2);
+
+	const XMVECTOR CC0 = g_XMCosCoefficients0;
+	vConstants = XM_PERMUTE_PS( CC0, _MM_SHUFFLE(3, 3, 3, 3) );
+	Result = _mm_add_ps(Result, vConstants);
+	Result = _mm_mul_ps(Result, x2);
+
+	vConstants = XM_PERMUTE_PS( CC0,  _MM_SHUFFLE(2, 2, 2, 2) );
+	Result = _mm_add_ps(Result, vConstants);
+	Result = _mm_mul_ps(Result, x2);
+
+	vConstants = XM_PERMUTE_PS( CC0,  _MM_SHUFFLE(1, 1, 1, 1) );
+	Result = _mm_add_ps(Result, vConstants);
+	Result = _mm_mul_ps(Result, x2);
+
+	vConstants = XM_PERMUTE_PS( CC0, _MM_SHUFFLE(0, 0, 0, 0) );
+	Result = _mm_add_ps(Result, vConstants);
+	Result = _mm_mul_ps(Result, x2);
+	Result = _mm_add_ps(Result, g_XMOne);
+	Result = _mm_mul_ps(Result, sign);
+	*VCosAngles = Result;
+}
+
 
 // Returns true if the vector contains a component that is either NAN or +/-infinite.
 inline bool VectorContainsNaNOrInfinite(const VectorRegister& Vec)
@@ -785,13 +891,9 @@ FORCEINLINE VectorRegister VectorFractional(const VectorRegister& X)
 	return VectorSubtract(X, VectorTruncate(X));
 }
 
-//TODO: Vectorize
 FORCEINLINE VectorRegister VectorMod(const VectorRegister& X, const VectorRegister& Y)
 {
-	return MakeVectorRegister((float)fmod(VectorGetComponent(X, 0), VectorGetComponent(Y, 0)),
-		(float)fmod(VectorGetComponent(X, 1), VectorGetComponent(Y, 1)),
-		(float)fmod(VectorGetComponent(X, 2), VectorGetComponent(Y, 2)),
-		(float)fmod(VectorGetComponent(X, 3), VectorGetComponent(Y, 3)));
+	return DirectX::XMVectorMod(X, Y);
 }
 
 //TODO: Vectorize
diff --git a/Engine/Source/Runtime/Core/Public/Math/UnrealMathFPU.h b/Engine/Source/Runtime/Core/Public/Math/UnrealMathFPU.h
--- a/Engine/Source/Runtime/Core/Public/Math/UnrealMathFPU.h
+++ b/Engine/Source/Runtime/Core/Public/Math/UnrealMathFPU.h
@@ -950,6 +950,28 @@ FORCEINLINE float VectorGetComponent(VectorRegister Vec, uint32 ComponentIndex)
  */
 #define VECTOR_ROUND_TOWARD_ZERO		0
 
+
+/**
+* Computes the sine and cosine of each component of a Vector.
+*
+* @param VSinAngles	VectorRegister Pointer to where the Sin result should be stored
+* @param VCosAngles	VectorRegister Pointer to where the Cos result should be stored
+* @param VAngles VectorRegister Pointer to the input angles 
+*/
+FORCEINLINE void VectorSinCos(  VectorRegister* VSinAngles, VectorRegister* VCosAngles, const VectorRegister* VAngles )
+{	
+	union { VectorRegister v; float f[4]; } VecSin, VecCos, VecAngles;
+	VecAngles.v = *VAngles;
+
+	FMath::SinCos(&VecSin.f[0], &VecCos.f[0], VecAngles.f[0]);
+	FMath::SinCos(&VecSin.f[1], &VecCos.f[1], VecAngles.f[1]);
+	FMath::SinCos(&VecSin.f[2], &VecCos.f[2], VecAngles.f[2]);
+	FMath::SinCos(&VecSin.f[3], &VecCos.f[3], VecAngles.f[3]);
+
+	*VSinAngles = VecSin.v;
+	*VCosAngles = VecCos.v;
+}
+
 // Returns true if the vector contains a component that is either NAN or +/-infinite.
 inline bool VectorContainsNaNOrInfinite(const VectorRegister& Vec)
 {
diff --git a/Engine/Source/Runtime/Core/Public/Math/UnrealMathNeon.h b/Engine/Source/Runtime/Core/Public/Math/UnrealMathNeon.h
--- a/Engine/Source/Runtime/Core/Public/Math/UnrealMathNeon.h
+++ b/Engine/Source/Runtime/Core/Public/Math/UnrealMathNeon.h
@@ -981,6 +981,27 @@ FORCEINLINE void VectorQuaternionMultiply( void* RESTRICT Result, const void* RE
 	*((VectorRegister *)Result) = VectorQuaternionMultiply2(*((const VectorRegister *)Quat1), *((const VectorRegister *)Quat2));
 }
 
+/**
+* Computes the sine and cosine of each component of a Vector.
+*
+* @param VSinAngles	VectorRegister Pointer to where the Sin result should be stored
+* @param VCosAngles	VectorRegister Pointer to where the Cos result should be stored
+* @param VAngles VectorRegister Pointer to the input angles 
+*/
+FORCEINLINE void VectorSinCos(  VectorRegister* VSinAngles, VectorRegister* VCosAngles, const VectorRegister* VAngles )
+{	
+	union { VectorRegister v; float f[4]; } VecSin, VecCos, VecAngles;
+	VecAngles.v = *VAngles;
+
+	FMath::SinCos(&VecSin.f[0], &VecCos.f[0], VecAngles.f[0]);
+	FMath::SinCos(&VecSin.f[1], &VecCos.f[1], VecAngles.f[1]);
+	FMath::SinCos(&VecSin.f[2], &VecCos.f[2], VecAngles.f[2]);
+	FMath::SinCos(&VecSin.f[3], &VecCos.f[3], VecAngles.f[3]);
+
+	*VSinAngles = VecSin.v;
+	*VCosAngles = VecCos.v;
+}
+
 // Returns true if the vector contains a component that is either NAN or +/-infinite.
 inline bool VectorContainsNaNOrInfinite(const VectorRegister& Vec)
 {
diff --git a/Engine/Source/Runtime/Core/Public/Math/UnrealMathSSE.h b/Engine/Source/Runtime/Core/Public/Math/UnrealMathSSE.h
--- a/Engine/Source/Runtime/Core/Public/Math/UnrealMathSSE.h
+++ b/Engine/Source/Runtime/Core/Public/Math/UnrealMathSSE.h
@@ -964,6 +964,28 @@ FORCEINLINE void VectorQuaternionMultiply( void* RESTRICT Result, const void* RE
 	*((VectorRegister *)Result) = VectorQuaternionMultiply2(*((const VectorRegister *)Quat1), *((const VectorRegister *)Quat2));
 }
 
+
+/**
+* Computes the sine and cosine of each component of a Vector.
+*
+* @param VSinAngles	VectorRegister Pointer to where the Sin result should be stored
+* @param VCosAngles	VectorRegister Pointer to where the Cos result should be stored
+* @param VAngles VectorRegister Pointer to the input angles 
+*/
+FORCEINLINE void VectorSinCos(  VectorRegister* VSinAngles, VectorRegister* VCosAngles, const VectorRegister* VAngles )
+{	
+	union { VectorRegister v; float f[4]; } VecSin, VecCos, VecAngles;
+	VecAngles.v = *VAngles;
+
+	FMath::SinCos(&VecSin.f[0], &VecCos.f[0], VecAngles.f[0]);
+	FMath::SinCos(&VecSin.f[1], &VecCos.f[1], VecAngles.f[1]);
+	FMath::SinCos(&VecSin.f[2], &VecCos.f[2], VecAngles.f[2]);
+	FMath::SinCos(&VecSin.f[3], &VecCos.f[3], VecAngles.f[3]);
+
+	*VSinAngles = VecSin.v;
+	*VCosAngles = VecCos.v;
+}
+
 // Returns true if the vector contains a component that is either NAN or +/-infinite.
 inline bool VectorContainsNaNOrInfinite(const VectorRegister& Vec)
 {
diff --git a/Engine/Source/Runtime/Core/Public/Math/UnrealMathUtility.h b/Engine/Source/Runtime/Core/Public/Math/UnrealMathUtility.h
--- a/Engine/Source/Runtime/Core/Public/Math/UnrealMathUtility.h
+++ b/Engine/Source/Runtime/Core/Public/Math/UnrealMathUtility.h
@@ -296,6 +296,80 @@ struct FMath : public FPlatformMath
 		return Loge(Value) * LogToLog2;
 	}
 
+	/**
+	* Computes the sine and cosine of a scalar float.
+	*
+	* @param ScalarSin	Pointer to where the Sin result should be stored
+	* @param ScalarCos	Pointer to where the Cos result should be stored
+	* @param Value  input angles 
+	*/
+	static FORCEINLINE void SinCos( float* ScalarSin, float* ScalarCos, float  Value )
+	{
+		// Map Value to y in [-pi,pi], x = 2*pi*quotient + remainder.
+		float quotient = (INV_PI*0.5f)*Value;
+		if (Value >= 0.0f)
+		{
+			quotient = (float)((int)(quotient + 0.5f));
+		}
+		else
+		{
+			quotient = (float)((int)(quotient - 0.5f));
+		}
+		float y = Value - (2.0f*PI)*quotient;
+
+		// Map y to [-pi/2,pi/2] with sin(y) = sin(Value).
+		float sign;
+		if (y > HALF_PI)
+		{
+			y = PI - y;
+			sign = -1.0f;
+		}
+		else if (y < -HALF_PI)
+		{
+			y = -PI - y;
+			sign = -1.0f;
+		}
+		else
+		{
+			sign = +1.0f;
+		}
+
+		float y2 = y * y;
+
+		// 11-degree minimax approximation
+		*ScalarSin = ( ( ( ( (-2.3889859e-08f * y2 + 2.7525562e-06f) * y2 - 0.00019840874f ) * y2 + 0.0083333310f ) * y2 - 0.16666667f ) * y2 + 1.0f ) * y;
+
+		// 10-degree minimax approximation
+		float p = ( ( ( ( -2.6051615e-07f * y2 + 2.4760495e-05f ) * y2 - 0.0013888378f ) * y2 + 0.041666638f ) * y2 - 0.5f ) * y2 + 1.0f;
+		*ScalarCos = sign*p;
+	}
+
+	/**
+	* Computes the ASin of a scalar float.
+	*
+	* @param Value  input angle 
+	* @return ASin of Value
+	*/
+	static FORCEINLINE float FastAsin( float Value )
+	{
+		// Clamp input to [-1,1].
+		bool nonnegative = (Value >= 0.0f);
+		float x = FMath::Abs(Value);
+		float omx = 1.0f - x;
+		if (omx < 0.0f)
+		{
+			omx = 0.0f;
+		}
+		float root = FMath::Sqrt(omx);
+
+		// 7-degree minimax approximation
+		float result = ( ( ( ( ( ( -0.0012624911f * x + 0.0066700901f ) * x - 0.0170881256f ) * x + 0.0308918810f ) * x - 0.0501743046f ) * x + 0.0889789874f ) * x - 0.2145988016f ) * x + 1.5707963050f;
+		result *= root;  // acos(|x|)
+
+		// acos(x) = pi - acos(-x) when x < 0, asin(x) = pi/2 - acos(x)
+		return (nonnegative ? HALF_PI - result : result - HALF_PI);
+	}
+
 
 	// Conversion Functions
 
diff --git a/Engine/Source/Runtime/Core/Public/Math/UnrealMathVectorCommon.h b/Engine/Source/Runtime/Core/Public/Math/UnrealMathVectorCommon.h
--- a/Engine/Source/Runtime/Core/Public/Math/UnrealMathVectorCommon.h
+++ b/Engine/Source/Runtime/Core/Public/Math/UnrealMathVectorCommon.h
@@ -109,6 +109,24 @@ FORCEINLINE VectorRegister VectorNormalizeQuaternion(const VectorRegister& Unnor
 	return VectorNormalizeSafe(UnnormalizedQuat, GlobalVectorConstants::Float0001);
 }
 
+// Normalize Rotator
+FORCEINLINE VectorRegister VectorNormalizeRotator(const VectorRegister& UnnormalizedRotator)
+{
+	static const VectorRegister VEC_360 = MakeVectorRegister(360.f, 360.f, 360.f, 360.f);
+	static const VectorRegister VEC_180 = MakeVectorRegister(180.f, 180.f, 180.f, 180.f);
+
+	// shift in the range [-360,360]
+	VectorRegister V0	= VectorMod( UnnormalizedRotator, VEC_360);
+	VectorRegister V1	= VectorAdd( V0, VEC_360 );
+	VectorRegister V2	= VectorSelect(VectorCompareGE(V0, VectorZero()), V0, V1);
+
+	// shift to [-180,180]
+	VectorRegister V3	= VectorSubtract( V2, VEC_360 );
+	VectorRegister V4	= VectorSelect(VectorCompareGT(V2, VEC_180), V3, V2);
+
+	return  V4;
+}
+
 /** 
  * Fast Linear Quaternion Interpolation for quaternions stored in VectorRegisters.
  * Result is NOT normalized.
diff --git a/Engine/Source/Runtime/Core/Public/Math/UnrealMathVectorConstants.h b/Engine/Source/Runtime/Core/Public/Math/UnrealMathVectorConstants.h
--- a/Engine/Source/Runtime/Core/Public/Math/UnrealMathVectorConstants.h
+++ b/Engine/Source/Runtime/Core/Public/Math/UnrealMathVectorConstants.h
@@ -25,6 +25,9 @@ namespace GlobalVectorConstants
 	static const VectorRegister QMULTI_SIGN_MASK1 = MakeVectorRegister( 1.f, 1.f, -1.f, -1.f );
 	static const VectorRegister QMULTI_SIGN_MASK2 = MakeVectorRegister( -1.f, 1.f, 1.f, -1.f );
 
+	static const VectorRegister DEG_TO_RAD = MakeVectorRegister(PI/(180.f), PI/(180.f), PI/(180.f), 0.0f);
+	static const VectorRegister DEG_TO_RAD_HALF = MakeVectorRegister((PI/180.f)*0.5f, (PI/180.f)*0.5f, (PI/180.f)*0.5f, 0.0f);
+	static const VectorRegister RAD_TO_DEG = MakeVectorRegister((180.f)/PI, (180.f)/PI, (180.f)/PI, 0.0f);
 
 	/** Bitmask to AND out the XYZ components in a vector */
 	static const VectorRegister XYZMask = MakeVectorRegister((uint32)0xffffffff, (uint32)0xffffffff, (uint32)0xffffffff, (uint32)0x00000000);
diff --git a/Engine/Source/Runtime/Core/Public/Math/Vector.h b/Engine/Source/Runtime/Core/Public/Math/Vector.h
--- a/Engine/Source/Runtime/Core/Public/Math/Vector.h
+++ b/Engine/Source/Runtime/Core/Public/Math/Vector.h
@@ -30,6 +30,9 @@ public:
 	/** Unreal forward vector (1,0,0) */
 	static CORE_API const FVector ForwardVector;
 
+	/** Unreal right vector (0,1,0) */
+	static CORE_API const FVector RightVector;
+
 public:
 
 #if ENABLE_NAN_DIAGNOSTIC
diff --git a/Engine/Source/Runtime/Engine/Private/Actor.cpp b/Engine/Source/Runtime/Engine/Private/Actor.cpp
--- a/Engine/Source/Runtime/Engine/Private/Actor.cpp
+++ b/Engine/Source/Runtime/Engine/Private/Actor.cpp
@@ -2750,17 +2750,17 @@ FRotator AActor::K2_GetActorRotation() const
 
 FVector AActor::GetActorForwardVector() const
 {
-	return GetTransform().GetUnitAxis(EAxis::X);
+	return GetActorQuat().RotateVector(FVector::ForwardVector);
 }
 
 FVector AActor::GetActorUpVector() const
 {
-	return GetTransform().GetUnitAxis(EAxis::Z);
+	return GetActorQuat().RotateVector(FVector::UpVector);
 }
 
 FVector AActor::GetActorRightVector() const
 {
-	return GetTransform().GetUnitAxis(EAxis::Y);
+	return GetActorQuat().RotateVector(FVector::RightVector);
 }
 
 USceneComponent* AActor::K2_GetRootComponent() const
