Stats - Culled hierarchical stats

diff --git a/Engine/Source/Runtime/Core/Private/Stats/StatsCommand.cpp b/Engine/Source/Runtime/Core/Private/Stats/StatsCommand.cpp
--- a/Engine/Source/Runtime/Core/Private/Stats/StatsCommand.cpp
+++ b/Engine/Source/Runtime/Core/Private/Stats/StatsCommand.cpp
@@ -48,8 +48,9 @@ struct FStatGroupParams
 	/** Default constructor. */
 	FStatGroupParams( const TCHAR* Cmd = nullptr )
 		: Group( Cmd, TEXT("group="), NAME_None )
-		, SortBy( Cmd, TEXT("sortby="), EStatCompareBy::Name )
+		, SortBy( Cmd, TEXT("sortby="), EStatCompareBy::Sum )
 		, MaxHistoryFrames( Cmd, TEXT("maxhistoryframes="), 60 )
+		, MaxHierarchyDepth( Cmd, TEXT("maxdepth="), 16 )
 		, bReset( FCString::Stristr( Cmd, TEXT("-reset") ) != nullptr )
 	{}
 
@@ -74,6 +75,12 @@ struct FStatGroupParams
 	// @TODO yrx 2014-08-21 Replace with TParsedValueWithDefaultAndRange
 	TParsedValueWithDefault<int32> MaxHistoryFrames;
 
+	/**
+	 *	Maximum depth for the hierarchy
+	 * -maxdepth=16
+	 */
+	TParsedValueWithDefault<int32> MaxHierarchyDepth;
+
 	/** Whether to reset all collected data. */
 	bool bReset;
 };
@@ -652,7 +659,8 @@ struct FHUDGroupManager
 		
 		// Generate root stats stack for current frame.
 		Stats.UncondenseStackStats( TargetFrame, NewFrame.HierarchyInclusive, &Filter, &NewFrame.NonStackStats );
-		NewFrame.HierarchyInclusive.AddNameHierarchy();
+		NewFrame.HierarchyInclusive.Cull( TNumericLimits<int64>::Max(), Params.MaxHierarchyDepth.Get() );
+		//NewFrame.HierarchyInclusive.AddNameHierarchy();
 		NewFrame.HierarchyInclusive.AddSelf();
 
 		{
@@ -907,6 +915,9 @@ struct FHUDGroupManager
 				}
 			}
 		}
+
+		out_EnabledItems.Add(NAME_Self);
+		out_EnabledItems.Add(NAME_OtherChildren);
 	}
 
 	void GetStatsForGroup( TSet<FName>& out_EnabledItems, const FName GroupName )
@@ -923,7 +934,10 @@ struct FHUDGroupManager
 			{
 				out_EnabledItems.Add(LongName->NameAndInfo.GetRawName()); // long name
 			}
-		}	
+		}
+
+		out_EnabledItems.Add(NAME_Self);
+		out_EnabledItems.Add(NAME_OtherChildren);
 	}
 
 	static FHUDGroupManager& Get(FStatsThreadState const& Stats)
