Fix for 'stat' commands so that they don't incorrectly report as unrecognized.

diff --git a/Engine/Source/Runtime/Core/Private/Stats/StatsCommand.cpp b/Engine/Source/Runtime/Core/Private/Stats/StatsCommand.cpp
--- a/Engine/Source/Runtime/Core/Private/Stats/StatsCommand.cpp
+++ b/Engine/Source/Runtime/Core/Private/Stats/StatsCommand.cpp
@@ -1426,12 +1426,12 @@ bool DirectStatsCommand(const TCHAR* Cmd, bool bBlockForCompletion /*= false*/,
 	{
 		FString AddArgs;
 		const TCHAR* TempCmd = Cmd;
-		bResult = true;
 
 		FString ArgNoWhitespaces = FDefaultValueHelper::RemoveWhitespaces(TempCmd);
 		const bool bIsEmpty = ArgNoWhitespaces.IsEmpty();
 #if STATS
-		if( bIsEmpty && Ar )
+		bResult = true;
+		if (bIsEmpty && Ar)
 		{
 			PrintStatsHelpToOutputDevice( *Ar );
 		}
@@ -1500,10 +1500,22 @@ bool DirectStatsCommand(const TCHAR* Cmd, bool bBlockForCompletion /*= false*/,
 		{
 		}
 		else
-#endif
 		{
-			bResult = false;
+			FString Param;
+			if (FParse::Token(TempCmd, Param, false))
+			{
+				const FName MaybeGroupFName = FName(*(FString(TEXT("STATGROUP_")) + Param));
+				FStatsThreadState& Stats = FStatsThreadState::GetLocalState();
+				bResult = Stats.Groups.Contains(MaybeGroupFName);
+			}
+			else
+			{
+				bResult = false;
+			}
 		}
+#else
+		bResult = false;
+#endif
 
 		check(IsInGameThread());
 		if( !bIsEmpty )
