Back out revision 58 from //depot/UE4/Engine/Source/Runtime/Core/Private/Stats/StatsCommand.cpp

diff --git a/Engine/Source/Runtime/Core/Private/Stats/StatsCommand.cpp b/Engine/Source/Runtime/Core/Private/Stats/StatsCommand.cpp
--- a/Engine/Source/Runtime/Core/Private/Stats/StatsCommand.cpp
+++ b/Engine/Source/Runtime/Core/Private/Stats/StatsCommand.cpp
@@ -1426,12 +1426,12 @@ bool DirectStatsCommand(const TCHAR* Cmd, bool bBlockForCompletion /*= false*/,
 	{
 		FString AddArgs;
 		const TCHAR* TempCmd = Cmd;
+		bResult = true;
 
 		FString ArgNoWhitespaces = FDefaultValueHelper::RemoveWhitespaces(TempCmd);
 		const bool bIsEmpty = ArgNoWhitespaces.IsEmpty();
 #if STATS
-		bResult = true;
-		if (bIsEmpty && Ar)
+		if( bIsEmpty && Ar )
 		{
 			PrintStatsHelpToOutputDevice( *Ar );
 		}
@@ -1500,22 +1500,10 @@ bool DirectStatsCommand(const TCHAR* Cmd, bool bBlockForCompletion /*= false*/,
 		{
 		}
 		else
+#endif
 		{
-			FString Param;
-			if (FParse::Token(TempCmd, Param, false))
-			{
-				const FName MaybeGroupFName = FName(*(FString(TEXT("STATGROUP_")) + Param));
-				FStatsThreadState& Stats = FStatsThreadState::GetLocalState();
-				bResult = Stats.Groups.Contains(MaybeGroupFName);
-			}
-			else
-			{
-				bResult = false;
-			}
+			bResult = false;
 		}
-#else
-		bResult = false;
-#endif
 
 		check(IsInGameThread());
 		if( !bIsEmpty )
