Stats - Added comment regarding accessing stats data

diff --git a/Engine/Source/Runtime/Core/Public/Stats/StatsData.h b/Engine/Source/Runtime/Core/Public/Stats/StatsData.h
--- a/Engine/Source/Runtime/Core/Public/Stats/StatsData.h
+++ b/Engine/Source/Runtime/Core/Public/Stats/StatsData.h
@@ -616,7 +616,7 @@ public:
 	/** Adds a frame worth of messages */
 	void AddMessages(TArray<FStatMessage>& InMessages);
 
-	/** Singleton to get the stats being collected by this executable **/
+	/** Singleton to get the stats being collected by this executable. Can be only accessed from the stats thread. **/
 	static FStatsThreadState& GetLocalState();
 };
 
