Async: Removed critical section from TFutureState (not needed, because FEvent already locks, if necessary)

diff --git a/Engine/Source/Runtime/Core/Public/Async/Future.h b/Engine/Source/Runtime/Core/Public/Async/Future.h
--- a/Engine/Source/Runtime/Core/Public/Async/Future.h
+++ b/Engine/Source/Runtime/Core/Public/Async/Future.h
@@ -47,8 +47,6 @@ public:
 	 */
 	bool WaitFor(const FTimespan& Duration) const
 	{
-		FScopeLock Lock(&CriticalSection);
-
 		if (CompletionEvent->Wait(Duration))
 		{
 			return true;
@@ -71,9 +69,6 @@ private:
 	/** Holds an event signaling that the result is available. */
 	FEvent* CompletionEvent;
 
-	/** Critical section for locking access to this class. */
-	mutable FCriticalSection CriticalSection;
-
 	/** Whether the asynchronous result is available. */
 	bool Complete;
 };
