Async: Using fire-and-forget for async graph tasks

diff --git a/Engine/Source/Runtime/Core/Public/Async/Async.h b/Engine/Source/Runtime/Core/Public/Async/Async.h
--- a/Engine/Source/Runtime/Core/Public/Async/Async.h
+++ b/Engine/Source/Runtime/Core/Public/Async/Async.h
@@ -91,7 +91,7 @@ public:
 	 */
 	static ESubsequentsMode::Type GetSubsequentsMode()
 	{
-		return ESubsequentsMode::TrackSubsequents;
+		return ESubsequentsMode::FireAndForget;
 	}
 
 private:
