Core: code & documentation cleanup pass

diff --git a/Engine/Source/Runtime/Core/Private/WinRT/WinRTEvent.h b/Engine/Source/Runtime/Core/Private/WinRT/WinRTEvent.h
--- a/Engine/Source/Runtime/Core/Private/WinRT/WinRTEvent.h
+++ b/Engine/Source/Runtime/Core/Private/WinRT/WinRTEvent.h
@@ -1,31 +1,24 @@
 // Copyright 1998-2015 Epic Games, Inc. All Rights Reserved.
 
-/*=============================================================================
-	WinRTEvent.h: Declares the FEventWin class.
-=============================================================================*/
-
 #pragma once
 
 #include "AllowWinRTPlatformTypes.h"
 
+
 /**
  * Implements the WinRT version of the FEvent interface.
  */
-class FEventWinRT : public FEvent
+class FEventWinRT
+	: public FEvent
 {
 public:
 
-	/**
-	 * Default constructor.
-	 */
+	/** Default constructor. */
 	FEventWinRT()
 		: Event(NULL)
-	{
-	}
+	{ }
 
-	/**
-	 * Destructor.
-	 */
+	/** Virtual destructor. */
 	virtual ~FEventWinRT()
 	{
 		if (Event != NULL)
@@ -34,9 +27,10 @@ public:
 		}
 	}
 
-
 public:
 
+	// FEvent interface
+
 	virtual bool Create (bool bIsManualReset = false) override
 	{
 		// Create the event and default it to non-signaled
@@ -62,11 +56,11 @@ public:
 
 	virtual bool Wait(uint32 WaitTime, const bool bIgnoreThreadIdleStats = false) override;
 
-
 private:
 
-	// Holds the handle to the event
+	/** Holds the handle to the event. */
 	HANDLE Event;
 };
 
+
 #include "HideWinRTPlatformTypes.h"
diff --git a/Engine/Source/Runtime/Core/Private/Windows/WindowsEvent.h b/Engine/Source/Runtime/Core/Private/Windows/WindowsEvent.h
--- a/Engine/Source/Runtime/Core/Private/Windows/WindowsEvent.h
+++ b/Engine/Source/Runtime/Core/Private/Windows/WindowsEvent.h
@@ -13,18 +13,13 @@ class FEventWin
 {
 public:
 
-	/**
-	 * Default constructor.
-	 */
-	FEventWin( )
+	/** Default constructor. */
+	FEventWin()
 		: Event(nullptr)
-	{
-	}
+	{ }
 
-	/**
-	 * Virtual destructor.
-	 */
-	virtual ~FEventWin( )
+	/** Virtual destructor. */
+	virtual ~FEventWin()
 	{
 		if (Event != nullptr)
 		{
@@ -34,6 +29,8 @@ public:
 
 public:
 
+	// FEvent interface
+
 	virtual bool Create( bool bIsManualReset = false ) override
 	{
 		// Create the event and default it to non-signaled
@@ -42,14 +39,14 @@ public:
 		return Event != nullptr;
 	}
 
-	virtual void Trigger( ) override
+	virtual void Trigger() override
 	{
 		check(Event);
 
 		SetEvent(Event);
 	}
 
-	virtual void Reset( ) override
+	virtual void Reset() override
 	{
 		check(Event);
 
@@ -60,7 +57,7 @@ public:
 
 private:
 
-	// Holds the handle to the event
+	/** Holds the handle to the event. */
 	HANDLE Event;
 };
 
diff --git a/Engine/Source/Runtime/Core/Public/HAL/PThreadEvent.h b/Engine/Source/Runtime/Core/Public/HAL/PThreadEvent.h
--- a/Engine/Source/Runtime/Core/Public/HAL/PThreadEvent.h
+++ b/Engine/Source/Runtime/Core/Public/HAL/PThreadEvent.h
@@ -93,6 +93,10 @@ public:
 		}
 	}
 
+public:
+
+	// FEvent interface
+
 	virtual bool Create( bool _bIsManualReset = false ) override
 	{
 		check(!bInitialized);
