Temporary workaround for crash-on-exit on Mac.

diff --git a/Engine/Source/Runtime/Core/Private/HAL/PThreadRunnableThread.cpp b/Engine/Source/Runtime/Core/Private/HAL/PThreadRunnableThread.cpp
--- a/Engine/Source/Runtime/Core/Private/HAL/PThreadRunnableThread.cpp
+++ b/Engine/Source/Runtime/Core/Private/HAL/PThreadRunnableThread.cpp
@@ -33,7 +33,7 @@ uint32 FRunnableThreadPThread::Run()
 	ThreadIsRunning = false;
 
 #if STATS
-	FThreadStats::Shutdown();
+//	FThreadStats::Shutdown();
 #endif
 
 	return ExitCode;
