Core: documentation cleanup pass

diff --git a/Engine/Source/Runtime/Core/Public/GenericPlatform/GenericPlatformProcess.h b/Engine/Source/Runtime/Core/Public/GenericPlatform/GenericPlatformProcess.h
--- a/Engine/Source/Runtime/Core/Public/GenericPlatform/GenericPlatformProcess.h
+++ b/Engine/Source/Runtime/Core/Public/GenericPlatform/GenericPlatformProcess.h
@@ -1,17 +1,15 @@
 // Copyright 1998-2015 Epic Games, Inc. All Rights Reserved.
 
-
-/*=============================================================================================
-	GenericPlatformProcess.h: Generic platform Process classes, mostly implemented with ANSI C++
-==============================================================================================*/
-
 #pragma once
+
 #include "Containers/ContainersFwd.h"
 #include "GenericPlatform/GenericPlatformMisc.h"
 #include "HAL/Platform.h"
 
+
 class FString;
 
+
 namespace EProcessResource
 {
 	enum Type
@@ -24,6 +22,7 @@ namespace EProcessResource
 	};
 }
 
+
 /** Not all platforms have different opening semantics, but Windows allows you to specify a specific verb when opening a file. */
 namespace ELaunchVerb
 {
@@ -31,11 +30,13 @@ namespace ELaunchVerb
 	{
 		/** Launch the application associated with opening file to 'view' */
 		Open,
+
 		/** Launch the application associated with opening file to 'edit' */
 		Edit,
 	};
 }
 
+
 /** Generic implementation for the process handle. */
 #if PLATFORM_WINDOWS && defined(__clang__)	// @todo clang: Clang on Windows has a bug with non-type template arguments
 static void* InvalidHandleValue = nullptr;
@@ -47,15 +48,16 @@ struct TProcHandle
 {
 	typedef T HandleType;
 public:
+
 	/** Default constructor. */
 	FORCEINLINE TProcHandle()
 		: Handle( InvalidHandleValue )
-	{}
+	{ }
 
 	/** Initialization constructor. */
 	FORCEINLINE explicit TProcHandle( T Other )
 		: Handle( Other )
-	{}
+	{ }
 
 	/** Assignment operator. */
 	FORCEINLINE TProcHandle& operator=( const TProcHandle& Other )
@@ -97,6 +99,7 @@ protected:
 	T Handle;
 };
 
+
 struct FProcHandle;
 
 /**
@@ -110,10 +113,13 @@ struct CORE_API FGenericPlatformProcess
 	struct FSemaphore
 	{
 		/** Returns the name of the object */
-		const TCHAR *	GetName() const			{ return Name; }
+		const TCHAR* GetName() const
+		{
+			return Name;
+		}
 
 		/** Acquires an exclusive access (also known as Wait()) */
-		virtual void	Lock() = 0;
+		virtual void Lock() = 0;
 
 		/** 
 		 * Tries to acquire and exclusive access for a specified amount of nanoseconds (also known as TryWait()).
@@ -121,20 +127,20 @@ struct CORE_API FGenericPlatformProcess
 		 * @param Nanoseconds (10^-9 seconds) to wait for, 
 		 * @return false if was not able to lock within given time
 		 */
-		virtual bool	TryLock(uint64 NanosecondsToWait) = 0;
+		virtual bool TryLock(uint64 NanosecondsToWait) = 0;
 
 		/** Relinquishes an exclusive access (also known as Release())*/
-		virtual void	Unlock() = 0;
+		virtual void Unlock() = 0;
 
 		/** 
-		 * Constructor
+		 * Creates and initializes a new instance with the specified name.
 		 *
 		 * @param InName name of the semaphore (all processes should use the same)
 		 */
 		FSemaphore(const FString& InName);
 
-		/** Destructor */
-		virtual ~FSemaphore()  {};
+		/** Virtual destructor. */
+		virtual ~FSemaphore() { };
 
 	protected:
 
@@ -184,20 +190,23 @@ struct CORE_API FGenericPlatformProcess
 	}
 
 	/**
-	 * Retrieves the ProcessId of this process
-	 * @return the ProcessId of this process
+	 * Retrieves the ProcessId of this process.
+	 *
+	 * @return the ProcessId of this process.
 	 */
 	static uint32 GetCurrentProcessId();
 	
 	/**	 
 	 * Change the thread processor affinity
 	 *
-	 * @param AffinityMask A bitfield indicating what processors the thread is allowed to run on
+	 * @param AffinityMask A bitfield indicating what processors the thread is allowed to run on.
 	 */
 	static void SetThreadAffinityMask( uint64 AffinityMask );
 
 	/**
-	 * Allow the platform to do anything it needs for game or render thread (this would use the NamedThreads enum if we could forward declare the enum)
+	 * Allow the platform to do anything it needs for game or render thread (this would use the NamedThreads enum if we could forward declare the enum).
+	 *
+	 * @param bIsRenderThread true if setting up the render thread, false for the game thread.
 	 */
 	static void SetupGameOrRenderThread(bool bIsRenderThread)
 	{
@@ -205,29 +214,31 @@ struct CORE_API FGenericPlatformProcess
 	
 	/** Get startup directory.  NOTE: Only one return value is valid at a time! **/
 	static const TCHAR* BaseDir();
+
 	/** Get user directory.  NOTE: Only one return value is valid at a time! **/
 	static const TCHAR* UserDir();
+
 	/** Get the user settings directory.  NOTE: Only one return value is valid at a time! **/
 	static const TCHAR *UserSettingsDir();
+
 	/** Get application settings directory.  NOTE: Only one return value is valid at a time! **/
 	static const TCHAR* ApplicationSettingsDir();
+
 	/** Get computer name.  NOTE: Only one return value is valid at a time! **/
 	static const TCHAR* ComputerName();
+
 	/** Get user name.  NOTE: Only one return value is valid at a time! **/
 	static const TCHAR* UserName(bool bOnlyAlphaNumeric = true);
 	static const TCHAR* ShaderDir();
 	static void SetShaderDir(const TCHAR*Where);
-	static void SetCurrentWorkingDirectoryToBaseDir()
-	{
-	}
+	static void SetCurrentWorkingDirectoryToBaseDir() { }
 
 	/**
 	 * Sets the process limits.
 	 *
-	 * @param Resource one of process resources
-	 * @param Limit the maximum amount of the resource (for some OS, this means both hard and soft limits)
-	 *
-	 * @return true if succeeded
+	 * @param Resource one of process resources.
+	 * @param Limit the maximum amount of the resource (for some OS, this means both hard and soft limits).
+	 * @return true on success, false otherwise.
 	 */
 	static bool SetProcessLimits(EProcessResource::Type Resource, uint64 Limit)
 	{
@@ -235,13 +246,13 @@ struct CORE_API FGenericPlatformProcess
 	}
 
 	/**
-	 *	Get the shader working directory
+	 * Get the shader working directory.
+	 *
+	 * @return The path to the directory.
 	 */
 	static const FString ShaderWorkingDir();
 
-	/**
-	 *	Clean the shader working directory
-	 */
+	/**	Clean the shader working directory. */
 	static void CleanShaderWorkingDir();
 
 	/**
@@ -252,7 +263,6 @@ struct CORE_API FGenericPlatformProcess
 	 */
 	static const TCHAR* ExecutableName(bool bRemoveExtension = true);
 
-
 	/**
 	 * Generates the path to the specified application or game.
 	 *
@@ -260,22 +270,25 @@ struct CORE_API FGenericPlatformProcess
 	 * executable's directory.For example, calling this method with "UE4" and EBuildConfigurations::Debug
 	 * on Windows 64-bit will generate the path "../Win64/UE4Editor-Win64-Debug.exe"
 	 *
-	 * @param AppName - The name of the application or game.
-	 * @param BuildConfiguration - The build configuration of the game.
-	 *
+	 * @param AppName The name of the application or game.
+	 * @param BuildConfiguration The build configuration of the game.
 	 * @return The generated application path.
 	 */
 	static FString GenerateApplicationPath( const FString& AppName, EBuildConfigurations::Type BuildConfiguration);
 
 	/**
 	 * Return the prefix of dynamic library (e.g. lib)
+	 *
+	 * @return The prefix string.
+	 * @see GetModuleExtension, GetModulesDirectory
 	 */
 	static const TCHAR* GetModulePrefix();
 
 	/**
 	 * Return the extension of dynamic library
 	 *
-	 * @return Extension of dynamic library
+	 * @return Extension of dynamic library.
+	 * @see GetModulePrefix, GetModulesDirectory
 	 */
 	static const TCHAR* GetModuleExtension();
 
@@ -286,14 +299,17 @@ struct CORE_API FGenericPlatformProcess
 
 	/**
 	 * Used only by platforms with DLLs, this gives the full path to the main directory containing modules
+	 *
+	 * @return The path to the directory.
+	 * @see GetModulePrefix, GetModuleExtension
 	 */
 	static const FString GetModulesDirectory();
 	
 	/**
-	* Launch a uniform resource locator (i.e. http://www.epicgames.com/unreal).
-	* This is expected to return immediately as the URL is launched by another
-	* task.
-	**/
+	 * Launch a uniform resource locator (i.e. http://www.epicgames.com/unreal).
+	 * This is expected to return immediately as the URL is launched by another
+	 * task.
+	 */
 	static void LaunchURL( const TCHAR* URL, const TCHAR* Parms, FString* Error );
 
 	/**
@@ -312,20 +328,23 @@ struct CORE_API FGenericPlatformProcess
 	 */
 	static FProcHandle CreateProc( const TCHAR* URL, const TCHAR* Parms, bool bLaunchDetached, bool bLaunchHidden, bool bLaunchReallyHidden, uint32* OutProcessID, int32 PriorityModifier, const TCHAR* OptionalWorkingDirectory, void* PipeWrite );
 
-	/** Returns true if the specified process is running 
+	/**
+	 * Returns true if the specified process is running 
 	 *
 	 * @param ProcessHandle handle returned from FPlatformProcess::CreateProc
 	 * @return true if the process is still running
 	 */
 	static bool IsProcRunning( FProcHandle & ProcessHandle );
 	
-	/** Waits for a process to stop
+	/**
+	 * Waits for a process to stop
 	 *
 	 * @param ProcessHandle handle returned from FPlatformProcess::CreateProc
 	 */
 	static void WaitForProc( FProcHandle & ProcessHandle );
 
-	/** Cleans up FProcHandle after we're done with it.
+	/**
+	 * Cleans up FProcHandle after we're done with it.
 	 *
 	 * @param ProcessHandle handle returned from FPlatformProcess::CreateProc.
 	 */
@@ -337,16 +356,22 @@ struct CORE_API FGenericPlatformProcess
 	 * @param KillTree Whether the entire process tree should be terminated.
 	 */
 	static void TerminateProc( FProcHandle & ProcessHandle, bool KillTree = false );
+
 	/** Retrieves the termination status of the specified process. **/
 	static bool GetProcReturnCode( FProcHandle & ProcHandle, int32* ReturnCode );
+
 	/** Returns true if the specified application is running */
 	static bool IsApplicationRunning( uint32 ProcessId );
+
 	/** Returns true if the specified application is running */
 	static bool IsApplicationRunning( const TCHAR* ProcName );
+
 	/** Returns the Name of process given by the PID.  Returns Empty string "" if PID not found. */
 	static FString GetApplicationName( uint32 ProcessId );
+
 	/** Outputs the virtual memory usage, of the process with the specified PID */
 	static bool GetApplicationMemoryUsage(uint32 ProcessId, SIZE_T* OutMemoryUsage);
+
 	/** Returns true if the specified application has a visible window, and that window is active/has focus/is selected */
 	static bool IsThisApplicationForeground();	
 
@@ -360,9 +385,9 @@ struct CORE_API FGenericPlatformProcess
 	static bool ExecProcess( const TCHAR* URL, const TCHAR* Params, int32* OutReturnCode, FString* OutStdOut, FString* OutStdErr );
 
 	/**
-	* Executes a process as administrator, requesting elevation as necessary. This
-	* call blocks until the process has returned.
-	*/
+	 * Executes a process as administrator, requesting elevation as necessary. This
+	 * call blocks until the process has returned.
+	 */
 	static bool ExecElevatedProcess(const TCHAR* URL, const TCHAR* Params, int32* OutReturnCode);
 
 	/**
@@ -393,18 +418,17 @@ struct CORE_API FGenericPlatformProcess
 #endif // PLATFORM_HAS_BSD_TIME
 
 	/**
-	 * Creates a new event
+	 * Creates a new event,
 	 *
-	 * @param bIsManualReset Whether the event requires manual reseting or not
-	 * @param InName Whether to use a commonly shared event or not. If so this
-	 * is the name of the event to share.
-	 *
-	 * @return Returns the new event object if successful, NULL otherwise
+	 * @param bIsManualReset Whether the event requires manual reseting or not.
+	 * @param InName Whether to use a commonly shared event or not. If so this is the name of the event to share.
+	 * @return Returns the new event object if successful, NULL otherwise.
 	 */
 	static class FEvent* CreateSynchEvent(bool bIsManualReset = 0);
 
 	/**
-	 * Creates the platform-specific runnable thread. This should only be called from FRunnableThread::Create
+	 * Creates the platform-specific runnable thread. This should only be called from FRunnableThread::Create.
+	 *
 	 * @return The newly created thread
 	 */
 	static class FRunnableThread* CreateRunnableThread();
@@ -412,11 +436,9 @@ struct CORE_API FGenericPlatformProcess
 	/**
 	 * Closes an anonymous pipe.
 	 *
-	 * @param ReadPipe - The handle to the read end of the pipe.
-	 * @param WritePipe - The handle to the write end of the pipe.
-	 *
-	 * @see CreatePipe
-	 * @see ReadPipe
+	 * @param ReadPipe The handle to the read end of the pipe.
+	 * @param WritePipe The handle to the write end of the pipe.
+	 * @see CreatePipe, ReadPipe
 	 */
 	static void ClosePipe( void* ReadPipe, void* WritePipe );
 
@@ -426,38 +448,29 @@ struct CORE_API FGenericPlatformProcess
 	 * Anonymous pipes can be used to capture and/or redirect STDOUT and STDERROR of a process.
 	 * The pipe created by this method can be passed into CreateProc as Write
 	 *
-	 * @param ReadPipe - Will hold the handle to the read end of the pipe.
-	 * @param WritePipe - Will hold the handle to the write end of the pipe.
-	 *
+	 * @param ReadPipe Will hold the handle to the read end of the pipe.
+	 * @param WritePipe Will hold the handle to the write end of the pipe.
 	 * @return true on success, false otherwise.
-	 *
-	 * @see ClosePipe
-	 * @see ReadPipe
+	 * @see ClosePipe, ReadPipe
 	 */
 	static bool CreatePipe( void*& ReadPipe, void*& WritePipe );
 
 	/**
 	 * Reads all pending data from an anonymous pipe, such as STDOUT or STDERROR of a process.
 	 *
-	 * @param Pipe - The handle to the pipe to read from.
-	 *
+	 * @param Pipe The handle to the pipe to read from.
 	 * @return A string containing the read data.
-	 *
-	 * @see ClosePipe
-	 * @see CreatePipe
+	 * @see ClosePipe, CreatePipe
 	 */
 	static FString ReadPipe( void* ReadPipe );
 
 	/**
 	 * Reads all pending data from an anonymous pipe, such as STDOUT or STDERROR of a process.
 	 *
-	 * @param Pipe - The handle to the pipe to read from.
-	 * @param Output - The data read.
-	 *
+	 * @param Pipe The handle to the pipe to read from.
+	 * @param Output The data read.
 	 * @return true if successful (i.e. any data was read)
-	 *
-	 * @see ClosePipe
-	 * @see CreatePipe
+	 * @see ClosePipe, CreatePipe
 	 */
 	static bool ReadPipeToArray(void* ReadPipe, TArray<uint8> & Output);
 
@@ -468,34 +481,29 @@ struct CORE_API FGenericPlatformProcess
 	 */
 	static bool SupportsMultithreading();
 	
-	/**
-	 * Enables Real Time Mode on the current thread
-	 *
-	 */
-	static void SetRealTimeMode()
-	{
-	}
+	/** Enables Real Time Mode on the current thread. */
+	static void SetRealTimeMode() { }
 
 	/**
 	 * Creates or opens an interprocess synchronization object.
 	 *
-	 * @param Name name (so we can use it across processes)
-	 * @param bCreate - if true, the function will try to create, otherwise will try to open existing
-	 * @param MaxLocks - maximum amount of locks that the semaphore can have (pass 1 to make it act as mutex)
+	 * @param Name name (so we can use it across processes).
+	 * @param bCreate If true, the function will try to create, otherwise will try to open existing.
+	 * @param MaxLocks Maximum amount of locks that the semaphore can have (pass 1 to make it act as mutex).
 	 */
 	static FSemaphore* NewInterprocessSynchObject(const FString& Name, bool bCreate, uint32 MaxLocks = 1);
 
 	/**
 	 * Deletes an interprocess synchronization object.
 	 *
-	 * @param Object object to destroy
+	 * @param Object object to destroy.
 	 */
 	static bool DeleteInterprocessSynchObject(FSemaphore * Object);
 
 	/**
 	 * Makes process run as a system service (daemon), i.e. detaches it from whatever user session it was initially run from.
 	 *
-	 * @return true if successful
+	 * @return true if successful, false otherwise.
 	 */
 	static bool Daemonize();
 };
