Async: Reverted local change that accidentally got checked in

diff --git a/Engine/Source/Runtime/Core/Private/Tests/Async/AsyncTest.cpp b/Engine/Source/Runtime/Core/Private/Tests/Async/AsyncTest.cpp
--- a/Engine/Source/Runtime/Core/Private/Tests/Async/AsyncTest.cpp
+++ b/Engine/Source/Runtime/Core/Private/Tests/Async/AsyncTest.cpp
@@ -7,7 +7,7 @@
 #include "Future.h"
 
 
-IMPLEMENT_SIMPLE_AUTOMATION_TEST(FAsyncTest, "Core.Async.Async", EAutomationTestFlags::ATF_SmokeTest)// ATF_Editor)
+IMPLEMENT_SIMPLE_AUTOMATION_TEST(FAsyncTest, "Core.Async.Async", EAutomationTestFlags::ATF_Editor)
 
 
 bool FAsyncTest::RunTest(const FString& Parameters)
