Added a transition API to some AI subsystems to not break 4.6 users' projects #UE4

diff --git a/Engine/Source/Runtime/AIModule/Classes/BehaviorTree/BTAuxiliaryNode.h b/Engine/Source/Runtime/AIModule/Classes/BehaviorTree/BTAuxiliaryNode.h
--- a/Engine/Source/Runtime/AIModule/Classes/BehaviorTree/BTAuxiliaryNode.h
+++ b/Engine/Source/Runtime/AIModule/Classes/BehaviorTree/BTAuxiliaryNode.h
@@ -72,4 +72,21 @@ protected:
 
 	/** sets next tick time */
 	void SetNextTickTime(uint8* NodeMemory, float RemainingTime) const;
+
+	//----------------------------------------------------------------------//
+	// DEPRECATED
+	//----------------------------------------------------------------------//
+public:
+	DEPRECATED(4.7, "This version is deprecated. Please use the one taking reference to UBehaviorTreeComponent rather than a pointer.")
+	void WrappedOnBecomeRelevant(UBehaviorTreeComponent* OwnerComp, uint8* NodeMemory) const;
+	DEPRECATED(4.7, "This version is deprecated. Please use the one taking reference to UBehaviorTreeComponent rather than a pointer.")
+	void WrappedOnCeaseRelevant(UBehaviorTreeComponent* OwnerComp, uint8* NodeMemory) const;
+	DEPRECATED(4.7, "This version is deprecated. Please use the one taking reference to UBehaviorTreeComponent rather than a pointer.")
+	void WrappedTickNode(UBehaviorTreeComponent* OwnerComp, uint8* NodeMemory, float DeltaSeconds) const;
+	DEPRECATED(4.7, "This version is deprecated. Please use the one taking reference to UBehaviorTreeComponent rather than a pointer.")
+	virtual void OnBecomeRelevant(UBehaviorTreeComponent* OwnerComp, uint8* NodeMemory);
+	DEPRECATED(4.7, "This version is deprecated. Please use the one taking reference to UBehaviorTreeComponent rather than a pointer.")
+	virtual void OnCeaseRelevant(UBehaviorTreeComponent* OwnerComp, uint8* NodeMemory);
+	DEPRECATED(4.7, "This version is deprecated. Please use the one taking reference to UBehaviorTreeComponent rather than a pointer.")
+	virtual void TickNode(UBehaviorTreeComponent* OwnerComp, uint8* NodeMemory, float DeltaSeconds);
 };
diff --git a/Engine/Source/Runtime/AIModule/Classes/BehaviorTree/BTDecorator.h b/Engine/Source/Runtime/AIModule/Classes/BehaviorTree/BTDecorator.h
--- a/Engine/Source/Runtime/AIModule/Classes/BehaviorTree/BTDecorator.h
+++ b/Engine/Source/Runtime/AIModule/Classes/BehaviorTree/BTDecorator.h
@@ -121,6 +121,12 @@ protected:
 	virtual bool CalculateRawConditionValue(UBehaviorTreeComponent& OwnerComp, uint8* NodeMemory) const;
 
 	friend FBehaviorDecoratorDetails;
+
+	//----------------------------------------------------------------------//
+	// DEPRECATED
+	//----------------------------------------------------------------------//
+	DEPRECATED(4.7, "This version is deprecated. Please use the one taking reference to UBehaviorTreeComponent rather than a pointer.")
+	virtual bool CalculateRawConditionValue(UBehaviorTreeComponent* OwnerComp, uint8* NodeMemory) const;
 };
 
 
diff --git a/Engine/Source/Runtime/AIModule/Classes/BehaviorTree/BTNode.h b/Engine/Source/Runtime/AIModule/Classes/BehaviorTree/BTNode.h
--- a/Engine/Source/Runtime/AIModule/Classes/BehaviorTree/BTNode.h
+++ b/Engine/Source/Runtime/AIModule/Classes/BehaviorTree/BTNode.h
@@ -180,6 +180,29 @@ protected:
 
 	/** if set, node is injected by subtree */
 	uint8 bIsInjected : 1;
+
+	//----------------------------------------------------------------------//
+	// DEPRECATED
+	//----------------------------------------------------------------------//
+public:
+	DEPRECATED(4.7, "This version is deprecated. Please use the one taking reference to UBehaviorTreeComponent rather than a pointer.")
+	virtual void InitializeMemory(UBehaviorTreeComponent* OwnerComp, uint8* NodeMemory, EBTMemoryInit::Type InitType) const;
+	DEPRECATED(4.7, "This version is deprecated. Please use the one taking reference to UBehaviorTreeComponent rather than a pointer.")
+	virtual void CleanupMemory(UBehaviorTreeComponent* OwnerComp, uint8* NodeMemory, EBTMemoryClear::Type CleanupType) const;
+	DEPRECATED(4.7, "This version is deprecated. Please use the one taking reference to UBehaviorTreeComponent rather than a pointer.")
+	virtual void DescribeRuntimeValues(const UBehaviorTreeComponent* OwnerComp, uint8* NodeMemory, EBTDescriptionVerbosity::Type Verbosity, TArray<FString>& Values) const;
+	DEPRECATED(4.7, "This version is deprecated. Please use the one taking reference to UBehaviorTreeComponent rather than a pointer.")
+	virtual void OnInstanceCreated(UBehaviorTreeComponent* OwnerComp);
+	DEPRECATED(4.7, "This version is deprecated. Please use the one taking reference to UBehaviorTreeComponent rather than a pointer.")
+	virtual void OnInstanceDestroyed(UBehaviorTreeComponent* OwnerComp);
+	DEPRECATED(4.7, "This version is deprecated. Please use the one taking reference to UBehaviorTreeComponent rather than a pointer.")
+	void InitializeInSubtree(UBehaviorTreeComponent* OwnerComp, uint8* NodeMemory, int32& NextInstancedIndex, EBTMemoryInit::Type InitType) const;
+	DEPRECATED(4.7, "This version is deprecated. Please use the one taking reference to UBehaviorTreeComponent rather than a pointer.")
+	void CleanupInSubtree(UBehaviorTreeComponent* OwnerComp, uint8* NodeMemory, EBTMemoryClear::Type CleanupType) const;
+	DEPRECATED(4.7, "This version is deprecated. Please use the one taking reference to UBehaviorTreeComponent rather than a pointer.")
+	UBTNode* GetNodeInstance(const UBehaviorTreeComponent* OwnerComp, uint8* NodeMemory) const;
+	DEPRECATED(4.7, "This version is deprecated. Please use the one taking reference to UBehaviorTreeComponent rather than a pointer.")
+	FString GetRuntimeDescription(const UBehaviorTreeComponent* OwnerComp, uint8* NodeMemory, EBTDescriptionVerbosity::Type Verbosity) const;
 };
 
 //////////////////////////////////////////////////////////////////////////
diff --git a/Engine/Source/Runtime/AIModule/Classes/BehaviorTree/BTTaskNode.h b/Engine/Source/Runtime/AIModule/Classes/BehaviorTree/BTTaskNode.h
--- a/Engine/Source/Runtime/AIModule/Classes/BehaviorTree/BTTaskNode.h
+++ b/Engine/Source/Runtime/AIModule/Classes/BehaviorTree/BTTaskNode.h
@@ -89,4 +89,34 @@ protected:
 	
 	/** unregister message observers */
 	void StopWaitingForMessages(UBehaviorTreeComponent& OwnerComp) const;
+
+	//----------------------------------------------------------------------//
+	// DEPRECATED
+	//----------------------------------------------------------------------//
+	DEPRECATED(4.7, "This version is deprecated. Please use the one taking reference to UBehaviorTreeComponent rather than a pointer.")
+	EBTNodeResult::Type WrappedExecuteTask(UBehaviorTreeComponent* OwnerComp, uint8* NodeMemory) const;
+	DEPRECATED(4.7, "This version is deprecated. Please use the one taking reference to UBehaviorTreeComponent rather than a pointer.")
+	EBTNodeResult::Type WrappedAbortTask(UBehaviorTreeComponent* OwnerComp, uint8* NodeMemory) const;
+	DEPRECATED(4.7, "This version is deprecated. Please use the one taking reference to UBehaviorTreeComponent rather than a pointer.")
+	void WrappedTickTask(UBehaviorTreeComponent* OwnerComp, uint8* NodeMemory, float DeltaSeconds) const;
+	DEPRECATED(4.7, "This version is deprecated. Please use the one taking reference to UBehaviorTreeComponent rather than a pointer.")
+	void WrappedOnTaskFinished(UBehaviorTreeComponent* OwnerComp, uint8* NodeMemory, EBTNodeResult::Type TaskResult) const;
+	DEPRECATED(4.7, "This version is deprecated. Please use the one taking reference to UBehaviorTreeComponent rather than a pointer.")
+	void FinishLatentTask(UBehaviorTreeComponent* OwnerComp, EBTNodeResult::Type TaskResult) const;
+	DEPRECATED(4.7, "This version is deprecated. Please use the one taking reference to UBehaviorTreeComponent rather than a pointer.")
+	void FinishLatentAbort(UBehaviorTreeComponent* OwnerComp) const;
+	DEPRECATED(4.7, "This version is deprecated. Please use the one taking reference to UBehaviorTreeComponent rather than a pointer.")
+	virtual EBTNodeResult::Type ExecuteTask(UBehaviorTreeComponent* OwnerComp, uint8* NodeMemory);
+	DEPRECATED(4.7, "This version is deprecated. Please use the one taking reference to UBehaviorTreeComponent rather than a pointer.")
+	virtual EBTNodeResult::Type AbortTask(UBehaviorTreeComponent* OwnerComp, uint8* NodeMemory);
+	DEPRECATED(4.7, "This version is deprecated. Please use the one taking reference to UBehaviorTreeComponent rather than a pointer.")
+	virtual void TickTask(UBehaviorTreeComponent* OwnerComp, uint8* NodeMemory, float DeltaSeconds);
+	DEPRECATED(4.7, "This version is deprecated. Please use the one taking reference to UBehaviorTreeComponent rather than a pointer.")
+	virtual void OnMessage(UBehaviorTreeComponent* OwnerComp, uint8* NodeMemory, FName Message, int32 RequestID, bool bSuccess);
+	DEPRECATED(4.7, "This version is deprecated. Please use the one taking reference to UBehaviorTreeComponent rather than a pointer.")
+	virtual void OnTaskFinished(UBehaviorTreeComponent* OwnerComp, uint8* NodeMemory, EBTNodeResult::Type TaskResult);
+	DEPRECATED(4.7, "This version is deprecated. Please use the one taking reference to UBehaviorTreeComponent rather than a pointer.")
+	void WaitForMessage(UBehaviorTreeComponent* OwnerComp, FName MessageType) const;
+	DEPRECATED(4.7, "This version is deprecated. Please use the one taking reference to UBehaviorTreeComponent rather than a pointer.")
+	void WaitForMessage(UBehaviorTreeComponent* OwnerComp, FName MessageType, int32 RequestID) const;
 };
diff --git a/Engine/Source/Runtime/AIModule/Private/AIController.cpp b/Engine/Source/Runtime/AIModule/Private/AIController.cpp
--- a/Engine/Source/Runtime/AIModule/Private/AIController.cpp
+++ b/Engine/Source/Runtime/AIModule/Private/AIController.cpp
@@ -540,7 +540,7 @@ EPathFollowingRequestResult::Type AAIController::MoveToLocation(const FVector& D
 	if (bCanRequestMove && bProjectDestinationToNavigation)
 	{
 		UNavigationSystem* NavSys = UNavigationSystem::GetCurrent(GetWorld());
-		const FNavAgentProperties& AgentProps = GetNavAgentProperties();
+		const FNavAgentProperties& AgentProps = GetNavAgentPropertiesRef();
 		FNavLocation ProjectedLocation;
 
 		if (NavSys && !NavSys->ProjectPointToNavigation(Dest, ProjectedLocation, AgentProps.GetExtent(), &AgentProps))
@@ -635,7 +635,7 @@ bool AAIController::PreparePathfinding(FPathFindingQuery& Query, const FVector&
 	if (NavSys)
 	{
 		ANavigationData* NavData = bUsePathfinding ?
-			NavSys->GetNavDataForProps(GetNavAgentProperties()) :
+			NavSys->GetNavDataForProps(GetNavAgentPropertiesRef()) :
 			NavSys->GetAbstractNavData();
 
 		FVector GoalLocation = Dest;
diff --git a/Engine/Source/Runtime/AIModule/Private/BehaviorTree/BTAuxiliaryNode.cpp b/Engine/Source/Runtime/AIModule/Private/BehaviorTree/BTAuxiliaryNode.cpp
--- a/Engine/Source/Runtime/AIModule/Private/BehaviorTree/BTAuxiliaryNode.cpp
+++ b/Engine/Source/Runtime/AIModule/Private/BehaviorTree/BTAuxiliaryNode.cpp
@@ -106,3 +106,49 @@ uint16 UBTAuxiliaryNode::GetSpecialMemorySize() const
 {
 	return bTickIntervals ? sizeof(FBTAuxiliaryMemory) : Super::GetSpecialMemorySize();
 }
+
+//----------------------------------------------------------------------//
+// DEPRECATED
+//----------------------------------------------------------------------//
+void UBTAuxiliaryNode::WrappedOnBecomeRelevant(UBehaviorTreeComponent* OwnerComp, uint8* NodeMemory) const
+{
+	if (OwnerComp)
+	{
+		WrappedOnBecomeRelevant(*OwnerComp, NodeMemory);
+	}
+}
+void UBTAuxiliaryNode::WrappedOnCeaseRelevant(UBehaviorTreeComponent* OwnerComp, uint8* NodeMemory) const
+{
+	if (OwnerComp)
+	{
+		WrappedOnCeaseRelevant(*OwnerComp, NodeMemory);
+	}
+}
+void UBTAuxiliaryNode::WrappedTickNode(UBehaviorTreeComponent* OwnerComp, uint8* NodeMemory, float DeltaSeconds) const
+{
+	if (OwnerComp)
+	{
+		WrappedTickNode(*OwnerComp, NodeMemory, DeltaSeconds);
+	}
+}
+void UBTAuxiliaryNode::OnBecomeRelevant(UBehaviorTreeComponent* OwnerComp, uint8* NodeMemory)
+{
+	if (OwnerComp)
+	{
+		OnBecomeRelevant(*OwnerComp, NodeMemory);
+	}
+}
+void UBTAuxiliaryNode::OnCeaseRelevant(UBehaviorTreeComponent* OwnerComp, uint8* NodeMemory)
+{
+	if (OwnerComp)
+	{
+		OnCeaseRelevant(*OwnerComp, NodeMemory);
+	}
+}
+void UBTAuxiliaryNode::TickNode(UBehaviorTreeComponent* OwnerComp, uint8* NodeMemory, float DeltaSeconds)
+{
+	if (OwnerComp)
+	{
+		TickNode(*OwnerComp, NodeMemory, DeltaSeconds);
+	}
+}
diff --git a/Engine/Source/Runtime/AIModule/Private/BehaviorTree/BTDecorator.cpp b/Engine/Source/Runtime/AIModule/Private/BehaviorTree/BTDecorator.cpp
--- a/Engine/Source/Runtime/AIModule/Private/BehaviorTree/BTDecorator.cpp
+++ b/Engine/Source/Runtime/AIModule/Private/BehaviorTree/BTDecorator.cpp
@@ -176,3 +176,11 @@ void UBTDecorator::UpdateFlowAbortMode()
 	}
 #endif
 }
+
+//----------------------------------------------------------------------//
+// DEPRECATED
+//----------------------------------------------------------------------//
+bool UBTDecorator::CalculateRawConditionValue(UBehaviorTreeComponent* OwnerComp, uint8* NodeMemory) const
+{
+	return OwnerComp ? CalculateRawConditionValue(*OwnerComp, NodeMemory) : false;
+}
diff --git a/Engine/Source/Runtime/AIModule/Private/BehaviorTree/BTNode.cpp b/Engine/Source/Runtime/AIModule/Private/BehaviorTree/BTNode.cpp
--- a/Engine/Source/Runtime/AIModule/Private/BehaviorTree/BTNode.cpp
+++ b/Engine/Source/Runtime/AIModule/Private/BehaviorTree/BTNode.cpp
@@ -186,4 +186,73 @@ bool UBTNode::UsesBlueprint() const
 	return false;
 }
 
-#endif
\ No newline at end of file
+#endif
+
+//----------------------------------------------------------------------//
+// DEPRECATED
+//----------------------------------------------------------------------//
+void UBTNode::InitializeMemory(UBehaviorTreeComponent* OwnerComp, uint8* NodeMemory, EBTMemoryInit::Type InitType) const
+{
+	if (OwnerComp)
+	{
+		InitializeMemory(*OwnerComp, NodeMemory, InitType);
+	}
+}
+void UBTNode::CleanupMemory(UBehaviorTreeComponent* OwnerComp, uint8* NodeMemory, EBTMemoryClear::Type CleanupType) const
+{
+	if (OwnerComp)
+	{
+		CleanupMemory(*OwnerComp, NodeMemory, CleanupType);
+	}
+}
+void UBTNode::DescribeRuntimeValues(const UBehaviorTreeComponent* OwnerComp, uint8* NodeMemory, EBTDescriptionVerbosity::Type Verbosity, TArray<FString>& Values) const
+{
+	if (OwnerComp)
+	{
+		DescribeRuntimeValues(*OwnerComp, NodeMemory, Verbosity, Values);
+	}
+}
+void UBTNode::OnInstanceCreated(UBehaviorTreeComponent* OwnerComp)
+{
+	if (OwnerComp)
+	{
+		OnInstanceCreated(*OwnerComp);
+	}
+}
+void UBTNode::OnInstanceDestroyed(UBehaviorTreeComponent* OwnerComp)
+{
+	if (OwnerComp)
+	{
+		OnInstanceDestroyed(*OwnerComp);
+	}
+}
+void UBTNode::InitializeInSubtree(UBehaviorTreeComponent* OwnerComp, uint8* NodeMemory, int32& NextInstancedIndex, EBTMemoryInit::Type InitType) const
+{
+	if (OwnerComp)
+	{
+		InitializeInSubtree(*OwnerComp, NodeMemory, NextInstancedIndex, InitType);
+	}
+}
+void UBTNode::CleanupInSubtree(UBehaviorTreeComponent* OwnerComp, uint8* NodeMemory, EBTMemoryClear::Type CleanupType) const
+{
+	if (OwnerComp)
+	{
+		CleanupInSubtree(*OwnerComp, NodeMemory, CleanupType);
+	}
+}
+UBTNode* UBTNode::GetNodeInstance(const UBehaviorTreeComponent* OwnerComp, uint8* NodeMemory) const
+{
+	if (OwnerComp)
+	{
+		return GetNodeInstance(*OwnerComp, NodeMemory);
+	}
+	return nullptr;
+}
+FString UBTNode::GetRuntimeDescription(const UBehaviorTreeComponent* OwnerComp, uint8* NodeMemory, EBTDescriptionVerbosity::Type Verbosity) const
+{
+	if (OwnerComp)
+	{
+		return GetRuntimeDescription(*OwnerComp, NodeMemory, Verbosity);
+	}
+	return TEXT("");
+}
\ No newline at end of file
diff --git a/Engine/Source/Runtime/AIModule/Private/BehaviorTree/BTTaskNode.cpp b/Engine/Source/Runtime/AIModule/Private/BehaviorTree/BTTaskNode.cpp
--- a/Engine/Source/Runtime/AIModule/Private/BehaviorTree/BTTaskNode.cpp
+++ b/Engine/Source/Runtime/AIModule/Private/BehaviorTree/BTTaskNode.cpp
@@ -127,4 +127,95 @@ FName UBTTaskNode::GetNodeIconName() const
 	return FName("BTEditor.Graph.BTNode.Task.Icon");
 }
 
-#endif	// WITH_EDITOR
\ No newline at end of file
+#endif	// WITH_EDITOR
+
+//----------------------------------------------------------------------//
+// DEPRECATED
+//----------------------------------------------------------------------//
+EBTNodeResult::Type UBTTaskNode::WrappedExecuteTask(UBehaviorTreeComponent* OwnerComp, uint8* NodeMemory) const
+{
+	return OwnerComp ? WrappedExecuteTask(*OwnerComp, NodeMemory) : EBTNodeResult::Failed;
+}
+EBTNodeResult::Type UBTTaskNode::WrappedAbortTask(UBehaviorTreeComponent* OwnerComp, uint8* NodeMemory) const
+{
+	return OwnerComp ? WrappedAbortTask(*OwnerComp, NodeMemory) : EBTNodeResult::Failed;
+}
+void UBTTaskNode::WrappedTickTask(UBehaviorTreeComponent* OwnerComp, uint8* NodeMemory, float DeltaSeconds) const
+{
+	if (OwnerComp)
+	{
+		WrappedTickTask(*OwnerComp, NodeMemory, DeltaSeconds);
+	}
+}
+void UBTTaskNode::WrappedOnTaskFinished(UBehaviorTreeComponent* OwnerComp, uint8* NodeMemory, EBTNodeResult::Type TaskResult) const
+{
+	if (OwnerComp)
+	{
+		WrappedOnTaskFinished(*OwnerComp, NodeMemory, TaskResult);
+	}
+}
+void UBTTaskNode::FinishLatentTask(UBehaviorTreeComponent* OwnerComp, EBTNodeResult::Type TaskResult) const
+{
+	if (OwnerComp)
+	{
+		FinishLatentTask(*OwnerComp, TaskResult);
+	}
+}
+void UBTTaskNode::FinishLatentAbort(UBehaviorTreeComponent* OwnerComp) const
+{
+	if (OwnerComp)
+	{
+		FinishLatentAbort(*OwnerComp);
+	}
+}
+EBTNodeResult::Type UBTTaskNode::ExecuteTask(UBehaviorTreeComponent* OwnerComp, uint8* NodeMemory)
+{
+	if (OwnerComp)
+	{
+		return ExecuteTask(*OwnerComp, NodeMemory);
+	}
+	return EBTNodeResult::Failed;
+}
+EBTNodeResult::Type UBTTaskNode::AbortTask(UBehaviorTreeComponent* OwnerComp, uint8* NodeMemory)
+{
+	if (OwnerComp)
+	{
+		return AbortTask(*OwnerComp, NodeMemory);
+	}
+	return EBTNodeResult::Failed;
+}
+void UBTTaskNode::TickTask(UBehaviorTreeComponent* OwnerComp, uint8* NodeMemory, float DeltaSeconds)
+{
+	if (OwnerComp)
+	{
+		TickTask(*OwnerComp, NodeMemory, DeltaSeconds);
+	}
+}
+void UBTTaskNode::OnMessage(UBehaviorTreeComponent* OwnerComp, uint8* NodeMemory, FName Message, int32 RequestID, bool bSuccess)
+{
+	if (OwnerComp)
+	{
+		OnMessage(*OwnerComp, NodeMemory, Message, RequestID, bSuccess);
+	}
+}
+void UBTTaskNode::OnTaskFinished(UBehaviorTreeComponent* OwnerComp, uint8* NodeMemory, EBTNodeResult::Type TaskResult)
+{
+	if (OwnerComp)
+	{
+		OnTaskFinished(*OwnerComp, NodeMemory, TaskResult);
+	}
+}
+void UBTTaskNode::WaitForMessage(UBehaviorTreeComponent* OwnerComp, FName MessageType) const
+{
+	if (OwnerComp)
+	{
+		WaitForMessage(*OwnerComp, MessageType);
+	}
+}
+void UBTTaskNode::WaitForMessage(UBehaviorTreeComponent* OwnerComp, FName MessageType, int32 RequestID) const
+{
+	if (OwnerComp)
+	{
+		WaitForMessage(*OwnerComp, MessageType, RequestID);
+	}
+}
\ No newline at end of file
diff --git a/Engine/Source/Runtime/AIModule/Private/BehaviorTree/BlueprintNodeHelpers.cpp b/Engine/Source/Runtime/AIModule/Private/BehaviorTree/BlueprintNodeHelpers.cpp
--- a/Engine/Source/Runtime/AIModule/Private/BehaviorTree/BlueprintNodeHelpers.cpp
+++ b/Engine/Source/Runtime/AIModule/Private/BehaviorTree/BlueprintNodeHelpers.cpp
@@ -253,4 +253,15 @@ namespace BlueprintNodeHelpers
 			}
 		}
 	}
+
+	//----------------------------------------------------------------------//
+	// DEPRECATED
+	//----------------------------------------------------------------------//
+	void AbortLatentActions(UActorComponent* OwnerOb, const UObject* Ob)
+	{
+		if (OwnerOb && Ob)
+		{
+			AbortLatentActions(OwnerOb, Ob);
+		}
+	}
 }
diff --git a/Engine/Source/Runtime/AIModule/Private/BehaviorTree/Decorators/BTDecorator_DoesPathExist.cpp b/Engine/Source/Runtime/AIModule/Private/BehaviorTree/Decorators/BTDecorator_DoesPathExist.cpp
--- a/Engine/Source/Runtime/AIModule/Private/BehaviorTree/Decorators/BTDecorator_DoesPathExist.cpp
+++ b/Engine/Source/Runtime/AIModule/Private/BehaviorTree/Decorators/BTDecorator_DoesPathExist.cpp
@@ -60,7 +60,7 @@ bool UBTDecorator_DoesPathExist::CalculateRawConditionValue(UBehaviorTreeCompone
 	if (NavSys && bHasPointA && bHasPointB)
 	{
 		const AAIController* AIOwner = OwnerComp.GetAIOwner();
-		const ANavigationData* NavData = AIOwner ? NavSys->GetNavDataForProps(AIOwner->GetNavAgentProperties()) : NULL;
+		const ANavigationData* NavData = AIOwner ? NavSys->GetNavDataForProps(AIOwner->GetNavAgentPropertiesRef()) : NULL;
 		TSharedPtr<const FNavigationQueryFilter> QueryFilter = UNavigationQueryFilter::GetQueryFilter(NavData, FilterClass);
 
 		if (PathQueryType == EPathExistanceQueryType::NavmeshRaycast2D)
diff --git a/Engine/Source/Runtime/AIModule/Private/EnvironmentQuery/EnvQueryTypes.cpp b/Engine/Source/Runtime/AIModule/Private/EnvironmentQuery/EnvQueryTypes.cpp
--- a/Engine/Source/Runtime/AIModule/Private/EnvironmentQuery/EnvQueryTypes.cpp
+++ b/Engine/Source/Runtime/AIModule/Private/EnvironmentQuery/EnvQueryTypes.cpp
@@ -251,7 +251,7 @@ namespace FEQSHelpers
 		INavAgentInterface* NavAgent = QueryInstance.Owner.IsValid() ? Cast<INavAgentInterface>(QueryInstance.Owner.Get()) : NULL;
 		if (NavAgent)
 		{
-			const FNavAgentProperties& NavAgentProps = NavAgent ? NavAgent->GetNavAgentProperties() : FNavAgentProperties::DefaultProperties;
+			const FNavAgentProperties& NavAgentProps = NavAgent ? NavAgent->GetNavAgentPropertiesRef() : FNavAgentProperties::DefaultProperties;
 			return Cast<const ARecastNavMesh>(NavSys->GetNavDataForProps(NavAgentProps));
 		}
 
diff --git a/Engine/Source/Runtime/AIModule/Private/EnvironmentQuery/Tests/EnvQueryTest_Pathfinding.cpp b/Engine/Source/Runtime/AIModule/Private/EnvironmentQuery/Tests/EnvQueryTest_Pathfinding.cpp
--- a/Engine/Source/Runtime/AIModule/Private/EnvironmentQuery/Tests/EnvQueryTest_Pathfinding.cpp
+++ b/Engine/Source/Runtime/AIModule/Private/EnvironmentQuery/Tests/EnvQueryTest_Pathfinding.cpp
@@ -234,7 +234,7 @@ ANavigationData* UEnvQueryTest_Pathfinding::FindNavigationData(UNavigationSystem
 	INavAgentInterface* NavAgent = Cast<INavAgentInterface>(Owner);
 	if (NavAgent)
 	{
-		return NavSys->GetNavDataForProps(NavAgent->GetNavAgentProperties());
+		return NavSys->GetNavDataForProps(NavAgent->GetNavAgentPropertiesRef());
 	}
 
 	return NavSys->GetMainNavData(FNavigationSystem::DontCreate);
diff --git a/Engine/Source/Runtime/AIModule/Private/Navigation/PathFollowingComponent.cpp b/Engine/Source/Runtime/AIModule/Private/Navigation/PathFollowingComponent.cpp
--- a/Engine/Source/Runtime/AIModule/Private/Navigation/PathFollowingComponent.cpp
+++ b/Engine/Source/Runtime/AIModule/Private/Navigation/PathFollowingComponent.cpp
@@ -433,7 +433,7 @@ void UPathFollowingComponent::SetMovementComponent(UNavMovementComponent* MoveCo
 
 	if (MoveComp != NULL)
 	{
-		const FNavAgentProperties& NavAgentProps = MoveComp->GetNavAgentProperties();
+		const FNavAgentProperties& NavAgentProps = MoveComp->GetNavAgentPropertiesRef();
 		MyDefaultAcceptanceRadius = NavAgentProps.AgentRadius;
 		MoveComp->PathFollowingComp = this;
 
diff --git a/Engine/Source/Runtime/AIModule/Public/BlueprintNodeHelpers.h b/Engine/Source/Runtime/AIModule/Public/BlueprintNodeHelpers.h
--- a/Engine/Source/Runtime/AIModule/Public/BlueprintNodeHelpers.h
+++ b/Engine/Source/Runtime/AIModule/Public/BlueprintNodeHelpers.h
@@ -32,4 +32,10 @@ namespace BlueprintNodeHelpers
 	{
 		return Ob->GetClass()->GetName().LeftChop(2);
 	}
+
+	//----------------------------------------------------------------------//
+	// DEPRECATED
+	//----------------------------------------------------------------------//
+	DEPRECATED(4.7, "This version is deprecated. Please use the one taking reference to UActorComponent rather than a pointer.")
+	void AbortLatentActions(UActorComponent* OwnerOb, const UObject* Ob);
 }
diff --git a/Engine/Source/Runtime/Core/Private/CorePrivatePCH.h b/Engine/Source/Runtime/Core/Private/CorePrivatePCH.h
--- a/Engine/Source/Runtime/Core/Private/CorePrivatePCH.h
+++ b/Engine/Source/Runtime/Core/Private/CorePrivatePCH.h
@@ -25,7 +25,6 @@ Low level includes.
 #include "MemoryOps.h"							// Functions for efficient handling of object arrays.
 
 #include "Misc/CoreDefines.h"
-
 // Container forward declarations
 #include "Containers/ContainerAllocationPolicies.h"
 
@@ -94,3 +93,4 @@ Includes.
 #include "CustomVersion.h"				// Custom versioning system.
 #include "MonitoredProcess.h"
 #include "Misc/DisableOldUETypes.h"
+#include "Misc/App.h"
\ No newline at end of file
diff --git a/Engine/Source/Runtime/Core/Private/Windows/WindowsPlatformCrashContext.cpp b/Engine/Source/Runtime/Core/Private/Windows/WindowsPlatformCrashContext.cpp
--- a/Engine/Source/Runtime/Core/Private/Windows/WindowsPlatformCrashContext.cpp
+++ b/Engine/Source/Runtime/Core/Private/Windows/WindowsPlatformCrashContext.cpp
@@ -125,71 +125,71 @@ void GetModuleVersion( TCHAR* ModuleName, TCHAR* StringBuffer, DWORD MaxSize )
  */
 void SetReportParameters( HREPORT ReportHandle, EXCEPTION_POINTERS* ExceptionInfo, const TCHAR* ErrorMessage )
 {
-	HRESULT Result;
-	TCHAR StringBuffer[MAX_SPRINTF] = {0};
-	TCHAR LocalBuffer[MAX_SPRINTF] = {0};
-
-	// Set the parameters for the standard problem signature
-	HMODULE ModuleHandle = GetModuleHandle( NULL );
-
-	StringCchPrintf( StringBuffer, MAX_SPRINTF, TEXT( "UE4-%s" ), FApp::GetGameName() );
-	Result = WerReportSetParameter( ReportHandle, WER_P0, TEXT( "Application Name" ), StringBuffer );
-
-	GetModuleFileName( ModuleHandle, LocalBuffer, MAX_SPRINTF );
-	PathStripPath( LocalBuffer );
-	GetModuleVersion( LocalBuffer, StringBuffer, MAX_SPRINTF );
-	Result = WerReportSetParameter( ReportHandle, WER_P1, TEXT( "Application Version" ), StringBuffer );
-
-	StringCchPrintf( StringBuffer, MAX_SPRINTF, TEXT( "%08x" ), GetTimestampForLoadedLibrary( ModuleHandle ) );
-	Result = WerReportSetParameter( ReportHandle, WER_P2, TEXT( "Application Timestamp" ), StringBuffer );
-
-	HMODULE FaultModuleHandle = NULL;
-	GetModuleHandleEx( GET_MODULE_HANDLE_EX_FLAG_FROM_ADDRESS | GET_MODULE_HANDLE_EX_FLAG_UNCHANGED_REFCOUNT, ( LPCTSTR )ExceptionInfo->ExceptionRecord->ExceptionAddress, &FaultModuleHandle );
-
-	GetModuleFileName( FaultModuleHandle, LocalBuffer, MAX_SPRINTF );
-	PathStripPath( LocalBuffer );
-	Result = WerReportSetParameter( ReportHandle, WER_P3, TEXT( "Fault Module Name" ), LocalBuffer );
-
-	GetModuleVersion( LocalBuffer, StringBuffer, MAX_SPRINTF );
-	Result = WerReportSetParameter( ReportHandle, WER_P4, TEXT( "Fault Module Version" ), StringBuffer );
-
-	StringCchPrintf( StringBuffer, MAX_SPRINTF, TEXT( "%08x" ), GetTimestampForLoadedLibrary( FaultModuleHandle ) );
-	Result = WerReportSetParameter( ReportHandle, WER_P5, TEXT( "Fault Module Timestamp" ), StringBuffer );
-
-	StringCchPrintf( StringBuffer, MAX_SPRINTF, TEXT( "%08x" ), ExceptionInfo->ExceptionRecord->ExceptionCode );
-	Result = WerReportSetParameter( ReportHandle, WER_P6, TEXT( "Exception Code" ), StringBuffer );
-
-	INT_PTR ExceptionOffset = ( char* )( ExceptionInfo->ExceptionRecord->ExceptionAddress ) - ( char* )FaultModuleHandle;
-	CA_SUPPRESS(6066) // The format specifier should probably be something like %tX, but VS 2013 doesn't support 't'.
-		StringCchPrintf( StringBuffer, MAX_SPRINTF, TEXT( "%p" ), ExceptionOffset );
-	Result = WerReportSetParameter( ReportHandle, WER_P7, TEXT( "Exception Offset" ), StringBuffer );
-
-	// Use LocalBuffer to store the error message.
-	FCString::Strncpy( LocalBuffer, ErrorMessage, MAX_SPRINTF );
-
-	// Replace " with ' and replace \n with #
-	for (TCHAR& Char: LocalBuffer)
-	{
-		if (Char == 0)
-		{
-			break;
-		}
-
-		switch (Char)
-		{
-			default: break;
-			case '"':	Char = '\'';	break;
-			case '\r':	Char = '#';		break;
-			case '\n':	Char = '#';		break;
-		}
-	}
-
-	// AssertLog should be ErrorMessage, but this require crash server changes, so don't change this.
-	StringCchPrintf( StringBuffer, MAX_SPRINTF, TEXT( "!%s!AssertLog=\"%s\"" ), FCommandLine::Get(), LocalBuffer );
-	Result = WerReportSetParameter( ReportHandle, WER_P8, TEXT( "Commandline" ), StringBuffer );
-
-	StringCchPrintf( StringBuffer, MAX_SPRINTF, TEXT( "%s!%s!%s!%d" ), TEXT( BRANCH_NAME ), FPlatformProcess::BaseDir(), FPlatformMisc::GetEngineMode(), BUILT_FROM_CHANGELIST );
-	Result = WerReportSetParameter( ReportHandle, WER_P9, TEXT( "BranchBaseDir" ), StringBuffer );
+	//HRESULT Result;
+	//TCHAR StringBuffer[MAX_SPRINTF] = {0};
+	//TCHAR LocalBuffer[MAX_SPRINTF] = {0};
+
+	//// Set the parameters for the standard problem signature
+	//HMODULE ModuleHandle = GetModuleHandle( NULL );
+
+	//StringCchPrintf( StringBuffer, MAX_SPRINTF, TEXT( "UE4-%s" ), FApp::GetGameName() );
+	//Result = WerReportSetParameter( ReportHandle, WER_P0, TEXT( "Application Name" ), StringBuffer );
+
+	//GetModuleFileName( ModuleHandle, LocalBuffer, MAX_SPRINTF );
+	//PathStripPath( LocalBuffer );
+	//GetModuleVersion( LocalBuffer, StringBuffer, MAX_SPRINTF );
+	//Result = WerReportSetParameter( ReportHandle, WER_P1, TEXT( "Application Version" ), StringBuffer );
+
+	//StringCchPrintf( StringBuffer, MAX_SPRINTF, TEXT( "%08x" ), GetTimestampForLoadedLibrary( ModuleHandle ) );
+	//Result = WerReportSetParameter( ReportHandle, WER_P2, TEXT( "Application Timestamp" ), StringBuffer );
+
+	//HMODULE FaultModuleHandle = NULL;
+	//GetModuleHandleEx( GET_MODULE_HANDLE_EX_FLAG_FROM_ADDRESS | GET_MODULE_HANDLE_EX_FLAG_UNCHANGED_REFCOUNT, ( LPCTSTR )ExceptionInfo->ExceptionRecord->ExceptionAddress, &FaultModuleHandle );
+
+	//GetModuleFileName( FaultModuleHandle, LocalBuffer, MAX_SPRINTF );
+	//PathStripPath( LocalBuffer );
+	//Result = WerReportSetParameter( ReportHandle, WER_P3, TEXT( "Fault Module Name" ), LocalBuffer );
+
+	//GetModuleVersion( LocalBuffer, StringBuffer, MAX_SPRINTF );
+	//Result = WerReportSetParameter( ReportHandle, WER_P4, TEXT( "Fault Module Version" ), StringBuffer );
+
+	//StringCchPrintf( StringBuffer, MAX_SPRINTF, TEXT( "%08x" ), GetTimestampForLoadedLibrary( FaultModuleHandle ) );
+	//Result = WerReportSetParameter( ReportHandle, WER_P5, TEXT( "Fault Module Timestamp" ), StringBuffer );
+
+	//StringCchPrintf( StringBuffer, MAX_SPRINTF, TEXT( "%08x" ), ExceptionInfo->ExceptionRecord->ExceptionCode );
+	//Result = WerReportSetParameter( ReportHandle, WER_P6, TEXT( "Exception Code" ), StringBuffer );
+
+	//INT_PTR ExceptionOffset = ( char* )( ExceptionInfo->ExceptionRecord->ExceptionAddress ) - ( char* )FaultModuleHandle;
+	//CA_SUPPRESS(6066) // The format specifier should probably be something like %tX, but VS 2013 doesn't support 't'.
+	//	StringCchPrintf( StringBuffer, MAX_SPRINTF, TEXT( "%p" ), ExceptionOffset );
+	//Result = WerReportSetParameter( ReportHandle, WER_P7, TEXT( "Exception Offset" ), StringBuffer );
+
+	//// Use LocalBuffer to store the error message.
+	//FCString::Strncpy( LocalBuffer, ErrorMessage, MAX_SPRINTF );
+
+	//// Replace " with ' and replace \n with #
+	//for (TCHAR& Char: LocalBuffer)
+	//{
+	//	if (Char == 0)
+	//	{
+	//		break;
+	//	}
+
+	//	switch (Char)
+	//	{
+	//		default: break;
+	//		case '"':	Char = '\'';	break;
+	//		case '\r':	Char = '#';		break;
+	//		case '\n':	Char = '#';		break;
+	//	}
+	//}
+
+	//// AssertLog should be ErrorMessage, but this require crash server changes, so don't change this.
+	//StringCchPrintf( StringBuffer, MAX_SPRINTF, TEXT( "!%s!AssertLog=\"%s\"" ), FCommandLine::Get(), LocalBuffer );
+	//Result = WerReportSetParameter( ReportHandle, WER_P8, TEXT( "Commandline" ), StringBuffer );
+
+	//StringCchPrintf( StringBuffer, MAX_SPRINTF, TEXT( "%s!%s!%s!%d" ), TEXT( BRANCH_NAME ), FPlatformProcess::BaseDir(), FPlatformMisc::GetEngineMode(), BUILT_FROM_CHANGELIST );
+	//Result = WerReportSetParameter( ReportHandle, WER_P9, TEXT( "BranchBaseDir" ), StringBuffer );
 }
 
 /**
@@ -213,11 +213,11 @@ void AddMiniDump(HREPORT ReportHandle, EXCEPTION_POINTERS* ExceptionInfo)
  */
 void AddMiscFiles( HREPORT ReportHandle )
 {
-	FString LogFileName = FPaths::GameLogDir() / FApp::GetGameName() + TEXT( ".log" );
+	/*FString LogFileName = FPaths::GameLogDir() / FApp::GetGameName() + TEXT( ".log" );
 	WerReportAddFile( ReportHandle, *LogFileName, WerFileTypeOther, WER_FILE_ANONYMOUS_DATA ); 
 
 	FString CrashVideoPath = FPaths::GameLogDir() / TEXT( "CrashVideo.avi" );
-	WerReportAddFile( ReportHandle, *CrashVideoPath, WerFileTypeOther, WER_FILE_ANONYMOUS_DATA );
+	WerReportAddFile( ReportHandle, *CrashVideoPath, WerFileTypeOther, WER_FILE_ANONYMOUS_DATA );*/
 }
 
 /**
@@ -280,95 +280,95 @@ private:
  */
 int32 ReportCrashUsingCrashReportClient(EXCEPTION_POINTERS* ExceptionInfo, const TCHAR* ErrorMessage, EErrorReportUI ReportUI)
 {
-	// Flush out the log
-	GLog->Flush();
-
-	// Prevent CrashReportClient from spawning another CrashReportClient.
-	const TCHAR* ExecutableName = FPlatformProcess::ExecutableName();
-	const bool bCanRunCrashReportClient = FCString::Stristr( ExecutableName, TEXT( "CrashReportClient" ) ) == nullptr;
-	if( bCanRunCrashReportClient )
-	{
-		// Set the report to force queue
-		FScopedWERQueuing ScopedQueueForcer;
-
-		// Construct the report details
-		WER_REPORT_INFORMATION ReportInformation = {sizeof( WER_REPORT_INFORMATION )};
-
-		StringCchCopy( ReportInformation.wzConsentKey, ARRAYSIZE( ReportInformation.wzConsentKey ), TEXT( "" ) );
-
-		StringCchCopy( ReportInformation.wzApplicationName, ARRAYSIZE( ReportInformation.wzApplicationName ), TEXT("UE4-") );
-		StringCchCat( ReportInformation.wzApplicationName, ARRAYSIZE( ReportInformation.wzApplicationName ), FApp::GetGameName() );
-
-		StringCchCopy( ReportInformation.wzApplicationPath, ARRAYSIZE( ReportInformation.wzApplicationPath ), FPlatformProcess::BaseDir() );
-		StringCchCat( ReportInformation.wzApplicationPath, ARRAYSIZE( ReportInformation.wzApplicationPath ), FPlatformProcess::ExecutableName() );
-		StringCchCat( ReportInformation.wzApplicationPath, ARRAYSIZE( ReportInformation.wzApplicationPath ), TEXT( ".exe" ) );
-
-		GetCrashDescription( ReportInformation );
-
-		// Create a crash event report
-		HREPORT ReportHandle = NULL;
-		if( WerReportCreate( APPCRASH_EVENT, WerReportApplicationCrash, &ReportInformation, &ReportHandle ) == S_OK )
-		{
-			// Set the standard set of a crash parameters
-			SetReportParameters( ReportHandle, ExceptionInfo, ErrorMessage );
-
-			// Add a manually generated minidump
-			AddMiniDump( ReportHandle, ExceptionInfo );
-
-			// Add the log and video
-			AddMiscFiles( ReportHandle );
-
-			// Submit
-			WER_SUBMIT_RESULT SubmitResult;
-			WerReportSubmit( ReportHandle, WerConsentAlwaysPrompt, WER_SUBMIT_QUEUE | WER_SUBMIT_BYPASS_DATA_THROTTLING, &SubmitResult );
-
-			// Cleanup
-			WerReportCloseHandle( ReportHandle );
-		}
-
-		// Build machines do not upload these automatically since it is not okay to have lingering processes after the build completes.
-		if( GIsBuildMachine )
-		{
-			return EXCEPTION_CONTINUE_EXECUTION;
-		}
-
-		FString CrashReportClientArguments;
-
-		// Suppress the user input dialog if we're running in unattended mode
-		bool bNoDialog = FApp::IsUnattended() || ReportUI == EErrorReportUI::ReportInUnattendedMode;
-		if( bNoDialog )
-		{
-			CrashReportClientArguments += TEXT( " -Unattended" );
-		}
-
-		if( FApp::IsInstalled() )
-		{
-			// Temporary workaround for CrashReportClient being built in Development, not Shipping (TTP328030). The
-			// following ensures that logs are saved to the user directory when UE4 is installed.
-			CrashReportClientArguments += TEXT( " -Installed" );
-		}
-
-		CrashReportClientArguments += FString( TEXT( " -AppName=" ) ) + ReportInformation.wzApplicationName;
-
-		static const TCHAR CrashReportClientExeName[] = TEXT( "CrashReportClient.exe" );
-		FString CrashClientPath = FString( TEXT( "..\\..\\..\\Engine\\Binaries" ) ) / FPlatformProcess::GetBinariesSubdirectory() / CrashReportClientExeName;
-
-		bool bCrashReporterRan = FPlatformProcess::CreateProc( *CrashClientPath, *CrashReportClientArguments, true, false, false, NULL, 0, NULL, NULL ).IsValid();
-
-		if( !bCrashReporterRan && !bNoDialog )
-		{
-			UE_LOG( LogWindows, Log, TEXT( "Could not start %s" ), CrashReportClientExeName );
-			FPlatformMemory::DumpStats( *GWarn );
-			FText MessageTitle( FText::Format(
-				NSLOCTEXT( "MessageDialog", "AppHasCrashed", "The {0} {1} has crashed and will close" ),
-				FText::FromString( ReportInformation.wzApplicationName ),
-				FText::FromString( FPlatformMisc::GetEngineMode() )
-				) );
-			FMessageDialog::Open( EAppMsgType::Ok, FText::FromString( GErrorHist ), &MessageTitle );
-		}
-	}
-
-	// Let the system take back over (return value only used by NewReportEnsure)
+	//// Flush out the log
+	//GLog->Flush();
+
+	//// Prevent CrashReportClient from spawning another CrashReportClient.
+	//const TCHAR* ExecutableName = FPlatformProcess::ExecutableName();
+	//const bool bCanRunCrashReportClient = FCString::Stristr( ExecutableName, TEXT( "CrashReportClient" ) ) == nullptr;
+	//if( bCanRunCrashReportClient )
+	//{
+	//	// Set the report to force queue
+	//	FScopedWERQueuing ScopedQueueForcer;
+
+	//	// Construct the report details
+	//	WER_REPORT_INFORMATION ReportInformation = {sizeof( WER_REPORT_INFORMATION )};
+
+	//	StringCchCopy( ReportInformation.wzConsentKey, ARRAYSIZE( ReportInformation.wzConsentKey ), TEXT( "" ) );
+
+	//	StringCchCopy( ReportInformation.wzApplicationName, ARRAYSIZE( ReportInformation.wzApplicationName ), TEXT("UE4-") );
+	//	StringCchCat( ReportInformation.wzApplicationName, ARRAYSIZE( ReportInformation.wzApplicationName ), FApp::GetGameName() );
+
+	//	StringCchCopy( ReportInformation.wzApplicationPath, ARRAYSIZE( ReportInformation.wzApplicationPath ), FPlatformProcess::BaseDir() );
+	//	StringCchCat( ReportInformation.wzApplicationPath, ARRAYSIZE( ReportInformation.wzApplicationPath ), FPlatformProcess::ExecutableName() );
+	//	StringCchCat( ReportInformation.wzApplicationPath, ARRAYSIZE( ReportInformation.wzApplicationPath ), TEXT( ".exe" ) );
+
+	//	GetCrashDescription( ReportInformation );
+
+	//	// Create a crash event report
+	//	HREPORT ReportHandle = NULL;
+	//	if( WerReportCreate( APPCRASH_EVENT, WerReportApplicationCrash, &ReportInformation, &ReportHandle ) == S_OK )
+	//	{
+	//		// Set the standard set of a crash parameters
+	//		SetReportParameters( ReportHandle, ExceptionInfo, ErrorMessage );
+
+	//		// Add a manually generated minidump
+	//		AddMiniDump( ReportHandle, ExceptionInfo );
+
+	//		// Add the log and video
+	//		AddMiscFiles( ReportHandle );
+
+	//		// Submit
+	//		WER_SUBMIT_RESULT SubmitResult;
+	//		WerReportSubmit( ReportHandle, WerConsentAlwaysPrompt, WER_SUBMIT_QUEUE | WER_SUBMIT_BYPASS_DATA_THROTTLING, &SubmitResult );
+
+	//		// Cleanup
+	//		WerReportCloseHandle( ReportHandle );
+	//	}
+
+	//	// Build machines do not upload these automatically since it is not okay to have lingering processes after the build completes.
+	//	if( GIsBuildMachine )
+	//	{
+	//		return EXCEPTION_CONTINUE_EXECUTION;
+	//	}
+
+	//	FString CrashReportClientArguments;
+
+	//	// Suppress the user input dialog if we're running in unattended mode
+	//	bool bNoDialog = FApp::IsUnattended() || ReportUI == EErrorReportUI::ReportInUnattendedMode;
+	//	if( bNoDialog )
+	//	{
+	//		CrashReportClientArguments += TEXT( " -Unattended" );
+	//	}
+
+	//	if( FApp::IsInstalled() )
+	//	{
+	//		// Temporary workaround for CrashReportClient being built in Development, not Shipping (TTP328030). The
+	//		// following ensures that logs are saved to the user directory when UE4 is installed.
+	//		CrashReportClientArguments += TEXT( " -Installed" );
+	//	}
+
+	//	CrashReportClientArguments += FString( TEXT( " -AppName=" ) ) + ReportInformation.wzApplicationName;
+
+	//	static const TCHAR CrashReportClientExeName[] = TEXT( "CrashReportClient.exe" );
+	//	FString CrashClientPath = FString( TEXT( "..\\..\\..\\Engine\\Binaries" ) ) / FPlatformProcess::GetBinariesSubdirectory() / CrashReportClientExeName;
+
+	//	bool bCrashReporterRan = FPlatformProcess::CreateProc( *CrashClientPath, *CrashReportClientArguments, true, false, false, NULL, 0, NULL, NULL ).IsValid();
+
+	//	if( !bCrashReporterRan && !bNoDialog )
+	//	{
+	//		UE_LOG( LogWindows, Log, TEXT( "Could not start %s" ), CrashReportClientExeName );
+	//		FPlatformMemory::DumpStats( *GWarn );
+	//		FText MessageTitle( FText::Format(
+	//			NSLOCTEXT( "MessageDialog", "AppHasCrashed", "The {0} {1} has crashed and will close" ),
+	//			FText::FromString( ReportInformation.wzApplicationName ),
+	//			FText::FromString( FPlatformMisc::GetEngineMode() )
+	//			) );
+	//		FMessageDialog::Open( EAppMsgType::Ok, FText::FromString( GErrorHist ), &MessageTitle );
+	//	}
+	//}
+
+	//// Let the system take back over (return value only used by NewReportEnsure)
 	return EXCEPTION_CONTINUE_EXECUTION;
 }
 
diff --git a/Engine/Source/Runtime/Engine/Classes/AI/Navigation/NavAgentInterface.h b/Engine/Source/Runtime/Engine/Classes/AI/Navigation/NavAgentInterface.h
--- a/Engine/Source/Runtime/Engine/Classes/AI/Navigation/NavAgentInterface.h
+++ b/Engine/Source/Runtime/Engine/Classes/AI/Navigation/NavAgentInterface.h
@@ -20,8 +20,10 @@ class INavAgentInterface
 
 	/**
 	 *	Retrieves FNavAgentProperties expressing navigation props and caps of represented agent
+	 *	@NOTE the function will be renamed to GetNavAgentProperties in 4.8. Current name was introduced
+	 *		to help with deprecating old GetNavAgentProperties function
 	 */
-	virtual const FNavAgentProperties& GetNavAgentProperties() const { return FNavAgentProperties::DefaultProperties; }
+	virtual const FNavAgentProperties& GetNavAgentPropertiesRef() const { return FNavAgentProperties::DefaultProperties; }
 
 	/**
 	 *	Retrieves Agent's location
@@ -41,4 +43,10 @@ class INavAgentInterface
 
 	/** Allows delaying repath requests */
 	virtual bool ShouldPostponePathUpdates() const { return false; }
+
+	//----------------------------------------------------------------------//
+	// DEPRECATED
+	//----------------------------------------------------------------------//
+	DEPRECATED(4.7, "This version is deprecated. Please use GetNavAgentPropertiesRef instead.")
+	virtual const FNavAgentProperties* GetNavAgentProperties() const { return &FNavAgentProperties::DefaultProperties; }
 };
diff --git a/Engine/Source/Runtime/Engine/Classes/AI/Navigation/NavigationTestingActor.h b/Engine/Source/Runtime/Engine/Classes/AI/Navigation/NavigationTestingActor.h
--- a/Engine/Source/Runtime/Engine/Classes/AI/Navigation/NavigationTestingActor.h
+++ b/Engine/Source/Runtime/Engine/Classes/AI/Navigation/NavigationTestingActor.h
@@ -153,7 +153,7 @@ public:
 #endif // WITH_EDITOR
 
 	// Begin INavAgentInterface Interface
-	virtual const FNavAgentProperties& GetNavAgentProperties() const override { return NavAgentProps; }
+	virtual const FNavAgentProperties& GetNavAgentPropertiesRef() const override { return NavAgentProps; }
 	virtual FVector GetNavAgentLocation() const override;
 	virtual void GetMoveGoalReachTest(class AActor* MovingActor, const FVector& MoveOffset, FVector& GoalOffset, float& GoalRadius, float& GoalHalfHeight) const override {}
 	// End INavAgentInterface Interface
diff --git a/Engine/Source/Runtime/Engine/Classes/GameFramework/Controller.h b/Engine/Source/Runtime/Engine/Classes/GameFramework/Controller.h
--- a/Engine/Source/Runtime/Engine/Classes/GameFramework/Controller.h
+++ b/Engine/Source/Runtime/Engine/Classes/GameFramework/Controller.h
@@ -265,7 +265,7 @@ public:
 	virtual void FailedToSpawnPawn();
 
 	// Begin INavAgentInterface Interface
-	virtual const struct FNavAgentProperties& GetNavAgentProperties() const override;
+	virtual const struct FNavAgentProperties& GetNavAgentPropertiesRef() const override;
 	virtual FVector GetNavAgentLocation() const override;
 	virtual void GetMoveGoalReachTest(class AActor* MovingActor, const FVector& MoveOffset, FVector& GoalOffset, float& GoalRadius, float& GoalHalfHeight) const override;
 	virtual bool ShouldPostponePathUpdates() const override;
diff --git a/Engine/Source/Runtime/Engine/Classes/GameFramework/NavMovementComponent.h b/Engine/Source/Runtime/Engine/Classes/GameFramework/NavMovementComponent.h
--- a/Engine/Source/Runtime/Engine/Classes/GameFramework/NavMovementComponent.h
+++ b/Engine/Source/Runtime/Engine/Classes/GameFramework/NavMovementComponent.h
@@ -79,9 +79,9 @@ public:
 	virtual bool CanStopPathFollowing() const;
 
 	/** @returns the NavAgentProps(const) */
-	FORCEINLINE const FNavAgentProperties& GetNavAgentProperties() const { return NavAgentProps; }
+	FORCEINLINE const FNavAgentProperties& GetNavAgentPropertiesRef() const { return NavAgentProps; }
 	/** @returns the NavAgentProps */
-	FORCEINLINE FNavAgentProperties& GetNavAgentProperties() { return NavAgentProps; }
+	FORCEINLINE FNavAgentProperties& GetNavAgentPropertiesRef() { return NavAgentProps; }
 
 	/** Resets runtime movement state to character's movement capabilities */
 	void ResetMoveState() { MovementState = NavAgentProps; }
@@ -130,6 +130,15 @@ public:
 	/** @return true if currently flying (moving through a non-fluid volume without resting on the ground) */
 	UFUNCTION(BlueprintCallable, Category="AI|Components|NavMovement")
 	virtual bool IsFlying() const;
+
+	//----------------------------------------------------------------------//
+	// DEPRECATED
+	//----------------------------------------------------------------------//
+public:
+	DEPRECATED(4.7, "This function is deprecated. Please use GetNavAgentPropertiesRef instead.")
+	const FNavAgentProperties* GetNavAgentProperties() const;
+	DEPRECATED(4.7, "This function is deprecated. Please use GetNavAgentPropertiesRef instead.")
+	FNavAgentProperties* GetNavAgentProperties();
 };
 
 
@@ -170,3 +179,15 @@ inline void UNavMovementComponent::StopMovementImmediately()
 	Super::StopMovementImmediately();
 	StopActiveMovement();
 }
+
+//----------------------------------------------------------------------//
+// DEPRECATED
+//----------------------------------------------------------------------//
+const FNavAgentProperties* UNavMovementComponent::GetNavAgentProperties() const 
+{ 
+	return &GetNavAgentPropertiesRef();
+}
+FNavAgentProperties* UNavMovementComponent::GetNavAgentProperties() 
+{ 
+	return &GetNavAgentPropertiesRef();
+}
\ No newline at end of file
diff --git a/Engine/Source/Runtime/Engine/Classes/GameFramework/Pawn.h b/Engine/Source/Runtime/Engine/Classes/GameFramework/Pawn.h
--- a/Engine/Source/Runtime/Engine/Classes/GameFramework/Pawn.h
+++ b/Engine/Source/Runtime/Engine/Classes/GameFramework/Pawn.h
@@ -245,7 +245,7 @@ public:
 	virtual void UpdateNavigationRelevance() {}
 
 	// Begin INavAgentInterface Interface
-	virtual const FNavAgentProperties& GetNavAgentProperties() const override;
+	virtual const FNavAgentProperties& GetNavAgentPropertiesRef() const override;
 	/** Basically retrieved pawn's location on navmesh */
 	UFUNCTION(BlueprintCallable, Category="Pawn")
 	virtual FVector GetNavAgentLocation() const override { return GetActorLocation() - FVector(0.f, 0.f, BaseEyeHeight); }
diff --git a/Engine/Source/Runtime/Engine/Private/AI/Navigation/NavigationSystem.cpp b/Engine/Source/Runtime/Engine/Private/AI/Navigation/NavigationSystem.cpp
--- a/Engine/Source/Runtime/Engine/Private/AI/Navigation/NavigationSystem.cpp
+++ b/Engine/Source/Runtime/Engine/Private/AI/Navigation/NavigationSystem.cpp
@@ -1021,7 +1021,7 @@ void UNavigationSystem::SimpleMoveToActor(AController* Controller, const AActor*
 	}
 	else
 	{
-		const ANavigationData* NavData = NavSys->GetNavDataForProps(Controller->GetNavAgentProperties());
+		const ANavigationData* NavData = NavSys->GetNavDataForProps(Controller->GetNavAgentPropertiesRef());
 		FPathFindingQuery Query(Controller, NavData, Controller->GetNavAgentLocation(), Goal->GetActorLocation());
 		FPathFindingResult Result = NavSys->FindPathSync(Query);
 		if (Result.IsSuccessful())
@@ -1077,7 +1077,7 @@ void UNavigationSystem::SimpleMoveToLocation(AController* Controller, const FVec
 	}
 	else
 	{
-		const ANavigationData* NavData = NavSys->GetNavDataForProps(Controller->GetNavAgentProperties());
+		const ANavigationData* NavData = NavSys->GetNavDataForProps(Controller->GetNavAgentPropertiesRef());
 		FPathFindingQuery Query(Controller, NavData, Controller->GetNavAgentLocation(), GoalLocation);
 		FPathFindingResult Result = NavSys->FindPathSync(Query);
 		if (Result.IsSuccessful())
@@ -1139,7 +1139,7 @@ UNavigationPath* UNavigationSystem::FindPathToLocationSynchronously(UObject* Wor
 			
 			if (NavAgent != NULL)
 			{
-				const FNavAgentProperties& AgentProps = NavAgent->GetNavAgentProperties();
+				const FNavAgentProperties& AgentProps = NavAgent->GetNavAgentPropertiesRef();
 				NavigationData = NavSys->GetNavDataForProps(AgentProps);
 				bValidPathContext = true;
 			}
@@ -1195,7 +1195,7 @@ bool UNavigationSystem::NavigationRaycast(UObject* WorldContextObject, const FVe
 		INavAgentInterface* MyNavAgent = Cast<INavAgentInterface>(Querier);
 		if (MyNavAgent)
 		{
-			const FNavAgentProperties& AgentProps = MyNavAgent->GetNavAgentProperties();
+			const FNavAgentProperties& AgentProps = MyNavAgent->GetNavAgentPropertiesRef();
 			NavData = NavSys->GetNavDataForProps(AgentProps);
 		}
 		if (NavData == NULL)
diff --git a/Engine/Source/Runtime/Engine/Private/CharacterMovementComponent.cpp b/Engine/Source/Runtime/Engine/Private/CharacterMovementComponent.cpp
--- a/Engine/Source/Runtime/Engine/Private/CharacterMovementComponent.cpp
+++ b/Engine/Source/Runtime/Engine/Private/CharacterMovementComponent.cpp
@@ -3766,7 +3766,7 @@ void UCharacterMovementComponent::PhysNavWalking(float deltaTime, int32 Iteratio
 	float SearchHeight = 100.0f;
 	if (MyNavAgent)
 	{
-		const FNavAgentProperties& AgentProps = MyNavAgent->GetNavAgentProperties();
+		const FNavAgentProperties& AgentProps = MyNavAgent->GetNavAgentPropertiesRef();
 		NavData = NavSys->GetNavDataForProps(AgentProps);
 		SearchRadius = AgentProps.AgentRadius * 2.0f;
 		SearchHeight = AgentProps.AgentHeight * 0.5f;
diff --git a/Engine/Source/Runtime/Engine/Private/Controller.cpp b/Engine/Source/Runtime/Engine/Private/Controller.cpp
--- a/Engine/Source/Runtime/Engine/Private/Controller.cpp
+++ b/Engine/Source/Runtime/Engine/Private/Controller.cpp
@@ -528,9 +528,9 @@ APawn* AController::K2_GetPawn() const
 	return GetPawn();
 }
 
-const FNavAgentProperties& AController::GetNavAgentProperties() const
+const FNavAgentProperties& AController::GetNavAgentPropertiesRef() const
 {
-	return Pawn && Pawn->GetMovementComponent() ? Pawn->GetMovementComponent()->GetNavAgentProperties() : FNavAgentProperties::DefaultProperties;
+	return Pawn && Pawn->GetMovementComponent() ? Pawn->GetMovementComponent()->GetNavAgentPropertiesRef() : FNavAgentProperties::DefaultProperties;
 }
 
 FVector AController::GetNavAgentLocation() const
diff --git a/Engine/Source/Runtime/Engine/Private/Pawn.cpp b/Engine/Source/Runtime/Engine/Private/Pawn.cpp
--- a/Engine/Source/Runtime/Engine/Private/Pawn.cpp
+++ b/Engine/Source/Runtime/Engine/Private/Pawn.cpp
@@ -1089,7 +1089,7 @@ void APawn::PawnMakeNoise(float Loudness, FVector NoiseLocation, bool bUseNoiseM
 	NoiseMaker->MakeNoise(Loudness, this, bUseNoiseMakerLocation ? NoiseMaker->GetActorLocation() : NoiseLocation);
 }
 
-const FNavAgentProperties& APawn::GetNavAgentProperties() const
+const FNavAgentProperties& APawn::GetNavAgentPropertiesRef() const
 {
-	return GetMovementComponent() ? GetMovementComponent()->GetNavAgentProperties() : FNavAgentProperties::DefaultProperties;
+	return GetMovementComponent() ? GetMovementComponent()->GetNavAgentPropertiesRef() : FNavAgentProperties::DefaultProperties;
 }
