Core: added support for event pooling to GenericPlatform

- added GetSynchEventFromPool(), ReturnSynchEventToPool() to FPlatformProcess
- moved FEventPool into separate private header file and templated it based on auto/manual reset
- simplified FScopedEvent
- cleaned up some related code & documentation

diff --git a/Engine/Source/Runtime/Core/Private/GenericPlatform/GenericPlatformProcess.cpp b/Engine/Source/Runtime/Core/Private/GenericPlatform/GenericPlatformProcess.cpp
--- a/Engine/Source/Runtime/Core/Private/GenericPlatform/GenericPlatformProcess.cpp
+++ b/Engine/Source/Runtime/Core/Private/GenericPlatform/GenericPlatformProcess.cpp
@@ -1,6 +1,7 @@
 // Copyright 1998-2015 Epic Games, Inc. All Rights Reserved.
 
 #include "CorePrivatePCH.h"
+#include "EventPool.h"
 #include "Public/Modules/ModuleVersion.h"
 
 
@@ -379,6 +380,25 @@ FEvent* FGenericPlatformProcess::CreateSynchEvent(bool bIsManualReset)
 }
 
 
+FEvent* FGenericPlatformProcess::GetSynchEventFromPool(bool bIsManualReset)
+{
+	return bIsManualReset
+		? FEventPool<EEventPoolTypes::ManualReset>::Get().GetEventFromPool()
+		: FEventPool<EEventPoolTypes::AutoReset>::Get().GetEventFromPool();
+}
+
+
+void FGenericPlatformProcess::ReturnSynchEventToPool(FEvent* Event)
+{
+	if (Event->IsManualReset())
+	{
+		FEventPool<EEventPoolTypes::ManualReset>::Get().ReturnToPool(Event);
+	}
+	else
+	{
+		FEventPool<EEventPoolTypes::AutoReset>::Get().ReturnToPool(Event);
+	}
+}
 
 
 #if PLATFORM_USE_PTHREADS
diff --git a/Engine/Source/Runtime/Core/Private/HAL/ThreadingBase.cpp b/Engine/Source/Runtime/Core/Private/HAL/ThreadingBase.cpp
--- a/Engine/Source/Runtime/Core/Private/HAL/ThreadingBase.cpp
+++ b/Engine/Source/Runtime/Core/Private/HAL/ThreadingBase.cpp
@@ -2,6 +2,7 @@
 
 
 #include "CorePrivatePCH.h"
+#include "EventPool.h"
 #include "LockFreeList.h"
 #include "StatsData.h"
 
@@ -174,40 +175,15 @@ FSingleThreadManager& FSingleThreadManager::Get()
 	return Singleton;
 }
 
-class FEventPool
-{
-	TLockFreePointerList<FEvent> Pool;
-public:
-	static FEventPool& Get()
-	{
-		static FEventPool Singleton;
-		return Singleton;
-	}
-	FEvent* GetEventFromPool()
-	{
-		FEvent* Result = Pool.Pop();
-		if (!Result)
-		{
-			Result = FPlatformProcess::CreateSynchEvent();
-		}
-		check(Result);
-		return Result;
-	}
-	void ReturnToPool(FEvent* Event)
-	{
-		check(Event);
-		Pool.Push(Event);
-	}
-};
-
-FEvent* FScopedEvent::GetEventFromPool()
-{
-	return FEventPool::Get().GetEventFromPool();
-}
+FScopedEvent::FScopedEvent()
+	: Event(FEventPool<EEventPoolTypes::AutoReset>::Get().GetEventFromPool())
+{ }
 
-void FScopedEvent::ReturnToPool(FEvent* Event)
+FScopedEvent::~FScopedEvent()
 {
-	FEventPool::Get().ReturnToPool(Event);
+	Event->Wait();
+	FEventPool<EEventPoolTypes::AutoReset>::Get().ReturnToPool(Event);
+	Event = nullptr;
 }
 
 FRunnableThread::~FRunnableThread()
diff --git a/Engine/Source/Runtime/Core/Private/WinRT/WinRTEvent.h b/Engine/Source/Runtime/Core/Private/WinRT/WinRTEvent.h
--- a/Engine/Source/Runtime/Core/Private/WinRT/WinRTEvent.h
+++ b/Engine/Source/Runtime/Core/Private/WinRT/WinRTEvent.h
@@ -36,10 +36,16 @@ public:
 		// Create the event and default it to non-signaled
 //		Event = CreateEvent(NULL, bIsManualReset, 0, NULL);
 		Event = CreateEventEx(NULL, NULL, bIsManualReset ? CREATE_EVENT_MANUAL_RESET : 0, EVENT_ALL_ACCESS);
+		ManualReset = bIsManualReset;
 
 		return Event != NULL;
 	}
 
+	virtual bool IsManualReset() override
+	{
+		return ManualReset;
+	}
+
 	virtual void Trigger () override
 	{
 		check(Event);
@@ -60,6 +66,9 @@ private:
 
 	/** Holds the handle to the event. */
 	HANDLE Event;
+
+	/** Whether the signaled state of the event needs to be reset manually. */
+	bool ManualReset;
 };
 
 
diff --git a/Engine/Source/Runtime/Core/Private/Windows/WindowsEvent.h b/Engine/Source/Runtime/Core/Private/Windows/WindowsEvent.h
--- a/Engine/Source/Runtime/Core/Private/Windows/WindowsEvent.h
+++ b/Engine/Source/Runtime/Core/Private/Windows/WindowsEvent.h
@@ -35,10 +35,16 @@ public:
 	{
 		// Create the event and default it to non-signaled
 		Event = CreateEvent(nullptr, bIsManualReset, 0, nullptr);
+		ManualReset = bIsManualReset;
 
 		return Event != nullptr;
 	}
 
+	virtual bool IsManualReset() override
+	{
+		return ManualReset;
+	}
+
 	virtual void Trigger() override
 	{
 		check(Event);
@@ -59,6 +65,9 @@ private:
 
 	/** Holds the handle to the event. */
 	HANDLE Event;
+
+	/** Whether the signaled state of the event needs to be reset manually. */
+	bool ManualReset;
 };
 
 
diff --git a/Engine/Source/Runtime/Core/Public/GenericPlatform/GenericPlatformProcess.h b/Engine/Source/Runtime/Core/Public/GenericPlatform/GenericPlatformProcess.h
--- a/Engine/Source/Runtime/Core/Public/GenericPlatform/GenericPlatformProcess.h
+++ b/Engine/Source/Runtime/Core/Public/GenericPlatform/GenericPlatformProcess.h
@@ -410,7 +410,7 @@ struct CORE_API FGenericPlatformProcess
 
 #if PLATFORM_HAS_BSD_TIME 
 
-	/** Sleep this thread for Seconds.  0.0 means release the current timeslice to let other threads get some attention. */
+	/** Sleep this thread for Seconds.  0.0 means release the current time slice to let other threads get some attention. */
 	static void Sleep( float Seconds );
 	/** Sleep this thread infinitely. */
 	static void SleepInfinite();
@@ -418,13 +418,30 @@ struct CORE_API FGenericPlatformProcess
 #endif // PLATFORM_HAS_BSD_TIME
 
 	/**
-	 * Creates a new event,
+	 * Creates a new event.
 	 *
 	 * @param bIsManualReset Whether the event requires manual reseting or not.
-	 * @param InName Whether to use a commonly shared event or not. If so this is the name of the event to share.
-	 * @return Returns the new event object if successful, NULL otherwise.
+	 * @return A new event, or nullptr none could be created.
+	 * @see GetSynchEventFromPool, ReturnSynchEventToPool
 	 */
-	static class FEvent* CreateSynchEvent(bool bIsManualReset = 0);
+	static class FEvent* CreateSynchEvent(bool bIsManualReset = false);
+
+	/**
+	 * Gets an event from the pool or creates a new one if necessary.
+	 *
+	 * @param bIsManualReset Whether the event requires manual reseting or not.
+	 * @return An event, or nullptr none could be created.
+	 * @see CreateSynchEvent, ReturnSynchEventToPool
+	 */
+	static class FEvent* GetSynchEventFromPool(bool bIsManualReset = false);
+
+	/**
+	 * Returns an event to the pool.
+	 *
+	 * @param Event The event to return.
+	 * @see CreateSynchEvent, GetSynchEventFromPool
+	 */
+	static void ReturnSynchEventToPool(FEvent* Event);
 
 	/**
 	 * Creates the platform-specific runnable thread. This should only be called from FRunnableThread::Create.
@@ -510,7 +527,7 @@ struct CORE_API FGenericPlatformProcess
 
 
 #if PLATFORM_USE_PTHREADS
-	// if we are using pthreads, we can set this up at the generic layer
+	// if we are using PThreads, we can set this up at the generic layer
 	#include "PThreadCriticalSection.h"
 	typedef FPThreadsCriticalSection FCriticalSection;
 #endif
diff --git a/Engine/Source/Runtime/Core/Public/HAL/PThreadEvent.h b/Engine/Source/Runtime/Core/Public/HAL/PThreadEvent.h
--- a/Engine/Source/Runtime/Core/Public/HAL/PThreadEvent.h
+++ b/Engine/Source/Runtime/Core/Public/HAL/PThreadEvent.h
@@ -119,6 +119,11 @@ public:
 		return RetVal;
 	}
 
+	virtual bool IsManualReset() override
+	{
+		return bIsManualReset;
+	}
+
 	virtual void Trigger() override
 	{
 		check(bInitialized);
diff --git a/Engine/Source/Runtime/Core/Public/HAL/ThreadingBase.h b/Engine/Source/Runtime/Core/Public/HAL/ThreadingBase.h
--- a/Engine/Source/Runtime/Core/Public/HAL/ThreadingBase.h
+++ b/Engine/Source/Runtime/Core/Public/HAL/ThreadingBase.h
@@ -44,10 +44,26 @@ public:
 	 */
 	virtual bool Create( bool bIsManualReset = false ) = 0;
 
-	/** Triggers the event so any waiting threads are activated */
+	/**
+	 * Whether the signaled state of this event needs to be reset manually.
+	 *
+	 * @return true if the state requires manual resetting, false otherwise.
+	 * @see Reset
+	 */
+	virtual bool IsManualReset() = 0;
+
+	/**
+	 * Triggers the event so any waiting threads are activated.
+	 *
+	 * @see IsManualReset, Reset
+	 */
 	virtual void Trigger() = 0;
 
-	/** Resets the event to an untriggered (waitable) state */
+	/**
+	 * Resets the event to an untriggered (waitable) state.
+	 *
+	 * @see IsManualReset, Trigger
+	 */
 	virtual void Reset() = 0;
 
 	/**
@@ -100,22 +116,15 @@ public:
  *		// MyEvent destructor is here, we wait here.
  * }
  */
-class CORE_API FScopedEvent
+class FScopedEvent
 {
 public:
 
 	/** Default constructor. */
-	FScopedEvent()
-		: Event(GetEventFromPool())
-	{ }
+	CORE_API FScopedEvent();
 
 	/** Destructor. */
-	~FScopedEvent()
-	{
-		Event->Wait();
-		ReturnToPool(Event);
-		Event = nullptr;
-	}
+	CORE_API ~FScopedEvent();
 
 	/** Triggers the event. */
 	void Trigger()
@@ -133,27 +142,9 @@ public:
 		return Event;
 	}
 
-protected:
-
-	/**
-	 * Returns an event object from the pool.
-	 *
-	 * @return An event object.
-	 * @see ReturnToPool
-	 */
-	static FEvent* GetEventFromPool();
-
-	/**
-	 * Returns an event object to the pool.
-	 *
-	 * @param Event The event object to return.
-	 * @see GetEventFromPool
-	 */
-	static void ReturnToPool( FEvent* Event );
-
 private:
 
-	// Holds the event.
+	/** Holds the event. */
 	FEvent* Event;
 };
 
@@ -446,23 +437,28 @@ public:
 
 	// FEvent Interface
 
-	virtual bool Create( bool bIsManualReset = false ) 
+	virtual bool Create( bool bIsManualReset = false ) override
 	{ 
 		bManualReset = bIsManualReset;
 		return true; 
 	}
 
-	virtual void Trigger()
+	virtual bool IsManualReset() override
+	{
+		return bManualReset;
+	}
+
+	virtual void Trigger() override
 	{
 		bTriggered = true;
 	}
 
-	virtual void Reset()
+	virtual void Reset() override
 	{
 		bTriggered = false;
 	}
 
-	virtual bool Wait( uint32 WaitTime, const bool bIgnoreThreadIdleStats = false ) 
+	virtual bool Wait( uint32 WaitTime, const bool bIgnoreThreadIdleStats = false ) override
 	{ 
 		// With only one thread it's assumed the event has been triggered
 		// before Wait is called, otherwise it would end up waiting forever or always fail.
