Async: using event pooling in TFuture for better performance

diff --git a/Engine/Source/Runtime/Core/Public/Async/Future.h b/Engine/Source/Runtime/Core/Public/Async/Future.h
--- a/Engine/Source/Runtime/Core/Public/Async/Future.h
+++ b/Engine/Source/Runtime/Core/Public/Async/Future.h
@@ -15,14 +15,15 @@ public:
 
 	/** Default constructor. */
 	FFutureState()
-		: CompletionEvent(FPlatformProcess::CreateSynchEvent(true))
+		: CompletionEvent(FPlatformProcess::GetSynchEventFromPool(true))
 		, Complete(false)
 	{ }
 
 	/** Destructor. */
 	~FFutureState()
 	{
-		delete CompletionEvent;
+		FPlatformProcess::ReturnSynchEventToPool(CompletionEvent);
+		CompletionEvent = nullptr;
 	}
 
 public:
