forgot to check in

diff --git a/Engine/Source/Runtime/Core/Private/Misc/EventPool.h b/Engine/Source/Runtime/Core/Private/Misc/EventPool.h
new file mode 100644
--- /dev/null
+++ b/Engine/Source/Runtime/Core/Private/Misc/EventPool.h
@@ -0,0 +1,88 @@
+// Copyright 1998-2015 Epic Games, Inc. All Rights Reserved.
+
+#pragma once
+
+
+/**
+ * Enumerates available event pool types.
+ */
+enum class EEventPoolTypes
+{
+	/** Creates events that have their signaled state reset automatically. */
+	AutoReset,
+
+	/** Creates events that have their signaled state reset manually. */
+	ManualReset
+};
+
+
+/**
+ * Template class for event pools.
+ *
+ * Events are expensive to create on most platforms. This pool allows for efficient
+ * recycling of event instances that are no longer used. Events can have their signaled
+ * state reset automatically or manually. The PoolType template parameter specifies
+ * which type of events the pool managers.
+ *
+ * @param PoolType Specifies the type of pool.
+ * @see FEvent
+ */
+template<EEventPoolTypes PoolType>
+class FEventPool
+{
+public:
+
+	/**
+	 * Gets the singleton instance of the event pool.
+	 *
+	 * @return Pool singleton.
+	 */
+	static FEventPool& Get()
+	{
+		static FEventPool Singleton;
+		return Singleton;
+	}
+
+	/**
+	 * Gets an event from the pool or creates one if necessary.
+	 *
+	 * @return The event.
+	 * @see ReturnToPool
+	 */
+	FEvent* GetEventFromPool()
+	{
+		FEvent* Result = Pool.Pop();
+
+		if (!Result)
+		{
+			Result = FPlatformProcess::CreateSynchEvent((PoolType == EEventPoolTypes::ManualReset));
+		}
+
+		check(Result);
+		return Result;
+	}
+
+	/**
+	 * Returns an event to the pool.
+	 *
+	 * @param Event The event to return.
+	 * @see GetEventFromPool
+	 */
+	void ReturnToPool(FEvent* Event)
+	{
+		check(Event);
+		check(Event->IsManualReset() == (PoolType == EEventPoolTypes::ManualReset));
+
+		if (PoolType == EEventPoolTypes::ManualReset)
+		{
+			Event->Reset();
+		}
+
+		Pool.Push(Event);
+	}
+
+private:
+
+	/** Holds the collection of recycled events. */
+	TLockFreePointerList<FEvent> Pool;
+};
