Core - More stats for memory profiling

diff --git a/Engine/Source/Runtime/Core/Private/Internationalization/TextLocalizationManager.cpp b/Engine/Source/Runtime/Core/Private/Internationalization/TextLocalizationManager.cpp
--- a/Engine/Source/Runtime/Core/Private/Internationalization/TextLocalizationManager.cpp
+++ b/Engine/Source/Runtime/Core/Private/Internationalization/TextLocalizationManager.cpp
@@ -21,6 +21,8 @@ void BeginInitTextLocalization()
 
 void EndInitTextLocalization()
 {
+	DECLARE_SCOPE_CYCLE_COUNTER(TEXT("EndInitTextLocalization"), STAT_EndInitTextLocalization, STATGROUP_LoadTime);
+
 	const bool ShouldLoadEditor = WITH_EDITOR;
 	const bool ShouldLoadGame = FApp::IsGame();
 
diff --git a/Engine/Source/Runtime/CoreUObject/Private/UObject/UObjectBase.cpp b/Engine/Source/Runtime/CoreUObject/Private/UObject/UObjectBase.cpp
--- a/Engine/Source/Runtime/CoreUObject/Private/UObject/UObjectBase.cpp
+++ b/Engine/Source/Runtime/CoreUObject/Private/UObject/UObjectBase.cpp
@@ -805,6 +805,8 @@ bool AnyNewlyLoadedUObjects()
 
 void ProcessNewlyLoadedUObjects()
 {
+	DECLARE_SCOPE_CYCLE_COUNTER(TEXT("ProcessNewlyLoadedUObjects"), STAT_ProcessNewlyLoadedUObjects, STATGROUP_ObjectVerbose);
+
 	UClassRegisterAllCompiledInClasses();
 
 	while( AnyNewlyLoadedUObjects() )
diff --git a/Engine/Source/Runtime/Engine/Private/GlobalShader.cpp b/Engine/Source/Runtime/Engine/Private/GlobalShader.cpp
--- a/Engine/Source/Runtime/Engine/Private/GlobalShader.cpp
+++ b/Engine/Source/Runtime/Engine/Private/GlobalShader.cpp
@@ -325,6 +325,8 @@ void SaveGlobalShaderMapToDerivedDataCache(EShaderPlatform Platform)
 
 TShaderMap<FGlobalShaderType>* GetGlobalShaderMap(EShaderPlatform Platform, bool bRefreshShaderMap)
 {
+	DECLARE_SCOPE_CYCLE_COUNTER(TEXT("GetGlobalShaderMap"), STAT_GetGlobalShaderMap, STATGROUP_LoadTime);
+
 	// No global shaders needed on dedicated server
 	if (FPlatformProperties::IsServerOnly())
 	{
