Changed the order of execution in FRunnableThreadPThread::Run() to avoid situation where thread destructor is proceeding with deleting stuff while Run() is still deallocating things in FThreadStats::Shutdown();

diff --git a/Engine/Source/Runtime/Core/Private/HAL/PThreadRunnableThread.cpp b/Engine/Source/Runtime/Core/Private/HAL/PThreadRunnableThread.cpp
--- a/Engine/Source/Runtime/Core/Private/HAL/PThreadRunnableThread.cpp
+++ b/Engine/Source/Runtime/Core/Private/HAL/PThreadRunnableThread.cpp
@@ -29,13 +29,13 @@ uint32 FRunnableThreadPThread::Run()
 		ThreadInitSyncEvent->Trigger();
 	}
 
-	// Clean ourselves up without waiting
-	ThreadIsRunning = false;
-
 #if STATS
-//	FThreadStats::Shutdown();
+	FThreadStats::Shutdown();
 #endif
 
+	// Clean ourselves up without waiting
+	ThreadIsRunning = false;
+
 	return ExitCode;
 }
 
