Enabled crash reporting (Back out revision 5 from //depot/UE4/Engine/Source/Runtime/Core/Private/Windows/WindowsPlatformCrashContext.cpp)

diff --git a/Engine/Source/Runtime/Core/Private/Windows/WindowsPlatformCrashContext.cpp b/Engine/Source/Runtime/Core/Private/Windows/WindowsPlatformCrashContext.cpp
--- a/Engine/Source/Runtime/Core/Private/Windows/WindowsPlatformCrashContext.cpp
+++ b/Engine/Source/Runtime/Core/Private/Windows/WindowsPlatformCrashContext.cpp
@@ -125,71 +125,71 @@ void GetModuleVersion( TCHAR* ModuleName, TCHAR* StringBuffer, DWORD MaxSize )
  */
 void SetReportParameters( HREPORT ReportHandle, EXCEPTION_POINTERS* ExceptionInfo, const TCHAR* ErrorMessage )
 {
-	//HRESULT Result;
-	//TCHAR StringBuffer[MAX_SPRINTF] = {0};
-	//TCHAR LocalBuffer[MAX_SPRINTF] = {0};
-
-	//// Set the parameters for the standard problem signature
-	//HMODULE ModuleHandle = GetModuleHandle( NULL );
-
-	//StringCchPrintf( StringBuffer, MAX_SPRINTF, TEXT( "UE4-%s" ), FApp::GetGameName() );
-	//Result = WerReportSetParameter( ReportHandle, WER_P0, TEXT( "Application Name" ), StringBuffer );
-
-	//GetModuleFileName( ModuleHandle, LocalBuffer, MAX_SPRINTF );
-	//PathStripPath( LocalBuffer );
-	//GetModuleVersion( LocalBuffer, StringBuffer, MAX_SPRINTF );
-	//Result = WerReportSetParameter( ReportHandle, WER_P1, TEXT( "Application Version" ), StringBuffer );
-
-	//StringCchPrintf( StringBuffer, MAX_SPRINTF, TEXT( "%08x" ), GetTimestampForLoadedLibrary( ModuleHandle ) );
-	//Result = WerReportSetParameter( ReportHandle, WER_P2, TEXT( "Application Timestamp" ), StringBuffer );
-
-	//HMODULE FaultModuleHandle = NULL;
-	//GetModuleHandleEx( GET_MODULE_HANDLE_EX_FLAG_FROM_ADDRESS | GET_MODULE_HANDLE_EX_FLAG_UNCHANGED_REFCOUNT, ( LPCTSTR )ExceptionInfo->ExceptionRecord->ExceptionAddress, &FaultModuleHandle );
-
-	//GetModuleFileName( FaultModuleHandle, LocalBuffer, MAX_SPRINTF );
-	//PathStripPath( LocalBuffer );
-	//Result = WerReportSetParameter( ReportHandle, WER_P3, TEXT( "Fault Module Name" ), LocalBuffer );
-
-	//GetModuleVersion( LocalBuffer, StringBuffer, MAX_SPRINTF );
-	//Result = WerReportSetParameter( ReportHandle, WER_P4, TEXT( "Fault Module Version" ), StringBuffer );
-
-	//StringCchPrintf( StringBuffer, MAX_SPRINTF, TEXT( "%08x" ), GetTimestampForLoadedLibrary( FaultModuleHandle ) );
-	//Result = WerReportSetParameter( ReportHandle, WER_P5, TEXT( "Fault Module Timestamp" ), StringBuffer );
-
-	//StringCchPrintf( StringBuffer, MAX_SPRINTF, TEXT( "%08x" ), ExceptionInfo->ExceptionRecord->ExceptionCode );
-	//Result = WerReportSetParameter( ReportHandle, WER_P6, TEXT( "Exception Code" ), StringBuffer );
-
-	//INT_PTR ExceptionOffset = ( char* )( ExceptionInfo->ExceptionRecord->ExceptionAddress ) - ( char* )FaultModuleHandle;
-	//CA_SUPPRESS(6066) // The format specifier should probably be something like %tX, but VS 2013 doesn't support 't'.
-	//	StringCchPrintf( StringBuffer, MAX_SPRINTF, TEXT( "%p" ), ExceptionOffset );
-	//Result = WerReportSetParameter( ReportHandle, WER_P7, TEXT( "Exception Offset" ), StringBuffer );
-
-	//// Use LocalBuffer to store the error message.
-	//FCString::Strncpy( LocalBuffer, ErrorMessage, MAX_SPRINTF );
-
-	//// Replace " with ' and replace \n with #
-	//for (TCHAR& Char: LocalBuffer)
-	//{
-	//	if (Char == 0)
-	//	{
-	//		break;
-	//	}
-
-	//	switch (Char)
-	//	{
-	//		default: break;
-	//		case '"':	Char = '\'';	break;
-	//		case '\r':	Char = '#';		break;
-	//		case '\n':	Char = '#';		break;
-	//	}
-	//}
-
-	//// AssertLog should be ErrorMessage, but this require crash server changes, so don't change this.
-	//StringCchPrintf( StringBuffer, MAX_SPRINTF, TEXT( "!%s!AssertLog=\"%s\"" ), FCommandLine::Get(), LocalBuffer );
-	//Result = WerReportSetParameter( ReportHandle, WER_P8, TEXT( "Commandline" ), StringBuffer );
-
-	//StringCchPrintf( StringBuffer, MAX_SPRINTF, TEXT( "%s!%s!%s!%d" ), TEXT( BRANCH_NAME ), FPlatformProcess::BaseDir(), FPlatformMisc::GetEngineMode(), BUILT_FROM_CHANGELIST );
-	//Result = WerReportSetParameter( ReportHandle, WER_P9, TEXT( "BranchBaseDir" ), StringBuffer );
+	HRESULT Result;
+	TCHAR StringBuffer[MAX_SPRINTF] = {0};
+	TCHAR LocalBuffer[MAX_SPRINTF] = {0};
+
+	// Set the parameters for the standard problem signature
+	HMODULE ModuleHandle = GetModuleHandle( NULL );
+
+	StringCchPrintf( StringBuffer, MAX_SPRINTF, TEXT( "UE4-%s" ), FApp::GetGameName() );
+	Result = WerReportSetParameter( ReportHandle, WER_P0, TEXT( "Application Name" ), StringBuffer );
+
+	GetModuleFileName( ModuleHandle, LocalBuffer, MAX_SPRINTF );
+	PathStripPath( LocalBuffer );
+	GetModuleVersion( LocalBuffer, StringBuffer, MAX_SPRINTF );
+	Result = WerReportSetParameter( ReportHandle, WER_P1, TEXT( "Application Version" ), StringBuffer );
+
+	StringCchPrintf( StringBuffer, MAX_SPRINTF, TEXT( "%08x" ), GetTimestampForLoadedLibrary( ModuleHandle ) );
+	Result = WerReportSetParameter( ReportHandle, WER_P2, TEXT( "Application Timestamp" ), StringBuffer );
+
+	HMODULE FaultModuleHandle = NULL;
+	GetModuleHandleEx( GET_MODULE_HANDLE_EX_FLAG_FROM_ADDRESS | GET_MODULE_HANDLE_EX_FLAG_UNCHANGED_REFCOUNT, ( LPCTSTR )ExceptionInfo->ExceptionRecord->ExceptionAddress, &FaultModuleHandle );
+
+	GetModuleFileName( FaultModuleHandle, LocalBuffer, MAX_SPRINTF );
+	PathStripPath( LocalBuffer );
+	Result = WerReportSetParameter( ReportHandle, WER_P3, TEXT( "Fault Module Name" ), LocalBuffer );
+
+	GetModuleVersion( LocalBuffer, StringBuffer, MAX_SPRINTF );
+	Result = WerReportSetParameter( ReportHandle, WER_P4, TEXT( "Fault Module Version" ), StringBuffer );
+
+	StringCchPrintf( StringBuffer, MAX_SPRINTF, TEXT( "%08x" ), GetTimestampForLoadedLibrary( FaultModuleHandle ) );
+	Result = WerReportSetParameter( ReportHandle, WER_P5, TEXT( "Fault Module Timestamp" ), StringBuffer );
+
+	StringCchPrintf( StringBuffer, MAX_SPRINTF, TEXT( "%08x" ), ExceptionInfo->ExceptionRecord->ExceptionCode );
+	Result = WerReportSetParameter( ReportHandle, WER_P6, TEXT( "Exception Code" ), StringBuffer );
+
+	INT_PTR ExceptionOffset = ( char* )( ExceptionInfo->ExceptionRecord->ExceptionAddress ) - ( char* )FaultModuleHandle;
+	CA_SUPPRESS(6066) // The format specifier should probably be something like %tX, but VS 2013 doesn't support 't'.
+		StringCchPrintf( StringBuffer, MAX_SPRINTF, TEXT( "%p" ), ExceptionOffset );
+	Result = WerReportSetParameter( ReportHandle, WER_P7, TEXT( "Exception Offset" ), StringBuffer );
+
+	// Use LocalBuffer to store the error message.
+	FCString::Strncpy( LocalBuffer, ErrorMessage, MAX_SPRINTF );
+
+	// Replace " with ' and replace \n with #
+	for (TCHAR& Char: LocalBuffer)
+	{
+		if (Char == 0)
+		{
+			break;
+		}
+
+		switch (Char)
+		{
+			default: break;
+			case '"':	Char = '\'';	break;
+			case '\r':	Char = '#';		break;
+			case '\n':	Char = '#';		break;
+		}
+	}
+
+	// AssertLog should be ErrorMessage, but this require crash server changes, so don't change this.
+	StringCchPrintf( StringBuffer, MAX_SPRINTF, TEXT( "!%s!AssertLog=\"%s\"" ), FCommandLine::Get(), LocalBuffer );
+	Result = WerReportSetParameter( ReportHandle, WER_P8, TEXT( "Commandline" ), StringBuffer );
+
+	StringCchPrintf( StringBuffer, MAX_SPRINTF, TEXT( "%s!%s!%s!%d" ), TEXT( BRANCH_NAME ), FPlatformProcess::BaseDir(), FPlatformMisc::GetEngineMode(), BUILT_FROM_CHANGELIST );
+	Result = WerReportSetParameter( ReportHandle, WER_P9, TEXT( "BranchBaseDir" ), StringBuffer );
 }
 
 /**
@@ -213,11 +213,11 @@ void AddMiniDump(HREPORT ReportHandle, EXCEPTION_POINTERS* ExceptionInfo)
  */
 void AddMiscFiles( HREPORT ReportHandle )
 {
-	/*FString LogFileName = FPaths::GameLogDir() / FApp::GetGameName() + TEXT( ".log" );
+	FString LogFileName = FPaths::GameLogDir() / FApp::GetGameName() + TEXT( ".log" );
 	WerReportAddFile( ReportHandle, *LogFileName, WerFileTypeOther, WER_FILE_ANONYMOUS_DATA ); 
 
 	FString CrashVideoPath = FPaths::GameLogDir() / TEXT( "CrashVideo.avi" );
-	WerReportAddFile( ReportHandle, *CrashVideoPath, WerFileTypeOther, WER_FILE_ANONYMOUS_DATA );*/
+	WerReportAddFile( ReportHandle, *CrashVideoPath, WerFileTypeOther, WER_FILE_ANONYMOUS_DATA );
 }
 
 /**
@@ -280,95 +280,95 @@ private:
  */
 int32 ReportCrashUsingCrashReportClient(EXCEPTION_POINTERS* ExceptionInfo, const TCHAR* ErrorMessage, EErrorReportUI ReportUI)
 {
-	//// Flush out the log
-	//GLog->Flush();
-
-	//// Prevent CrashReportClient from spawning another CrashReportClient.
-	//const TCHAR* ExecutableName = FPlatformProcess::ExecutableName();
-	//const bool bCanRunCrashReportClient = FCString::Stristr( ExecutableName, TEXT( "CrashReportClient" ) ) == nullptr;
-	//if( bCanRunCrashReportClient )
-	//{
-	//	// Set the report to force queue
-	//	FScopedWERQueuing ScopedQueueForcer;
-
-	//	// Construct the report details
-	//	WER_REPORT_INFORMATION ReportInformation = {sizeof( WER_REPORT_INFORMATION )};
-
-	//	StringCchCopy( ReportInformation.wzConsentKey, ARRAYSIZE( ReportInformation.wzConsentKey ), TEXT( "" ) );
-
-	//	StringCchCopy( ReportInformation.wzApplicationName, ARRAYSIZE( ReportInformation.wzApplicationName ), TEXT("UE4-") );
-	//	StringCchCat( ReportInformation.wzApplicationName, ARRAYSIZE( ReportInformation.wzApplicationName ), FApp::GetGameName() );
-
-	//	StringCchCopy( ReportInformation.wzApplicationPath, ARRAYSIZE( ReportInformation.wzApplicationPath ), FPlatformProcess::BaseDir() );
-	//	StringCchCat( ReportInformation.wzApplicationPath, ARRAYSIZE( ReportInformation.wzApplicationPath ), FPlatformProcess::ExecutableName() );
-	//	StringCchCat( ReportInformation.wzApplicationPath, ARRAYSIZE( ReportInformation.wzApplicationPath ), TEXT( ".exe" ) );
-
-	//	GetCrashDescription( ReportInformation );
-
-	//	// Create a crash event report
-	//	HREPORT ReportHandle = NULL;
-	//	if( WerReportCreate( APPCRASH_EVENT, WerReportApplicationCrash, &ReportInformation, &ReportHandle ) == S_OK )
-	//	{
-	//		// Set the standard set of a crash parameters
-	//		SetReportParameters( ReportHandle, ExceptionInfo, ErrorMessage );
-
-	//		// Add a manually generated minidump
-	//		AddMiniDump( ReportHandle, ExceptionInfo );
-
-	//		// Add the log and video
-	//		AddMiscFiles( ReportHandle );
-
-	//		// Submit
-	//		WER_SUBMIT_RESULT SubmitResult;
-	//		WerReportSubmit( ReportHandle, WerConsentAlwaysPrompt, WER_SUBMIT_QUEUE | WER_SUBMIT_BYPASS_DATA_THROTTLING, &SubmitResult );
-
-	//		// Cleanup
-	//		WerReportCloseHandle( ReportHandle );
-	//	}
-
-	//	// Build machines do not upload these automatically since it is not okay to have lingering processes after the build completes.
-	//	if( GIsBuildMachine )
-	//	{
-	//		return EXCEPTION_CONTINUE_EXECUTION;
-	//	}
-
-	//	FString CrashReportClientArguments;
-
-	//	// Suppress the user input dialog if we're running in unattended mode
-	//	bool bNoDialog = FApp::IsUnattended() || ReportUI == EErrorReportUI::ReportInUnattendedMode;
-	//	if( bNoDialog )
-	//	{
-	//		CrashReportClientArguments += TEXT( " -Unattended" );
-	//	}
-
-	//	if( FApp::IsInstalled() )
-	//	{
-	//		// Temporary workaround for CrashReportClient being built in Development, not Shipping (TTP328030). The
-	//		// following ensures that logs are saved to the user directory when UE4 is installed.
-	//		CrashReportClientArguments += TEXT( " -Installed" );
-	//	}
-
-	//	CrashReportClientArguments += FString( TEXT( " -AppName=" ) ) + ReportInformation.wzApplicationName;
-
-	//	static const TCHAR CrashReportClientExeName[] = TEXT( "CrashReportClient.exe" );
-	//	FString CrashClientPath = FString( TEXT( "..\\..\\..\\Engine\\Binaries" ) ) / FPlatformProcess::GetBinariesSubdirectory() / CrashReportClientExeName;
-
-	//	bool bCrashReporterRan = FPlatformProcess::CreateProc( *CrashClientPath, *CrashReportClientArguments, true, false, false, NULL, 0, NULL, NULL ).IsValid();
-
-	//	if( !bCrashReporterRan && !bNoDialog )
-	//	{
-	//		UE_LOG( LogWindows, Log, TEXT( "Could not start %s" ), CrashReportClientExeName );
-	//		FPlatformMemory::DumpStats( *GWarn );
-	//		FText MessageTitle( FText::Format(
-	//			NSLOCTEXT( "MessageDialog", "AppHasCrashed", "The {0} {1} has crashed and will close" ),
-	//			FText::FromString( ReportInformation.wzApplicationName ),
-	//			FText::FromString( FPlatformMisc::GetEngineMode() )
-	//			) );
-	//		FMessageDialog::Open( EAppMsgType::Ok, FText::FromString( GErrorHist ), &MessageTitle );
-	//	}
-	//}
-
-	//// Let the system take back over (return value only used by NewReportEnsure)
+	// Flush out the log
+	GLog->Flush();
+
+	// Prevent CrashReportClient from spawning another CrashReportClient.
+	const TCHAR* ExecutableName = FPlatformProcess::ExecutableName();
+	const bool bCanRunCrashReportClient = FCString::Stristr( ExecutableName, TEXT( "CrashReportClient" ) ) == nullptr;
+	if( bCanRunCrashReportClient )
+	{
+		// Set the report to force queue
+		FScopedWERQueuing ScopedQueueForcer;
+
+		// Construct the report details
+		WER_REPORT_INFORMATION ReportInformation = {sizeof( WER_REPORT_INFORMATION )};
+
+		StringCchCopy( ReportInformation.wzConsentKey, ARRAYSIZE( ReportInformation.wzConsentKey ), TEXT( "" ) );
+
+		StringCchCopy( ReportInformation.wzApplicationName, ARRAYSIZE( ReportInformation.wzApplicationName ), TEXT("UE4-") );
+		StringCchCat( ReportInformation.wzApplicationName, ARRAYSIZE( ReportInformation.wzApplicationName ), FApp::GetGameName() );
+
+		StringCchCopy( ReportInformation.wzApplicationPath, ARRAYSIZE( ReportInformation.wzApplicationPath ), FPlatformProcess::BaseDir() );
+		StringCchCat( ReportInformation.wzApplicationPath, ARRAYSIZE( ReportInformation.wzApplicationPath ), FPlatformProcess::ExecutableName() );
+		StringCchCat( ReportInformation.wzApplicationPath, ARRAYSIZE( ReportInformation.wzApplicationPath ), TEXT( ".exe" ) );
+
+		GetCrashDescription( ReportInformation );
+
+		// Create a crash event report
+		HREPORT ReportHandle = NULL;
+		if( WerReportCreate( APPCRASH_EVENT, WerReportApplicationCrash, &ReportInformation, &ReportHandle ) == S_OK )
+		{
+			// Set the standard set of a crash parameters
+			SetReportParameters( ReportHandle, ExceptionInfo, ErrorMessage );
+
+			// Add a manually generated minidump
+			AddMiniDump( ReportHandle, ExceptionInfo );
+
+			// Add the log and video
+			AddMiscFiles( ReportHandle );
+
+			// Submit
+			WER_SUBMIT_RESULT SubmitResult;
+			WerReportSubmit( ReportHandle, WerConsentAlwaysPrompt, WER_SUBMIT_QUEUE | WER_SUBMIT_BYPASS_DATA_THROTTLING, &SubmitResult );
+
+			// Cleanup
+			WerReportCloseHandle( ReportHandle );
+		}
+
+		// Build machines do not upload these automatically since it is not okay to have lingering processes after the build completes.
+		if( GIsBuildMachine )
+		{
+			return EXCEPTION_CONTINUE_EXECUTION;
+		}
+
+		FString CrashReportClientArguments;
+
+		// Suppress the user input dialog if we're running in unattended mode
+		bool bNoDialog = FApp::IsUnattended() || ReportUI == EErrorReportUI::ReportInUnattendedMode;
+		if( bNoDialog )
+		{
+			CrashReportClientArguments += TEXT( " -Unattended" );
+		}
+
+		if( FApp::IsInstalled() )
+		{
+			// Temporary workaround for CrashReportClient being built in Development, not Shipping (TTP328030). The
+			// following ensures that logs are saved to the user directory when UE4 is installed.
+			CrashReportClientArguments += TEXT( " -Installed" );
+		}
+
+		CrashReportClientArguments += FString( TEXT( " -AppName=" ) ) + ReportInformation.wzApplicationName;
+
+		static const TCHAR CrashReportClientExeName[] = TEXT( "CrashReportClient.exe" );
+		FString CrashClientPath = FString( TEXT( "..\\..\\..\\Engine\\Binaries" ) ) / FPlatformProcess::GetBinariesSubdirectory() / CrashReportClientExeName;
+
+		bool bCrashReporterRan = FPlatformProcess::CreateProc( *CrashClientPath, *CrashReportClientArguments, true, false, false, NULL, 0, NULL, NULL ).IsValid();
+
+		if( !bCrashReporterRan && !bNoDialog )
+		{
+			UE_LOG( LogWindows, Log, TEXT( "Could not start %s" ), CrashReportClientExeName );
+			FPlatformMemory::DumpStats( *GWarn );
+			FText MessageTitle( FText::Format(
+				NSLOCTEXT( "MessageDialog", "AppHasCrashed", "The {0} {1} has crashed and will close" ),
+				FText::FromString( ReportInformation.wzApplicationName ),
+				FText::FromString( FPlatformMisc::GetEngineMode() )
+				) );
+			FMessageDialog::Open( EAppMsgType::Ok, FText::FromString( GErrorHist ), &MessageTitle );
+		}
+	}
+
+	// Let the system take back over (return value only used by NewReportEnsure)
 	return EXCEPTION_CONTINUE_EXECUTION;
 }
 
