Engine - StaticArray now has a Num() function.  Useful when you just want to provide the length in the template arg, and just want to indirectly refer to it later.

diff --git a/Engine/Source/Runtime/Core/Public/Containers/StaticArray.h b/Engine/Source/Runtime/Core/Public/Containers/StaticArray.h
--- a/Engine/Source/Runtime/Core/Public/Containers/StaticArray.h
+++ b/Engine/Source/Runtime/Core/Public/Containers/StaticArray.h
@@ -73,6 +73,7 @@ public:
 		check(Index < NumElements);
 		return *(TElement*)&Elements[Index];
 	}
+
 	const TElement& operator[](uint32 Index) const
 	{
 		check(Index < NumElements);
@@ -91,6 +92,7 @@ public:
 		}
 		return true;
 	}
+
 	friend bool operator!=(const TStaticArray& A,const TStaticArray& B)
 	{
 		for(uint32 ElementIndex = 0;ElementIndex < NumElements;++ElementIndex)
@@ -102,6 +104,9 @@ public:
 		}
 		return false;
 	}
+
+	/** The number of elements in the array. */
+	int32 Num() const { return NumElements; }
 	
 	/** Hash function. */
 	friend uint32 GetTypeHash(const TStaticArray& Array)
