Added support for android to use the nonufs manifest for file timestamps when running cook on the fly non streaming mode.

diff --git a/Engine/Source/Programs/AutomationTool/Scripts/CopyBuildToStagingDirectory.Automation.cs b/Engine/Source/Programs/AutomationTool/Scripts/CopyBuildToStagingDirectory.Automation.cs
--- a/Engine/Source/Programs/AutomationTool/Scripts/CopyBuildToStagingDirectory.Automation.cs
+++ b/Engine/Source/Programs/AutomationTool/Scripts/CopyBuildToStagingDirectory.Automation.cs
@@ -411,12 +411,15 @@ public partial class Project : CommandUtils
 
 	public static void DumpTargetManifest(Dictionary<string, string> Mapping, string Filename, string StageDir, List<string> CRCFiles)
 	{
-		if (Mapping.Count > 0)
+        // const string Iso8601DateTimeFormat = "%Y-%m-%dT%H:%M:%S.%sZ";
+        // const string Iso8601DateTimeFormat = "yyyy'-'MM'-'dd HH':'mm':'ss'Z'";
+        const string Iso8601DateTimeFormat = "u"; // predefined universal Iso standard format
+        if (Mapping.Count > 0)
 		{
 			var Lines = new List<string>();
 			foreach (var Pair in Mapping)
 			{
-				string TimeStamp = File.GetLastWriteTimeUtc(Pair.Key).ToString();
+                string TimeStamp = File.GetLastWriteTimeUtc(Pair.Key).ToString(Iso8601DateTimeFormat);
 				if (CRCFiles.Contains(Pair.Value))
 				{
 					byte[] FileData = File.ReadAllBytes(StageDir + "/" + Pair.Value);
diff --git a/Engine/Source/Programs/AutomationTool/Scripts/RunProjectCommand.Automation.cs b/Engine/Source/Programs/AutomationTool/Scripts/RunProjectCommand.Automation.cs
--- a/Engine/Source/Programs/AutomationTool/Scripts/RunProjectCommand.Automation.cs
+++ b/Engine/Source/Programs/AutomationTool/Scripts/RunProjectCommand.Automation.cs
@@ -732,7 +732,7 @@ public partial class Project : CommandUtils
                 }
 				TempCmdLine += " ";
 
-				if (!Params.Stage)
+				if (!Params.CookOnTheFlyStreaming)
 				{
 					TempCmdLine += "-streaming ";
 				}
diff --git a/Engine/Source/Runtime/Core/Private/Android/AndroidFile.cpp b/Engine/Source/Runtime/Core/Private/Android/AndroidFile.cpp
--- a/Engine/Source/Runtime/Core/Private/Android/AndroidFile.cpp
+++ b/Engine/Source/Runtime/Core/Private/Android/AndroidFile.cpp
@@ -21,7 +21,8 @@
 DEFINE_LOG_CATEGORY_STATIC(LogAndroidFile, Log, All);
 
 #define LOG_ANDROID_FILE 0
-// #define LOG_ANDROID_FILE 1
+
+#define LOG_ANDROID_FILE_MANIFEST 0
 
 // Support 64 bit file access.
 #define UE_ANDROID_FILE_64 0
@@ -30,6 +31,9 @@ DEFINE_LOG_CATEGORY_STATIC(LogAndroidFile, Log, All);
 // make an FTimeSpan object that represents the "epoch" for time_t (from a stat struct)
 const FDateTime AndroidEpoch(1970, 1, 1);
 
+#define USE_UTIME 0
+
+
 // AndroidProcess uses this for executable name
 FString GPackageName;
 static int32 GPackageVersion = 0;
@@ -280,6 +284,174 @@ public:
 	}
 };
 
+
+class FManifestReader
+{
+private:
+	bool bInitialized;
+	FString ManifestFileName;
+	TMap<FString, FDateTime> ManifestEntries;
+public:
+
+	FManifestReader( const FString& InManifestFileName ) : ManifestFileName(InManifestFileName), bInitialized(false)
+	{
+	}
+
+	bool GetFileTimeStamp( const FString& FileName, FDateTime& DateTime ) 
+	{
+		if ( bInitialized == false )
+		{
+			Read();
+			bInitialized = true;
+		}
+
+		const FDateTime* Result = ManifestEntries.Find( FileName );
+		if ( Result )
+		{
+			DateTime = *Result;
+#if LOG_ANDROID_FILE_MANIFEST
+			FPlatformMisc::LowLevelOutputDebugStringf(TEXT("Found time stamp for '%s', %s"), *FileName, *DateTime.ToString());
+#endif
+			return true;
+		}
+#if LOG_ANDROID_FILE_MANIFEST
+		FPlatformMisc::LowLevelOutputDebugStringf(TEXT("Didn't find time stamp for '%s'"), *FileName);
+#endif
+		return false;
+	}
+
+
+	bool SetFileTimeStamp( const FString& FileName, const FDateTime& DateTime )
+	{
+		FDateTime* Result = ManifestEntries.Find( FileName );
+		if ( Result == NULL )
+		{
+			ManifestEntries.Add(FileName, DateTime );
+		}
+		else
+		{
+			*Result = DateTime;
+		}
+#if LOG_ANDROID_FILE_MANIFEST
+		FPlatformMisc::LowLevelOutputDebugStringf(TEXT("SetFileTimeStamp '%s', %s"), *FileName, *DateTime.ToString());
+#endif
+		return true;
+	}
+
+	// read manifest from disk
+	void Read()
+	{
+		// Local filepaths are directly in the deployment directory.
+		static const FString BasePath = GFilePathBase + FString("/") + FApp::GetGameName() + FString("/");
+		const FString ManfiestPath = BasePath + ManifestFileName;
+
+		ManifestEntries.Empty();
+
+		// int Handle = open( TCHAR_TO_UTF8(ManifestFileName), O_RDWR );
+		int Handle = open( TCHAR_TO_UTF8(*ManfiestPath), O_RDONLY );
+
+		if ( Handle == -1 )
+		{
+#if LOG_ANDROID_FILE_MANIFEST
+			FPlatformMisc::LowLevelOutputDebugStringf(TEXT("Failed to open file for read'%s'"), *ManifestFileName);
+#endif
+			return;
+		}
+
+		FString EntireFile;
+		char Buffer[1024];
+		Buffer[1023] = '\0';
+		int BytesRead = 1023;
+		while ( BytesRead == 1023 )
+		{
+			BytesRead = read(Handle, Buffer, 1023);
+			check( Buffer[1023] == '\0');
+			EntireFile.Append(FString(UTF8_TO_TCHAR(Buffer)));
+		}
+
+		close( Handle );
+
+		TArray<FString> Lines;
+		EntireFile.ParseIntoArrayLines(&Lines);
+
+#if LOG_ANDROID_FILE_MANIFEST
+		FPlatformMisc::LowLevelOutputDebugStringf(TEXT("Loaded manifest file %s"), *ManifestFileName);
+		for( const auto& Line : Lines )
+		{
+			FPlatformMisc::LowLevelOutputDebugStringf(TEXT("Manifest Line %s"), *Line );
+		}
+#endif
+		for ( const auto& Line : Lines) 
+		{
+#if LOG_ANDROID_FILE_MANIFEST
+			FPlatformMisc::LowLevelOutputDebugStringf(TEXT("Processing line '%s' "), *Line);
+#endif
+			FString Filename;
+			FString DateTimeString;
+			if ( Line.Split(TEXT("\t"), &Filename, &DateTimeString) )
+			{
+				FDateTime ModifiedDate;
+				if ( FDateTime::ParseIso8601(*DateTimeString, ModifiedDate) )
+				{
+#if LOG_ANDROID_FILE_MANIFEST
+					FPlatformMisc::LowLevelOutputDebugStringf(TEXT("Read time stamp '%s' %s"), *Filename, *ModifiedDate.ToString());
+#endif
+					ManifestEntries.Emplace( MoveTemp(Filename), ModifiedDate );
+				}
+				else
+				{
+#if LOG_ANDROID_FILE_MANIFEST
+					FPlatformMisc::LowLevelOutputDebugStringf(TEXT("Failed to parse date for file '%s' %s"), *Filename, *DateTimeString);
+#endif					
+				}
+			}
+#if LOG_ANDROID_FILE_MANIFEST
+			else
+			{
+				FPlatformMisc::LowLevelOutputDebugStringf(TEXT("Unable to split line '%s'"), *Line);
+			}
+#endif
+		}
+	}
+
+	void Write()
+	{
+		
+		// Local filepaths are directly in the deployment directory.
+		static const FString BasePath = GFilePathBase + FString("/") + FApp::GetGameName() + FString("/");
+		const FString ManfiestPath = BasePath + ManifestFileName;
+
+
+		int Handle = open( TCHAR_TO_UTF8(*ManfiestPath), O_WRONLY );
+
+		if ( Handle == -1 )
+		{
+#if LOG_ANDROID_FILE_MANIFEST
+			FPlatformMisc::LowLevelOutputDebugStringf(TEXT("Failed to open file for write '%s'"), *ManifestFileName);
+#endif
+			return;
+		}
+
+
+		for ( const auto& Line : ManifestEntries )
+		{
+			const int BufferSize = 4096;
+			char Buffer[BufferSize] = {"\0"}; 
+			const FString RawDateTimeString = Line.Value.ToIso8601();
+			const FString DateTimeString = FString::Printf(TEXT("%s\r\n"), *RawDateTimeString);
+			strncpy(Buffer, TCHAR_TO_UTF8(*Line.Key), BufferSize-1);
+			strncat(Buffer, "\t", BufferSize-1 );
+			strncat(Buffer, TCHAR_TO_UTF8(*DateTimeString), BufferSize-1 );
+			write( Handle, Buffer, strlen( Buffer ) );
+		}
+
+		close( Handle );
+	}
+};
+
+FManifestReader NonUFSManifest(TEXT("Manifest_NonUFSFiles.txt"));
+FManifestReader UFSManifest(TEXT("Manifest_UFSFiles.txt"));
+
 /*
 	Access to files in multiple ZIP archives.
 */
@@ -886,7 +1058,7 @@ public:
 
 	FDateTime GetTimeStamp(const TCHAR* Filename, bool AllowLocal)
 	{
-#if LOG_ANDROID_FILE
+#if LOG_ANDROID_FILE_MANIFEST
 		FPlatformMisc::LowLevelOutputDebugStringf(TEXT("FAndroidPlatformFile::GetTimeStamp('%s')"), Filename);
 #endif
 		FString LocalPath;
@@ -895,6 +1067,8 @@ public:
 
 		if (IsLocal(LocalPath))
 		{
+
+#if USE_UTIME
 			struct stat FileInfo;
 			if (stat(TCHAR_TO_UTF8(*LocalPath), &FileInfo) == -1)
 			{
@@ -903,6 +1077,19 @@ public:
 			// convert _stat time to FDateTime
 			FTimespan TimeSinceEpoch(0, 0, FileInfo.st_mtime);
 			return AndroidEpoch + TimeSinceEpoch;
+#else
+			FDateTime Result;
+			if ( NonUFSManifest.GetFileTimeStamp(AssetPath,Result) )
+			{
+				return Result;
+			}
+
+#if LOG_ANDROID_FILE_MANIFEST
+			FPlatformMisc::LowLevelOutputDebugStringf(TEXT("Failed to find time stamp in NonUFSManifest for file '%s'"), Filename);
+#endif
+			return FDateTime::MinValue();
+
+#endif
 		}
 		else if (IsResource(AssetPath))
 		{
@@ -933,6 +1120,7 @@ public:
 		// Can only set time stamp on local files
 		if (IsLocal(LocalPath))
 		{
+#if USE_UTIME
 			// get file times
 			struct stat FileInfo;
 			if (stat(TCHAR_TO_UTF8(*LocalPath), &FileInfo) == -1)
@@ -944,6 +1132,11 @@ public:
 			Times.actime = FileInfo.st_atime;
 			Times.modtime = (DateTime - AndroidEpoch).GetTotalSeconds();
 			utime(TCHAR_TO_UTF8(*LocalPath), &Times);
+#else
+			// do something better as utime isn't supported on android very well...
+			NonUFSManifest.SetFileTimeStamp( AssetPath, DateTime );
+			NonUFSManifest.Write();
+#endif
 		}
 	}
 
diff --git a/Engine/Source/Runtime/Core/Private/Android/AndroidMisc.cpp b/Engine/Source/Runtime/Core/Private/Android/AndroidMisc.cpp
--- a/Engine/Source/Runtime/Core/Private/Android/AndroidMisc.cpp
+++ b/Engine/Source/Runtime/Core/Private/Android/AndroidMisc.cpp
@@ -299,26 +299,55 @@ public:
 		if(FileHandle)
 		{
 			fclose(FileHandle);
-			struct utimbuf Times;
-			Times.actime = 0;
-			Times.modtime = 0;
-			int Result = utime(TestFilePathChar, &Times);
-			Supported = -1 != Result;
-			unlink(TestFilePathChar);
-		}
 
+			// get file times
+			struct stat FileInfo;
+			if (stat(TestFilePathChar, &FileInfo) == -1)
+			{
+				FPlatformMisc::LowLevelOutputDebugStringf(TEXT("Unable to get file stamp for file ('%s')"), TestFilePathChar);
+			}
+			else
+			{
+				struct utimbuf Times;
+			    Times.actime = 0;
+			    Times.modtime = 0;
+				int Result = utime(TestFilePathChar, &Times);
+				Supported = -1 != Result;
+				unlink(TestFilePathChar);
+
+				FPlatformMisc::LowLevelOutputDebugStringf(TEXT("UTime failed for local caching supported test, with error code %d\n"), Result);
+			}
+			
+		}
+		else
+		{
+			FPlatformMisc::LowLevelOutputDebugStringf(TEXT("Failed to create file for Local cache file test\n"), Supported);
+		}
 		FPlatformMisc::LowLevelOutputDebugStringf(TEXT("Is Local Caching Supported? %d\n"), Supported);
 	}
 
 	bool Supported;
 };
 
-bool FAndroidMisc::SupportsLocalCaching()
+bool SupportsUTime()
 {
 	static FTestUtime Test;
 	return Test.Supported;
 }
 
+
+bool FAndroidMisc::SupportsLocalCaching()
+{
+	return true;
+
+	/*if ( SupportsUTime() )
+	{
+		return true;
+	}*/
+
+
+}
+
 /**
  * Good enough default crash reporter.
  */
