#perf - Convert more uses of individual Sin/Cos calls to use new faster SinCos(). Cleaned up some hardcoded degree to radian conversion to use the function template instead.

diff --git a/Engine/Source/Runtime/Core/Private/Math/UnrealMath.cpp b/Engine/Source/Runtime/Core/Private/Math/UnrealMath.cpp
--- a/Engine/Source/Runtime/Core/Private/Math/UnrealMath.cpp
+++ b/Engine/Source/Runtime/Core/Private/Math/UnrealMath.cpp
@@ -771,8 +771,8 @@ FQuat FQuat::FindBetween(const FVector& vec1, const FVector& vec2)
 		angle = PI - angle;
 	}
 
-	const float sinHalfAng = FMath::Sin(0.5f * angle);
-	const float cosHalfAng = FMath::Cos(0.5f * angle);
+	float sinHalfAng, cosHalfAng;
+	FMath::SinCos(&sinHalfAng, &cosHalfAng, 0.5f * angle);
 	const FVector axis = cross / crossMag;
 
 	return FQuat(
diff --git a/Engine/Source/Runtime/Core/Public/Math/Quat.h b/Engine/Source/Runtime/Core/Public/Math/Quat.h
--- a/Engine/Source/Runtime/Core/Public/Math/Quat.h
+++ b/Engine/Source/Runtime/Core/Public/Math/Quat.h
@@ -622,8 +622,8 @@ FORCEINLINE FQuat& FQuat::operator=( const FQuat& Other )
 FORCEINLINE FQuat::FQuat( FVector Axis, float AngleRad )
 {
 	const float half_a = 0.5f * AngleRad;
-	const float s = FMath::Sin(half_a);
-	const float c = FMath::Cos(half_a);
+	float s, c;
+	FMath::SinCos(&s, &c, half_a);
 
 	X = s * Axis.X;
 	Y = s * Axis.Y;
diff --git a/Engine/Source/Runtime/Core/Public/Math/RotationTranslationMatrix.h b/Engine/Source/Runtime/Core/Public/Math/RotationTranslationMatrix.h
--- a/Engine/Source/Runtime/Core/Public/Math/RotationTranslationMatrix.h
+++ b/Engine/Source/Runtime/Core/Public/Math/RotationTranslationMatrix.h
@@ -44,13 +44,11 @@ FORCEINLINE FRotationTranslationMatrix::FRotationTranslationMatrix(const FRotato
 
 #else
 	
-	static const float DEG_TO_RAD = PI/(180.f);
-
 	float SP, SY, SR;
 	float CP, CY, CR;
-	FMath::SinCos(&SP, &CP, Rot.Pitch * DEG_TO_RAD);
-	FMath::SinCos(&SY, &CY, Rot.Yaw * DEG_TO_RAD);
-	FMath::SinCos(&SR, &CR, Rot.Roll * DEG_TO_RAD);
+	FMath::SinCos(&SP, &CP, FMath::DegreesToRadians(Rot.Pitch));
+	FMath::SinCos(&SY, &CY, FMath::DegreesToRadians(Rot.Yaw));
+	FMath::SinCos(&SR, &CR, FMath::DegreesToRadians(Rot.Roll));
 
 #endif
 
diff --git a/Engine/Source/Runtime/Core/Public/Math/ScaleRotationTranslationMatrix.h b/Engine/Source/Runtime/Core/Public/Math/ScaleRotationTranslationMatrix.h
--- a/Engine/Source/Runtime/Core/Public/Math/ScaleRotationTranslationMatrix.h
+++ b/Engine/Source/Runtime/Core/Public/Math/ScaleRotationTranslationMatrix.h
@@ -22,12 +22,11 @@ public:
 
 FORCEINLINE FScaleRotationTranslationMatrix::FScaleRotationTranslationMatrix(const FVector& Scale, const FRotator& Rot, const FVector& Origin)
 {
-	const float	SR	= FMath::Sin(Rot.Roll * PI / 180.f);
-	const float	SP	= FMath::Sin(Rot.Pitch * PI / 180.f);
-	const float	SY	= FMath::Sin(Rot.Yaw * PI / 180.f);
-	const float	CR	= FMath::Cos(Rot.Roll * PI / 180.f);
-	const float	CP	= FMath::Cos(Rot.Pitch * PI / 180.f);
-	const float	CY	= FMath::Cos(Rot.Yaw * PI / 180.f);
+	float SP, SY, SR;
+	float CP, CY, CR;
+	FMath::SinCos(&SP, &CP, FMath::DegreesToRadians(Rot.Pitch));
+	FMath::SinCos(&SY, &CY, FMath::DegreesToRadians(Rot.Yaw));
+	FMath::SinCos(&SR, &CR, FMath::DegreesToRadians(Rot.Roll));
 
 	M[0][0]	= (CP * CY) * Scale.X;
 	M[0][1]	= (CP * SY) * Scale.X;
diff --git a/Engine/Source/Runtime/Core/Public/Math/Vector.h b/Engine/Source/Runtime/Core/Public/Math/Vector.h
--- a/Engine/Source/Runtime/Core/Public/Math/Vector.h
+++ b/Engine/Source/Runtime/Core/Public/Math/Vector.h
@@ -1023,8 +1023,8 @@ FORCEINLINE FVector::FVector( const FVector2D V, float InZ )
 
 inline FVector FVector::RotateAngleAxis( const float AngleDeg, const FVector& Axis ) const
 {
-	const float S	= FMath::Sin(AngleDeg * PI / 180.f);
-	const float C	= FMath::Cos(AngleDeg * PI / 180.f);
+	float S, C;
+	FMath::SinCos(&S, &C, FMath::DegreesToRadians(AngleDeg));
 
 	const float XX	= Axis.X * Axis.X;
 	const float YY	= Axis.Y * Axis.Y;
