Add IsInside location test for sphere

diff --git a/Engine/Source/Runtime/Core/Public/Math/Sphere.h b/Engine/Source/Runtime/Core/Public/Math/Sphere.h
--- a/Engine/Source/Runtime/Core/Public/Math/Sphere.h
+++ b/Engine/Source/Runtime/Core/Public/Math/Sphere.h
@@ -91,6 +91,18 @@ public:
 		return (Center - Other.Center).SizeSquared() <= FMath::Square(Other.W - Tolerance - W);
 	}
 
+
+	/**
+	* Checks whether the given location is inside this sphere.
+	*
+	* @param In The location to test for inside the bounding volume.
+	* @return true if location is inside this volume.
+	*/
+	bool IsInside(const FVector& In, float Tolerance = KINDA_SMALL_NUMBER) const
+	{
+		return (Center - In).SizeSquared() <= FMath::Square(W + Tolerance);
+	}
+
 	
 	/**
 	 * Test whether this sphere intersects another.
