Core - Misc malloc binned stats fixes

diff --git a/Engine/Source/Runtime/Core/Public/HAL/MallocBinned.h b/Engine/Source/Runtime/Core/Public/HAL/MallocBinned.h
--- a/Engine/Source/Runtime/Core/Public/HAL/MallocBinned.h
+++ b/Engine/Source/Runtime/Core/Public/HAL/MallocBinned.h
@@ -37,8 +37,8 @@ DECLARE_MEMORY_STAT_EXTERN(TEXT("Binned Waste Current"),	STAT_Binned_WasteCurren
 DECLARE_MEMORY_STAT_EXTERN(TEXT("Binned Waste Peak"),		STAT_Binned_WastePeak,STATGROUP_MemoryAllocator, CORE_API);
 DECLARE_MEMORY_STAT_EXTERN(TEXT("Binned Used Current"),		STAT_Binned_UsedCurrent,STATGROUP_MemoryAllocator, CORE_API);
 DECLARE_MEMORY_STAT_EXTERN(TEXT("Binned Used Peak"),		STAT_Binned_UsedPeak,STATGROUP_MemoryAllocator, CORE_API);
-DECLARE_MEMORY_STAT_EXTERN(TEXT("Binned Current Allocs"),	STAT_Binned_CurrentAllocs,STATGROUP_MemoryAllocator, CORE_API);
-DECLARE_MEMORY_STAT_EXTERN(TEXT("Binned Total Allocs"),		STAT_Binned_TotalAllocs,STATGROUP_MemoryAllocator, CORE_API);
+DECLARE_DWORD_COUNTER_STAT_EXTERN(TEXT("Binned Current Allocs"),	STAT_Binned_CurrentAllocs,STATGROUP_MemoryAllocator, CORE_API);
+DECLARE_DWORD_COUNTER_STAT_EXTERN(TEXT("Binned Total Allocs"),		STAT_Binned_TotalAllocs,STATGROUP_MemoryAllocator, CORE_API);
 DECLARE_MEMORY_STAT_EXTERN(TEXT("Binned Slack Current"),	STAT_Binned_SlackCurrent,STATGROUP_MemoryAllocator, CORE_API);
 
 
@@ -1131,6 +1131,7 @@ public:
 	/** Called once per frame, gathers and sets all memory allocator statistics into the corresponding stats. */
 	virtual void UpdateStats() override
 	{
+		FMalloc::UpdateStats();
 #if STATS
 		SIZE_T	LocalOsCurrent = 0;
 		SIZE_T	LocalOsPeak = 0;
@@ -1166,8 +1167,8 @@ public:
 		SET_MEMORY_STAT( STAT_Binned_WastePeak, LocalWastePeak );
 		SET_MEMORY_STAT( STAT_Binned_UsedCurrent, LocalUsedCurrent );
 		SET_MEMORY_STAT( STAT_Binned_UsedPeak, LocalUsedPeak );
-		SET_MEMORY_STAT( STAT_Binned_CurrentAllocs, LocalCurrentAllocs );
-		SET_MEMORY_STAT( STAT_Binned_TotalAllocs, LocalTotalAllocs );
+		SET_DWORD_STAT( STAT_Binned_CurrentAllocs, LocalCurrentAllocs );
+		SET_DWORD_STAT( STAT_Binned_TotalAllocs, LocalTotalAllocs );
 		SET_MEMORY_STAT( STAT_Binned_SlackCurrent, LocalSlackCurrent );
 #endif
 	}
