Core - Added operator+= the chunked array

diff --git a/Engine/Source/Runtime/Core/Public/Containers/ChunkedArray.h b/Engine/Source/Runtime/Core/Public/Containers/ChunkedArray.h
--- a/Engine/Source/Runtime/Core/Public/Containers/ChunkedArray.h
+++ b/Engine/Source/Runtime/Core/Public/Containers/ChunkedArray.h
@@ -152,6 +152,18 @@ public:
 		return *this; 
 	}
 
+	FORCEINLINE TChunkedArray& operator+=(const TChunkedArray& Other) 
+	{ 
+		if( (UPTRINT*)this != (UPTRINT*)&Other )
+		{
+			for( int32 Index = 0; Index < Other.Num(); ++Index )
+			{
+				AddElement(Other[Index]);
+			}
+		}
+		return *this; 
+	}
+
 	int32 Add( int32 Count=1 )
 	{
 		check(Count>=0);
