Avoid calling ::MoveWindow win32 function if nothing changed.

https://github.com/EpicGames/UnrealEngine/pull/752

diff --git a/Engine/Source/Runtime/Core/Private/Windows/WindowsWindow.cpp b/Engine/Source/Runtime/Core/Private/Windows/WindowsWindow.cpp
--- a/Engine/Source/Runtime/Core/Private/Windows/WindowsWindow.cpp
+++ b/Engine/Source/Runtime/Core/Private/Windows/WindowsWindow.cpp
@@ -355,8 +355,16 @@ void FWindowsWindow::ReshapeWindow( int32 NewX, int32 NewY, int32 NewWidth, int3
 
 
 	// NOTE: MoveWindow will trigger a WM_SIZE and our SWindow's cached size will be updated
-	const bool bRepaint = true;
-	::MoveWindow( HWnd, WindowX, WindowY, NewWidth, NewHeight, bRepaint );
+
+	const int32 OldWidth = WindowInfo.rcClient.right - WindowInfo.rcClient.left;
+	const int32 OldHeight = WindowInfo.rcClient.bottom - WindowInfo.rcClient.top;
+
+	if( NewX != WindowInfo.rcClient.left || NewY != WindowInfo.rcClient.top || NewWidth != OldWidth || NewHeight != OldHeight  )
+	{
+		const bool bRepaint = true;
+		::MoveWindow(HWnd, WindowX, WindowY, NewWidth, NewHeight, bRepaint);
+	}
+
 
 	if( Definition->SizeWillChangeOften && bVirtualSizeChanged )
 	{
