Valid stats commands now appear in the console auto-complete list and no longer incorrectly output an 'unrecognized command' message.

diff --git a/Engine/Source/Runtime/Core/Private/Stats/StatsCommand.cpp b/Engine/Source/Runtime/Core/Private/Stats/StatsCommand.cpp
--- a/Engine/Source/Runtime/Core/Private/Stats/StatsCommand.cpp
+++ b/Engine/Source/Runtime/Core/Private/Stats/StatsCommand.cpp
@@ -1426,12 +1426,12 @@ bool DirectStatsCommand(const TCHAR* Cmd, bool bBlockForCompletion /*= false*/,
 	{
 		FString AddArgs;
 		const TCHAR* TempCmd = Cmd;
-		bResult = true;
 
 		FString ArgNoWhitespaces = FDefaultValueHelper::RemoveWhitespaces(TempCmd);
 		const bool bIsEmpty = ArgNoWhitespaces.IsEmpty();
 #if STATS
-		if( bIsEmpty && Ar )
+		bResult = true;
+		if (bIsEmpty && Ar)
 		{
 			PrintStatsHelpToOutputDevice( *Ar );
 		}
@@ -1500,10 +1500,25 @@ bool DirectStatsCommand(const TCHAR* Cmd, bool bBlockForCompletion /*= false*/,
 		{
 		}
 		else
-#endif
 		{
 			bResult = false;
+
+			FString MaybeGroup;
+			if (FParse::Token(TempCmd, MaybeGroup, false) && MaybeGroup.Len() > 0)
+			{
+				// If there is + at the end of the group name, remove it
+				const int32 PlusPos = MaybeGroup.Len() - 1;
+				const bool bHierarchy = MaybeGroup[MaybeGroup.Len() - 1] == TEXT('+');
+				if (bHierarchy)
+				{
+					MaybeGroup.RemoveAt(PlusPos, 1, false);
+				}
+
+				const FName MaybeGroupFName = FName(*(FString(TEXT("STATGROUP_")) + MaybeGroup));
+				bResult = FStatGroupGameThreadNotifier::Get().StatGroupNames.Contains(MaybeGroupFName);
+			}
 		}
+#endif
 
 		check(IsInGameThread());
 		if( !bIsEmpty )
diff --git a/Engine/Source/Runtime/Core/Public/Stats/StatsData.h b/Engine/Source/Runtime/Core/Public/Stats/StatsData.h
--- a/Engine/Source/Runtime/Core/Public/Stats/StatsData.h
+++ b/Engine/Source/Runtime/Core/Public/Stats/StatsData.h
@@ -861,6 +861,11 @@ public:
 	void NewData(FStatNameAndInfo NameAndInfo)
 	{
 		NameAndInfos.Add(NameAndInfo);
+		const FName GroupName = NameAndInfo.GetGroupName();
+		if (!GroupName.IsNone() && GroupName != NAME_Groups)
+		{
+			StatGroupNames.Add(GroupName);
+		}
 	}
 
 	void SendData()
@@ -883,6 +888,9 @@ public:
 	DECLARE_DELEGATE_OneParam(FOnNewStatGroupRegistered, const TArray<FStatNameAndInfo>&);
 	FOnNewStatGroupRegistered NewStatGroupDelegate;
 
+	/** A set of all the stat group names which have been registered */
+	TSet<FName> StatGroupNames;
+
 private:
 	FStatGroupGameThreadNotifier(){}
 	~FStatGroupGameThreadNotifier(){}
diff --git a/Engine/Source/Runtime/Engine/Private/UserInterface/Console.cpp b/Engine/Source/Runtime/Engine/Private/UserInterface/Console.cpp
--- a/Engine/Source/Runtime/Engine/Private/UserInterface/Console.cpp
+++ b/Engine/Source/Runtime/Engine/Private/UserInterface/Console.cpp
@@ -11,6 +11,7 @@
 #include "DefaultValueHelper.h"
 #include "ConsoleSettings.h"
 #include "GameFramework/InputSettings.h"
+#include "Stats/StatsData.h"
 
 static const uint32 MAX_AUTOCOMPLETION_LINES = 20;
 
@@ -219,6 +220,23 @@ void UConsole::BuildRuntimeAutoCompleteList(bool bForce)
 		AutoCompleteList[NewIdx].Command = FString(TEXT("open 127.0.0.1"));
 		AutoCompleteList[NewIdx].Desc = FString(TEXT("open 127.0.0.1 (opens connection to localhost)"));
 	}
+
+	// stat commands
+	{
+		const TSet<FName>& StatGroupNames = FStatGroupGameThreadNotifier::Get().StatGroupNames;
+
+		int32 NewIdx = AutoCompleteList.AddZeroed(StatGroupNames.Num());
+		for (const FName& StatGroupName : StatGroupNames)
+		{
+			FString Command = FString(TEXT("Stat "));
+			Command += StatGroupName.ToString().RightChop(sizeof("STATGROUP_") - 1);
+
+			AutoCompleteList[NewIdx].Command = Command;
+			AutoCompleteList[NewIdx].Desc = FString();
+			NewIdx++;
+		}
+	}
+
 	// build the magic tree!
 	for (int32 ListIdx = 0; ListIdx < AutoCompleteList.Num(); ListIdx++)
 	{
