Added generic UserTempDir() to platform abstraction.

- Moved Windows code from DesktopPlatform module to Core.
- All other platforms aren't implemented yet, but should compile (fixes CIS).

diff --git a/Engine/Source/Developer/DesktopPlatform/Private/Linux/DesktopPlatformLinux.cpp b/Engine/Source/Developer/DesktopPlatform/Private/Linux/DesktopPlatformLinux.cpp
--- a/Engine/Source/Developer/DesktopPlatform/Private/Linux/DesktopPlatformLinux.cpp
+++ b/Engine/Source/Developer/DesktopPlatform/Private/Linux/DesktopPlatformLinux.cpp
@@ -380,7 +380,7 @@ FFeedbackContext* FDesktopPlatformLinux::GetNativeFeedbackContext()
 
 FString FDesktopPlatformLinux::GetUserTempPath()
 {
-	//@todo - get the path to the user's temp dir
+	return FString(FPlatformProcess::UserTempDir());
 }
 
 #undef LOCTEXT_NAMESPACE
diff --git a/Engine/Source/Developer/DesktopPlatform/Private/Mac/DesktopPlatformMac.cpp b/Engine/Source/Developer/DesktopPlatform/Private/Mac/DesktopPlatformMac.cpp
--- a/Engine/Source/Developer/DesktopPlatform/Private/Mac/DesktopPlatformMac.cpp
+++ b/Engine/Source/Developer/DesktopPlatform/Private/Mac/DesktopPlatformMac.cpp
@@ -754,7 +754,7 @@ FFeedbackContext* FDesktopPlatformMac::GetNativeFeedbackContext()
 
 FString FDesktopPlatformMac::GetUserTempPath()
 {
-	//@todo - get the path to the user's temp dir
+	return FString(FPlatformProcess::UserTempDir());
 }
 
 #undef LOCTEXT_NAMESPACE
diff --git a/Engine/Source/Developer/DesktopPlatform/Private/Windows/DesktopPlatformWindows.cpp b/Engine/Source/Developer/DesktopPlatform/Private/Windows/DesktopPlatformWindows.cpp
--- a/Engine/Source/Developer/DesktopPlatform/Private/Windows/DesktopPlatformWindows.cpp
+++ b/Engine/Source/Developer/DesktopPlatform/Private/Windows/DesktopPlatformWindows.cpp
@@ -580,12 +580,7 @@ FFeedbackContext* FDesktopPlatformWindows::GetNativeFeedbackContext()
 
 FString FDesktopPlatformWindows::GetUserTempPath()
 {
-	TCHAR TempPath[MAX_PATH];
-	ZeroMemory(TempPath, sizeof(TCHAR) * MAX_PATH);
-
-	::GetTempPath(MAX_PATH, TempPath);
-
-	return FString(TempPath).Replace(TEXT("\\"), TEXT("/"));
+	return FString(FPlatformProcess::UserTempDir());
 }
 
 void FDesktopPlatformWindows::GetRequiredRegistrySettings(TIndirectArray<FRegistryRootedKey> &RootedKeys)
diff --git a/Engine/Source/Runtime/Core/Private/GenericPlatform/GenericPlatformProcess.cpp b/Engine/Source/Runtime/Core/Private/GenericPlatform/GenericPlatformProcess.cpp
--- a/Engine/Source/Runtime/Core/Private/GenericPlatform/GenericPlatformProcess.cpp
+++ b/Engine/Source/Runtime/Core/Private/GenericPlatform/GenericPlatformProcess.cpp
@@ -57,6 +57,12 @@ const TCHAR *FGenericPlatformProcess::UserSettingsDir()
 	return FPlatformMisc::RootDir();
 }
 
+const TCHAR *FGenericPlatformProcess::UserTempDir()
+{
+	// default to the root directory
+	return FPlatformMisc::RootDir();
+}
+
 const TCHAR* FGenericPlatformProcess::ApplicationSettingsDir()
 {
 	// default to the root directory
diff --git a/Engine/Source/Runtime/Core/Private/Windows/WindowsPlatformProcess.cpp b/Engine/Source/Runtime/Core/Private/Windows/WindowsPlatformProcess.cpp
--- a/Engine/Source/Runtime/Core/Private/Windows/WindowsPlatformProcess.cpp
+++ b/Engine/Source/Runtime/Core/Private/Windows/WindowsPlatformProcess.cpp
@@ -748,6 +748,21 @@ const TCHAR* FWindowsPlatformProcess::UserDir()
 	return *WindowsUserDir;
 }
 
+const TCHAR* FWindowsPlatformProcess::UserTempDir()
+{
+	static FString WindowsUserTempDir;
+	if( !WindowsUserTempDir.Len() )
+	{
+		TCHAR TempPath[MAX_PATH];
+		ZeroMemory(TempPath, sizeof(TCHAR) * MAX_PATH);
+
+		::GetTempPath(MAX_PATH, TempPath);
+
+		WindowsUserTempDir = FString(TempPath).Replace(TEXT("\\"), TEXT("/"));
+	}
+	return *WindowsUserTempDir;
+}
+
 const TCHAR* FWindowsPlatformProcess::UserSettingsDir()
 {
 	static FString WindowsUserSettingsDir;
diff --git a/Engine/Source/Runtime/Core/Public/GenericPlatform/GenericPlatformProcess.h b/Engine/Source/Runtime/Core/Public/GenericPlatform/GenericPlatformProcess.h
--- a/Engine/Source/Runtime/Core/Public/GenericPlatform/GenericPlatformProcess.h
+++ b/Engine/Source/Runtime/Core/Public/GenericPlatform/GenericPlatformProcess.h
@@ -221,6 +221,9 @@ struct CORE_API FGenericPlatformProcess
 	/** Get the user settings directory.  NOTE: Only one return value is valid at a time! **/
 	static const TCHAR *UserSettingsDir();
 
+	/** Get the user temporary directory.  NOTE: Only one return value is valid at a time! **/
+	static const TCHAR *UserTempDir();
+
 	/** Get application settings directory.  NOTE: Only one return value is valid at a time! **/
 	static const TCHAR* ApplicationSettingsDir();
 
diff --git a/Engine/Source/Runtime/Core/Public/Windows/WindowsPlatformProcess.h b/Engine/Source/Runtime/Core/Public/Windows/WindowsPlatformProcess.h
--- a/Engine/Source/Runtime/Core/Public/Windows/WindowsPlatformProcess.h
+++ b/Engine/Source/Runtime/Core/Public/Windows/WindowsPlatformProcess.h
@@ -86,6 +86,7 @@ public:
 	static void SetThreadAffinityMask( uint64 AffinityMask );
 	static const TCHAR* BaseDir();
 	static const TCHAR* UserDir();
+	static const TCHAR* UserTempDir();
 	static const TCHAR* UserSettingsDir();
 	static const TCHAR* ApplicationSettingsDir();
 	static const TCHAR* ComputerName();
