ON_SCOPE_EXIT - generalized scope guard:

{
    FileHandle* Handle = GetFileHandle();
    ON_SCOPE_EXIT
    {
        CloseFile(Handle);
    };

    DoSomethingWithFile( Handle );

    // File will be closed automatically no matter how the scope is exited, e.g.:
    // * Any return statement.
    // * break or continue (if the scope is a loop body).
    // * An exception is thrown outside the block.
    // * Execution reaches the end of the block.
}

diff --git a/Engine/Source/Runtime/Core/Public/Misc/ScopeExit.h b/Engine/Source/Runtime/Core/Public/Misc/ScopeExit.h
new file mode 100644
--- /dev/null
+++ b/Engine/Source/Runtime/Core/Public/Misc/ScopeExit.h
@@ -0,0 +1,88 @@
+// Copyright 1998-2015 Epic Games, Inc. All Rights Reserved.
+
+#pragma once
+
+#include "HAL/Platform.h"
+#include "Optional.h"
+
+namespace ScopeExitSupport
+{
+	/**
+	 * Not meant for direct consumption : use ON_SCOPE_EXIT instead.
+	 *
+	 * RAII class that calls a lambda when it is destroyed.
+	 */
+	template <typename FuncType>
+	class TScopeGuard : public FNoncopyable
+	{
+	public:
+		// Given a lambda, constructs an RAII scope guard.
+		explicit TScopeGuard(FuncType&& InFunc)
+			: Func(MoveTemp(InFunc))
+		{
+		}
+
+		// This constructor needs to be available for the code to compile.
+		// It will be almost definitely be RVOed out (even in DEBUG).
+		TScopeGuard(TScopeGuard&& Other)
+			: Func(MoveTemp(Other.Func))
+		{
+			Other.Func.Reset();
+		}
+
+		// Causes
+		~TScopeGuard()
+		{
+			if (Func.IsSet())
+			{
+				Func.GetValue()();
+			}
+		}
+
+	private:
+		// The lambda to be executed when this guard goes out of scope.
+		TOptional<FuncType> Func;
+	};
+
+	struct FScopeGuardSyntaxSupport
+	{
+		template <typename FuncType>
+		TScopeGuard<FuncType> operator+(FuncType&& InFunc)
+		{
+			return TScopeGuard<FuncType>(Forward<FuncType>(InFunc));
+		}
+	};
+}
+
+
+#ifdef __COUNTER__
+	// Created a variable with a unique name
+	#define ANONYMOUS_VARIABLE( Name ) Name##__COUNTER__
+#else
+	// Created a variable with a unique name.
+	// Less reliable than the __COUNTER__ version.
+	#define ANONYMOUS_VARIABLE( Name ) Name##__LINE__
+#endif
+
+
+/**
+ * Enables a lambda to be executed on scope exit.
+ *
+ * Example:
+ *    {
+ *      FileHandle* Handle = GetFileHandle();
+ *      ON_SCOPE_EXIT
+ *      {
+ *          CloseFile(Handle);
+ *      };
+ *
+ *      DoSomethingWithFile( Handle );
+ *
+ *      // File will be closed automatically no matter how the scope is exited, e.g.:
+ *      // * Any return statement.
+ *      // * break or continue (if the scope is a loop body).
+ *      // * An exception is thrown outside the block.
+ *      // * Execution reaches the end of the block.
+ *    }
+ */
+#define ON_SCOPE_EXIT const auto ANONYMOUS_VARIABLE(ScopeGuard_) = ::ScopeExitSupport::FScopeGuardSyntaxSupport() + [&]()
\ No newline at end of file
