Implement conditional and no stats version of Sleep.

diff --git a/Engine/Source/Runtime/Core/Private/GenericPlatform/GenericPlatformProcess.cpp b/Engine/Source/Runtime/Core/Private/GenericPlatform/GenericPlatformProcess.cpp
--- a/Engine/Source/Runtime/Core/Private/GenericPlatform/GenericPlatformProcess.cpp
+++ b/Engine/Source/Runtime/Core/Private/GenericPlatform/GenericPlatformProcess.cpp
@@ -3,6 +3,7 @@
 #include "CorePrivatePCH.h"
 #include "EventPool.h"
 #include "Public/Modules/ModuleVersion.h"
+#include "Templates/Function.h"
 
 
 #if PLATFORM_HAS_BSD_TIME 
@@ -10,6 +11,7 @@
 	#include <sched.h>  // for sched_yield
 #endif
 
+DEFINE_STAT(STAT_Sleep);
 
 void* FGenericPlatformProcess::GetDllHandle( const TCHAR* Filename )
 {
@@ -252,12 +254,16 @@ void FGenericPlatformProcess::ExploreFolder( const TCHAR* FilePath )
 
 
 #if PLATFORM_HAS_BSD_TIME
-DECLARE_CYCLE_STAT(TEXT("CPU Stall - Sleep"),STAT_Sleep,STATGROUP_CPUStalls);
 
 void FGenericPlatformProcess::Sleep( float Seconds )
 {
 	SCOPE_CYCLE_COUNTER(STAT_Sleep);
 	FThreadIdleStats::FScopeIdle Scope;
+	SleepNoStats(Seconds);
+}
+
+void FGenericPlatformProcess::SleepNoStats( float Seconds )
+{
 	const int32 usec = FPlatformMath::TruncToInt(Seconds * 1000000.0f);
 	if (usec > 0)
 	{
@@ -277,6 +283,21 @@ void FGenericPlatformProcess::SleepInfinite()
 
 #endif // PLATFORM_HAS_BSD_TIME 
 
+void FGenericPlatformProcess::ConditionalSleep(const TFunction<bool()>& Condition)
+{
+	if (Condition())
+	{
+		return;
+	}
+
+	SCOPE_CYCLE_COUNTER(STAT_Sleep);
+	FThreadIdleStats::FScopeIdle Scope;
+	do
+	{
+		FPlatformProcess::SleepNoStats(0.0f);
+	} while (!Condition());
+}
+
 #if PLATFORM_USE_PTHREADS
 
 #include "PThreadEvent.h"
diff --git a/Engine/Source/Runtime/Core/Private/HTML5/HTML5PlatformProcess.cpp b/Engine/Source/Runtime/Core/Private/HTML5/HTML5PlatformProcess.cpp
--- a/Engine/Source/Runtime/Core/Private/HTML5/HTML5PlatformProcess.cpp
+++ b/Engine/Source/Runtime/Core/Private/HTML5/HTML5PlatformProcess.cpp
@@ -55,6 +55,13 @@ void FHTML5PlatformProcess::Sleep( float Seconds )
 {
 	SCOPE_CYCLE_COUNTER(STAT_HTML5Sleep);
 	FThreadIdleStats::FScopeIdle Scope;
+
+	SleepNoStats(Seconds);
+	// @todo html5 threading: actually put the thread to sleep for awhile (UE3 did nothing, we do the same now)
+}
+
+void FHTML5PlatformProcess::SleepNoStats(float Seconds)
+{
 	// @todo html5 threading: actually put the thread to sleep for awhile (UE3 did nothing, we do the same now)
 }
 
diff --git a/Engine/Source/Runtime/Core/Private/WinRT/WinRTProcess.cpp b/Engine/Source/Runtime/Core/Private/WinRT/WinRTProcess.cpp
--- a/Engine/Source/Runtime/Core/Private/WinRT/WinRTProcess.cpp
+++ b/Engine/Source/Runtime/Core/Private/WinRT/WinRTProcess.cpp
@@ -138,8 +138,6 @@ const TCHAR* FWinRTProcess::BaseDir()
 
 #include "HideWinRTPlatformTypes.h"
 
-DECLARE_CYCLE_STAT(TEXT("CPU Stall - Sleep"),STAT_Sleep,STATGROUP_CPUStalls);
-
 void FWinRTProcess::Sleep( float Seconds )
 {
 	SCOPE_CYCLE_COUNTER(STAT_Sleep);
diff --git a/Engine/Source/Runtime/Core/Private/Windows/WindowsPlatformProcess.cpp b/Engine/Source/Runtime/Core/Private/Windows/WindowsPlatformProcess.cpp
--- a/Engine/Source/Runtime/Core/Private/Windows/WindowsPlatformProcess.cpp
+++ b/Engine/Source/Runtime/Core/Private/Windows/WindowsPlatformProcess.cpp
@@ -964,13 +964,16 @@ bool FWindowsPlatformProcess::ResolveNetworkPath( FString InUNCPath, FString& Ou
 	return false;
 }
 
-DECLARE_CYCLE_STAT(TEXT("CPU Stall - Sleep"),STAT_Sleep,STATGROUP_CPUStalls);
-
 void FWindowsPlatformProcess::Sleep( float Seconds )
 {
 	SCOPE_CYCLE_COUNTER(STAT_Sleep);
 	FThreadIdleStats::FScopeIdle Scope;
-	::Sleep( (uint32)(Seconds * 1000.0) );
+	SleepNoStats(Seconds);
+}
+
+void FWindowsPlatformProcess::SleepNoStats(float Seconds)
+{
+	::Sleep((uint32)(Seconds * 1000.0));
 }
 
 void FWindowsPlatformProcess::SleepInfinite()
diff --git a/Engine/Source/Runtime/Core/Public/GenericPlatform/GenericPlatformProcess.h b/Engine/Source/Runtime/Core/Public/GenericPlatform/GenericPlatformProcess.h
--- a/Engine/Source/Runtime/Core/Public/GenericPlatform/GenericPlatformProcess.h
+++ b/Engine/Source/Runtime/Core/Public/GenericPlatform/GenericPlatformProcess.h
@@ -8,7 +8,7 @@
 
 
 class FString;
-
+template <typename FuncType> class TFunction;
 
 namespace EProcessResource
 {
@@ -413,13 +413,22 @@ struct CORE_API FGenericPlatformProcess
 
 #if PLATFORM_HAS_BSD_TIME 
 
-	/** Sleep this thread for Seconds.  0.0 means release the current time slice to let other threads get some attention. */
+	/** Sleep this thread for Seconds.  0.0 means release the current time slice to let other threads get some attention. Uses stats.*/
 	static void Sleep( float Seconds );
+	/** Sleep this thread for Seconds.  0.0 means release the current time slice to let other threads get some attention. */
+	static void SleepNoStats( float Seconds );
 	/** Sleep this thread infinitely. */
 	static void SleepInfinite();
 
 #endif // PLATFORM_HAS_BSD_TIME
 
+	/**
+	* Sleep thread until condition is satisfied.
+	*
+	* @param	Condition	Condition to evaluate.
+	*/
+	static void ConditionalSleep(const TFunction<bool()>& Condition);
+
 	/**
 	 * Creates a new event.
 	 *
diff --git a/Engine/Source/Runtime/Core/Public/HTML5/HTML5PlatformProcess.h b/Engine/Source/Runtime/Core/Public/HTML5/HTML5PlatformProcess.h
--- a/Engine/Source/Runtime/Core/Public/HTML5/HTML5PlatformProcess.h
+++ b/Engine/Source/Runtime/Core/Public/HTML5/HTML5PlatformProcess.h
@@ -30,7 +30,8 @@ struct CORE_API FHTML5PlatformProcess : public FGenericPlatformProcess
 {
 	static const TCHAR* ComputerName();
 	static const TCHAR* BaseDir();
-	static void Sleep( float Seconds );
+	static void Sleep(float Seconds);
+	static void SleepNoStats(float Seconds);
 	static void SleepInfinite();
 	static class FEvent* CreateSynchEvent(bool bIsManualReset = 0);
 	static class FRunnableThread* CreateRunnableThread();
diff --git a/Engine/Source/Runtime/Core/Public/Misc/CoreStats.h b/Engine/Source/Runtime/Core/Public/Misc/CoreStats.h
--- a/Engine/Source/Runtime/Core/Public/Misc/CoreStats.h
+++ b/Engine/Source/Runtime/Core/Public/Misc/CoreStats.h
@@ -87,3 +87,5 @@ DECLARE_FLOAT_COUNTER_STAT_EXTERN(TEXT("Bandwidth (MByte/ sec)"),STAT_AsyncIO_Ba
 
 DECLARE_CYCLE_STAT_EXTERN(TEXT("Convert Rotator to Quat"), STAT_MathConvertRotatorToQuat, STATGROUP_Math, CORE_API);
 DECLARE_CYCLE_STAT_EXTERN(TEXT("Convert Quat to Rotator"), STAT_MathConvertQuatToRotator, STATGROUP_Math, CORE_API);
+
+DECLARE_CYCLE_STAT_EXTERN(TEXT("CPU Stall - Sleep"), STAT_Sleep, STATGROUP_CPUStalls, CORE_API);
diff --git a/Engine/Source/Runtime/Core/Public/Windows/WindowsPlatformProcess.h b/Engine/Source/Runtime/Core/Public/Windows/WindowsPlatformProcess.h
--- a/Engine/Source/Runtime/Core/Public/Windows/WindowsPlatformProcess.h
+++ b/Engine/Source/Runtime/Core/Public/Windows/WindowsPlatformProcess.h
@@ -112,7 +112,8 @@ public:
 	static void LaunchFileInDefaultExternalApplication( const TCHAR* FileName, const TCHAR* Parms = NULL, ELaunchVerb::Type Verb = ELaunchVerb::Open );
 	static void ExploreFolder( const TCHAR* FilePath );
 	static bool ResolveNetworkPath( FString InUNCPath, FString& OutPath ); 
-	static void Sleep( float Seconds );
+	static void Sleep(float Seconds);
+	static void SleepNoStats(float Seconds);
 	static void SleepInfinite();
 	static class FEvent* CreateSynchEvent(bool bIsManualReset = false);
 	static class FRunnableThread* CreateRunnableThread();
