Back out revision 141 from //depot/UE4/Engine/Source/Runtime/Core/Private/Misc/ConfigCacheIni.cpp

Fixes URLs no longer being quoted in config files, which caused them to lose eveything after the http: when reloaded. This was causing source control integration in the editor to fail when connecting via URL (i.e. SVN, P4 with Assembla etc.)

diff --git a/Engine/Source/Runtime/Core/Private/Misc/ConfigCacheIni.cpp b/Engine/Source/Runtime/Core/Private/Misc/ConfigCacheIni.cpp
--- a/Engine/Source/Runtime/Core/Private/Misc/ConfigCacheIni.cpp
+++ b/Engine/Source/Runtime/Core/Private/Misc/ConfigCacheIni.cpp
@@ -552,8 +552,8 @@ void FConfigFile::Read( const FString& Filename )
 bool FConfigFile::ShouldExportQuotedString(const FString& PropertyValue) const
 {
 	// The value should be exported as quoted string if it begins with a space (which is stripped on import) or
-	// when it contains ' //' (interpreted as a comment when importing) or starts with '//'
-	return **PropertyValue == TEXT(' ') || FCString::Strstr(*PropertyValue, TEXT(" //")) != NULL || PropertyValue.StartsWith(TEXT("//"));
+	// when it contains '//' (interpreted as a comment when importing).
+	return **PropertyValue == TEXT(' ') || FCString::Strstr(*PropertyValue, TEXT("//")) != NULL;
 }
 
 
