Do not try to copy file in FFileManagerGeneric::Copy if dest and src paths are the same

diff --git a/Engine/Source/Runtime/Core/Private/HAL/FileManagerGeneric.cpp b/Engine/Source/Runtime/Core/Private/HAL/FileManagerGeneric.cpp
--- a/Engine/Source/Runtime/Core/Private/HAL/FileManagerGeneric.cpp
+++ b/Engine/Source/Runtime/Core/Private/HAL/FileManagerGeneric.cpp
@@ -122,7 +122,11 @@ int64 FFileManagerGeneric::FileSize( const TCHAR* Filename )
 uint32 FFileManagerGeneric::Copy( const TCHAR* InDestFile, const TCHAR* InSrcFile, bool ReplaceExisting, bool EvenIfReadOnly, bool Attributes, FCopyProgress* Progress )
 {
 	uint32	Result = COPY_OK;
-	if( Progress )
+	if( FPaths::ConvertRelativePathToFull( InDestFile ) == FPaths::ConvertRelativePathToFull( InSrcFile ) )
+	{
+		Result = COPY_Fail;
+	}
+	else if( Progress )
 	{
 		Result = CopyWithProgress( InDestFile, InSrcFile, ReplaceExisting, EvenIfReadOnly, Attributes, Progress );
 	}
