Only cache Android window size if valid (greater than 8x8)
#android
#ue4

diff --git a/Engine/Source/Runtime/Core/Private/Android/AndroidWindow.cpp b/Engine/Source/Runtime/Core/Private/Android/AndroidWindow.cpp
--- a/Engine/Source/Runtime/Core/Private/Android/AndroidWindow.cpp
+++ b/Engine/Source/Runtime/Core/Private/Android/AndroidWindow.cpp
@@ -71,7 +71,7 @@ int32 FAndroidWindow::GetDepthBufferPreference()
 FPlatformRect FAndroidWindow::GetScreenRect()
 {
 	// since orientation and resolution won't change on Android, use cached results if valid
-	if (WindowWidth > 0)
+	if (WindowWidth > 8)
 	{
 		FPlatformRect ScreenRect;
 		ScreenRect.Left = 0;
