Core: Async code & documentation improvements

diff --git a/Engine/Source/Runtime/Core/Private/Tests/Async/AsyncTest.cpp b/Engine/Source/Runtime/Core/Private/Tests/Async/AsyncTest.cpp
--- a/Engine/Source/Runtime/Core/Private/Tests/Async/AsyncTest.cpp
+++ b/Engine/Source/Runtime/Core/Private/Tests/Async/AsyncTest.cpp
@@ -12,14 +12,14 @@ IMPLEMENT_SIMPLE_AUTOMATION_TEST(FAsyncTest, "Core.Async.Async", EAutomationTest
 
 bool FAsyncTest::RunTest(const FString& Parameters)
 {
-	auto Task = []()
+	TFunction<int()> Task = []()
 	{
 		return 123;
 	};
 
 	// task graph task
 	{
-		auto Future = Async<int>(EAsyncExecution::TaskGraph, Task);
+		auto Future = Async(EAsyncExecution::TaskGraph, Task);
 		int Result = Future.Get();
 
 		TestEqual(TEXT("Task graph task must return expected value"), Result, 123);
@@ -27,7 +27,7 @@ bool FAsyncTest::RunTest(const FString& Parameters)
 
 	// thread task
 	{
-		auto Future = Async<int>(EAsyncExecution::Thread, Task);
+		auto Future = Async(EAsyncExecution::Thread, Task);
 		int Result = Future.Get();
 
 		TestEqual(TEXT("Threaded task must return expected value"), Result, 123);
@@ -35,7 +35,7 @@ bool FAsyncTest::RunTest(const FString& Parameters)
 
 	// thread pool task
 	{
-		auto Future = Async<int>(EAsyncExecution::ThreadPool, Task);
+		auto Future = Async(EAsyncExecution::ThreadPool, Task);
 		int Result = Future.Get();
 
 		TestEqual(TEXT("Thread pool task must return expected value"), Result, 123);
diff --git a/Engine/Source/Runtime/Core/Public/Async/Async.h b/Engine/Source/Runtime/Core/Public/Async/Async.h
--- a/Engine/Source/Runtime/Core/Public/Async/Async.h
+++ b/Engine/Source/Runtime/Core/Public/Async/Async.h
@@ -22,11 +22,51 @@ enum class EAsyncExecution
 };
 
 
+/**
+ * Base class for asynchronous functions that are executed in the Task Graph system.
+ */
+class FAsyncTaskBase
+{
+public:
+
+	/**
+	 * Returns the name of the thread that this task should run on.
+	 *
+	 * @return Always run on any thread.
+	 */
+	ENamedThreads::Type GetDesiredThread()
+	{
+		return ENamedThreads::AnyThread;
+	}
+
+	/**
+	 * Gets the task's stats tracking identifier.
+	 *
+	 * @return Stats identifier.
+	 */
+	TStatId GetStatId() const
+	{
+		return GET_STATID(STAT_TaskGraph_OtherTasks);
+	}
+
+	/**
+	 * Gets the mode for tracking subsequent tasks.
+	 *
+	 * @return Always track subsequent tasks.
+	 */
+	static ESubsequentsMode::Type GetSubsequentsMode()
+	{
+		return ESubsequentsMode::FireAndForget;
+	}
+};
+
+
 /**
  * Template for asynchronous functions that are executed in the Task Graph system.
  */
 template<typename ResultType>
 class TAsyncTask
+	: public FAsyncTaskBase
 {
 public:
 
@@ -54,16 +94,6 @@ public:
 		Promise.SetValue(Function());
 	}
 
-	/**
-	 * Returns the name of the thread that this task should run on.
-	 *
-	 * @return Always run on any thread.
-	 */
-	ENamedThreads::Type GetDesiredThread()
-	{
-		return ENamedThreads::AnyThread;
-	}
-
 	/**
 	 * Gets the future that will hold the asynchronous result.
 	 *
@@ -74,26 +104,6 @@ public:
 		return Promise.GetFuture();
 	}
 
-	/**
-	 * Gets the task's stats tracking identifier.
-	 *
-	 * @return Stats identifier.
-	 */
-	TStatId GetStatId() const
-	{
-		return GET_STATID(STAT_TaskGraph_OtherTasks);
-	}
-
-	/**
-	 * Gets the mode for tracking subsequent tasks.
-	 *
-	 * @return Always track subsequent tasks.
-	 */
-	static ESubsequentsMode::Type GetSubsequentsMode()
-	{
-		return ESubsequentsMode::FireAndForget;
-	}
-
 private:
 
 	/** The function to execute on the Task Graph. */
@@ -203,6 +213,37 @@ private:
 /**
  * Executes a given function asynchronously.
  *
+ * While we still support VS2012, we need to help the compiler deduce the ResultType
+ * template argument for Async<T>(), which can be accomplished in two ways:
+ *
+ *	- call Async() with the template parameter, i.e. Async<int>(..)
+ *	- use a TFunction binding that explicitly specifies the type, i.e. TFunction<int()> F = []() { .. }
+ *
+ * Usage examples:
+ *
+ *	// using global function
+ *		int TestFunc()
+ *		{
+ *			return 123;
+ *		}
+ *
+ *		TFunction<int()> Task = TestFunc();
+ *		auto Result = Async(EAsyncExecution::Thread, Task);
+ *
+ *	// using lambda
+ *		TFunction<int()> Task = []()
+ *		{
+ *			return 123;
+ *		}
+ *
+ *		auto Result = Async(EAsyncExecution::Thread, Task);
+ *
+ *
+ *	// using inline lambda
+ *		auto Result = Async<int>(EAsyncExecution::Thread, []() {
+ *			return 123;
+ *		}
+ *
  * @param Execution The execution method to use, i.e. on Task Graph or in a separate thread.
  * @param Function The function to execute.
  * @result A TFuture object that will receive the return value from the function.
