Core: code and documentation cleanup pass

diff --git a/Engine/Source/Runtime/Core/Private/Containers/String.cpp b/Engine/Source/Runtime/Core/Private/Containers/String.cpp
--- a/Engine/Source/Runtime/Core/Private/Containers/String.cpp
+++ b/Engine/Source/Runtime/Core/Private/Containers/String.cpp
@@ -22,7 +22,7 @@ void FString::TrimToNullTerminator()
 
 int32 FString::Find( const TCHAR* SubStr, ESearchCase::Type SearchCase, ESearchDir::Type SearchDir, int32 StartPosition ) const
 {
-	if ( SubStr == NULL )
+	if ( SubStr == nullptr )
 	{
 		return INDEX_NONE;
 	}
@@ -215,7 +215,7 @@ FString FString::TrimQuotes( bool* bQuotesRemoved ) const
 		}
 	}
 
-	if ( bQuotesRemoved != NULL )
+	if ( bQuotesRemoved != nullptr )
 	{
 		*bQuotesRemoved = bQuotesWereRemoved;
 	}
@@ -795,7 +795,7 @@ int32 FString::ParseIntoArray(TArray<FString>* InArray, const TCHAR** DelimArray
 
 FString FString::Replace(const TCHAR* From, const TCHAR* To, ESearchCase::Type SearchCase) const
 {
-	// Previous code used to accidentally accept a NULL replacement string - this is no longer accepted.
+	// Previous code used to accidentally accept a nullptr replacement string - this is no longer accepted.
 	check(To);
 
 	if (IsEmpty() || !From || !*From)
@@ -838,8 +838,8 @@ int32 FString::ReplaceInline( const TCHAR* SearchText, const TCHAR* ReplacementT
 	int32 ReplacementCount = 0;
 
 	if (Len() > 0
-	&&	SearchText != NULL && *SearchText != 0
-	&&	ReplacementText != NULL && (SearchCase == ESearchCase::IgnoreCase || FCString::Strcmp(SearchText, ReplacementText) != 0) )
+	&&	SearchText != nullptr && *SearchText != 0
+	&&	ReplacementText != nullptr && (SearchCase == ESearchCase::IgnoreCase || FCString::Strcmp(SearchText, ReplacementText) != 0) )
 	{
 		const int32 NumCharsToReplace=FCString::Strlen(SearchText);
 		const int32 NumCharsToInsert=FCString::Strlen(ReplacementText);
@@ -847,7 +847,7 @@ int32 FString::ReplaceInline( const TCHAR* SearchText, const TCHAR* ReplacementT
 		if ( NumCharsToInsert == NumCharsToReplace )
 		{
 			TCHAR* Pos = SearchCase == ESearchCase::IgnoreCase ? FCString::Stristr(&(*this)[0], SearchText) : FCString::Strstr(&(*this)[0], SearchText);
-			while ( Pos != NULL )
+			while ( Pos != nullptr )
 			{
 				ReplacementCount++;
 
@@ -876,7 +876,7 @@ int32 FString::ReplaceInline( const TCHAR* SearchText, const TCHAR* ReplacementT
 			TCHAR* WritePosition = (TCHAR*)Copy.Data.GetData();
 			// look for From in the remaining string
 			TCHAR* SearchPosition = SearchCase == ESearchCase::IgnoreCase ? FCString::Stristr(WritePosition, SearchText) : FCString::Strstr(WritePosition, SearchText);
-			while ( SearchPosition != NULL )
+			while ( SearchPosition != nullptr )
 			{
 				ReplacementCount++;
 
@@ -962,14 +962,14 @@ static const uint32 MaxSupportedEscapeChars = ARRAY_COUNT(CharToEscapeSeqMap);
  *
  * @return	a string with all control characters replaced by the escaped version.
  */
-FString FString::ReplaceCharWithEscapedChar( const TArray<TCHAR>* Chars/*=NULL*/ ) const
+FString FString::ReplaceCharWithEscapedChar( const TArray<TCHAR>* Chars/*=nullptr*/ ) const
 {
-	if ( Len() > 0 && (Chars == NULL || Chars->Num() > 0) )
+	if ( Len() > 0 && (Chars == nullptr || Chars->Num() > 0) )
 	{
 		FString Result(*this);
 		for ( int32 ChIdx = 0; ChIdx < MaxSupportedEscapeChars; ChIdx++ )
 		{
-			if ( Chars == NULL || Chars->Contains(*(CharToEscapeSeqMap[ChIdx][0])) )
+			if ( Chars == nullptr || Chars->Contains(*(CharToEscapeSeqMap[ChIdx][0])) )
 			{
 				// use ReplaceInline as that won't create a copy of the string if the character isn't found
 				Result.ReplaceInline(CharToEscapeSeqMap[ChIdx][0], CharToEscapeSeqMap[ChIdx][1]);
@@ -984,15 +984,15 @@ FString FString::ReplaceCharWithEscapedChar( const TArray<TCHAR>* Chars/*=NULL*/
  * Removes the escape backslash for all supported characters, replacing the escape and character with the non-escaped version.  (i.e.
  * replaces "\\n" with "\n".  Counterpart to ReplaceCharWithEscapedChar().
  */
-FString FString::ReplaceEscapedCharWithChar( const TArray<TCHAR>* Chars/*=NULL*/ ) const
+FString FString::ReplaceEscapedCharWithChar( const TArray<TCHAR>* Chars/*=nullptr*/ ) const
 {
-	if ( Len() > 0 && (Chars == NULL || Chars->Num() > 0) )
+	if ( Len() > 0 && (Chars == nullptr || Chars->Num() > 0) )
 	{
 		FString Result(*this);
 		// Spin CharToEscapeSeqMap backwards to ensure we're doing the inverse of ReplaceCharWithEscapedChar
 		for ( int32 ChIdx = MaxSupportedEscapeChars - 1; ChIdx >= 0; ChIdx-- )
 		{
-			if ( Chars == NULL || Chars->Contains(*(CharToEscapeSeqMap[ChIdx][0])) )
+			if ( Chars == nullptr || Chars->Contains(*(CharToEscapeSeqMap[ChIdx][0])) )
 			{
 				// use ReplaceInline as that won't create a copy of the string if the character isn't found
 				Result.ReplaceInline( CharToEscapeSeqMap[ChIdx][1], CharToEscapeSeqMap[ChIdx][0] );
@@ -1056,7 +1056,7 @@ VARARG_BODY( FString, FString::Printf, const TCHAR*, VARARG_NONE )
 	// If that fails, start allocating regular memory
 	if( Result == -1 )
 	{
-		Buffer = NULL;
+		Buffer = nullptr;
 		while(Result == -1)
 		{
 			BufferSize *= 2;
diff --git a/Engine/Source/Runtime/Core/Public/Containers/AllocatorFixedSizeFreeList.h b/Engine/Source/Runtime/Core/Public/Containers/AllocatorFixedSizeFreeList.h
--- a/Engine/Source/Runtime/Core/Public/Containers/AllocatorFixedSizeFreeList.h
+++ b/Engine/Source/Runtime/Core/Public/Containers/AllocatorFixedSizeFreeList.h
@@ -1,12 +1,7 @@
 // Copyright 1998-2015 Epic Games, Inc. All Rights Reserved.
 
-/*=============================================================================
-	Custom allocator.
-
-=============================================================================*/
 #pragma once
-#ifndef __ALLOCATORFIXEDSIZEFREELIST_H__
-#define __ALLOCATORFIXEDSIZEFREELIST_H__
+
 
 /**
  * Fixed-sized allocator that uses a free list to cache allocations.
@@ -36,7 +31,7 @@ public:
 	 * @param	InitialBlockSize		- number of allocations to warm the cache with
 	 */
 	TAllocatorFixedSizeFreeList(uint32 InitialBlockSize=0)
-		:FreeList(NULL)
+		:FreeList(nullptr)
 		,NumAllocated(0)
 		,NumLive(0)
 	{
@@ -133,6 +128,7 @@ public:
 	}
 
 private:
+
 	struct FreeListNode
 	{
 		FreeListNode* NextFreeAllocation;
@@ -143,12 +139,14 @@ private:
 		checkSlow(NumAllocated >= NumLive);
 	}
 
+private:
+
 	/** Linked List of free memory blocks. */
 	FreeListNode* FreeList;
+
 	/** The number of objects that have been allocated. */
 	uint32 NumAllocated;
+
 	/** The number of objects that are constructed and "out in the wild". */
 	uint32 NumLive;
 };
-
-#endif
diff --git a/Engine/Source/Runtime/Core/Public/Containers/Array.h b/Engine/Source/Runtime/Core/Public/Containers/Array.h
--- a/Engine/Source/Runtime/Core/Public/Containers/Array.h
+++ b/Engine/Source/Runtime/Core/Public/Containers/Array.h
@@ -1,10 +1,7 @@
 // Copyright 1998-2015 Epic Games, Inc. All Rights Reserved.
 
-/*=============================================================================
-	Array.h: Dynamic array definitions.
-=============================================================================*/
-
 #pragma once
+
 #include "Containers/ContainerAllocationPolicies.h"
 #include "HAL/Platform.h"
 #include "Serialization/ArchiveBase.h"
@@ -12,8 +9,10 @@
 #include "Templates/Sorting.h"
 #include "Templates/UnrealTemplate.h"
 
+
 #define DEBUG_HEAP 0
 
+
 /**
  * Generic iterator which can operate on types that expose the following:
  * - A type called ElementType representing the contained type.
@@ -333,14 +332,14 @@ public:
 
 	FORCEINLINE bool operator()( ElementType* A, ElementType* B ) const 
 	{
-		check( A != NULL );
-		check( B != NULL );
+		check( A != nullptr );
+		check( B != nullptr );
 		return Predicate( *B, *A ); 
 	}
 	FORCEINLINE bool operator()( const ElementType* A, const ElementType* B ) const 
 	{
-		check( A != NULL );
-		check( B != NULL );
+		check( A != nullptr );
+		check( B != nullptr );
 		return Predicate( *B, *A ); 
 	}
 };
@@ -587,7 +586,7 @@ public:
 	/**
 	 * Helper function for returning a typed pointer to the first array entry.
 	 *
-	 * @returns Pointer to first array entry or NULL if ArrayMax == 0.
+	 * @returns Pointer to first array entry or nullptr if ArrayMax == 0.
 	 */
 	DEPRECATED(4.6, "This function is deprecated as it does the same as GetData(). Please use GetData() instead.")
 	FORCEINLINE ElementType* GetTypedData()
@@ -598,7 +597,7 @@ public:
 	/**
 	 * Helper function for returning a typed pointer to the first array entry.
 	 *
-	 * @returns Pointer to first array entry or NULL if ArrayMax == 0.
+	 * @returns Pointer to first array entry or nullptr if ArrayMax == 0.
 	 */
 	FORCEINLINE ElementType* GetData()
 	{
@@ -608,7 +607,7 @@ public:
 	/**
 	 * Helper function for returning a typed pointer to the first array entry.
 	 *
-	 * @returns Pointer to first array entry or NULL if ArrayMax == 0.
+	 * @returns Pointer to first array entry or nullptr if ArrayMax == 0.
 	 */
 	DEPRECATED(4.6, "This function is deprecated as it does the same as GetData(). Please use GetData() instead.")
 	FORCEINLINE const ElementType* GetTypedData() const
@@ -619,7 +618,7 @@ public:
 	/**
 	 * Helper function for returning a typed pointer to the first array entry.
 	 *
-	 * @returns Pointer to first array entry or NULL if ArrayMax == 0.
+	 * @returns Pointer to first array entry or nullptr if ArrayMax == 0.
 	 */
 	FORCEINLINE const ElementType* GetData() const
 	{
@@ -989,7 +988,7 @@ public:
 	 *
 	 * @param Key The key to search by.
 	 *
-	 * @returns Pointer to the first matching element, or NULL if none is found.
+	 * @returns Pointer to the first matching element, or nullptr if none is found.
 	 */
 	template <typename KeyType>
 	FORCEINLINE const ElementType* FindByKey(const KeyType& Key) const
@@ -1003,7 +1002,7 @@ public:
 	 *
 	 * @param Key The key to search by.
 	 *
-	 * @returns Pointer to the first matching element, or NULL if none is found.
+	 * @returns Pointer to the first matching element, or nullptr if none is found.
 	 */
 	template <typename KeyType>
 	ElementType* FindByKey(const KeyType& Key)
@@ -1016,7 +1015,7 @@ public:
 			}
 		}
 
-		return NULL;
+		return nullptr;
 	}
 
 	/**
@@ -1025,7 +1024,7 @@ public:
 	 * @param Pred The functor to apply to each element.
 	 *
 	 * @returns Pointer to the first element for which the predicate returns
-	 *          true, or NULL if none is found.
+	 *          true, or nullptr if none is found.
 	 */
 	template <typename Predicate>
 	FORCEINLINE const ElementType* FindByPredicate(Predicate Pred) const
@@ -1039,7 +1038,7 @@ public:
 	 * @param Pred The functor to apply to each element.
 	 *
 	 * @return Pointer to the first element for which the predicate returns
-	 *         true, or NULL if none is found.
+	 *         true, or nullptr if none is found.
 	 */
 	template <typename Predicate>
 	ElementType* FindByPredicate(Predicate Pred)
@@ -1052,7 +1051,7 @@ public:
 			}
 		}
 
-		return NULL;
+		return nullptr;
 	}
 
 	/**
@@ -1104,7 +1103,7 @@ public:
 	template <typename Predicate>
 	FORCEINLINE bool ContainsByPredicate(Predicate Pred) const
 	{
-		return FindByPredicate(Pred) != NULL;
+		return FindByPredicate(Pred) != nullptr;
 	}
 
 	/**
@@ -1354,7 +1353,7 @@ public:
 	 */
 	int32 Insert(const ElementType* Ptr, int32 Count, int32 Index)
 	{
-		check(Ptr != NULL);
+		check(Ptr != nullptr);
 
 		InsertUninitialized(Index, Count);
 		ConstructItems<ElementType>(GetData() + Index, Ptr, Count);
@@ -1653,7 +1652,7 @@ public:
 	 */
 	void Append(const ElementType* Ptr, int32 Count)
 	{
-		check(Ptr != NULL);
+		check(Ptr != nullptr);
 
 		int32 Pos = AddUninitialized(Count);
 		ConstructItems<ElementType>(GetData() + Pos, Ptr, Count);
@@ -2079,18 +2078,18 @@ public:
 	 * @returns True if element was found. False otherwise.
 	 */
 	template<typename SearchType>
-	bool FindItemByClass(SearchType **Item = NULL, int32 *ItemIndex = NULL, int32 StartIndex = 0) const
+	bool FindItemByClass(SearchType **Item = nullptr, int32 *ItemIndex = nullptr, int32 StartIndex = 0) const
 	{
 		UClass* SearchClass = SearchType::StaticClass();
 		for (int32 Idx = StartIndex; Idx < ArrayNum; Idx++)
 		{
-			if ((*this)[Idx] != NULL && (*this)[Idx]->IsA(SearchClass))
+			if ((*this)[Idx] != nullptr && (*this)[Idx]->IsA(SearchClass))
 			{
-				if (Item != NULL)
+				if (Item != nullptr)
 				{
 					*Item = (SearchType*)((*this)[Idx]);
 				}
-				if (ItemIndex != NULL)
+				if (ItemIndex != nullptr)
 				{
 					*ItemIndex = Idx;
 				}
@@ -2880,7 +2879,7 @@ public:
 	/**
 	 * Helper function for returning a typed pointer to the first array entry.
 	 *
-	 * @returns Pointer to first array entry or NULL if this->ArrayMax == 0.
+	 * @returns Pointer to first array entry or nullptr if this->ArrayMax == 0.
 	 */
 	DEPRECATED(4.6, "This function is deprecated as it does the same as GetData(). Please use GetData() instead.")
 	FORCEINLINE T** GetTypedData()
@@ -2891,7 +2890,7 @@ public:
 	/**
 	 * Helper function for returning a typed pointer to the first array entry.
 	 *
-	 * @returns Pointer to first array entry or NULL if this->ArrayMax == 0.
+	 * @returns Pointer to first array entry or nullptr if this->ArrayMax == 0.
 	 */
 	FORCEINLINE T** GetData()
 	{
@@ -2901,7 +2900,7 @@ public:
 	/**
 	 * Helper function for returning a typed pointer to the first array entry.
 	 *
-	 * @returns Pointer to first array entry or NULL if this->ArrayMax == 0.
+	 * @returns Pointer to first array entry or nullptr if this->ArrayMax == 0.
 	 */
 	DEPRECATED(4.6, "This function is deprecated as it does the same as GetData(). Please use GetData() instead.")
 	FORCEINLINE const T** GetTypedData() const
@@ -2912,7 +2911,7 @@ public:
 	/**
 	 * Helper function for returning a typed pointer to the first array entry.
 	 *
-	 * @returns Pointer to first array entry or NULL if this->ArrayMax == 0.
+	 * @returns Pointer to first array entry or nullptr if this->ArrayMax == 0.
 	 */
 	FORCEINLINE const T** GetData() const
 	{
diff --git a/Engine/Source/Runtime/Core/Public/Containers/ArrayBuilder.h b/Engine/Source/Runtime/Core/Public/Containers/ArrayBuilder.h
--- a/Engine/Source/Runtime/Core/Public/Containers/ArrayBuilder.h
+++ b/Engine/Source/Runtime/Core/Public/Containers/ArrayBuilder.h
@@ -1,9 +1,5 @@
 // Copyright 1998-2015 Epic Games, Inc. All Rights Reserved.
 
-/*=============================================================================
-	ArrayBuilder.h: Declares the TArrayBuilder template.
-=============================================================================*/
-
 #pragma once
 
 
@@ -17,10 +13,8 @@ template<typename ElementType, typename Allocator = FDefaultAllocator> class TAr
 {
 public:
 
-	/**
-	 * Default constructor.
-	 */
-	TArrayBuilder( ) { }
+	/** Default constructor. */
+	TArrayBuilder() { }
 
 	/**
 	 * Creates and initializes an array builder from an array of items.
@@ -81,7 +75,7 @@ public:
 	 *
 	 * @return A new array.
 	 */
-	TArray<ElementType, Allocator> Build( )
+	TArray<ElementType, Allocator> Build()
 	{
 		return Array;
 	}
@@ -91,13 +85,13 @@ public:
 	 *
 	 * @return A new array.
 	 */
-	operator TArray<ElementType, Allocator>( )
+	operator TArray<ElementType, Allocator>()
 	{
 		return Build();
 	}
 
 private:
 
-	// Holds the array.
+	/** Holds the array. */
 	TArray<ElementType, Allocator> Array;
 };
\ No newline at end of file
diff --git a/Engine/Source/Runtime/Core/Public/Containers/BitArray.h b/Engine/Source/Runtime/Core/Public/Containers/BitArray.h
--- a/Engine/Source/Runtime/Core/Public/Containers/BitArray.h
+++ b/Engine/Source/Runtime/Core/Public/Containers/BitArray.h
@@ -1,12 +1,10 @@
 // Copyright 1998-2015 Epic Games, Inc. All Rights Reserved.
 
-/*=============================================================================
-	BitArray.h: Bit array definition.
-=============================================================================*/
-
 #pragma once
+
 #include "Math/UnrealMathUtility.h"
 
+
 // Functions for manipulating bit sets.
 struct FBitSet
 {
@@ -20,6 +18,7 @@ struct FBitSet
 	}
 };
 
+
 // Forward declaration.
 template<typename Allocator = FDefaultBitArrayAllocator>
 class TBitArray;
@@ -30,6 +29,7 @@ class TConstSetBitIterator;
 template<typename Allocator = FDefaultBitArrayAllocator,typename OtherAllocator = FDefaultBitArrayAllocator>
 class TConstDualSetBitIterator;
 
+
 /**
  * Serializer (predefined for no friend injection in gcc 411)
  */
@@ -106,6 +106,7 @@ private:
 	uint32 Mask;
 };
 
+
 /** Used to read a bit in the array as a bool. */
 class FConstBitReference
 {
@@ -126,6 +127,7 @@ private:
 	uint32 Mask;
 };
 
+
 /** Used to reference a bit in an unspecified bit array. */
 class FRelativeBitReference
 {
@@ -147,6 +149,7 @@ protected:
 	uint32 Mask;
 };
 
+
 /**
  * A dynamically sized bit array.
  * An array of Booleans.  They stored in one bit/Boolean.  There are iterators that efficiently iterate over only set bits.
@@ -653,6 +656,7 @@ private:
 	}
 };
 
+
 template<typename Allocator>
 struct TContainerTraits<TBitArray<Allocator> > : public TContainerTraitsBase<TBitArray<Allocator> >
 {
@@ -661,6 +665,7 @@ struct TContainerTraits<TBitArray<Allocator> > : public TContainerTraitsBase<TBi
 		TAllocatorTraits<Allocator>::SupportsMove };
 };
 
+
 /** An iterator which only iterates over set bits. */
 template<typename Allocator>
 class TConstSetBitIterator : public FRelativeBitReference
@@ -765,6 +770,7 @@ private:
 	}
 };
 
+
 /** An iterator which only iterates over the bits which are set in both of two bit-arrays. */
 template<typename Allocator,typename OtherAllocator>
 class TConstDualSetBitIterator : public FRelativeBitReference
@@ -874,10 +880,10 @@ private:
 	}
 };
 
+
 /** A specialization of the exchange macro that avoids reallocating when exchanging two bit arrays. */
 template<typename Allocator>
 FORCEINLINE void Exchange(TBitArray<Allocator>& A,TBitArray<Allocator>& B)
 {
 	FMemory::Memswap(&A,&B,sizeof(TBitArray<Allocator>));
 }
-
diff --git a/Engine/Source/Runtime/Core/Public/Containers/ChunkedArray.h b/Engine/Source/Runtime/Core/Public/Containers/ChunkedArray.h
--- a/Engine/Source/Runtime/Core/Public/Containers/ChunkedArray.h
+++ b/Engine/Source/Runtime/Core/Public/Containers/ChunkedArray.h
@@ -1,13 +1,10 @@
 // Copyright 1998-2015 Epic Games, Inc. All Rights Reserved.
 
-/*=============================================================================
-	ChunkedArray.h: Chunked array definition.
-=============================================================================*/
-
 #pragma once
 
 #include "Array.h"
 
+
 /** An array that uses multiple allocations to avoid allocation failure due to fragmentation. */
 template<typename ElementType, uint32 TargetBytesPerChunk = 16384 >
 class TChunkedArray
@@ -208,6 +205,7 @@ public:
 	}
 
 protected:
+
 	friend struct TContainerTraits<TChunkedArray<ElementType, TargetBytesPerChunk>>;
 
 	enum { NumElementsPerChunk = TargetBytesPerChunk / sizeof(ElementType) };
@@ -227,6 +225,7 @@ protected:
 	int32 NumElements;
 };
 
+
 template <typename ElementType, uint32 TargetBytesPerChunk>
 struct TContainerTraits<TChunkedArray<ElementType, TargetBytesPerChunk> > : public TContainerTraitsBase<TChunkedArray<ElementType, TargetBytesPerChunk> >
 {
@@ -235,6 +234,7 @@ struct TContainerTraits<TChunkedArray<ElementType, TargetBytesPerChunk> > : publ
 		TContainerTraits<typename TChunkedArray<ElementType, TargetBytesPerChunk>::ChunksType>::MoveWillEmptyContainer };
 };
 
+
 template <typename T,uint32 TargetBytesPerChunk> void* operator new( size_t Size, TChunkedArray<T,TargetBytesPerChunk>& ChunkedArray )
 {
 	check(Size == sizeof(T));
@@ -242,6 +242,7 @@ template <typename T,uint32 TargetBytesPerChunk> void* operator new( size_t Size
 	return &ChunkedArray(Index);
 }
 
+
 /**
  * A specialization of the exchange macro that avoids reallocating when
  * exchanging two arrays.
@@ -253,4 +254,4 @@ template <typename T, uint32 TargetBytesPerChunk>
 inline void Exchange(TChunkedArray<T,TargetBytesPerChunk>& FirstArrayToExchange, TChunkedArray<T,TargetBytesPerChunk>& SecondArrayToExchange)
 {
 	FMemory::Memswap(&FirstArrayToExchange, &SecondArrayToExchange, sizeof(TChunkedArray<T,TargetBytesPerChunk>));
-}
\ No newline at end of file
+}
diff --git a/Engine/Source/Runtime/Core/Public/Containers/CircularBuffer.h b/Engine/Source/Runtime/Core/Public/Containers/CircularBuffer.h
--- a/Engine/Source/Runtime/Core/Public/Containers/CircularBuffer.h
+++ b/Engine/Source/Runtime/Core/Public/Containers/CircularBuffer.h
@@ -1,9 +1,5 @@
 // Copyright 1998-2015 Epic Games, Inc. All Rights Reserved.
 
-/*=============================================================================
-	CircularBuffer.h: Declares the TCircularBuffer template class.
-=============================================================================*/
-
 #pragma once
 
 
@@ -77,7 +73,7 @@ public:
 	 *
 	 * @return Buffer capacity.
 	 */
-	FORCEINLINE uint32 Capacity( ) const
+	FORCEINLINE uint32 Capacity() const
 	{
 		return Elements.Num();
 	}
@@ -95,9 +91,9 @@ public:
 
 private:
 
-	// Holds the mask for indexing the buffer's elements.
+	/** Holds the mask for indexing the buffer's elements. */
 	uint32 IndexMask;
 
-	// Holds the buffer's elements.
+	/** Holds the buffer's elements. */
 	TArray<ElementType> Elements;
 };
diff --git a/Engine/Source/Runtime/Core/Public/Containers/CircularQueue.h b/Engine/Source/Runtime/Core/Public/Containers/CircularQueue.h
--- a/Engine/Source/Runtime/Core/Public/Containers/CircularQueue.h
+++ b/Engine/Source/Runtime/Core/Public/Containers/CircularQueue.h
@@ -1,9 +1,5 @@
 // Copyright 1998-2015 Epic Games, Inc. All Rights Reserved.
 
-/*=============================================================================
-	CircularQueue.h: Declares the TCircularQueue template.
-=============================================================================*/
-
 #pragma once
 
 
@@ -35,10 +31,8 @@ public:
 		, Tail(0)
 	{ }
 
-	/**
-	 * Destructor.
-	 */
-	virtual ~TCircularQueue( ) { }
+	/** Virtual destructor. */
+	virtual ~TCircularQueue() { }
 
 public:
 
@@ -47,7 +41,7 @@ public:
 	 *
 	 * @return Number of queued elements.
 	 */
-	uint32 Count( ) const
+	uint32 Count() const
 	{
 		int32 Count = Head - Tail;
 
@@ -83,7 +77,7 @@ public:
 	 *
 	 * @see IsEmpty
 	 */
-	void Empty( )
+	void Empty()
 	{
 		Head = Tail;
 	}
@@ -115,7 +109,7 @@ public:
 	 * @return true if the queue is empty, false otherwise.
 	 * @see Empty, IsFull
 	 */
-	FORCEINLINE bool IsEmpty( ) const
+	FORCEINLINE bool IsEmpty() const
 	{
 		return (Head == Tail);
 	}
@@ -126,7 +120,7 @@ public:
 	 * @return true if the queue is full, false otherwise.
 	 * @see IsEmpty
 	 */
-	bool IsFull( ) const
+	bool IsFull() const
 	{
 		return (Buffer.GetNextIndex(Tail) == Head);
 	}
@@ -151,12 +145,12 @@ public:
 
 private:
 
-	// Holds the buffer.
+	/** Holds the buffer. */
 	TCircularBuffer<ElementType> Buffer;
 
-	// Holds the index to the first item in the buffer.
+	/** Holds the index to the first item in the buffer. */
 	MS_ALIGN(CACHE_LINE_SIZE) volatile uint32 Head GCC_ALIGN(CACHE_LINE_SIZE);
 
-	// Holds the index to the last item in the buffer.
+	/** Holds the index to the last item in the buffer. */
 	MS_ALIGN(CACHE_LINE_SIZE) volatile uint32 Tail GCC_ALIGN(CACHE_LINE_SIZE);
 };
diff --git a/Engine/Source/Runtime/Core/Public/Containers/ContainerAllocationPolicies.h b/Engine/Source/Runtime/Core/Public/Containers/ContainerAllocationPolicies.h
--- a/Engine/Source/Runtime/Core/Public/Containers/ContainerAllocationPolicies.h
+++ b/Engine/Source/Runtime/Core/Public/Containers/ContainerAllocationPolicies.h
@@ -1,14 +1,11 @@
 // Copyright 1998-2015 Epic Games, Inc. All Rights Reserved.
 
-/*=============================================================================
-	ContainerAllocationPolicies.h: Defines allocation policies for containers.
-=============================================================================*/
-
 #pragma once
 
 #include "MemoryBase.h"
 #include "Misc/OutputDevice.h"
 
+
 /** Used to determine the alignment of an element type. */
 template<typename ElementType>
 class TElementAlignmentCalculator
@@ -33,7 +30,6 @@ public:
 #define ALIGNOF(T) (TElementAlignmentCalculator<T>::Alignment)
 
 
-
 /** branchless pointer selection
 * return A ? A : B;
 **/
@@ -198,7 +194,7 @@ public:
 
 		/** Default constructor. */
 		ForAnyElementType()
-			: Data(NULL)
+			: Data(nullptr)
 		{}
 
 		/**
@@ -216,7 +212,7 @@ public:
 			}
 
 			Data       = Other.Data;
-			Other.Data = NULL;
+			Other.Data = nullptr;
 		}
 
 		/** Destructor. */
@@ -239,7 +235,7 @@ public:
 			SIZE_T NumBytesPerElement
 			)
 		{
-			// Avoid calling FMemory::Realloc( NULL, 0 ) as ANSI C mandates returning a valid pointer which is not what we want.
+			// Avoid calling FMemory::Realloc( nullptr, 0 ) as ANSI C mandates returning a valid pointer which is not what we want.
 			if( Data || NumElements )
 			{
 				//checkSlow(((uint64)NumElements*(uint64)ElementTypeInfo.GetSize() < (uint64)INT_MAX));
@@ -304,7 +300,7 @@ public:
 	public:
 		/** Default constructor. */
 		ForAnyElementType()
-			: Data(NULL)
+			: Data(nullptr)
 		{}
 
 		/**
@@ -322,7 +318,7 @@ public:
 			}
 
 			Data       = Other.Data;
-			Other.Data = NULL;
+			Other.Data = nullptr;
 		}
 
 		/** Destructor. */
@@ -341,7 +337,7 @@ public:
 		}
 		void ResizeAllocation(int32 PreviousNumElements,int32 NumElements,SIZE_T NumBytesPerElement)
 		{
-			// Avoid calling FMemory::Realloc( NULL, 0 ) as ANSI C mandates returning a valid pointer which is not what we want.
+			// Avoid calling FMemory::Realloc( nullptr, 0 ) as ANSI C mandates returning a valid pointer which is not what we want.
 			if( Data || NumElements )
 			{
 				//checkSlow(((uint64)NumElements*(uint64)ElementTypeInfo.GetSize() < (uint64)INT_MAX));
diff --git a/Engine/Source/Runtime/Core/Public/Containers/DynamicRHIResourceArray.h b/Engine/Source/Runtime/Core/Public/Containers/DynamicRHIResourceArray.h
--- a/Engine/Source/Runtime/Core/Public/Containers/DynamicRHIResourceArray.h
+++ b/Engine/Source/Runtime/Core/Public/Containers/DynamicRHIResourceArray.h
@@ -1,11 +1,8 @@
 // Copyright 1998-2015 Epic Games, Inc. All Rights Reserved.
 
-/*=============================================================================
-	DynamicRHIResourceArray.h: Resource array definitions for dynamically bound RHIs.
-=============================================================================*/
-
 #pragma once
 
+
 /** alignment for supported resource types */
 enum EResourceAlignment
 {
diff --git a/Engine/Source/Runtime/Core/Public/Containers/EnumAsByte.h b/Engine/Source/Runtime/Core/Public/Containers/EnumAsByte.h
--- a/Engine/Source/Runtime/Core/Public/Containers/EnumAsByte.h
+++ b/Engine/Source/Runtime/Core/Public/Containers/EnumAsByte.h
@@ -14,9 +14,7 @@ class TEnumAsByte
 public:
 	typedef TEnum EnumType;
 
-	/**
-	 * Default Constructor (no initialization).
-	 */
+	/** Default Constructor (no initialization). */
 	FORCEINLINE TEnumAsByte() { }
 
 	/**
diff --git a/Engine/Source/Runtime/Core/Public/Containers/List.h b/Engine/Source/Runtime/Core/Public/Containers/List.h
--- a/Engine/Source/Runtime/Core/Public/Containers/List.h
+++ b/Engine/Source/Runtime/Core/Public/Containers/List.h
@@ -2,18 +2,18 @@
 
 #pragma once
 
+
 template <class ContainerType, class ElementType>
 class TLinkedListIterator
 {
 public:
 	explicit TLinkedListIterator(ContainerType* FirstLink)
 		: CurrentLink(FirstLink)
-	{
-	}
+	{ }
 
 	/**
-		* Advances the iterator to the next element.
-		*/
+	 * Advances the iterator to the next element.
+	 */
 	void Next()
 	{
 		checkSlow(CurrentLink);
@@ -53,12 +53,14 @@ public:
 	}
 
 private:
+
 	ContainerType* CurrentLink;
 
 	friend bool operator==(const TLinkedListIterator& Lhs, const TLinkedListIterator& Rhs) { return Lhs.CurrentLink == Rhs.CurrentLink; }
 	friend bool operator!=(const TLinkedListIterator& Lhs, const TLinkedListIterator& Rhs) { return Lhs.CurrentLink != Rhs.CurrentLink; }
 };
 
+
 /**
  * Encapsulates a link in a single linked list with constant access time.
  */
@@ -66,25 +68,29 @@ template <class ElementType>
 class TLinkedList
 {
 public:
+
 	/**
 	 * Used to iterate over the elements of a linked list.
 	 */
 	typedef TLinkedListIterator<TLinkedList,       ElementType> TIterator;
 	typedef TLinkedListIterator<TLinkedList, const ElementType> TConstIterator;
 
-	// Constructors.
+	/** Default constructor (empty list). */
 	TLinkedList()
 		: NextLink(nullptr)
 		, PrevLink(nullptr)
-	{
-	}
+	{ }
 
+	/**
+	 * Creates a new linked list with a single element.
+	 *
+	 * @param InElement The element to add to the list.
+	 */
 	explicit TLinkedList(const ElementType& InElement)
 		: Element (InElement)
 		, NextLink(nullptr)
 		, PrevLink(nullptr)
-	{
-	}
+	{ }
 
 	/**
 	 * Removes this element from the list in constant time.
@@ -102,8 +108,8 @@ public:
 			*PrevLink = NextLink;
 		}
 		// Make it safe to call Unlink again.
-		NextLink = NULL;
-		PrevLink = NULL;
+		NextLink = nullptr;
+		PrevLink = nullptr;
 	}
 
 	/**
@@ -129,7 +135,7 @@ public:
 	 */
 	bool IsLinked()
 	{
-		return PrevLink != NULL;
+		return PrevLink != nullptr;
 	}
 
 	TLinkedList* GetPrevLink() const
@@ -180,10 +186,10 @@ template <class NodeType, class ElementType>
 class TDoubleLinkedListIterator
 {
 public:
+
 	explicit TDoubleLinkedListIterator(NodeType* StartingNode)
 		: CurrentNode(StartingNode)
-	{
-	}
+	{ }
 
 	/** conversion to "bool" returning true if the iterator is valid. */
 	FORCEINLINE_EXPLICIT_OPERATOR_BOOL() const
@@ -245,6 +251,7 @@ private:
 	friend bool operator!=(const TDoubleLinkedListIterator& Lhs, const TDoubleLinkedListIterator& Rhs) { return Lhs.CurrentNode != Rhs.CurrentNode; }
 };
 
+
 /**
  * Double linked list.
  */
@@ -259,9 +266,8 @@ public:
 
 		/** Constructor */
 		TDoubleLinkedListNode( const ElementType& InValue )
-		: Value(InValue), NextNode(NULL), PrevNode(NULL)
-		{
-		}
+			: Value(InValue), NextNode(nullptr), PrevNode(nullptr)
+		{ }
 
 		const ElementType& GetValue() const
 		{
@@ -300,8 +306,7 @@ public:
 		: HeadNode(nullptr)
 		, TailNode(nullptr)
 		, ListSize(0)
-	{
-	}
+	{ }
 
 	/** Destructor */
 	virtual ~TDoubleLinkedList()
@@ -310,23 +315,24 @@ public:
 	}
 
 	// Adding/Removing methods
+
 	/**
 	 * Add the specified value to the beginning of the list, making that value the new head of the list.
 	 *
 	 * @param	InElement	the value to add to the list.
-	 *
-	 * @return	whether the node was successfully added into the list
+	 * @return	whether the node was successfully added into the list.
+	 * @see GetHead, InsertNode, RemoveNode
 	 */
 	bool AddHead( const ElementType& InElement )
 	{
 		TDoubleLinkedListNode* NewNode = new TDoubleLinkedListNode(InElement);
-		if ( NewNode == NULL )
+		if ( NewNode == nullptr )
 		{
 			return false;
 		}
 
 		// have an existing head node - change the head node to point to this one
-		if ( HeadNode != NULL )
+		if ( HeadNode != nullptr )
 		{
 			NewNode->NextNode = HeadNode;
 			HeadNode->PrevNode = NewNode;
@@ -345,18 +351,18 @@ public:
 	 * Append the specified value to the end of the list
 	 *
 	 * @param	InElement	the value to add to the list.
-	 *
 	 * @return	whether the node was successfully added into the list
+	 * @see GetTail, InsertNode, RemoveNode
 	 */
 	bool AddTail( const ElementType& InElement )
 	{
 		TDoubleLinkedListNode* NewNode = new TDoubleLinkedListNode(InElement);
-		if ( NewNode == NULL )
+		if ( NewNode == nullptr )
 		{
 			return false;
 		}
 
-		if ( TailNode != NULL )
+		if ( TailNode != nullptr )
 		{
 			TailNode->NextNode = NewNode;
 			NewNode->PrevNode = TailNode;
@@ -376,19 +382,19 @@ public:
 	 *
 	 * @param	InElement			the value to insert into the list
 	 * @param	NodeToInsertBefore	the new node will be inserted into the list at the current location of this node
-	 *								if NULL, the new node will become the new head of the list
-	 *
+	 *								if nullptr, the new node will become the new head of the list
 	 * @return	whether the node was successfully added into the list
+	 * @see Empty, RemoveNode
 	 */
-	bool InsertNode( const ElementType& InElement, TDoubleLinkedListNode* NodeToInsertBefore=NULL )
+	bool InsertNode( const ElementType& InElement, TDoubleLinkedListNode* NodeToInsertBefore=nullptr )
 	{
-		if ( NodeToInsertBefore == NULL || NodeToInsertBefore == HeadNode )
+		if ( NodeToInsertBefore == nullptr || NodeToInsertBefore == HeadNode )
 		{
 			return AddHead(InElement);
 		}
 
 		TDoubleLinkedListNode* NewNode = new TDoubleLinkedListNode(InElement);
-		if ( NewNode == NULL )
+		if ( NewNode == nullptr )
 		{
 			return false;
 		}
@@ -404,9 +410,10 @@ public:
 	}
 
 	/**
-	 * Remove the node corresponding to InElement
+	 * Remove the node corresponding to InElement.
 	 *
-	 * @param	InElement	the value to remove from the list
+	 * @param InElement The value to remove from the list.
+	 * @see Empty, InsertNode
 	 */
 	void RemoveNode( const ElementType& InElement )
 	{
@@ -416,10 +423,13 @@ public:
 
 	/**
 	 * Removes the node specified.
+	 *
+	 * @param NodeToRemove The node to remove.
+	 * @see Empty, InsertNode
 	 */
 	void RemoveNode( TDoubleLinkedListNode* NodeToRemove )
 	{
-		if ( NodeToRemove != NULL )
+		if ( NodeToRemove != nullptr )
 		{
 			// if we only have one node, just call Clear() so that we don't have to do lots of extra checks in the code below
 			if ( Num() == 1 )
@@ -432,13 +442,13 @@ public:
 			if ( NodeToRemove == HeadNode )
 			{
 				HeadNode = HeadNode->NextNode;
-				HeadNode->PrevNode = NULL;
+				HeadNode->PrevNode = nullptr;
 			}
 
 			else if ( NodeToRemove == TailNode )
 			{
 				TailNode = TailNode->PrevNode;
-				TailNode->NextNode = NULL;
+				TailNode->NextNode = nullptr;
 			}
 			else
 			{
@@ -447,32 +457,33 @@ public:
 			}
 
 			delete NodeToRemove;
-			NodeToRemove = NULL;
+			NodeToRemove = nullptr;
 			SetListSize(ListSize - 1);
 		}
 	}
 
-	/**
-	 * Removes all nodes from the list.
-	 */
+	/** Removes all nodes from the list. */
 	void Empty()
 	{
 		TDoubleLinkedListNode* pNode;
-		while ( HeadNode != NULL )
+		while ( HeadNode != nullptr )
 		{
 			pNode = HeadNode->NextNode;
 			delete HeadNode;
 			HeadNode = pNode;
 		}
 
-		HeadNode = TailNode = NULL;
+		HeadNode = TailNode = nullptr;
 		SetListSize(0);
 	}
 
 	// Accessors.
 
 	/**
-	 * Returns the node at the head of the list
+	 * Returns the node at the head of the list.
+	 *
+	 * @return Pointer to the first node.
+	 * @see GetTail
 	 */
 	TDoubleLinkedListNode* GetHead() const
 	{
@@ -480,7 +491,10 @@ public:
 	}
 
 	/**
-	 * Returns the node at the end of the list
+	 * Returns the node at the end of the list.
+	 *
+	 * @return Pointer to the last node.
+	 * @see GetHead
 	 */
 	TDoubleLinkedListNode* GetTail() const
 	{
@@ -491,13 +505,12 @@ public:
 	 * Finds the node corresponding to the value specified
 	 *
 	 * @param	InElement	the value to find
-	 *
-	 * @return	a pointer to the node that contains the value specified, or NULL of the value couldn't be found
+	 * @return	a pointer to the node that contains the value specified, or nullptr of the value couldn't be found
 	 */
 	TDoubleLinkedListNode* FindNode( const ElementType& InElement )
 	{
 		TDoubleLinkedListNode* pNode = HeadNode;
-		while ( pNode != NULL )
+		while ( pNode != nullptr )
 		{
 			if ( pNode->GetValue() == InElement )
 			{
@@ -511,7 +524,9 @@ public:
 	}
 
 	/**
-	 * Returns the size of the list.
+	 * Returns the number of items in the list.
+	 *
+	 * @return Item count.
 	 */
 	int32 Num() const
 	{
@@ -562,7 +577,7 @@ public:
 
 	// Constructor.
 
-	TList(const ElementType &InElement, TList<ElementType>* InNext = NULL)
+	TList(const ElementType &InElement, TList<ElementType>* InNext = nullptr)
 	{
 		Element = InElement;
 		Next = InNext;
diff --git a/Engine/Source/Runtime/Core/Public/Containers/LockFreeFixedSizeAllocator.h b/Engine/Source/Runtime/Core/Public/Containers/LockFreeFixedSizeAllocator.h
--- a/Engine/Source/Runtime/Core/Public/Containers/LockFreeFixedSizeAllocator.h
+++ b/Engine/Source/Runtime/Core/Public/Containers/LockFreeFixedSizeAllocator.h
@@ -4,11 +4,16 @@
 
 #include "LockFreeList.h"
 
-/** Thread safe, lock free pooling allocator of fixed size blocks that never returns free space until program shutdown. **/
+
+/**
+ * Thread safe, lock free pooling allocator of fixed size blocks that
+ * never returns free space until program shutdown.
+ */
 template<int32 SIZE, typename TTrackingCounter = FNoopCounter>
 class TLockFreeFixedSizeAllocator	// alignment isn't handled, assumes FMemory::Malloc will work
 {
 public:
+
 	/** Destructor, returns all memory via FMemory::Free **/
 	~TLockFreeFixedSizeAllocator()
 	{
@@ -20,7 +25,13 @@ public:
 		}
 		check(!NumFree.GetValue());
 	}
-	/** Returns a memory block of size SIZE **/
+
+	/**
+	 * Allocates a memory block of size SIZE.
+	 *
+	 * @return Pointer to the allocated memory.
+	 * @see Free
+	 */
 	void* Allocate()
 	{
 		NumUsed.Increment();
@@ -35,7 +46,13 @@ public:
 		}
 		return Memory;
 	}
-	/** Given a memory block (obtained from Allocate) put it on the free list for future use. **/
+
+	/**
+	 * Puts a memory block previously obtained from Allocate() back on the free list for future use.
+	 *
+	 * @param Item The item to free.
+	 * @see Allocate
+	 */
 	void Free(void *Item)
 	{
 		NumUsed.Decrement();
@@ -43,44 +60,79 @@ public:
 		NumFree.Increment();
 	}
 
+	/**
+	 * Gets the number of allocated memory blocks that are currently in use.
+	 *
+	 * @return Number of used memory blocks.
+	 * @see GetNumFree
+	 */
 	const TTrackingCounter& GetNumUsed() const
 	{
 		return NumUsed;
 	}
 
+	/**
+	 * Gets the number of allocated memory blocks that are currently unused.
+	 *
+	 * @return Number of unused memory blocks.
+	 * @see GetNumUsed
+	 */
 	const TTrackingCounter& GetNumFree() const
 	{
 		return NumFree;
 	}
 
 private:
-	/** Lock free list of free memory blocks **/
-	TLockFreePointerList<void>		FreeList;
 
-	/** Total number of blocks outstanding and not in the free list **/
-	TTrackingCounter	NumUsed; 
-	/** Total number of blocks in the free list **/
-	TTrackingCounter	NumFree;
-};
+	/** Lock free list of free memory blocks. */
+	TLockFreePointerList<void> FreeList;
 
+	/** Total number of blocks outstanding and not in the free list. */
+	TTrackingCounter NumUsed; 
+
+	/** Total number of blocks in the free list. */
+	TTrackingCounter NumFree;
+};
 
 
-/** Thread safe, lock free pooling allocator of memory for instances of T. Never returns free space until program shutdown. **/
+/**
+ * Thread safe, lock free pooling allocator of memory for instances of T.
+ *
+ * Never returns free space until program shutdown.
+ */
 template<class T>
 class TLockFreeClassAllocator : private TLockFreeFixedSizeAllocator<sizeof(T)>
 {
 public:
-	/** Returns a memory block of size sizeof(T) **/
+
+	/**
+	 * Returns a memory block of size sizeof(T).
+	 *
+	 * @return Pointer to the allocated memory.
+	 * @see Free, New
+	 */
 	void* Allocate()
 	{
 		return TLockFreeFixedSizeAllocator<sizeof(T)>::Allocate();
 	}
-	/** Returns a new T using the default contructor **/
+
+	/**
+	 * Returns a new T using the default constructor.
+	 *
+	 * @return Pointer to the new object.
+	 * @see Allocate, Free
+	 */
 	T* New()
 	{
 		return new (Allocate()) T();
 	}
-	/** Calls a destructor on Item and returns the memory to the free list for recycling. **/
+
+	/**
+	 * Calls a destructor on Item and returns the memory to the free list for recycling.
+	 *
+	 * @param Item The item whose memory to free.
+	 * @see Allocate, New
+	 */
 	void Free(T *Item)
 	{
 		Item->~T();
diff --git a/Engine/Source/Runtime/Core/Public/Containers/Map.h b/Engine/Source/Runtime/Core/Public/Containers/Map.h
--- a/Engine/Source/Runtime/Core/Public/Containers/Map.h
+++ b/Engine/Source/Runtime/Core/Public/Containers/Map.h
@@ -1,10 +1,12 @@
 // Copyright 1998-2015 Epic Games, Inc. All Rights Reserved.
 
 #pragma once
+
 #include "Containers/Array.h"
 #include "Containers/Set.h"
 #include "Containers/UnrealString.h"
 
+
 #define ExchangeB(A,B) {bool T=A; A=B; B=T;}
 
 
@@ -367,7 +369,7 @@ public:
 	/**
 	 * Returns the key associated with the specified value.  The time taken is O(N) in the number of pairs.
 	 * @param	Value - The value to search for
-	 * @return	A pointer to the key associated with the specified value, or NULL if the value isn't contained in this map.  The pointer
+	 * @return	A pointer to the key associated with the specified value, or nullptr if the value isn't contained in this map.  The pointer
 	 *			is only valid until the next change to any key in the map.
 	 */
 	const KeyType* FindKey(ValueInitType Value) const
@@ -379,13 +381,13 @@ public:
 				return &PairIt->Key;
 			}
 		}
-		return NULL;
+		return nullptr;
 	}
 
 	/**
 	 * Returns the value associated with a specified key.
 	 * @param	Key - The key to search for.
-	 * @return	A pointer to the value associated with the specified key, or NULL if the key isn't contained in this map.  The pointer
+	 * @return	A pointer to the value associated with the specified key, or nullptr if the key isn't contained in this map.  The pointer
 	 *			is only valid until the next change to any key in the map.
 	 */
 	FORCEINLINE ValueType* Find(KeyConstPointerType Key)
@@ -395,7 +397,7 @@ public:
 			return &Pair->Value;
 		}
 
-		return NULL;
+		return nullptr;
 	}
 	FORCEINLINE const ValueType* Find(KeyConstPointerType Key) const
 	{
@@ -456,7 +458,7 @@ public:
 	FORCEINLINE const ValueType& FindChecked(KeyConstPointerType Key) const
 	{
 		const auto* Pair = Pairs.Find(Key);
-		check( Pair != NULL );
+		check( Pair != nullptr );
 		return Pair->Value;
 	}
 
@@ -468,7 +470,7 @@ public:
 	FORCEINLINE ValueType& FindChecked(KeyConstPointerType Key)
 	{
 		auto* Pair = Pairs.Find(Key);
-		check( Pair != NULL );
+		check( Pair != nullptr );
 		return Pair->Value;
 	}
 
@@ -1162,7 +1164,7 @@ public:
 	 * Finds an association between a specified key and value. (const)
 	 * @param Key - The key to find.
 	 * @param Value - The value to find.
-	 * @return If the map contains a matching association, a pointer to the value in the map is returned.  Otherwise NULL is returned.
+	 * @return If the map contains a matching association, a pointer to the value in the map is returned.  Otherwise nullptr is returned.
 	 *			The pointer is only valid as long as the map isn't changed.
 	 */
 	FORCEINLINE const ValueType* FindPair(KeyInitType Key,ValueInitType Value) const
@@ -1174,7 +1176,7 @@ public:
 	 * Finds an association between a specified key and value.
 	 * @param Key - The key to find.
 	 * @param Value - The value to find.
-	 * @return If the map contains a matching association, a pointer to the value in the map is returned.  Otherwise NULL is returned.
+	 * @return If the map contains a matching association, a pointer to the value in the map is returned.  Otherwise nullptr is returned.
 	 *			The pointer is only valid as long as the map isn't changed.
 	 */
 	ValueType* FindPair(KeyInitType Key,ValueInitType Value)
@@ -1189,7 +1191,7 @@ public:
 			}
 		}
 
-		return NULL;
+		return nullptr;
 	}
 
 	/** Returns the number of values within this map associated with the specified key */
diff --git a/Engine/Source/Runtime/Core/Public/Containers/MapBuilder.h b/Engine/Source/Runtime/Core/Public/Containers/MapBuilder.h
--- a/Engine/Source/Runtime/Core/Public/Containers/MapBuilder.h
+++ b/Engine/Source/Runtime/Core/Public/Containers/MapBuilder.h
@@ -14,10 +14,8 @@ template<typename KeyType, typename ValueType, typename SetAllocator = FDefaultS
 {
 public:
 
-	/**
-	 * Default constructor.
-	 */
-	TMapBuilder( ) { }
+	/** Default constructor. */
+	TMapBuilder() { }
 
 	/**
 	 * Creates and initializes a new map builder from another map.
@@ -64,7 +62,7 @@ public:
 	 *
 	 * @return A new map.
 	 */
-	TMap<KeyType, ValueType, SetAllocator> Build( )
+	TMap<KeyType, ValueType, SetAllocator> Build()
 	{
 		return Map;
 	}
@@ -74,13 +72,13 @@ public:
 	 *
 	 * @return A new map.
 	 */
-	operator TMap<KeyType, ValueType, SetAllocator>( )
+	operator TMap<KeyType, ValueType, SetAllocator>()
 	{
 		return Build();
 	}
 
 private:
 
-	// Holds the map being built.
+	/** Holds the map being built. */
 	TMap<KeyType, ValueType, SetAllocator> Map;
-};
\ No newline at end of file
+};
diff --git a/Engine/Source/Runtime/Core/Public/Containers/Queue.h b/Engine/Source/Runtime/Core/Public/Containers/Queue.h
--- a/Engine/Source/Runtime/Core/Public/Containers/Queue.h
+++ b/Engine/Source/Runtime/Core/Public/Containers/Queue.h
@@ -32,25 +32,21 @@ namespace EQueueMode
  *
  * @param ItemType The type of items stored in the queue.
  * @param Mode The queue mode (single-producer, single-consumer by default).
-  * @todo gmp: Implement node pooling.
+ * @todo gmp: Implement node pooling.
  */
 template<typename ItemType, EQueueMode::Type Mode = EQueueMode::Spsc>
 class TQueue
 {
 public:
 
-	/**
-	 * Default constructor.
-	 */
-	TQueue( )
+	/** Default constructor. */
+	TQueue()
 	{
 		Head = Tail = new TNode();
 	}
 
-	/**
-	 * Destructor.
-	 */
-	~TQueue( )
+	/** Destructor. */
+	~TQueue()
 	{
 		while (Tail != nullptr)
 		{
@@ -68,6 +64,7 @@ public:
 	 *
 	 * @param OutValue Will hold the returned value.
 	 * @return true if a value was returned, false if the queue was empty.
+	 * @see Enqueue, IsEmpty, Peek
 	 */
 	bool Dequeue( ItemType& OutItem )
 	{
@@ -95,6 +92,7 @@ public:
 	 *
 	 * @param Item The item to add.
 	 * @return true if the item was added, false otherwise.
+	 * @see Dequeue, IsEmpty, Peek
 	 */
 	bool Enqueue( const ItemType& Item )
 	{
@@ -126,8 +124,9 @@ public:
 	 * Checks whether the queue is empty.
 	 *
 	 * @return true if the queue is empty, false otherwise.
+	 * @see Dequeue, Enqueue, Peek
 	 */
-	bool IsEmpty( ) const
+	bool IsEmpty() const
 	{
 		return (Tail->NextNode == nullptr);
 	}
@@ -137,6 +136,7 @@ public:
 	 *
 	 * @param OutItem Will hold the peeked at item.
 	 * @return true if an item was returned, false if the queue was empty.
+	 * @see Dequeue, Enqueue, IsEmpty
 	 */
 	bool Peek( ItemType& OutItem )
 	{
@@ -152,22 +152,22 @@ public:
 
 private:
 
-	// Structure for the internal linked list.
+	/** Structure for the internal linked list. */
 	struct TNode
 	{
-		// Holds a pointer to the next node in the list.
+		/** Holds a pointer to the next node in the list. */
 		TNode* volatile NextNode;
 
-		// Holds the node's item.
+		/** Holds the node's item. */
 		ItemType Item;
 
 
-		// Default constructor.
-		TNode( )
+		/** Default constructor. */
+		TNode()
 			: NextNode(nullptr)
 		{ }
 
-		// Creates and initializes a new node.
+		/** Creates and initializes a new node. */
 		TNode( const ItemType& InItem )
 			: NextNode(nullptr)
 			, Item(InItem)
@@ -175,9 +175,9 @@ private:
 	};
 
 
-	// Holds a pointer to the head of the list.
+	/** Holds a pointer to the head of the list. */
 	MS_ALIGN(16) TNode* volatile Head GCC_ALIGN(16);
 
-	// Holds a pointer to the tail of the list.
+	/** Holds a pointer to the tail of the list. */
 	TNode* Tail;
 };
diff --git a/Engine/Source/Runtime/Core/Public/Containers/ResourceArray.h b/Engine/Source/Runtime/Core/Public/Containers/ResourceArray.h
--- a/Engine/Source/Runtime/Core/Public/Containers/ResourceArray.h
+++ b/Engine/Source/Runtime/Core/Public/Containers/ResourceArray.h
@@ -20,9 +20,7 @@ public:
 	 */
 	virtual uint32 GetResourceDataSize() const = 0;
 
-	/**
-	 * Called on non-UMA systems after the RHI has copied the resource data, and no longer needs the CPU's copy.
-	 */
+	/** Called on non-UMA systems after the RHI has copied the resource data, and no longer needs the CPU's copy. */
 	virtual void Discard() = 0;
 
 	/**
@@ -43,19 +41,22 @@ public:
 
 
 /**
- * allows for direct GPU mem allocation for bulk resource types
+ * Allows for direct GPU mem allocation for bulk resource types.
  */
 class FResourceBulkDataInterface
 {
 public:
+
 	/** 
 	 * @return ptr to the resource memory which has been preallocated
 	 */
 	virtual const void* GetResourceBulkData() const = 0;
+
 	/** 
 	 * @return size of resource memory
 	 */
 	virtual uint32 GetResourceBulkDataSize() const = 0;
+
 	/**
 	 * Free memory after it has been used to initialize RHI resource 
 	 */
@@ -64,28 +65,33 @@ public:
 
 
 /**
-* allows for direct GPU mem allocation for texture resource
-*/
+ * Allows for direct GPU mem allocation for texture resource.
+ */
 class FTexture2DResourceMem : public FResourceBulkDataInterface
 {
 public:
+
 	/**
 	 * @param MipIdx index for mip to retrieve
 	 * @return ptr to the offset in bulk memory for the given mip
 	 */
 	virtual void* GetMipData(int32 MipIdx) = 0;
+
 	/**
 	 * @return total number of mips stored in this resource
 	 */
 	virtual int32	GetNumMips() = 0;
+
 	/** 
 	 * @return width of texture stored in this resource
 	 */
 	virtual int32 GetSizeX() = 0;
+
 	/** 
 	 * @return height of texture stored in this resource
 	 */
 	virtual int32 GetSizeY() = 0;
+
 	/**
 	 * @return Whether the resource memory is properly allocated or not.
 	 */
@@ -96,14 +102,10 @@ public:
 	 */
 	virtual bool HasAsyncAllocationCompleted() const = 0;
 
-	/**
-	 * Blocks the calling thread until the allocation has been completed.
-	 */
+	/** Blocks the calling thread until the allocation has been completed. */
 	virtual void FinishAsyncAllocation() = 0;
 
-	/**
-	 * Cancels any async allocation.
-	 */
+	/** Cancels any async allocation. */
 	virtual void CancelAsyncAllocation() = 0;
 
 	virtual ~FTexture2DResourceMem() {}
diff --git a/Engine/Source/Runtime/Core/Public/Containers/StaticBitArray.h b/Engine/Source/Runtime/Core/Public/Containers/StaticBitArray.h
--- a/Engine/Source/Runtime/Core/Public/Containers/StaticBitArray.h
+++ b/Engine/Source/Runtime/Core/Public/Containers/StaticBitArray.h
@@ -1,11 +1,8 @@
 // Copyright 1998-2015 Epic Games, Inc. All Rights Reserved.
 
-/*=============================================================================
-	StaticBitArray.h: Static bit array definition
-=============================================================================*/
-
 #pragma once
 
+
 /** Used to read/write a bit in the static array as a bool. */
 template<typename T>
 class TStaticBitReference
@@ -47,9 +44,9 @@ class TConstStaticBitReference
 public:
 
 	TConstStaticBitReference(const T& InData,T InMask)
-	:	Data(InData)
-	,	Mask(InMask)
-	{}
+		: Data(InData)
+		, Mask(InMask)
+	{ }
 
 	FORCEINLINE operator bool() const
 	{
@@ -64,7 +61,7 @@ private:
 
 
 /**
- * A statically sized bit array 
+ * A statically sized bit array.
  */
 template<uint32 NumBits>
 class TStaticBitArray
@@ -77,17 +74,13 @@ class TStaticBitArray
 
 public:
 
-	/**
-	 * Minimal initialization constructor
-	 */
+	/** Minimal initialization constructor */
 	FORCEINLINE TStaticBitArray()
 	{
 		Clear_();
 	}
 
-	/**
-	 * Constructor that allows initializing by assignment from 0
-	 */
+	/** Constructor that allows initializing by assignment from 0 */
 	FORCEINLINE TStaticBitArray(UnspecifiedZeroType)
 	{
 		Clear_();
@@ -96,7 +89,7 @@ public:
 	/**
 	 * Constructor to initialize to a single bit
 	 */
-	FORCEINLINE TStaticBitArray(bool,uint32 InBitIndex)
+	FORCEINLINE TStaticBitArray(bool, uint32 InBitIndex)
 	{
 		/***********************************************************************
 
diff --git a/Engine/Source/Runtime/Core/Public/Containers/StringConv.h b/Engine/Source/Runtime/Core/Public/Containers/StringConv.h
--- a/Engine/Source/Runtime/Core/Public/Containers/StringConv.h
+++ b/Engine/Source/Runtime/Core/Public/Containers/StringConv.h
@@ -1,19 +1,19 @@
 // Copyright 1998-2015 Epic Games, Inc. All Rights Reserved.
 
-/**
- *
- * This file contains the classes used when converting strings between
- * standards (ANSI, UNICODE, etc.)
- */
+// This file contains the classes used when converting strings between
+// standards (ANSI, UNICODE, etc.)
 
 #pragma once
+
 #include "Containers/ContainerAllocationPolicies.h"
 #include "HAL/Platform.h"
 #include "HAL/PlatformString.h"
 
+
 #define DEFAULT_STRING_CONVERSION_SIZE 128u
 #define UNICODE_BOGUS_CHAR_CODEPOINT   '?'
 
+
 template <typename From, typename To>
 class TStringConvert
 {
diff --git a/Engine/Source/Runtime/Core/Public/Containers/Ticker.h b/Engine/Source/Runtime/Core/Public/Containers/Ticker.h
--- a/Engine/Source/Runtime/Core/Public/Containers/Ticker.h
+++ b/Engine/Source/Runtime/Core/Public/Containers/Ticker.h
@@ -24,13 +24,12 @@ public:
 
 	FTicker()
 		: CurrentTime(0.0)
-	{
-	}
+	{ }
 
 	/**
 	 * Add a new ticker with a given delay / interval
-	 * @param InDelegate	delegate to fire after the delay
-	 * @param InDelay		Delay until next fire; 0 means "next frame"
+	 * @param InDelegate Delegate to fire after the delay
+	 * @param InDelay Delay until next fire; 0 means "next frame"
 	 */
 	FDelegateHandle AddTicker(const FTickerDelegate& InDelegate, float InDelay = 0.0f)
 	{
@@ -42,7 +41,7 @@ public:
 	/**
 	 * Removes a previously added ticker delegate.
 	 *
-	 * @param Handle - The handle of the ticker to remove.
+	 * @param Handle The handle of the ticker to remove.
 	 */
 	void RemoveTicker(FDelegateHandle Handle)
 	{
@@ -52,7 +51,7 @@ public:
 	/**
 	 * Removes a previously added ticker delegate.
 	 *
-	 * @param Delegate - The delegate to remove.
+	 * @param Delegate The delegate to remove.
 	 */
 	DELEGATE_DEPRECATED("This RemoveTicker overload is deprecated - please remove tickers using the FDelegateHandle returned by the AddTicker function.")
 	void RemoveTicker(const FTickerDelegate& Delegate)
@@ -137,6 +136,7 @@ private:
 	TArray<FElement> PriorityQueue;
 };
 
+
 /**
  * Base class for ticker objects
  */
@@ -156,9 +156,7 @@ public:
 		TickHandle = FTicker::GetCoreTicker().AddTicker(TickDelegate, InDelay);
 	}
 
-	/**
-	 * Destructor
-	 */
+	/** Virtual destructor. */
 	virtual ~FTickerObjectBase()
 	{
 		// Unregister ticker delegate
diff --git a/Engine/Source/Runtime/Core/Public/Containers/UnrealString.h b/Engine/Source/Runtime/Core/Public/Containers/UnrealString.h
--- a/Engine/Source/Runtime/Core/Public/Containers/UnrealString.h
+++ b/Engine/Source/Runtime/Core/Public/Containers/UnrealString.h
@@ -1,12 +1,10 @@
 // Copyright 1998-2015 Epic Games, Inc. All Rights Reserved.
 
-/*=============================================================================
-	UnrealString.h: Dynamic string definitions.  This needed to be 
-		UnrealString.h to avoid conflicting with the Windows platform 
-		SDK string.h
-=============================================================================*/
+// This needed to be UnrealString.h to avoid conflicting with
+// the Windows platform SDK string.h
 
 #pragma once
+
 #include "Containers/Array.h"
 #include "Math/NumericLimits.h"
 #include "Math/UnrealMathUtility.h"
@@ -14,6 +12,7 @@
 #include "Misc/CString.h"
 #include "Templates/MemoryOps.h"
 
+
 /** Determines case sensitivity options for string comparisons. */
 namespace ESearchCase
 {
@@ -121,7 +120,7 @@ public:
 	 * @param In array of TCHAR
 	 */
 	template <typename CharType>
-	FORCEINLINE FString(const CharType* Src, typename TEnableIf<TIsCharType<CharType>::Value>::Type* Dummy = NULL) // This TEnableIf is to ensure we don't instantiate this constructor for non-char types, like id* in Obj-C
+	FORCEINLINE FString(const CharType* Src, typename TEnableIf<TIsCharType<CharType>::Value>::Type* Dummy = nullptr) // This TEnableIf is to ensure we don't instantiate this constructor for non-char types, like id* in Obj-C
 	{
 		if (Src && *Src)
 		{
@@ -166,7 +165,7 @@ public:
 #else
 			FMemory::Memcpy(Data.GetData(), [In cStringUsingEncoding:NSUTF16StringEncoding], Len * sizeof(TCHAR));
 #endif
-			// for non-UTF-8 encodings this may not be terminated with a NULL!
+			// for non-UTF-8 encodings this may not be terminated with a nullptr!
 			Data[Len-1] = '\0';
 		}
 	}
@@ -1192,7 +1191,7 @@ public:
 	/**
 	 * Returns a copy of this string with wrapping quotation marks removed.
 	 */
-	FString TrimQuotes( bool* bQuotesRemoved=NULL ) const;
+	FString TrimQuotes( bool* bQuotesRemoved = nullptr ) const;
 
 	/**
 	 * Breaks up a delimited string into elements of a string array.
@@ -1215,7 +1214,7 @@ public:
 	 *
 	 * @return	The number of elements in InArray
 	 */
-	int32 ParseIntoArrayWS( TArray<FString>* InArray, const TCHAR* pchExtraDelim = NULL ) const;
+	int32 ParseIntoArrayWS( TArray<FString>* InArray, const TCHAR* pchExtraDelim = nullptr ) const;
 
 	/**
 	* Breaks up a delimited string into elements of a string array, using line ending characters
@@ -1292,14 +1291,14 @@ public:
 	 *
 	 * @return	a string with all control characters replaced by the escaped version.
 	 */
-	FString ReplaceCharWithEscapedChar( const TArray<TCHAR>* Chars=NULL ) const;
+	FString ReplaceCharWithEscapedChar( const TArray<TCHAR>* Chars = nullptr ) const;
 
 	/**
 	 * Removes the escape backslash for all supported characters, replacing the escape and character with the non-escaped version.  (i.e.
 	 * replaces "\\n" with "\n".  Counterpart to ReplaceCharWithEscapedChar().
 	 * @return copy of this string with replacement made
 	 */
-	FString ReplaceEscapedCharWithChar( const TArray<TCHAR>* Chars=NULL ) const;
+	FString ReplaceEscapedCharWithChar( const TArray<TCHAR>* Chars = nullptr ) const;
 
 	/** 
 	 * Replaces all instances of '\t' with TabWidth number of spaces
@@ -1586,7 +1585,7 @@ namespace LexicalConversion
 	inline void FromString(uint8& OutValue,		const TCHAR* Buffer)	{	OutValue = FCString::Atoi(Buffer);		}
 	inline void FromString(uint16& OutValue, 	const TCHAR* Buffer)	{	OutValue = FCString::Atoi(Buffer);		}
 	inline void FromString(uint32& OutValue, 	const TCHAR* Buffer)	{	OutValue = FCString::Atoi64(Buffer);	}	//64 because this unsigned and so Atoi might overflow
-	inline void FromString(uint64& OutValue, 	const TCHAR* Buffer)	{	OutValue = FCString::Strtoui64(Buffer, NULL, 0); }
+	inline void FromString(uint64& OutValue, 	const TCHAR* Buffer)	{	OutValue = FCString::Strtoui64(Buffer, nullptr, 0); }
 	inline void FromString(float& OutValue,		const TCHAR* Buffer)	{	OutValue = FCString::Atof(Buffer);		}
 	inline void FromString(double& OutValue, 	const TCHAR* Buffer)	{	OutValue = FCString::Atod(Buffer);		}
 	inline void FromString(bool& OutValue, 		const TCHAR* Buffer)	{	OutValue = FCString::ToBool(Buffer);	}
