Adding a cache to speed up multiple reads of the ini files.
Fixes UE-6759.

diff --git a/Engine/Source/Runtime/Core/Private/Misc/ConfigCacheIni.cpp b/Engine/Source/Runtime/Core/Private/Misc/ConfigCacheIni.cpp
--- a/Engine/Source/Runtime/Core/Private/Misc/ConfigCacheIni.cpp
+++ b/Engine/Source/Runtime/Core/Private/Misc/ConfigCacheIni.cpp
@@ -10,9 +10,42 @@
 #include "EngineBuildSettings.h"
 #include "Paths.h"
 
+#ifdef WITH_EDITOR
+#	define INI_CACHE (1)
+#else
+#	define INI_CACHE (0)
+#endif
 
 DEFINE_LOG_CATEGORY(LogConfig);
+namespace 
+{
+	FName GenerateHierarchyCacheKey(const TArray<FIniFilename>& IniHierarchy, const FString& IniPath) 
+	{
+#if !INI_CACHE
+		return FName();
+#else
+		// A Hierarchy Key is a combined list of all ini file paths that affect that inis data set.
+		FString HierKey;
+		//
+		auto KeyLen = IniPath.Len();
+		for (const auto& Ini : IniHierarchy)
+		{
+			KeyLen += Ini.Filename.Len();
+		}
+		HierKey.Reserve(KeyLen);
+		for (const auto& Ini : IniHierarchy)
+		{
+			HierKey += Ini.Filename;
+		}
+		HierKey += IniPath;
 
+		return FName(*HierKey);
+#endif
+	}
+#if INI_CACHE
+	TMap<FName, FConfigFile> HierarchyCache;
+#endif
+}
 
 class FTextFriendHelper
 {
@@ -664,33 +697,83 @@ static bool LoadIniFileHierarchy(const TArray<FIniFilename>& HierarchyToLoad, FC
 		}
 	}
 
+	int32 FirstCacheIndex = 0;
+#if INI_CACHE
+	// Find the last value in the hierarchy that is cached. We can start the load from there
+	for( int32 IniIndex = 0; IniIndex < HierarchyToLoad.Num(); IniIndex++ )
+	{
+		if (HierarchyCache.Find(HierarchyToLoad[IniIndex].CacheKey)) 
+		{
+			FirstCacheIndex = IniIndex;
+		}
+	}
+#endif
+
 	TArray<FDateTime> TimestampsOfInis;
 
 	// Traverse ini list back to front, merging along the way.
-	for( int32 IniIndex = 0; IniIndex < HierarchyToLoad.Num(); IniIndex++ )
+	for( int32 IniIndex = FirstCacheIndex; IniIndex < HierarchyToLoad.Num(); IniIndex++ )
 	{
 		const FIniFilename& IniToLoad = HierarchyToLoad[ IniIndex ];
 		const FString& IniFileName = IniToLoad.Filename;
-
-		// Spit out friendly error if there was a problem locating .inis (e.g. bad command line parameter or missing folder, ...).
-		if (IsUsingLocalIniFile(*IniFileName, NULL) && (IFileManager::Get().FileSize(*IniFileName) < 0))
+		bool bDoProcess = true;
+#if INI_CACHE
+		if (!IniToLoad.CacheKey.IsNone())
 		{
-			if (IniToLoad.bRequired)
+			auto* CachedConfigFile = HierarchyCache.Find(IniToLoad.CacheKey);
+			if (CachedConfigFile) 
 			{
-				//UE_LOG(LogConfig, Error, TEXT("Couldn't locate '%s' which is required to run '%s'"), *IniToLoad.Filename, FApp::GetGameName() );
-				return false;
+				ConfigFile = *CachedConfigFile;
+				bDoProcess = false;
 			}
-			else
+			ConfigFile.CacheKey = IniToLoad.CacheKey;
+		}
+		else
+		{
+			ConfigFile.CacheKey = FName();
+		}
+#endif
+		if (bDoProcess)
+		{
+			// Spit out friendly error if there was a problem locating .inis (e.g. bad command line parameter or missing folder, ...).
+			if (IsUsingLocalIniFile(*IniFileName, NULL) && (IFileManager::Get().FileSize(*IniFileName) < 0))
 			{
-				continue;
+				if (IniToLoad.bRequired)
+				{
+					//UE_LOG(LogConfig, Error, TEXT("Couldn't locate '%s' which is required to run '%s'"), *IniToLoad.Filename, FApp::GetGameName() );
+					return false;
+				}
+				else
+				{
+#if INI_CACHE
+					if (!IniToLoad.CacheKey.IsNone())
+					{
+						if (bDoProcess)
+						{
+							HierarchyCache.Add(IniToLoad.CacheKey, ConfigFile);
+						}
+					}
+#endif
+					continue;
+				}
 			}
+
+			bool bDoEmptyConfig = false;
+			bool bDoCombine = (IniIndex != 0);
+			bool bDoWrite = false;
+			//UE_LOG(LogConfig, Log,  TEXT( "Combining configFile: %s" ), *IniList(IniIndex) );
+			ProcessIniContents(*HierarchyToLoad.Last().Filename, *IniFileName, &ConfigFile, bDoEmptyConfig, bDoCombine, bDoWrite);
 		}
 
-		bool bDoEmptyConfig = false;
-		bool bDoCombine = (IniIndex != 0);
-		bool bDoWrite = false;
-		//UE_LOG(LogConfig, Log,  TEXT( "Combining configFile: %s" ), *IniList(IniIndex) );
-		ProcessIniContents(*HierarchyToLoad.Last().Filename, *IniFileName, &ConfigFile, bDoEmptyConfig, bDoCombine, bDoWrite);
+#if INI_CACHE
+		if (!IniToLoad.CacheKey.IsNone())
+		{
+			if (bDoProcess)
+			{
+				HierarchyCache.Add(IniToLoad.CacheKey, ConfigFile);
+			}
+		}
+#endif
 	}
 
 	// Set this configs files source ini hierarchy to show where it was loaded from.
@@ -902,6 +985,14 @@ bool FConfigFile::Write( const FString& Filename, bool bDoRemoteWrite/* = true*/
 	}
 	bool bResult = FFileHelper::SaveStringToFile( Text, *Filename );
 
+#if INI_CACHE
+	// if we wrote the config successfully
+	if (bResult && !CacheKey.IsNone()) 
+	{
+		HierarchyCache.Add(CacheKey, *this);
+	}
+#endif
+
 	// File is still dirty if it didn't save.
 	Dirty = !bResult;
 
@@ -2631,12 +2722,14 @@ static void GetSourceIniHierarchyFilenames(const TCHAR* InBaseIniName, const TCH
 		OutHierarchy.Add(FIniFilename(FString::Printf(TEXT("%s/DedicatedServer%s.ini"), SourceConfigDir, InBaseIniName), false));
 	}
 	// Game/Config/User* ini
-	OutHierarchy.Add(FIniFilename(FString::Printf(TEXT("%s/User%s.ini"), SourceConfigDir, InBaseIniName), false));
+	FString FullIniPath = FString::Printf(TEXT("%s/User%s.ini"), SourceConfigDir, InBaseIniName);
+	OutHierarchy.Add(FIniFilename(FullIniPath, false, GenerateHierarchyCacheKey(OutHierarchy, FullIniPath)));
 	
 	// Engine/Config/Platform/Platform* ini
 	OutHierarchy.Add( FIniFilename(FString::Printf(TEXT("%s%s/%s%s.ini"), EngineConfigDir, *PlatformName, *PlatformName, InBaseIniName), false) );
 	// Game/Config/Platform/Platform* ini
-	OutHierarchy.Add( FIniFilename(FString::Printf(TEXT("%s%s/%s%s.ini"), SourceConfigDir, *PlatformName, *PlatformName, InBaseIniName), false) );
+	FullIniPath = FString::Printf(TEXT("%s%s/%s%s.ini"), SourceConfigDir, *PlatformName, *PlatformName, InBaseIniName);
+	OutHierarchy.Add(FIniFilename(FullIniPath, false, GenerateHierarchyCacheKey(OutHierarchy, FullIniPath)) );
 }
 
 /**
@@ -2674,7 +2767,7 @@ void FConfigCacheIni::InitializeConfigSystem()
 {
 	// create GConfig
 	GConfig = new FConfigCacheIni;
-
+	
 	// load the main .ini files (unless we're running a program or a gameless UE4Editor.exe, DefaultEngine.ini is required).
 	const bool bIsGamelessExe = !FApp::HasGameName();
 	const bool bDefaultEngineIniRequired = !bIsGamelessExe && (GIsGameAgnosticExe || FApp::IsGameNameEmpty());
@@ -2767,7 +2860,7 @@ bool FConfigCacheIni::LoadGlobalIniFile(FString& FinalIniFilename, const TCHAR*
 
 	// Keep a record of the original settings
 	NewConfigFile.SourceConfigFile = new FConfigFile();
-	LoadIniFileHierarchy(NewConfigFile.SourceIniHierarchy, *NewConfigFile.SourceConfigFile );
+	LoadIniFileHierarchy(NewConfigFile.SourceIniHierarchy, *NewConfigFile.SourceConfigFile);
 
 	// now generate and make sure it's up to date
 	bool bResult = GenerateDestIniFile(NewConfigFile, FinalIniFilename, NewConfigFile.SourceIniHierarchy, bAllowGeneratedIniWhenCooked);
@@ -2809,7 +2902,7 @@ void FConfigCacheIni::LoadLocalIniFile(FConfigFile& ConfigFile, const TCHAR* Ini
 		
 		// Keep a record of the original settings
 		ConfigFile.SourceConfigFile = new FConfigFile();
-		LoadIniFileHierarchy(ConfigFile.SourceIniHierarchy, *ConfigFile.SourceConfigFile );
+		LoadIniFileHierarchy(ConfigFile.SourceIniHierarchy, *ConfigFile.SourceConfigFile);
 
 		// now generate and make sure it's up to date (using IniName as a Base for an ini filename)
 		const bool bAllowGeneratedINIs = true;
@@ -2835,7 +2928,7 @@ void FConfigCacheIni::LoadExternalIniFile(FConfigFile& ConfigFile, const TCHAR*
 
 		// Keep a record of the original settings
 		ConfigFile.SourceConfigFile = new FConfigFile();
-		LoadIniFileHierarchy(ConfigFile.SourceIniHierarchy, *ConfigFile.SourceConfigFile );
+		LoadIniFileHierarchy(ConfigFile.SourceIniHierarchy, *ConfigFile.SourceConfigFile);
 
 		// now generate and make sure it's up to date (using IniName as a Base for an ini filename)
 		const bool bAllowGeneratedINIs = true;
diff --git a/Engine/Source/Runtime/Core/Public/Misc/ConfigCacheIni.h b/Engine/Source/Runtime/Core/Public/Misc/ConfigCacheIni.h
--- a/Engine/Source/Runtime/Core/Public/Misc/ConfigCacheIni.h
+++ b/Engine/Source/Runtime/Core/Public/Misc/ConfigCacheIni.h
@@ -31,10 +31,13 @@ struct FIniFilename
 	FString Filename;
 	/** If true this ini file is required to generate the output ini. */
 	bool bRequired;
+	/** Used as ID for looking up an INI Hierarchy */
+	FName CacheKey;
 
-	FIniFilename(const FString& InFilename, bool InIsRequired)
+	FIniFilename(const FString& InFilename, bool InIsRequired, FName InCacheKey=FName())
 		: Filename(InFilename)
 		, bRequired(InIsRequired) 
+		, CacheKey(InCacheKey)
 	{}
 };
 
@@ -63,6 +66,9 @@ public:
 	/** The untainted config file which contains the coalesced base/default options. I.e. No Saved/ options*/
 	FConfigFile* SourceConfigFile;
 
+	/** Key to the cache to speed up ini parsing */
+	FName CacheKey;
+
 #if !UE_BUILD_SHIPPING
 	/** The collection of overrides which stemmed from the commandline */
 	TArray<FConfigCommandlineOverride> CommandlineOptions;
