UHT: Allows defining multiple UCLASSes in one header file.

diff --git a/Engine/Source/Runtime/Core/Public/Delegates/DelegateCombinations_Variadics.h b/Engine/Source/Runtime/Core/Public/Delegates/DelegateCombinations_Variadics.h
--- a/Engine/Source/Runtime/Core/Public/Delegates/DelegateCombinations_Variadics.h
+++ b/Engine/Source/Runtime/Core/Public/Delegates/DelegateCombinations_Variadics.h
@@ -1,69 +1,67 @@
 // Copyright 1998-2015 Epic Games, Inc. All Rights Reserved.
 
-// NOTE: This source file was automatically generated using DelegateHeaderTool
-
 #pragma once
 
 #define DECLARE_DELEGATE( DelegateName ) FUNC_DECLARE_DELEGATE( NoParams, DelegateName, void )
 #define DECLARE_MULTICAST_DELEGATE( DelegateName ) FUNC_DECLARE_MULTICAST_DELEGATE( NoParams, DelegateName, void )
 #define DECLARE_EVENT( OwningType, EventName ) FUNC_DECLARE_EVENT( OwningType, EventName, NoParams, void )
-#define DECLARE_DYNAMIC_DELEGATE( DelegateName ) FUNC_DECLARE_DYNAMIC_DELEGATE( FWeakObjectPtr, NoParams, DelegateName, DelegateName##_DelegateWrapper, , FUNC_CONCAT( *this ), void )
-#define DECLARE_DYNAMIC_MULTICAST_DELEGATE( DelegateName ) FUNC_DECLARE_DYNAMIC_MULTICAST_DELEGATE( FWeakObjectPtr, NoParams, DelegateName, DelegateName##_DelegateWrapper, , FUNC_CONCAT( *this ), void )
+#define DECLARE_DYNAMIC_DELEGATE( DelegateName ) BODY_MACRO_COMBINE(CURRENT_FILE_ID,_,__LINE__,_DELEGATE) FUNC_DECLARE_DYNAMIC_DELEGATE( FWeakObjectPtr, NoParams, DelegateName, DelegateName##_DelegateWrapper, , FUNC_CONCAT( *this ), void )
+#define DECLARE_DYNAMIC_MULTICAST_DELEGATE( DelegateName ) BODY_MACRO_COMBINE(CURRENT_FILE_ID,_,__LINE__,_DELEGATE) FUNC_DECLARE_DYNAMIC_MULTICAST_DELEGATE( FWeakObjectPtr, NoParams, DelegateName, DelegateName##_DelegateWrapper, , FUNC_CONCAT( *this ), void )
 #define DECLARE_DELEGATE_RetVal( ReturnValueType, DelegateName ) FUNC_DECLARE_DELEGATE( RetVal_NoParams, DelegateName, ReturnValueType )
-#define DECLARE_DYNAMIC_DELEGATE_RetVal( ReturnValueType, DelegateName ) FUNC_DECLARE_DYNAMIC_DELEGATE_RETVAL( FWeakObjectPtr, RetVal_NoParams, DelegateName, DelegateName##_DelegateWrapper, ReturnValueType, , FUNC_CONCAT( *this ), ReturnValueType )
+#define DECLARE_DYNAMIC_DELEGATE_RetVal( ReturnValueType, DelegateName ) BODY_MACRO_COMBINE(CURRENT_FILE_ID,_,__LINE__,_DELEGATE) FUNC_DECLARE_DYNAMIC_DELEGATE_RETVAL( FWeakObjectPtr, RetVal_NoParams, DelegateName, DelegateName##_DelegateWrapper, ReturnValueType, , FUNC_CONCAT( *this ), ReturnValueType )
 #define DECLARE_DELEGATE_OneParam( DelegateName, Param1Type ) FUNC_DECLARE_DELEGATE( OneParam, DelegateName, void, Param1Type )
 #define DECLARE_MULTICAST_DELEGATE_OneParam( DelegateName, Param1Type ) FUNC_DECLARE_MULTICAST_DELEGATE( OneParam, DelegateName, void, Param1Type )
 #define DECLARE_EVENT_OneParam( OwningType, EventName, Param1Type ) FUNC_DECLARE_EVENT( OwningType, EventName, OneParam, void, Param1Type )
-#define DECLARE_DYNAMIC_DELEGATE_OneParam( DelegateName, Param1Type, Param1Name ) FUNC_DECLARE_DYNAMIC_DELEGATE( FWeakObjectPtr, OneParam, DelegateName, DelegateName##_DelegateWrapper, FUNC_CONCAT( Param1Type InParam1 ), FUNC_CONCAT( *this, InParam1 ), void, Param1Type )
-#define DECLARE_DYNAMIC_MULTICAST_DELEGATE_OneParam( DelegateName, Param1Type, Param1Name ) FUNC_DECLARE_DYNAMIC_MULTICAST_DELEGATE( FWeakObjectPtr, OneParam, DelegateName, DelegateName##_DelegateWrapper, FUNC_CONCAT( Param1Type InParam1 ), FUNC_CONCAT( *this, InParam1 ), void, Param1Type )
+#define DECLARE_DYNAMIC_DELEGATE_OneParam( DelegateName, Param1Type, Param1Name ) BODY_MACRO_COMBINE(CURRENT_FILE_ID,_,__LINE__,_DELEGATE) FUNC_DECLARE_DYNAMIC_DELEGATE( FWeakObjectPtr, OneParam, DelegateName, DelegateName##_DelegateWrapper, FUNC_CONCAT( Param1Type InParam1 ), FUNC_CONCAT( *this, InParam1 ), void, Param1Type )
+#define DECLARE_DYNAMIC_MULTICAST_DELEGATE_OneParam( DelegateName, Param1Type, Param1Name ) BODY_MACRO_COMBINE(CURRENT_FILE_ID,_,__LINE__,_DELEGATE) FUNC_DECLARE_DYNAMIC_MULTICAST_DELEGATE( FWeakObjectPtr, OneParam, DelegateName, DelegateName##_DelegateWrapper, FUNC_CONCAT( Param1Type InParam1 ), FUNC_CONCAT( *this, InParam1 ), void, Param1Type )
 #define DECLARE_DELEGATE_RetVal_OneParam( ReturnValueType, DelegateName, Param1Type ) FUNC_DECLARE_DELEGATE( RetVal_OneParam, DelegateName, ReturnValueType, Param1Type )
-#define DECLARE_DYNAMIC_DELEGATE_RetVal_OneParam( ReturnValueType, DelegateName, Param1Type, Param1Name ) FUNC_DECLARE_DYNAMIC_DELEGATE_RETVAL( FWeakObjectPtr, RetVal_OneParam, DelegateName, DelegateName##_DelegateWrapper, ReturnValueType, FUNC_CONCAT( Param1Type InParam1 ), FUNC_CONCAT( *this, InParam1 ), ReturnValueType, Param1Type )
+#define DECLARE_DYNAMIC_DELEGATE_RetVal_OneParam( ReturnValueType, DelegateName, Param1Type, Param1Name ) BODY_MACRO_COMBINE(CURRENT_FILE_ID,_,__LINE__,_DELEGATE) FUNC_DECLARE_DYNAMIC_DELEGATE_RETVAL( FWeakObjectPtr, RetVal_OneParam, DelegateName, DelegateName##_DelegateWrapper, ReturnValueType, FUNC_CONCAT( Param1Type InParam1 ), FUNC_CONCAT( *this, InParam1 ), ReturnValueType, Param1Type )
 #define DECLARE_DELEGATE_TwoParams( DelegateName, Param1Type, Param2Type ) FUNC_DECLARE_DELEGATE( TwoParams, DelegateName, void, Param1Type, Param2Type )
 #define DECLARE_MULTICAST_DELEGATE_TwoParams( DelegateName, Param1Type, Param2Type ) FUNC_DECLARE_MULTICAST_DELEGATE( TwoParams, DelegateName, void, Param1Type, Param2Type )
 #define DECLARE_EVENT_TwoParams( OwningType, EventName, Param1Type, Param2Type ) FUNC_DECLARE_EVENT( OwningType, EventName, TwoParams, void, Param1Type, Param2Type )
-#define DECLARE_DYNAMIC_DELEGATE_TwoParams( DelegateName, Param1Type, Param1Name, Param2Type, Param2Name ) FUNC_DECLARE_DYNAMIC_DELEGATE( FWeakObjectPtr, TwoParams, DelegateName, DelegateName##_DelegateWrapper, FUNC_CONCAT( Param1Type InParam1, Param2Type InParam2 ), FUNC_CONCAT( *this, InParam1, InParam2 ), void, Param1Type, Param2Type )
-#define DECLARE_DYNAMIC_MULTICAST_DELEGATE_TwoParams( DelegateName, Param1Type, Param1Name, Param2Type, Param2Name ) FUNC_DECLARE_DYNAMIC_MULTICAST_DELEGATE( FWeakObjectPtr, TwoParams, DelegateName, DelegateName##_DelegateWrapper, FUNC_CONCAT( Param1Type InParam1, Param2Type InParam2 ), FUNC_CONCAT( *this, InParam1, InParam2 ), void, Param1Type, Param2Type )
+#define DECLARE_DYNAMIC_DELEGATE_TwoParams( DelegateName, Param1Type, Param1Name, Param2Type, Param2Name ) BODY_MACRO_COMBINE(CURRENT_FILE_ID,_,__LINE__,_DELEGATE) FUNC_DECLARE_DYNAMIC_DELEGATE( FWeakObjectPtr, TwoParams, DelegateName, DelegateName##_DelegateWrapper, FUNC_CONCAT( Param1Type InParam1, Param2Type InParam2 ), FUNC_CONCAT( *this, InParam1, InParam2 ), void, Param1Type, Param2Type )
+#define DECLARE_DYNAMIC_MULTICAST_DELEGATE_TwoParams( DelegateName, Param1Type, Param1Name, Param2Type, Param2Name ) BODY_MACRO_COMBINE(CURRENT_FILE_ID,_,__LINE__,_DELEGATE) FUNC_DECLARE_DYNAMIC_MULTICAST_DELEGATE( FWeakObjectPtr, TwoParams, DelegateName, DelegateName##_DelegateWrapper, FUNC_CONCAT( Param1Type InParam1, Param2Type InParam2 ), FUNC_CONCAT( *this, InParam1, InParam2 ), void, Param1Type, Param2Type )
 #define DECLARE_DELEGATE_RetVal_TwoParams( ReturnValueType, DelegateName, Param1Type, Param2Type ) FUNC_DECLARE_DELEGATE( RetVal_TwoParams, DelegateName, ReturnValueType, Param1Type, Param2Type )
-#define DECLARE_DYNAMIC_DELEGATE_RetVal_TwoParams( ReturnValueType, DelegateName, Param1Type, Param1Name, Param2Type, Param2Name ) FUNC_DECLARE_DYNAMIC_DELEGATE_RETVAL( FWeakObjectPtr, RetVal_TwoParams, DelegateName, DelegateName##_DelegateWrapper, ReturnValueType, FUNC_CONCAT( Param1Type InParam1, Param2Type InParam2 ), FUNC_CONCAT( *this, InParam1, InParam2 ), ReturnValueType, Param1Type, Param2Type )
+#define DECLARE_DYNAMIC_DELEGATE_RetVal_TwoParams( ReturnValueType, DelegateName, Param1Type, Param1Name, Param2Type, Param2Name ) BODY_MACRO_COMBINE(CURRENT_FILE_ID,_,__LINE__,_DELEGATE) FUNC_DECLARE_DYNAMIC_DELEGATE_RETVAL( FWeakObjectPtr, RetVal_TwoParams, DelegateName, DelegateName##_DelegateWrapper, ReturnValueType, FUNC_CONCAT( Param1Type InParam1, Param2Type InParam2 ), FUNC_CONCAT( *this, InParam1, InParam2 ), ReturnValueType, Param1Type, Param2Type )
 #define DECLARE_DELEGATE_ThreeParams( DelegateName, Param1Type, Param2Type, Param3Type ) FUNC_DECLARE_DELEGATE( ThreeParams, DelegateName, void, Param1Type, Param2Type, Param3Type )
 #define DECLARE_MULTICAST_DELEGATE_ThreeParams( DelegateName, Param1Type, Param2Type, Param3Type ) FUNC_DECLARE_MULTICAST_DELEGATE( ThreeParams, DelegateName, void, Param1Type, Param2Type, Param3Type )
 #define DECLARE_EVENT_ThreeParams( OwningType, EventName, Param1Type, Param2Type, Param3Type ) FUNC_DECLARE_EVENT( OwningType, EventName, ThreeParams, void, Param1Type, Param2Type, Param3Type )
-#define DECLARE_DYNAMIC_DELEGATE_ThreeParams( DelegateName, Param1Type, Param1Name, Param2Type, Param2Name, Param3Type, Param3Name ) FUNC_DECLARE_DYNAMIC_DELEGATE( FWeakObjectPtr, ThreeParams, DelegateName, DelegateName##_DelegateWrapper, FUNC_CONCAT( Param1Type InParam1, Param2Type InParam2, Param3Type InParam3 ), FUNC_CONCAT( *this, InParam1, InParam2, InParam3 ), void, Param1Type, Param2Type, Param3Type )
-#define DECLARE_DYNAMIC_MULTICAST_DELEGATE_ThreeParams( DelegateName, Param1Type, Param1Name, Param2Type, Param2Name, Param3Type, Param3Name ) FUNC_DECLARE_DYNAMIC_MULTICAST_DELEGATE( FWeakObjectPtr, ThreeParams, DelegateName, DelegateName##_DelegateWrapper, FUNC_CONCAT( Param1Type InParam1, Param2Type InParam2, Param3Type InParam3 ), FUNC_CONCAT( *this, InParam1, InParam2, InParam3 ), void, Param1Type, Param2Type, Param3Type )
+#define DECLARE_DYNAMIC_DELEGATE_ThreeParams( DelegateName, Param1Type, Param1Name, Param2Type, Param2Name, Param3Type, Param3Name ) BODY_MACRO_COMBINE(CURRENT_FILE_ID,_,__LINE__,_DELEGATE) FUNC_DECLARE_DYNAMIC_DELEGATE( FWeakObjectPtr, ThreeParams, DelegateName, DelegateName##_DelegateWrapper, FUNC_CONCAT( Param1Type InParam1, Param2Type InParam2, Param3Type InParam3 ), FUNC_CONCAT( *this, InParam1, InParam2, InParam3 ), void, Param1Type, Param2Type, Param3Type )
+#define DECLARE_DYNAMIC_MULTICAST_DELEGATE_ThreeParams( DelegateName, Param1Type, Param1Name, Param2Type, Param2Name, Param3Type, Param3Name ) BODY_MACRO_COMBINE(CURRENT_FILE_ID,_,__LINE__,_DELEGATE) FUNC_DECLARE_DYNAMIC_MULTICAST_DELEGATE( FWeakObjectPtr, ThreeParams, DelegateName, DelegateName##_DelegateWrapper, FUNC_CONCAT( Param1Type InParam1, Param2Type InParam2, Param3Type InParam3 ), FUNC_CONCAT( *this, InParam1, InParam2, InParam3 ), void, Param1Type, Param2Type, Param3Type )
 #define DECLARE_DELEGATE_RetVal_ThreeParams( ReturnValueType, DelegateName, Param1Type, Param2Type, Param3Type ) FUNC_DECLARE_DELEGATE( RetVal_ThreeParams, DelegateName, ReturnValueType, Param1Type, Param2Type, Param3Type )
-#define DECLARE_DYNAMIC_DELEGATE_RetVal_ThreeParams( ReturnValueType, DelegateName, Param1Type, Param1Name, Param2Type, Param2Name, Param3Type, Param3Name ) FUNC_DECLARE_DYNAMIC_DELEGATE_RETVAL( FWeakObjectPtr, RetVal_ThreeParams, DelegateName, DelegateName##_DelegateWrapper, ReturnValueType, FUNC_CONCAT( Param1Type InParam1, Param2Type InParam2, Param3Type InParam3 ), FUNC_CONCAT( *this, InParam1, InParam2, InParam3 ), ReturnValueType, Param1Type, Param2Type, Param3Type )
+#define DECLARE_DYNAMIC_DELEGATE_RetVal_ThreeParams( ReturnValueType, DelegateName, Param1Type, Param1Name, Param2Type, Param2Name, Param3Type, Param3Name ) BODY_MACRO_COMBINE(CURRENT_FILE_ID,_,__LINE__,_DELEGATE) FUNC_DECLARE_DYNAMIC_DELEGATE_RETVAL( FWeakObjectPtr, RetVal_ThreeParams, DelegateName, DelegateName##_DelegateWrapper, ReturnValueType, FUNC_CONCAT( Param1Type InParam1, Param2Type InParam2, Param3Type InParam3 ), FUNC_CONCAT( *this, InParam1, InParam2, InParam3 ), ReturnValueType, Param1Type, Param2Type, Param3Type )
 #define DECLARE_DELEGATE_FourParams( DelegateName, Param1Type, Param2Type, Param3Type, Param4Type ) FUNC_DECLARE_DELEGATE( FourParams, DelegateName, void, Param1Type, Param2Type, Param3Type, Param4Type )
 #define DECLARE_MULTICAST_DELEGATE_FourParams( DelegateName, Param1Type, Param2Type, Param3Type, Param4Type ) FUNC_DECLARE_MULTICAST_DELEGATE( FourParams, DelegateName, void, Param1Type, Param2Type, Param3Type, Param4Type )
 #define DECLARE_EVENT_FourParams( OwningType, EventName, Param1Type, Param2Type, Param3Type, Param4Type ) FUNC_DECLARE_EVENT( OwningType, EventName, FourParams, void, Param1Type, Param2Type, Param3Type, Param4Type )
-#define DECLARE_DYNAMIC_DELEGATE_FourParams( DelegateName, Param1Type, Param1Name, Param2Type, Param2Name, Param3Type, Param3Name, Param4Type, Param4Name ) FUNC_DECLARE_DYNAMIC_DELEGATE( FWeakObjectPtr, FourParams, DelegateName, DelegateName##_DelegateWrapper, FUNC_CONCAT( Param1Type InParam1, Param2Type InParam2, Param3Type InParam3, Param4Type InParam4 ), FUNC_CONCAT( *this, InParam1, InParam2, InParam3, InParam4 ), void, Param1Type, Param2Type, Param3Type, Param4Type )
-#define DECLARE_DYNAMIC_MULTICAST_DELEGATE_FourParams( DelegateName, Param1Type, Param1Name, Param2Type, Param2Name, Param3Type, Param3Name, Param4Type, Param4Name ) FUNC_DECLARE_DYNAMIC_MULTICAST_DELEGATE( FWeakObjectPtr, FourParams, DelegateName, DelegateName##_DelegateWrapper, FUNC_CONCAT( Param1Type InParam1, Param2Type InParam2, Param3Type InParam3, Param4Type InParam4 ), FUNC_CONCAT( *this, InParam1, InParam2, InParam3, InParam4 ), void, Param1Type, Param2Type, Param3Type, Param4Type )
+#define DECLARE_DYNAMIC_DELEGATE_FourParams( DelegateName, Param1Type, Param1Name, Param2Type, Param2Name, Param3Type, Param3Name, Param4Type, Param4Name ) BODY_MACRO_COMBINE(CURRENT_FILE_ID,_,__LINE__,_DELEGATE) FUNC_DECLARE_DYNAMIC_DELEGATE( FWeakObjectPtr, FourParams, DelegateName, DelegateName##_DelegateWrapper, FUNC_CONCAT( Param1Type InParam1, Param2Type InParam2, Param3Type InParam3, Param4Type InParam4 ), FUNC_CONCAT( *this, InParam1, InParam2, InParam3, InParam4 ), void, Param1Type, Param2Type, Param3Type, Param4Type )
+#define DECLARE_DYNAMIC_MULTICAST_DELEGATE_FourParams( DelegateName, Param1Type, Param1Name, Param2Type, Param2Name, Param3Type, Param3Name, Param4Type, Param4Name ) BODY_MACRO_COMBINE(CURRENT_FILE_ID,_,__LINE__,_DELEGATE) FUNC_DECLARE_DYNAMIC_MULTICAST_DELEGATE( FWeakObjectPtr, FourParams, DelegateName, DelegateName##_DelegateWrapper, FUNC_CONCAT( Param1Type InParam1, Param2Type InParam2, Param3Type InParam3, Param4Type InParam4 ), FUNC_CONCAT( *this, InParam1, InParam2, InParam3, InParam4 ), void, Param1Type, Param2Type, Param3Type, Param4Type )
 #define DECLARE_DELEGATE_RetVal_FourParams( ReturnValueType, DelegateName, Param1Type, Param2Type, Param3Type, Param4Type ) FUNC_DECLARE_DELEGATE( RetVal_FourParams, DelegateName, ReturnValueType, Param1Type, Param2Type, Param3Type, Param4Type )
-#define DECLARE_DYNAMIC_DELEGATE_RetVal_FourParams( ReturnValueType, DelegateName, Param1Type, Param1Name, Param2Type, Param2Name, Param3Type, Param3Name, Param4Type, Param4Name ) FUNC_DECLARE_DYNAMIC_DELEGATE_RETVAL( FWeakObjectPtr, RetVal_FourParams, DelegateName, DelegateName##_DelegateWrapper, ReturnValueType, FUNC_CONCAT( Param1Type InParam1, Param2Type InParam2, Param3Type InParam3, Param4Type InParam4 ), FUNC_CONCAT( *this, InParam1, InParam2, InParam3, InParam4 ), ReturnValueType, Param1Type, Param2Type, Param3Type, Param4Type )
+#define DECLARE_DYNAMIC_DELEGATE_RetVal_FourParams( ReturnValueType, DelegateName, Param1Type, Param1Name, Param2Type, Param2Name, Param3Type, Param3Name, Param4Type, Param4Name ) BODY_MACRO_COMBINE(CURRENT_FILE_ID,_,__LINE__,_DELEGATE) FUNC_DECLARE_DYNAMIC_DELEGATE_RETVAL( FWeakObjectPtr, RetVal_FourParams, DelegateName, DelegateName##_DelegateWrapper, ReturnValueType, FUNC_CONCAT( Param1Type InParam1, Param2Type InParam2, Param3Type InParam3, Param4Type InParam4 ), FUNC_CONCAT( *this, InParam1, InParam2, InParam3, InParam4 ), ReturnValueType, Param1Type, Param2Type, Param3Type, Param4Type )
 #define DECLARE_DELEGATE_FiveParams( DelegateName, Param1Type, Param2Type, Param3Type, Param4Type, Param5Type ) FUNC_DECLARE_DELEGATE( FiveParams, DelegateName, void, Param1Type, Param2Type, Param3Type, Param4Type, Param5Type )
 #define DECLARE_MULTICAST_DELEGATE_FiveParams( DelegateName, Param1Type, Param2Type, Param3Type, Param4Type, Param5Type ) FUNC_DECLARE_MULTICAST_DELEGATE( FiveParams, DelegateName, void, Param1Type, Param2Type, Param3Type, Param4Type, Param5Type )
 #define DECLARE_EVENT_FiveParams( OwningType, EventName, Param1Type, Param2Type, Param3Type, Param4Type, Param5Type ) FUNC_DECLARE_EVENT( OwningType, EventName, FiveParams, void, Param1Type, Param2Type, Param3Type, Param4Type, Param5Type )
-#define DECLARE_DYNAMIC_DELEGATE_FiveParams( DelegateName, Param1Type, Param1Name, Param2Type, Param2Name, Param3Type, Param3Name, Param4Type, Param4Name, Param5Type, Param5Name ) FUNC_DECLARE_DYNAMIC_DELEGATE( FWeakObjectPtr, FiveParams, DelegateName, DelegateName##_DelegateWrapper, FUNC_CONCAT( Param1Type InParam1, Param2Type InParam2, Param3Type InParam3, Param4Type InParam4, Param5Type InParam5 ), FUNC_CONCAT( *this, InParam1, InParam2, InParam3, InParam4, InParam5 ), void, Param1Type, Param2Type, Param3Type, Param4Type, Param5Type )
-#define DECLARE_DYNAMIC_MULTICAST_DELEGATE_FiveParams( DelegateName, Param1Type, Param1Name, Param2Type, Param2Name, Param3Type, Param3Name, Param4Type, Param4Name, Param5Type, Param5Name ) FUNC_DECLARE_DYNAMIC_MULTICAST_DELEGATE( FWeakObjectPtr, FiveParams, DelegateName, DelegateName##_DelegateWrapper, FUNC_CONCAT( Param1Type InParam1, Param2Type InParam2, Param3Type InParam3, Param4Type InParam4, Param5Type InParam5 ), FUNC_CONCAT( *this, InParam1, InParam2, InParam3, InParam4, InParam5 ), void, Param1Type, Param2Type, Param3Type, Param4Type, Param5Type )
+#define DECLARE_DYNAMIC_DELEGATE_FiveParams( DelegateName, Param1Type, Param1Name, Param2Type, Param2Name, Param3Type, Param3Name, Param4Type, Param4Name, Param5Type, Param5Name ) BODY_MACRO_COMBINE(CURRENT_FILE_ID,_,__LINE__,_DELEGATE) FUNC_DECLARE_DYNAMIC_DELEGATE( FWeakObjectPtr, FiveParams, DelegateName, DelegateName##_DelegateWrapper, FUNC_CONCAT( Param1Type InParam1, Param2Type InParam2, Param3Type InParam3, Param4Type InParam4, Param5Type InParam5 ), FUNC_CONCAT( *this, InParam1, InParam2, InParam3, InParam4, InParam5 ), void, Param1Type, Param2Type, Param3Type, Param4Type, Param5Type )
+#define DECLARE_DYNAMIC_MULTICAST_DELEGATE_FiveParams( DelegateName, Param1Type, Param1Name, Param2Type, Param2Name, Param3Type, Param3Name, Param4Type, Param4Name, Param5Type, Param5Name ) BODY_MACRO_COMBINE(CURRENT_FILE_ID,_,__LINE__,_DELEGATE) FUNC_DECLARE_DYNAMIC_MULTICAST_DELEGATE( FWeakObjectPtr, FiveParams, DelegateName, DelegateName##_DelegateWrapper, FUNC_CONCAT( Param1Type InParam1, Param2Type InParam2, Param3Type InParam3, Param4Type InParam4, Param5Type InParam5 ), FUNC_CONCAT( *this, InParam1, InParam2, InParam3, InParam4, InParam5 ), void, Param1Type, Param2Type, Param3Type, Param4Type, Param5Type )
 #define DECLARE_DELEGATE_RetVal_FiveParams( ReturnValueType, DelegateName, Param1Type, Param2Type, Param3Type, Param4Type, Param5Type ) FUNC_DECLARE_DELEGATE( RetVal_FiveParams, DelegateName, ReturnValueType, Param1Type, Param2Type, Param3Type, Param4Type, Param5Type )
-#define DECLARE_DYNAMIC_DELEGATE_RetVal_FiveParams( ReturnValueType, DelegateName, Param1Type, Param1Name, Param2Type, Param2Name, Param3Type, Param3Name, Param4Type, Param4Name, Param5Type, Param5Name ) FUNC_DECLARE_DYNAMIC_DELEGATE_RETVAL( FWeakObjectPtr, RetVal_FiveParams, DelegateName, DelegateName##_DelegateWrapper, ReturnValueType, FUNC_CONCAT( Param1Type InParam1, Param2Type InParam2, Param3Type InParam3, Param4Type InParam4, Param5Type InParam5 ), FUNC_CONCAT( *this, InParam1, InParam2, InParam3, InParam4, InParam5 ), ReturnValueType, Param1Type, Param2Type, Param3Type, Param4Type, Param5Type )
+#define DECLARE_DYNAMIC_DELEGATE_RetVal_FiveParams( ReturnValueType, DelegateName, Param1Type, Param1Name, Param2Type, Param2Name, Param3Type, Param3Name, Param4Type, Param4Name, Param5Type, Param5Name ) BODY_MACRO_COMBINE(CURRENT_FILE_ID,_,__LINE__,_DELEGATE) FUNC_DECLARE_DYNAMIC_DELEGATE_RETVAL( FWeakObjectPtr, RetVal_FiveParams, DelegateName, DelegateName##_DelegateWrapper, ReturnValueType, FUNC_CONCAT( Param1Type InParam1, Param2Type InParam2, Param3Type InParam3, Param4Type InParam4, Param5Type InParam5 ), FUNC_CONCAT( *this, InParam1, InParam2, InParam3, InParam4, InParam5 ), ReturnValueType, Param1Type, Param2Type, Param3Type, Param4Type, Param5Type )
 #define DECLARE_DELEGATE_SixParams( DelegateName, Param1Type, Param2Type, Param3Type, Param4Type, Param5Type, Param6Type ) FUNC_DECLARE_DELEGATE( SixParams, DelegateName, void, Param1Type, Param2Type, Param3Type, Param4Type, Param5Type, Param6Type )
 #define DECLARE_MULTICAST_DELEGATE_SixParams( DelegateName, Param1Type, Param2Type, Param3Type, Param4Type, Param5Type, Param6Type ) FUNC_DECLARE_MULTICAST_DELEGATE( SixParams, DelegateName, void, Param1Type, Param2Type, Param3Type, Param4Type, Param5Type, Param6Type )
 #define DECLARE_EVENT_SixParams( OwningType, EventName, Param1Type, Param2Type, Param3Type, Param4Type, Param5Type, Param6Type ) FUNC_DECLARE_EVENT( OwningType, EventName, SixParams, void, Param1Type, Param2Type, Param3Type, Param4Type, Param5Type, Param6Type )
-#define DECLARE_DYNAMIC_DELEGATE_SixParams( DelegateName, Param1Type, Param1Name, Param2Type, Param2Name, Param3Type, Param3Name, Param4Type, Param4Name, Param5Type, Param5Name, Param6Type, Param6Name ) FUNC_DECLARE_DYNAMIC_DELEGATE( FWeakObjectPtr, SixParams, DelegateName, DelegateName##_DelegateWrapper, FUNC_CONCAT( Param1Type InParam1, Param2Type InParam2, Param3Type InParam3, Param4Type InParam4, Param5Type InParam5, Param6Type InParam6 ), FUNC_CONCAT( *this, InParam1, InParam2, InParam3, InParam4, InParam5, InParam6 ), void, Param1Type, Param2Type, Param3Type, Param4Type, Param5Type, Param6Type )
-#define DECLARE_DYNAMIC_MULTICAST_DELEGATE_SixParams( DelegateName, Param1Type, Param1Name, Param2Type, Param2Name, Param3Type, Param3Name, Param4Type, Param4Name, Param5Type, Param5Name, Param6Type, Param6Name ) FUNC_DECLARE_DYNAMIC_MULTICAST_DELEGATE( FWeakObjectPtr, SixParams, DelegateName, DelegateName##_DelegateWrapper, FUNC_CONCAT( Param1Type InParam1, Param2Type InParam2, Param3Type InParam3, Param4Type InParam4, Param5Type InParam5, Param6Type InParam6 ), FUNC_CONCAT( *this, InParam1, InParam2, InParam3, InParam4, InParam5, InParam6 ), void, Param1Type, Param2Type, Param3Type, Param4Type, Param5Type, Param6Type )
+#define DECLARE_DYNAMIC_DELEGATE_SixParams( DelegateName, Param1Type, Param1Name, Param2Type, Param2Name, Param3Type, Param3Name, Param4Type, Param4Name, Param5Type, Param5Name, Param6Type, Param6Name ) BODY_MACRO_COMBINE(CURRENT_FILE_ID,_,__LINE__,_DELEGATE) FUNC_DECLARE_DYNAMIC_DELEGATE( FWeakObjectPtr, SixParams, DelegateName, DelegateName##_DelegateWrapper, FUNC_CONCAT( Param1Type InParam1, Param2Type InParam2, Param3Type InParam3, Param4Type InParam4, Param5Type InParam5, Param6Type InParam6 ), FUNC_CONCAT( *this, InParam1, InParam2, InParam3, InParam4, InParam5, InParam6 ), void, Param1Type, Param2Type, Param3Type, Param4Type, Param5Type, Param6Type )
+#define DECLARE_DYNAMIC_MULTICAST_DELEGATE_SixParams( DelegateName, Param1Type, Param1Name, Param2Type, Param2Name, Param3Type, Param3Name, Param4Type, Param4Name, Param5Type, Param5Name, Param6Type, Param6Name ) BODY_MACRO_COMBINE(CURRENT_FILE_ID,_,__LINE__,_DELEGATE) FUNC_DECLARE_DYNAMIC_MULTICAST_DELEGATE( FWeakObjectPtr, SixParams, DelegateName, DelegateName##_DelegateWrapper, FUNC_CONCAT( Param1Type InParam1, Param2Type InParam2, Param3Type InParam3, Param4Type InParam4, Param5Type InParam5, Param6Type InParam6 ), FUNC_CONCAT( *this, InParam1, InParam2, InParam3, InParam4, InParam5, InParam6 ), void, Param1Type, Param2Type, Param3Type, Param4Type, Param5Type, Param6Type )
 #define DECLARE_DELEGATE_RetVal_SixParams( ReturnValueType, DelegateName, Param1Type, Param2Type, Param3Type, Param4Type, Param5Type, Param6Type ) FUNC_DECLARE_DELEGATE( RetVal_SixParams, DelegateName, ReturnValueType, Param1Type, Param2Type, Param3Type, Param4Type, Param5Type, Param6Type )
-#define DECLARE_DYNAMIC_DELEGATE_RetVal_SixParams( ReturnValueType, DelegateName, Param1Type, Param1Name, Param2Type, Param2Name, Param3Type, Param3Name, Param4Type, Param4Name, Param5Type, Param5Name, Param6Type, Param6Name ) FUNC_DECLARE_DYNAMIC_DELEGATE_RETVAL( FWeakObjectPtr, RetVal_SixParams, DelegateName, DelegateName##_DelegateWrapper, ReturnValueType, FUNC_CONCAT( Param1Type InParam1, Param2Type InParam2, Param3Type InParam3, Param4Type InParam4, Param5Type InParam5, Param6Type InParam6 ), FUNC_CONCAT( *this, InParam1, InParam2, InParam3, InParam4, InParam5, InParam6 ), ReturnValueType, Param1Type, Param2Type, Param3Type, Param4Type, Param5Type, Param6Type )
+#define DECLARE_DYNAMIC_DELEGATE_RetVal_SixParams( ReturnValueType, DelegateName, Param1Type, Param1Name, Param2Type, Param2Name, Param3Type, Param3Name, Param4Type, Param4Name, Param5Type, Param5Name, Param6Type, Param6Name ) BODY_MACRO_COMBINE(CURRENT_FILE_ID,_,__LINE__,_DELEGATE) FUNC_DECLARE_DYNAMIC_DELEGATE_RETVAL( FWeakObjectPtr, RetVal_SixParams, DelegateName, DelegateName##_DelegateWrapper, ReturnValueType, FUNC_CONCAT( Param1Type InParam1, Param2Type InParam2, Param3Type InParam3, Param4Type InParam4, Param5Type InParam5, Param6Type InParam6 ), FUNC_CONCAT( *this, InParam1, InParam2, InParam3, InParam4, InParam5, InParam6 ), ReturnValueType, Param1Type, Param2Type, Param3Type, Param4Type, Param5Type, Param6Type )
 #define DECLARE_DELEGATE_SevenParams( DelegateName, Param1Type, Param2Type, Param3Type, Param4Type, Param5Type, Param6Type, Param7Type ) FUNC_DECLARE_DELEGATE( SevenParams, DelegateName, void, Param1Type, Param2Type, Param3Type, Param4Type, Param5Type, Param6Type, Param7Type )
 #define DECLARE_MULTICAST_DELEGATE_SevenParams( DelegateName, Param1Type, Param2Type, Param3Type, Param4Type, Param5Type, Param6Type, Param7Type ) FUNC_DECLARE_MULTICAST_DELEGATE( SevenParams, DelegateName, void, Param1Type, Param2Type, Param3Type, Param4Type, Param5Type, Param6Type, Param7Type )
 #define DECLARE_EVENT_SevenParams( OwningType, EventName, Param1Type, Param2Type, Param3Type, Param4Type, Param5Type, Param6Type, Param7Type ) FUNC_DECLARE_EVENT( OwningType, EventName, SevenParams, void, Param1Type, Param2Type, Param3Type, Param4Type, Param5Type, Param6Type, Param7Type )
-#define DECLARE_DYNAMIC_DELEGATE_SevenParams( DelegateName, Param1Type, Param1Name, Param2Type, Param2Name, Param3Type, Param3Name, Param4Type, Param4Name, Param5Type, Param5Name, Param6Type, Param6Name, Param7Type, Param7Name ) FUNC_DECLARE_DYNAMIC_DELEGATE( FWeakObjectPtr, SevenParams, DelegateName, DelegateName##_DelegateWrapper, FUNC_CONCAT( Param1Type InParam1, Param2Type InParam2, Param3Type InParam3, Param4Type InParam4, Param5Type InParam5, Param6Type InParam6, Param7Type InParam7 ), FUNC_CONCAT( *this, InParam1, InParam2, InParam3, InParam4, InParam5, InParam6, InParam7 ), void, Param1Type, Param2Type, Param3Type, Param4Type, Param5Type, Param6Type, Param7Type )
-#define DECLARE_DYNAMIC_MULTICAST_DELEGATE_SevenParams( DelegateName, Param1Type, Param1Name, Param2Type, Param2Name, Param3Type, Param3Name, Param4Type, Param4Name, Param5Type, Param5Name, Param6Type, Param6Name, Param7Type, Param7Name ) FUNC_DECLARE_DYNAMIC_MULTICAST_DELEGATE( FWeakObjectPtr, SevenParams, DelegateName, DelegateName##_DelegateWrapper, FUNC_CONCAT( Param1Type InParam1, Param2Type InParam2, Param3Type InParam3, Param4Type InParam4, Param5Type InParam5, Param6Type InParam6, Param7Type InParam7 ), FUNC_CONCAT( *this, InParam1, InParam2, InParam3, InParam4, InParam5, InParam6, InParam7 ), void, Param1Type, Param2Type, Param3Type, Param4Type, Param5Type, Param6Type, Param7Type )
+#define DECLARE_DYNAMIC_DELEGATE_SevenParams( DelegateName, Param1Type, Param1Name, Param2Type, Param2Name, Param3Type, Param3Name, Param4Type, Param4Name, Param5Type, Param5Name, Param6Type, Param6Name, Param7Type, Param7Name ) BODY_MACRO_COMBINE(CURRENT_FILE_ID,_,__LINE__,_DELEGATE) FUNC_DECLARE_DYNAMIC_DELEGATE( FWeakObjectPtr, SevenParams, DelegateName, DelegateName##_DelegateWrapper, FUNC_CONCAT( Param1Type InParam1, Param2Type InParam2, Param3Type InParam3, Param4Type InParam4, Param5Type InParam5, Param6Type InParam6, Param7Type InParam7 ), FUNC_CONCAT( *this, InParam1, InParam2, InParam3, InParam4, InParam5, InParam6, InParam7 ), void, Param1Type, Param2Type, Param3Type, Param4Type, Param5Type, Param6Type, Param7Type )
+#define DECLARE_DYNAMIC_MULTICAST_DELEGATE_SevenParams( DelegateName, Param1Type, Param1Name, Param2Type, Param2Name, Param3Type, Param3Name, Param4Type, Param4Name, Param5Type, Param5Name, Param6Type, Param6Name, Param7Type, Param7Name ) BODY_MACRO_COMBINE(CURRENT_FILE_ID,_,__LINE__,_DELEGATE) FUNC_DECLARE_DYNAMIC_MULTICAST_DELEGATE( FWeakObjectPtr, SevenParams, DelegateName, DelegateName##_DelegateWrapper, FUNC_CONCAT( Param1Type InParam1, Param2Type InParam2, Param3Type InParam3, Param4Type InParam4, Param5Type InParam5, Param6Type InParam6, Param7Type InParam7 ), FUNC_CONCAT( *this, InParam1, InParam2, InParam3, InParam4, InParam5, InParam6, InParam7 ), void, Param1Type, Param2Type, Param3Type, Param4Type, Param5Type, Param6Type, Param7Type )
 #define DECLARE_DELEGATE_RetVal_SevenParams( ReturnValueType, DelegateName, Param1Type, Param2Type, Param3Type, Param4Type, Param5Type, Param6Type, Param7Type ) FUNC_DECLARE_DELEGATE( RetVal_SevenParams, DelegateName, ReturnValueType, Param1Type, Param2Type, Param3Type, Param4Type, Param5Type, Param6Type, Param7Type )
-#define DECLARE_DYNAMIC_DELEGATE_RetVal_SevenParams( ReturnValueType, DelegateName, Param1Type, Param1Name, Param2Type, Param2Name, Param3Type, Param3Name, Param4Type, Param4Name, Param5Type, Param5Name, Param6Type, Param6Name, Param7Type, Param7Name ) FUNC_DECLARE_DYNAMIC_DELEGATE_RETVAL( FWeakObjectPtr, RetVal_SevenParams, DelegateName, DelegateName##_DelegateWrapper, ReturnValueType, FUNC_CONCAT( Param1Type InParam1, Param2Type InParam2, Param3Type InParam3, Param4Type InParam4, Param5Type InParam5, Param6Type InParam6, Param7Type InParam7 ), FUNC_CONCAT( *this, InParam1, InParam2, InParam3, InParam4, InParam5, InParam6, InParam7 ), ReturnValueType, Param1Type, Param2Type, Param3Type, Param4Type, Param5Type, Param6Type, Param7Type )
+#define DECLARE_DYNAMIC_DELEGATE_RetVal_SevenParams( ReturnValueType, DelegateName, Param1Type, Param1Name, Param2Type, Param2Name, Param3Type, Param3Name, Param4Type, Param4Name, Param5Type, Param5Name, Param6Type, Param6Name, Param7Type, Param7Name ) BODY_MACRO_COMBINE(CURRENT_FILE_ID,_,__LINE__,_DELEGATE) FUNC_DECLARE_DYNAMIC_DELEGATE_RETVAL( FWeakObjectPtr, RetVal_SevenParams, DelegateName, DelegateName##_DelegateWrapper, ReturnValueType, FUNC_CONCAT( Param1Type InParam1, Param2Type InParam2, Param3Type InParam3, Param4Type InParam4, Param5Type InParam5, Param6Type InParam6, Param7Type InParam7 ), FUNC_CONCAT( *this, InParam1, InParam2, InParam3, InParam4, InParam5, InParam6, InParam7 ), ReturnValueType, Param1Type, Param2Type, Param3Type, Param4Type, Param5Type, Param6Type, Param7Type )
 #define DECLARE_DELEGATE_EightParams( DelegateName, Param1Type, Param2Type, Param3Type, Param4Type, Param5Type, Param6Type, Param7Type, Param8Type ) FUNC_DECLARE_DELEGATE( EightParams, DelegateName, void, Param1Type, Param2Type, Param3Type, Param4Type, Param5Type, Param6Type, Param7Type, Param8Type )
 #define DECLARE_MULTICAST_DELEGATE_EightParams( DelegateName, Param1Type, Param2Type, Param3Type, Param4Type, Param5Type, Param6Type, Param7Type, Param8Type ) FUNC_DECLARE_MULTICAST_DELEGATE( EightParams, DelegateName, void, Param1Type, Param2Type, Param3Type, Param4Type, Param5Type, Param6Type, Param7Type, Param8Type )
 #define DECLARE_EVENT_EightParams( OwningType, EventName, Param1Type, Param2Type, Param3Type, Param4Type, Param5Type, Param6Type, Param7Type, Param8Type ) FUNC_DECLARE_EVENT( OwningType, EventName, EightParams, void, Param1Type, Param2Type, Param3Type, Param4Type, Param5Type, Param6Type, Param7Type, Param8Type )
-#define DECLARE_DYNAMIC_DELEGATE_EightParams( DelegateName, Param1Type, Param1Name, Param2Type, Param2Name, Param3Type, Param3Name, Param4Type, Param4Name, Param5Type, Param5Name, Param6Type, Param6Name, Param7Type, Param7Name, Param8Type, Param8Name ) FUNC_DECLARE_DYNAMIC_DELEGATE( FWeakObjectPtr, EightParams, DelegateName, DelegateName##_DelegateWrapper, FUNC_CONCAT( Param1Type InParam1, Param2Type InParam2, Param3Type InParam3, Param4Type InParam4, Param5Type InParam5, Param6Type InParam6, Param7Type InParam7, Param8Type InParam8 ), FUNC_CONCAT( *this, InParam1, InParam2, InParam3, InParam4, InParam5, InParam6, InParam7, InParam8 ), void, Param1Type, Param2Type, Param3Type, Param4Type, Param5Type, Param6Type, Param7Type, Param8Type )
-#define DECLARE_DYNAMIC_MULTICAST_DELEGATE_EightParams( DelegateName, Param1Type, Param1Name, Param2Type, Param2Name, Param3Type, Param3Name, Param4Type, Param4Name, Param5Type, Param5Name, Param6Type, Param6Name, Param7Type, Param7Name, Param8Type, Param8Name ) FUNC_DECLARE_DYNAMIC_MULTICAST_DELEGATE( FWeakObjectPtr, EightParams, DelegateName, DelegateName##_DelegateWrapper, FUNC_CONCAT( Param1Type InParam1, Param2Type InParam2, Param3Type InParam3, Param4Type InParam4, Param5Type InParam5, Param6Type InParam6, Param7Type InParam7, Param8Type InParam8 ), FUNC_CONCAT( *this, InParam1, InParam2, InParam3, InParam4, InParam5, InParam6, InParam7, InParam8 ), void, Param1Type, Param2Type, Param3Type, Param4Type, Param5Type, Param6Type, Param7Type, Param8Type )
+#define DECLARE_DYNAMIC_DELEGATE_EightParams( DelegateName, Param1Type, Param1Name, Param2Type, Param2Name, Param3Type, Param3Name, Param4Type, Param4Name, Param5Type, Param5Name, Param6Type, Param6Name, Param7Type, Param7Name, Param8Type, Param8Name ) BODY_MACRO_COMBINE(CURRENT_FILE_ID,_,__LINE__,_DELEGATE) FUNC_DECLARE_DYNAMIC_DELEGATE( FWeakObjectPtr, EightParams, DelegateName, DelegateName##_DelegateWrapper, FUNC_CONCAT( Param1Type InParam1, Param2Type InParam2, Param3Type InParam3, Param4Type InParam4, Param5Type InParam5, Param6Type InParam6, Param7Type InParam7, Param8Type InParam8 ), FUNC_CONCAT( *this, InParam1, InParam2, InParam3, InParam4, InParam5, InParam6, InParam7, InParam8 ), void, Param1Type, Param2Type, Param3Type, Param4Type, Param5Type, Param6Type, Param7Type, Param8Type )
+#define DECLARE_DYNAMIC_MULTICAST_DELEGATE_EightParams( DelegateName, Param1Type, Param1Name, Param2Type, Param2Name, Param3Type, Param3Name, Param4Type, Param4Name, Param5Type, Param5Name, Param6Type, Param6Name, Param7Type, Param7Name, Param8Type, Param8Name ) BODY_MACRO_COMBINE(CURRENT_FILE_ID,_,__LINE__,_DELEGATE) FUNC_DECLARE_DYNAMIC_MULTICAST_DELEGATE( FWeakObjectPtr, EightParams, DelegateName, DelegateName##_DelegateWrapper, FUNC_CONCAT( Param1Type InParam1, Param2Type InParam2, Param3Type InParam3, Param4Type InParam4, Param5Type InParam5, Param6Type InParam6, Param7Type InParam7, Param8Type InParam8 ), FUNC_CONCAT( *this, InParam1, InParam2, InParam3, InParam4, InParam5, InParam6, InParam7, InParam8 ), void, Param1Type, Param2Type, Param3Type, Param4Type, Param5Type, Param6Type, Param7Type, Param8Type )
 #define DECLARE_DELEGATE_RetVal_EightParams( ReturnValueType, DelegateName, Param1Type, Param2Type, Param3Type, Param4Type, Param5Type, Param6Type, Param7Type, Param8Type ) FUNC_DECLARE_DELEGATE( RetVal_EightParams, DelegateName, ReturnValueType, Param1Type, Param2Type, Param3Type, Param4Type, Param5Type, Param6Type, Param7Type, Param8Type )
-#define DECLARE_DYNAMIC_DELEGATE_RetVal_EightParams( ReturnValueType, DelegateName, Param1Type, Param1Name, Param2Type, Param2Name, Param3Type, Param3Name, Param4Type, Param4Name, Param5Type, Param5Name, Param6Type, Param6Name, Param7Type, Param7Name, Param8Type, Param8Name ) FUNC_DECLARE_DYNAMIC_DELEGATE_RETVAL( FWeakObjectPtr, RetVal_EightParams, DelegateName, DelegateName##_DelegateWrapper, ReturnValueType, FUNC_CONCAT( Param1Type InParam1, Param2Type InParam2, Param3Type InParam3, Param4Type InParam4, Param5Type InParam5, Param6Type InParam6, Param7Type InParam7, Param8Type InParam8 ), FUNC_CONCAT( *this, InParam1, InParam2, InParam3, InParam4, InParam5, InParam6, InParam7, InParam8 ), ReturnValueType, Param1Type, Param2Type, Param3Type, Param4Type, Param5Type, Param6Type, Param7Type, Param8Type )
+#define DECLARE_DYNAMIC_DELEGATE_RetVal_EightParams( ReturnValueType, DelegateName, Param1Type, Param1Name, Param2Type, Param2Name, Param3Type, Param3Name, Param4Type, Param4Name, Param5Type, Param5Name, Param6Type, Param6Name, Param7Type, Param7Name, Param8Type, Param8Name ) BODY_MACRO_COMBINE(CURRENT_FILE_ID,_,__LINE__,_DELEGATE) FUNC_DECLARE_DYNAMIC_DELEGATE_RETVAL( FWeakObjectPtr, RetVal_EightParams, DelegateName, DelegateName##_DelegateWrapper, ReturnValueType, FUNC_CONCAT( Param1Type InParam1, Param2Type InParam2, Param3Type InParam3, Param4Type InParam4, Param5Type InParam5, Param6Type InParam6, Param7Type InParam7, Param8Type InParam8 ), FUNC_CONCAT( *this, InParam1, InParam2, InParam3, InParam4, InParam5, InParam6, InParam7, InParam8 ), ReturnValueType, Param1Type, Param2Type, Param3Type, Param4Type, Param5Type, Param6Type, Param7Type, Param8Type )
diff --git a/Engine/Source/Runtime/Core/Public/Math/RandomStream.h b/Engine/Source/Runtime/Core/Public/Math/RandomStream.h
--- a/Engine/Source/Runtime/Core/Public/Math/RandomStream.h
+++ b/Engine/Source/Runtime/Core/Public/Math/RandomStream.h
@@ -15,7 +15,7 @@ class FRandomStream
 #else
 	#define MAYBE_COREUOBJECT_API COREUOBJECT_API
 #endif
-	friend MAYBE_COREUOBJECT_API class UScriptStruct* Z_Construct_UScriptStruct_UObject_FRandomStream();
+	friend MAYBE_COREUOBJECT_API class UScriptStruct* Z_Construct_UScriptStruct_FRandomStream();
 	
 public:
 
diff --git a/Engine/Source/Runtime/Core/Public/Math/Transform.h b/Engine/Source/Runtime/Core/Public/Math/Transform.h
--- a/Engine/Source/Runtime/Core/Public/Math/Transform.h
+++ b/Engine/Source/Runtime/Core/Public/Math/Transform.h
@@ -26,7 +26,7 @@ class FTransform
 #else
 	#define MAYBE_COREUOBJECT_API COREUOBJECT_API
 #endif
-	friend MAYBE_COREUOBJECT_API class UScriptStruct* Z_Construct_UScriptStruct_UObject_FTransform();
+	friend MAYBE_COREUOBJECT_API class UScriptStruct* Z_Construct_UScriptStruct_FTransform();
 
 protected:
 	/** Rotation of this transformation, as a quaternion. */
diff --git a/Engine/Source/Runtime/Core/Public/Math/TransformVectorized.h b/Engine/Source/Runtime/Core/Public/Math/TransformVectorized.h
--- a/Engine/Source/Runtime/Core/Public/Math/TransformVectorized.h
+++ b/Engine/Source/Runtime/Core/Public/Math/TransformVectorized.h
@@ -28,7 +28,7 @@ MS_ALIGN(16) class FTransform
 #else
 	#define MAYBE_COREUOBJECT_API COREUOBJECT_API
 #endif 
-	friend MAYBE_COREUOBJECT_API class UScriptStruct* Z_Construct_UScriptStruct_UObject_FTransform();
+	friend MAYBE_COREUOBJECT_API class UScriptStruct* Z_Construct_UScriptStruct_FTransform();
 
 protected:
 	/** Rotation of this transformation, as a quaternion */
diff --git a/Engine/Source/Runtime/Core/Public/Misc/Paths.h b/Engine/Source/Runtime/Core/Public/Misc/Paths.h
--- a/Engine/Source/Runtime/Core/Public/Misc/Paths.h
+++ b/Engine/Source/Runtime/Core/Public/Misc/Paths.h
@@ -324,6 +324,16 @@ public:
 	/** Convert all / and \ to TEXT("/") */
 	static void NormalizeFilename(FString& InPath);
 
+	/**
+	 * Checks if two paths are the same.
+	 *
+	 * @param PathA First path to check.
+	 * @param PathB Second path to check.
+	 *
+	 * @returns True if both paths are the same. False otherwise.
+	 */
+	static bool IsSamePath(const FString& PathA, const FString& PathB);
+
 	/** Normalize all / and \ to TEXT("/") and remove any trailing TEXT("/") if the character before that is not a TEXT("/") or a colon */
 	static void NormalizeDirectoryName(FString& InPath);
 
diff --git a/Engine/Source/Runtime/CoreUObject/Private/UObject/Class.cpp b/Engine/Source/Runtime/CoreUObject/Private/UObject/Class.cpp
--- a/Engine/Source/Runtime/CoreUObject/Private/UObject/Class.cpp
+++ b/Engine/Source/Runtime/CoreUObject/Private/UObject/Class.cpp
@@ -3829,8 +3829,24 @@ bool UFunction::IsSignatureCompatibleWith(const UFunction* OtherFunction, uint64
 	return !(IteratorB && (IteratorB->PropertyFlags & CPF_Parm));
 }
 
+UScriptStruct* GetBaseStructure(const TCHAR* Name)
+{
+	static auto* CoreUObjectPkg = FindObjectChecked<UPackage>(nullptr, TEXT("/Script/CoreUObject"));
+	return FindObjectChecked<UScriptStruct>(CoreUObjectPkg, Name);
+}
 
 IMPLEMENT_CORE_INTRINSIC_CLASS(UFunction, UStruct,
 	{
 	}
 );
+
+UDelegateFunction::UDelegateFunction(const FObjectInitializer& ObjectInitializer, UFunction* InSuperFunction, uint32 InFunctionFlags, uint16 InRepOffset, SIZE_T ParamsSize)
+	: UFunction(ObjectInitializer, InSuperFunction, InFunctionFlags, InRepOffset, ParamsSize)
+{
+
+}
+
+IMPLEMENT_CORE_INTRINSIC_CLASS(UDelegateFunction, UFunction,
+	{
+	}
+);
\ No newline at end of file
diff --git a/Engine/Source/Runtime/CoreUObject/Private/UObject/UObjectGlobals.cpp b/Engine/Source/Runtime/CoreUObject/Private/UObject/UObjectGlobals.cpp
--- a/Engine/Source/Runtime/CoreUObject/Private/UObject/UObjectGlobals.cpp
+++ b/Engine/Source/Runtime/CoreUObject/Private/UObject/UObjectGlobals.cpp
@@ -169,6 +169,17 @@ UObject* StaticFindObject( UClass* ObjectClass, UObject* InObjectPackage, const
 	FName ObjectName(*InName, FNAME_Add, true);
 	MatchingObject = StaticFindObjectFast( ObjectClass, ObjectPackage, ObjectName, ExactClass, bAnyPackage );
 
+	if (!MatchingObject && ObjectPackage != nullptr &&
+			(
+				ObjectClass == UEnum::StaticClass()	// Enums
+				|| ObjectClass == UScriptStruct::StaticClass() // Structs
+				|| (ObjectClass == UFunction::StaticClass() && FString(OrigInName).EndsWith(HEADER_GENERATED_DELEGATE_SIGNATURE_SUFFIX)) // Delegates
+			)
+		)
+	{
+		MatchingObject = StaticFindObject(ObjectClass, ObjectPackage->GetOutermost(), *ObjectName.ToString(), ExactClass);
+	}
+
 	return MatchingObject;
 }
 
@@ -2592,7 +2603,7 @@ FArchive& FScriptInterface::Serialize(FArchive& Ar, UClass* InterfaceType)
 /** A struct used as stub for deleted ones. */
 UScriptStruct* GetFallbackStruct()
 {
-	static UScriptStruct* FallbackStruct = FindObjectChecked<UScriptStruct>(UObject::StaticClass(), TEXT("FallbackStruct"));
+	static UScriptStruct* FallbackStruct = GetBaseStructure(TEXT("FallbackStruct"));
 	return FallbackStruct;
 }
 
@@ -2654,3 +2665,15 @@ UObject* FObjectInitializer::CreateEditorOnlyDefaultSubobject(UObject* Outer, FN
 #endif
 	return nullptr;
 }
+
+COREUOBJECT_API UFunction* FindDelegateSignature(FName DelegateSignatureName)
+{
+	FString StringName = DelegateSignatureName.ToString();
+
+	if (StringName.EndsWith(HEADER_GENERATED_DELEGATE_SIGNATURE_SUFFIX))
+	{
+		return FindObject<UFunction>(ANY_PACKAGE, *StringName);
+	}
+
+	return nullptr;
+}
diff --git a/Engine/Source/Runtime/CoreUObject/Public/UObject/Class.h b/Engine/Source/Runtime/CoreUObject/Public/UObject/Class.h
--- a/Engine/Source/Runtime/CoreUObject/Public/UObject/Class.h
+++ b/Engine/Source/Runtime/CoreUObject/Public/UObject/Class.h
@@ -414,7 +414,7 @@ enum EStructFlags
 	
 	STRUCT_HasInstancedReference= 0x00000004,
 
-	// Unused entry				= 0x00000008,
+	STRUCT_NoExport				= 0x00000008,
 
 	/** Indicates that this struct should always be serialized as a single unit */
 	STRUCT_Atomic				= 0x00000010,
@@ -1289,6 +1289,14 @@ public:
 	bool IsSignatureCompatibleWith(const UFunction* OtherFunction, uint64 IgnoreFlags) const;
 };
 
+class COREUOBJECT_API UDelegateFunction : public UFunction
+{
+	DECLARE_CASTED_CLASS_INTRINSIC(UDelegateFunction, UFunction, 0, CoreUObject, CASTCLASS_UDelegateFunction)
+	DECLARE_WITHIN(UObject)
+public:
+	explicit UDelegateFunction(const FObjectInitializer& ObjectInitializer, UFunction* InSuperFunction, uint32 InFunctionFlags = 0, uint16 InRepOffset = 0, SIZE_T ParamsSize = 0);
+};
+
 /*-----------------------------------------------------------------------------
 	UEnum.
 -----------------------------------------------------------------------------*/
@@ -2578,3 +2586,5 @@ inline T* GetMutableDefault(UClass *Class)
 	checkSlow(Class->GetDefaultObject()->IsA(T::StaticClass()));
 	return (T*)Class->GetDefaultObject();
 }
+
+COREUOBJECT_API UScriptStruct* GetBaseStructure(const TCHAR* Name);
\ No newline at end of file
diff --git a/Engine/Source/Runtime/CoreUObject/Public/UObject/ObjectBase.h b/Engine/Source/Runtime/CoreUObject/Public/UObject/ObjectBase.h
--- a/Engine/Source/Runtime/CoreUObject/Public/UObject/ObjectBase.h
+++ b/Engine/Source/Runtime/CoreUObject/Public/UObject/ObjectBase.h
@@ -280,6 +280,7 @@ typedef uint64 EClassCastFlags;
 #define CASTCLASS_USkinnedMeshComponent			DECLARE_UINT64(0x0000020000000000)
 #define CASTCLASS_USkeletalMeshComponent		DECLARE_UINT64(0x0000040000000000)
 #define CASTCLASS_UBlueprint					DECLARE_UINT64(0x0000080000000000)
+#define CASTCLASS_UDelegateFunction				DECLARE_UINT64(0x0000100000000000)
 
 #define CASTCLASS_AllFlags						DECLARE_UINT64(0xFFFFFFFFFFFFFFFF)
 
@@ -499,11 +500,25 @@ struct COREUOBJECT_API FReferencerInformationList
 #define UENUM(...)
 #define UDELEGATE(...)
 
-// This pair of macros is used to help implement GENERATED_USTRUCT_BODY()
-#define USTRUCT_COMBINE_INNER(A,B,C) A##B##C
-#define USTRUCT_COMBINE(A,B,C) USTRUCT_COMBINE_INNER(A,B,C)
+// This pair of macros is used to help implement GENERATED_BODY() and GENERATED_USTRUCT_BODY()
+#define BODY_MACRO_COMBINE_INNER(A,B,C,D) A##B##C##D
+#define BODY_MACRO_COMBINE(A,B,C,D) BODY_MACRO_COMBINE_INNER(A,B,C,D)
 
-#define GENERATED_USTRUCT_BODY(...) USTRUCT_COMBINE(UCLASS_CURRENT_FILE_NAME,_USTRUCT_BODY_LINE_,__LINE__)
+#define GENERATED_BODY_LEGACY(...) BODY_MACRO_COMBINE(CURRENT_FILE_ID,_,__LINE__,_GENERATED_BODY_LEGACY)
+#define GENERATED_BODY(...) BODY_MACRO_COMBINE(CURRENT_FILE_ID,_,__LINE__,_GENERATED_BODY)
+
+#define GENERATED_USTRUCT_BODY(...) GENERATED_BODY()
+#define GENERATED_UCLASS_BODY(...) GENERATED_BODY_LEGACY()
+#define GENERATED_UINTERFACE_BODY(...) GENERATED_BODY_LEGACY()
+#define GENERATED_IINTERFACE_BODY(...) GENERATED_BODY()
+
+#if UE_BUILD_DOCS
+#define UCLASS(...)
+#else
+#define UCLASS(...) BODY_MACRO_COMBINE(CURRENT_FILE_ID,_,__LINE__,_PROLOG)
+#endif
+
+#define UINTERFACE(...) UCLASS()
 
 // This macro is used to declare a thunk function in autogenerated boilerplate code
 #define DECLARE_FUNCTION(func) void func( FFrame& Stack, RESULT_DECL )
diff --git a/Engine/Source/Runtime/CoreUObject/Public/UObject/UObjectGlobals.h b/Engine/Source/Runtime/CoreUObject/Public/UObject/UObjectGlobals.h
--- a/Engine/Source/Runtime/CoreUObject/Public/UObject/UObjectGlobals.h
+++ b/Engine/Source/Runtime/CoreUObject/Public/UObject/UObjectGlobals.h
@@ -1153,6 +1153,11 @@ inline T* GetMutableDefault()
 template< class T > 
 inline T* GetMutableDefault(UClass *Class);
 
+/**
+ * Looks for delegate signature with given name.
+ */
+COREUOBJECT_API UFunction* FindDelegateSignature(FName DelegateSignatureName);
+
 /**
  * Determines whether the specified array contains objects of the specified class.
  *
diff --git a/Engine/Source/Runtime/Engine/Classes/Camera/CameraActor.h b/Engine/Source/Runtime/Engine/Classes/Camera/CameraActor.h
--- a/Engine/Source/Runtime/Engine/Classes/Camera/CameraActor.h
+++ b/Engine/Source/Runtime/Engine/Classes/Camera/CameraActor.h
@@ -1,6 +1,9 @@
 // Copyright 1998-2015 Epic Games, Inc. All Rights Reserved.
 
 #pragma once
+
+#include "Engine/Scene.h"
+
 #include "CameraActor.generated.h"
 
 /** 
diff --git a/Engine/Source/Runtime/Engine/Classes/Camera/CameraAnimInst.h b/Engine/Source/Runtime/Engine/Classes/Camera/CameraAnimInst.h
--- a/Engine/Source/Runtime/Engine/Classes/Camera/CameraAnimInst.h
+++ b/Engine/Source/Runtime/Engine/Classes/Camera/CameraAnimInst.h
@@ -5,6 +5,9 @@
  */
 
 #pragma once
+
+#include "Camera/CameraTypes.h"
+
 #include "CameraAnimInst.generated.h"
 
 UCLASS(notplaceable, BlueprintType, transient)
diff --git a/Engine/Source/Runtime/Engine/Classes/Matinee/InterpTrackInstToggle.h b/Engine/Source/Runtime/Engine/Classes/Matinee/InterpTrackInstToggle.h
--- a/Engine/Source/Runtime/Engine/Classes/Matinee/InterpTrackInstToggle.h
+++ b/Engine/Source/Runtime/Engine/Classes/Matinee/InterpTrackInstToggle.h
@@ -2,6 +2,9 @@
 
 
 #pragma once
+
+#include "Matinee/InterpTrackToggle.h"
+
 #include "InterpTrackInstToggle.generated.h"
 
 UCLASS()
diff --git a/Engine/Source/Runtime/Engine/Classes/Matinee/InterpTrackInstVisibility.h b/Engine/Source/Runtime/Engine/Classes/Matinee/InterpTrackInstVisibility.h
--- a/Engine/Source/Runtime/Engine/Classes/Matinee/InterpTrackInstVisibility.h
+++ b/Engine/Source/Runtime/Engine/Classes/Matinee/InterpTrackInstVisibility.h
@@ -1,7 +1,9 @@
 // Copyright 1998-2015 Epic Games, Inc. All Rights Reserved.
 
-
 #pragma once
+
+#include "InterpTrackVisibility.h"
+
 #include "InterpTrackInstVisibility.generated.h"
 
 UCLASS()
diff --git a/Engine/Source/Runtime/Engine/Private/ActorConstruction.cpp b/Engine/Source/Runtime/Engine/Private/ActorConstruction.cpp
--- a/Engine/Source/Runtime/Engine/Private/ActorConstruction.cpp
+++ b/Engine/Source/Runtime/Engine/Private/ActorConstruction.cpp
@@ -123,7 +123,7 @@ static TArray<FRandomStream*> FindRandomStreams(AActor* InActor)
 {
 	check(InActor);
 	TArray<FRandomStream*> OutStreams;
-	UScriptStruct* RandomStreamStruct = FindObjectChecked<UScriptStruct>(UObject::StaticClass(), TEXT("RandomStream"));
+	UScriptStruct* RandomStreamStruct = GetBaseStructure(TEXT("RandomStream"));
 	for( TFieldIterator<UStructProperty> It(InActor->GetClass()) ; It ; ++It )
 	{
 		UStructProperty* StructProp = *It;
diff --git a/Engine/Source/Runtime/Engine/Private/KismetSystemLibrary.cpp b/Engine/Source/Runtime/Engine/Private/KismetSystemLibrary.cpp
--- a/Engine/Source/Runtime/Engine/Private/KismetSystemLibrary.cpp
+++ b/Engine/Source/Runtime/Engine/Private/KismetSystemLibrary.cpp
@@ -667,7 +667,7 @@ void UKismetSystemLibrary::SetVectorPropertyByName(UObject* Object, FName Proper
 {
 	if(Object != NULL)
 	{
-		UScriptStruct* VectorStruct = FindObjectChecked<UScriptStruct>(UObject::StaticClass(), TEXT("Vector"));
+		UScriptStruct* VectorStruct = GetBaseStructure(TEXT("Vector"));
 		UStructProperty* VectorProp = FindField<UStructProperty>(Object->GetClass(), PropertyName);
 		if(VectorProp != NULL && VectorProp->Struct == VectorStruct)
 		{
@@ -680,7 +680,7 @@ void UKismetSystemLibrary::SetRotatorPropertyByName(UObject* Object, FName Prope
 {
 	if(Object != NULL)
 	{
-		UScriptStruct* RotatorStruct = FindObjectChecked<UScriptStruct>(UObject::StaticClass(), TEXT("Rotator"));
+		UScriptStruct* RotatorStruct = GetBaseStructure(TEXT("Rotator"));
 		UStructProperty* RotatorProp = FindField<UStructProperty>(Object->GetClass(), PropertyName);
 		if(RotatorProp != NULL && RotatorProp->Struct == RotatorStruct)
 		{
@@ -693,7 +693,7 @@ void UKismetSystemLibrary::SetLinearColorPropertyByName(UObject* Object, FName P
 {
 	if(Object != NULL)
 	{
-		UScriptStruct* ColorStruct = FindObjectChecked<UScriptStruct>(UObject::StaticClass(), TEXT("LinearColor"));
+		UScriptStruct* ColorStruct = GetBaseStructure(TEXT("LinearColor"));
 		UStructProperty* ColorProp = FindField<UStructProperty>(Object->GetClass(), PropertyName);
 		if(ColorProp != NULL && ColorProp->Struct == ColorStruct)
 		{
@@ -706,7 +706,7 @@ void UKismetSystemLibrary::SetTransformPropertyByName(UObject* Object, FName Pro
 {
 	if(Object != NULL)
 	{
-		UScriptStruct* TransformStruct = FindObjectChecked<UScriptStruct>(UObject::StaticClass(), TEXT("Transform"));
+		UScriptStruct* TransformStruct = GetBaseStructure(TEXT("Transform"));
 		UStructProperty* TransformProp = FindField<UStructProperty>(Object->GetClass(), PropertyName);
 		if(TransformProp != NULL && TransformProp->Struct == TransformStruct)
 		{
diff --git a/Engine/Source/Runtime/Engine/Private/Materials/MaterialExpressions.cpp b/Engine/Source/Runtime/Engine/Private/Materials/MaterialExpressions.cpp
--- a/Engine/Source/Runtime/Engine/Private/Materials/MaterialExpressions.cpp
+++ b/Engine/Source/Runtime/Engine/Private/Materials/MaterialExpressions.cpp
@@ -1260,7 +1260,7 @@ static bool VerifySamplerType(
 		EMaterialSamplerType CorrectSamplerType = UMaterialExpressionTextureBase::GetSamplerTypeForTexture( Texture );
 		if ( SamplerType != CorrectSamplerType )
 		{
-			UEnum* SamplerTypeEnum = FindObject<UEnum>( NULL, TEXT("/Script/Engine.EngineTypes.EMaterialSamplerType") );
+			UEnum* SamplerTypeEnum = FindObject<UEnum>( NULL, TEXT("/Script/Engine.EMaterialSamplerType") );
 			check( SamplerTypeEnum );
 
 			FString SamplerTypeDisplayName = SamplerTypeEnum->GetEnumText(SamplerType).ToString();
@@ -1275,7 +1275,7 @@ static bool VerifySamplerType(
 		}
 		if((SamplerType == SAMPLERTYPE_Normal || SamplerType == SAMPLERTYPE_Masks) && Texture->SRGB)
 		{
-			UEnum* SamplerTypeEnum = FindObject<UEnum>( NULL, TEXT("/Script/Engine.EngineTypes.EMaterialSamplerType") );
+			UEnum* SamplerTypeEnum = FindObject<UEnum>( NULL, TEXT("/Script/Engine.EMaterialSamplerType") );
 			check( SamplerTypeEnum );
 
 			FString SamplerTypeDisplayName = SamplerTypeEnum->GetEnumText(SamplerType).ToString();
@@ -4796,7 +4796,7 @@ int32 UMaterialExpressionSceneTexture::Compile(class FMaterialCompiler* Compiler
 
 void UMaterialExpressionSceneTexture::GetCaption(TArray<FString>& OutCaptions) const
 {
-	UEnum* Enum = FindObject<UEnum>(NULL, TEXT("Engine.MaterialExpressionSceneTexture.ESceneTextureId"));
+	UEnum* Enum = FindObject<UEnum>(NULL, TEXT("Engine.ESceneTextureId"));
 
 	check(Enum);
 
diff --git a/Engine/Source/Runtime/Engine/Private/SoundGroups.cpp b/Engine/Source/Runtime/Engine/Private/SoundGroups.cpp
--- a/Engine/Source/Runtime/Engine/Private/SoundGroups.cpp
+++ b/Engine/Source/Runtime/Engine/Private/SoundGroups.cpp
@@ -24,7 +24,7 @@ void USoundGroups::Initialize() const
 	}
 
 #if WITH_EDITOR
-	UEnum* SoundGroupEnum = FindObjectChecked<UEnum>(NULL, TEXT("/Script/Engine.SoundGroups.ESoundGroup"));
+	UEnum* SoundGroupEnum = FindObjectChecked<UEnum>(NULL, TEXT("/Script/Engine.ESoundGroup"));
 
 	for (const auto& It : SoundGroupMap)
 	{
@@ -56,7 +56,7 @@ const FSoundGroup& USoundGroups::GetSoundGroup(const ESoundGroup SoundGroup) con
 	const FSoundGroup* SG = SoundGroupMap.Find(SoundGroup);
 	if (SG == NULL)
 	{
-		UEnum* SoundGroupEnum = FindObjectChecked<UEnum>(NULL, TEXT("/Script/Engine.SoundGroups.ESoundGroup"));
+		UEnum* SoundGroupEnum = FindObjectChecked<UEnum>(NULL, TEXT("/Script/Engine.ESoundGroup"));
 		
 		UE_LOG(LogAudio, Warning, TEXT("Requested SoundGroup %s does not have defined profile.  Using SOUNDGROUP_Default."), *SoundGroupEnum->GetEnumText(SoundGroup).ToString());
 		return SoundGroupMap.FindChecked(SOUNDGROUP_Default);
diff --git a/Engine/Source/Runtime/Engine/Private/Timeline.cpp b/Engine/Source/Runtime/Engine/Private/Timeline.cpp
--- a/Engine/Source/Runtime/Engine/Private/Timeline.cpp
+++ b/Engine/Source/Runtime/Engine/Private/Timeline.cpp
@@ -736,7 +736,7 @@ ETimelineSigType UTimelineComponent::GetTimelineSignatureForFunction(const UFunc
 
 UFunction* UTimelineComponent::GetTimelineEventSignature()
 {
-	UFunction* TimelineEventSig = FindObject<UFunction>(UTimelineComponent::StaticClass(), TEXT("OnTimelineEvent__DelegateSignature"));
+	UFunction* TimelineEventSig = FindObject<UFunction>(FindPackage(nullptr, TEXT("/Script/Engine")), TEXT("OnTimelineEvent__DelegateSignature"));
 	check(TimelineEventSig != NULL);
 	return TimelineEventSig;
 }
@@ -744,7 +744,7 @@ UFunction* UTimelineComponent::GetTimelineEventSignature()
 
 UFunction* UTimelineComponent::GetTimelineFloatSignature()
 {
-	UFunction* TimelineFloatSig = FindObject<UFunction>(UTimelineComponent::StaticClass(), TEXT("OnTimelineFloat__DelegateSignature"));
+	UFunction* TimelineFloatSig = FindObject<UFunction>(FindPackage(nullptr, TEXT("/Script/Engine")), TEXT("OnTimelineFloat__DelegateSignature"));
 	check(TimelineFloatSig != NULL);
 	return TimelineFloatSig;
 }
@@ -752,7 +752,7 @@ UFunction* UTimelineComponent::GetTimelineFloatSignature()
 
 UFunction* UTimelineComponent::GetTimelineVectorSignature()
 {
-	UFunction* TimelineVectorSig = FindObject<UFunction>(UTimelineComponent::StaticClass(), TEXT("OnTimelineVector__DelegateSignature"));
+	UFunction* TimelineVectorSig = FindObject<UFunction>(FindPackage(nullptr, TEXT("/Script/Engine")), TEXT("OnTimelineVector__DelegateSignature"));
 	check(TimelineVectorSig != NULL);
 	return TimelineVectorSig;
 }
@@ -760,7 +760,7 @@ UFunction* UTimelineComponent::GetTimelineVectorSignature()
 
 UFunction* UTimelineComponent::GetTimelineLinearColorSignature()
 {
-	UFunction* TimelineVectorSig = FindObject<UFunction>(UTimelineComponent::StaticClass(), TEXT("OnTimelineLinearColor__DelegateSignature"));
+	UFunction* TimelineVectorSig = FindObject<UFunction>(FindPackage(nullptr, TEXT("/Script/Engine")), TEXT("OnTimelineLinearColor__DelegateSignature"));
 	check(TimelineVectorSig != NULL);
 	return TimelineVectorSig;
 }
diff --git a/Engine/Source/Runtime/UMG/Public/Binding/DynamicPropertyPath.h b/Engine/Source/Runtime/UMG/Public/Binding/DynamicPropertyPath.h
--- a/Engine/Source/Runtime/UMG/Public/Binding/DynamicPropertyPath.h
+++ b/Engine/Source/Runtime/UMG/Public/Binding/DynamicPropertyPath.h
@@ -97,7 +97,7 @@ inline bool IsConcreteTypeCompatibleWithReflectedType<FString>(UProperty* Proper
 template<>
 inline bool IsConcreteTypeCompatibleWithReflectedType<FLinearColor>(UProperty* Property)
 {
-	static const UScriptStruct* LinearColorStruct = FindObjectChecked<UScriptStruct>(UObject::StaticClass(), TEXT("LinearColor"));
+	static const UScriptStruct* LinearColorStruct = GetBaseStructure(TEXT("LinearColor"));
 	
 	if ( UStructProperty* StructProperty = Cast<UStructProperty>(Property) )
 	{
