Convert early command line argument parsing to getdelim().

- Part of pull request #728, contributed by [REDACTED_PERSON].

diff --git a/Engine/Source/Runtime/Core/Private/Linux/LinuxPlatformMemory.cpp b/Engine/Source/Runtime/Core/Private/Linux/LinuxPlatformMemory.cpp
--- a/Engine/Source/Runtime/Core/Private/Linux/LinuxPlatformMemory.cpp
+++ b/Engine/Source/Runtime/Core/Private/Linux/LinuxPlatformMemory.cpp
@@ -38,40 +38,30 @@ class FMalloc* FLinuxPlatformMemory::BaseAllocator()
 	{
 		if (FILE* CmdLineFile = fopen("/proc/self/cmdline", "r"))
 		{
-			char CmdLineBuffer[4096] = { 0 };
-			FPlatformMemory::Memzero(CmdLineBuffer, sizeof(CmdLineBuffer));
-			if (fgets(CmdLineBuffer, sizeof(CmdLineBuffer)-2, CmdLineFile))	// -2 to guarantee that there are always two zeroes even if cmdline is too long
+			char * Arg = nullptr;
+			size_t Size = 0;
+			while(getdelim(&Arg, &Size, 0, CmdLineFile) != -1)
 			{
-				char * Arg = CmdLineBuffer;
-				while (*Arg != 0 || Arg - CmdLineBuffer >= sizeof(CmdLineBuffer))
+#if PLATFORM_SUPPORTS_JEMALLOC
+				if (FCStringAnsi::Stricmp(Arg, "-jemalloc") == 0)
 				{
-					if (FCStringAnsi::Stricmp(Arg, "-jemalloc") == 0)
-					{
-						AllocatorToUse = EAllocatorToUse::Jemalloc;
-						break;
-					}
-
-					if (FCStringAnsi::Stricmp(Arg, "-ansimalloc") == 0)
-					{
-						AllocatorToUse = EAllocatorToUse::Ansi;
-						break;
-					}
-
-					if (FCStringAnsi::Stricmp(Arg, "-binnedmalloc") == 0)
-					{
-						AllocatorToUse = EAllocatorToUse::Jemalloc;
-						break;
-					}
-
-					// advance till zero
-					while (*Arg)
-					{
-						++Arg;
-					}
-					++Arg;	// and skip the zero
+					AllocatorToUse = EAllocatorToUse::Jemalloc;
+					break;
+				}
+#endif // PLATFORM_SUPPORTS_JEMALLOC
+				if (FCStringAnsi::Stricmp(Arg, "-ansimalloc") == 0)
+				{
+					AllocatorToUse = EAllocatorToUse::Ansi;
+					break;
 				}
-			}
 
+				if (FCStringAnsi::Stricmp(Arg, "-binnedmalloc") == 0)
+				{
+					AllocatorToUse = EAllocatorToUse::Jemalloc;
+					break;
+				}	
+			}
+			free(Arg);
 			fclose(CmdLineFile);
 		}
 	}
