Optimized Mac FTextInputMethodChangeNotifier so it doesn't call expensive GetWindow() until it's really needed. The new code assumes (hopefully correctly) that [[TextView inputContext] invalidateCharacterCoordinates]; is not needed for newly created contexts.

diff --git a/Engine/Source/Runtime/Core/Private/Mac/MacTextInputMethodSystem.cpp b/Engine/Source/Runtime/Core/Private/Mac/MacTextInputMethodSystem.cpp
--- a/Engine/Source/Runtime/Core/Private/Mac/MacTextInputMethodSystem.cpp
+++ b/Engine/Source/Runtime/Core/Private/Mac/MacTextInputMethodSystem.cpp
@@ -17,13 +17,12 @@ namespace
 		FTextInputMethodChangeNotifier(const TSharedRef<ITextInputMethodContext>& InContext)
 		:	Context(InContext)
 		{
-			ContextWindow = Context.Pin()->GetWindow();
 		}
 		
 		virtual ~FTextInputMethodChangeNotifier() {}
 		
 		void SetContextWindow(TSharedPtr<FGenericWindow> Window);
-		TSharedPtr<FGenericWindow> GetContextWindow(void) const;
+		TSharedPtr<FGenericWindow> GetContextWindow();
 		
 		virtual void NotifyLayoutChanged(const ELayoutChangeType ChangeType) override;
 		virtual void NotifySelectionChanged() override;
@@ -40,16 +39,20 @@ namespace
 		ContextWindow = Window;
 	}
 	
-	TSharedPtr<FGenericWindow> FTextInputMethodChangeNotifier::GetContextWindow(void) const
+	TSharedPtr<FGenericWindow> FTextInputMethodChangeNotifier::GetContextWindow()
 	{
+		if (!ContextWindow.IsValid())
+		{
+			ContextWindow = Context.Pin()->GetWindow();
+		}
 		return ContextWindow;
 	}
 	
 	void FTextInputMethodChangeNotifier::NotifyLayoutChanged(const ELayoutChangeType ChangeType)
 	{
-		if(ContextWindow.IsValid())
+		if(ChangeType != ELayoutChangeType::Created && GetContextWindow().IsValid())
 		{
-			FCocoaWindow* CocoaWindow = (FCocoaWindow*)ContextWindow->GetOSWindowHandle();
+			FCocoaWindow* CocoaWindow = (FCocoaWindow*)GetContextWindow()->GetOSWindowHandle();
 			MainThreadCall(^{
 				SCOPED_AUTORELEASE_POOL;
 				if(CocoaWindow && [CocoaWindow openGLView])
@@ -63,9 +66,9 @@ namespace
 	
 	void FTextInputMethodChangeNotifier::NotifySelectionChanged()
 	{
-		if(ContextWindow.IsValid())
+		if(GetContextWindow().IsValid())
 		{
-			FCocoaWindow* CocoaWindow = (FCocoaWindow*)ContextWindow->GetOSWindowHandle();
+			FCocoaWindow* CocoaWindow = (FCocoaWindow*)GetContextWindow()->GetOSWindowHandle();
 			MainThreadCall(^{
 				SCOPED_AUTORELEASE_POOL;
 				if(CocoaWindow && [CocoaWindow openGLView])
@@ -79,9 +82,9 @@ namespace
 	
 	void FTextInputMethodChangeNotifier::NotifyTextChanged(const uint32 BeginIndex, const uint32 OldLength, const uint32 NewLength)
 	{
-		if(ContextWindow.IsValid())
+		if(GetContextWindow().IsValid())
 		{
-			FCocoaWindow* CocoaWindow = (FCocoaWindow*)ContextWindow->GetOSWindowHandle();
+			FCocoaWindow* CocoaWindow = (FCocoaWindow*)GetContextWindow()->GetOSWindowHandle();
 			MainThreadCall(^{
 				SCOPED_AUTORELEASE_POOL;
 				if(CocoaWindow && [CocoaWindow openGLView])
@@ -95,9 +98,9 @@ namespace
 	
 	void FTextInputMethodChangeNotifier::CancelComposition()
 	{
-		if(ContextWindow.IsValid())
+		if(GetContextWindow().IsValid())
 		{
-			FCocoaWindow* CocoaWindow = (FCocoaWindow*)ContextWindow->GetOSWindowHandle();
+			FCocoaWindow* CocoaWindow = (FCocoaWindow*)GetContextWindow()->GetOSWindowHandle();
 			MainThreadCall(^{
 				SCOPED_AUTORELEASE_POOL;
 				if(CocoaWindow && [CocoaWindow openGLView])
