Bump the object version after merging changes to the 4.7 branch, so we can detect malformed assets later on.

diff --git a/Engine/Source/Runtime/Core/Public/UObject/ObjectVersion.h b/Engine/Source/Runtime/Core/Public/UObject/ObjectVersion.h
--- a/Engine/Source/Runtime/Core/Public/UObject/ObjectVersion.h
+++ b/Engine/Source/Runtime/Core/Public/UObject/ObjectVersion.h
@@ -456,6 +456,8 @@ enum EUnrealEngineObjectUE4Version
 	VER_UE4_ADD_MODIFIERS_RUNTIME_GENERATION,
 	// Tidied up material's handling of masked blend mode.
 	VER_UE4_MATERIAL_MASKED_BLENDMODE_TIDY,
+	// After merging VER_UE4_LANDSCAPE_GRASS_COOKING-VER_UE4_MATERIAL_MASKED_BLENDMODE_TIDY to the 4.7 branch to enable reading earlier assets without writing them
+	VER_UE4_POST_4_7_OBJECTVERSION_MERGE,
 
 	// -----<new versions can be added before this line>-------------------------------------------------
 	// - this needs to be the last line (see note below)
