Discard recast navmesh data that's saved by the 4.7 editor after merging VER_UE4_MERGED_ADD_MODIFIERS_RUNTIME_GENERATION_TO_4_7

diff --git a/Engine/Source/Runtime/Core/Public/UObject/ObjectVersion.h b/Engine/Source/Runtime/Core/Public/UObject/ObjectVersion.h
--- a/Engine/Source/Runtime/Core/Public/UObject/ObjectVersion.h
+++ b/Engine/Source/Runtime/Core/Public/UObject/ObjectVersion.h
@@ -456,8 +456,10 @@ enum EUnrealEngineObjectUE4Version
 	VER_UE4_ADD_MODIFIERS_RUNTIME_GENERATION,
 	// Tidied up material's handling of masked blend mode.
 	VER_UE4_MATERIAL_MASKED_BLENDMODE_TIDY,
-	// After merging VER_UE4_LANDSCAPE_GRASS_COOKING-VER_UE4_MATERIAL_MASKED_BLENDMODE_TIDY to the 4.7 branch to enable reading earlier assets without writing them
-	VER_UE4_POST_4_7_OBJECTVERSION_MERGE,
+	// Merged VER_UE4_ADD_MODIFIERS_RUNTIME_GENERATION into 4.7 branch. 
+	VER_UE4_MERGED_ADD_MODIFIERS_RUNTIME_GENERATION_TO_4_7,
+	// After adding VER_UE4_MERGED_ADD_MODIFIERS_RUNTIME_GENERATION_TO_4_7 into 4.7 branch. 
+	VER_UE4_AFTER_MERGED_ADD_MODIFIERS_RUNTIME_GENERATION_TO_4_7,
 
 	// -----<new versions can be added before this line>-------------------------------------------------
 	// - this needs to be the last line (see note below)
diff --git a/Engine/Source/Runtime/Engine/Private/AI/Navigation/RecastNavMesh.cpp b/Engine/Source/Runtime/Engine/Private/AI/Navigation/RecastNavMesh.cpp
--- a/Engine/Source/Runtime/Engine/Private/AI/Navigation/RecastNavMesh.cpp
+++ b/Engine/Source/Runtime/Engine/Private/AI/Navigation/RecastNavMesh.cpp
@@ -697,7 +697,7 @@ void ARecastNavMesh::Serialize( FArchive& Ar )
 
 	if (Ar.IsLoading())
 	{
-		if (NavMeshVersion < NAVMESHVER_MIN_COMPATIBLE)
+		if (NavMeshVersion < NAVMESHVER_MIN_COMPATIBLE || (Ar.UE4Ver() >= VER_UE4_MERGED_ADD_MODIFIERS_RUNTIME_GENERATION_TO_4_7 && Ar.UE4Ver() < VER_UE4_AFTER_MERGED_ADD_MODIFIERS_RUNTIME_GENERATION_TO_4_7))
 		{
 			// incompatible, just skip over this data.  navmesh needs rebuilt.
 			Ar.Seek( RecastNavMeshSizePos + RecastNavMeshSizeBytes );
