SetEpicAccountId now returns a bool so you can tell if it *thinks* it worked. #UETOOL-244

diff --git a/Engine/Source/Runtime/Core/Private/GenericPlatform/GenericPlatformMisc.cpp b/Engine/Source/Runtime/Core/Private/GenericPlatform/GenericPlatformMisc.cpp
--- a/Engine/Source/Runtime/Core/Private/GenericPlatform/GenericPlatformMisc.cpp
+++ b/Engine/Source/Runtime/Core/Private/GenericPlatform/GenericPlatformMisc.cpp
@@ -283,9 +283,7 @@ bool FGenericPlatformMisc::SetStoredValue(const FString& InStoreId, const FStrin
 	KeyValue = InValue;
 
 	ConfigFile.Dirty = true;
-	ConfigFile.Write(ConfigPath);
-
-	return true;
+	return ConfigFile.Write(ConfigPath);
 }
 
 bool FGenericPlatformMisc::GetStoredValue(const FString& InStoreId, const FString& InSectionName, const FString& InKeyName, FString& OutValue)
@@ -849,9 +847,9 @@ FString FGenericPlatformMisc::GetEpicAccountId()
 	return EpicAccountId;
 }
 
-void FGenericPlatformMisc::SetEpicAccountId( const FString& AccountId )
+bool FGenericPlatformMisc::SetEpicAccountId( const FString& AccountId )
 {
-	FPlatformMisc::SetStoredValue( TEXT( "Epic Games" ), TEXT( "Unreal Engine/Identifiers" ), TEXT( "AccountId" ), AccountId );
+	return FPlatformMisc::SetStoredValue( TEXT( "Epic Games" ), TEXT( "Unreal Engine/Identifiers" ), TEXT( "AccountId" ), AccountId );
 }
 
 const TCHAR* FGenericPlatformMisc::GetEngineMode()
diff --git a/Engine/Source/Runtime/Core/Public/GenericPlatform/GenericPlatformMisc.h b/Engine/Source/Runtime/Core/Public/GenericPlatform/GenericPlatformMisc.h
--- a/Engine/Source/Runtime/Core/Public/GenericPlatform/GenericPlatformMisc.h
+++ b/Engine/Source/Runtime/Core/Public/GenericPlatform/GenericPlatformMisc.h
@@ -824,14 +824,16 @@ public:
 	static FGuid GetMachineId();
 
 	/**
-	 * Get the Epic account ID for the user who last used the Launcher
+	 * Get the Epic account ID for the user who last used the Launcher.
+	 * @return an empty string if the account ID was not present or it failed to read it for any reason.
 	 */
 	static FString GetEpicAccountId();
 
 	/**
 	 * Set the Epic account ID for the user who last used the Launcher
+	 * @return true if the account ID was set successfully, false if something failed and it was not set.
 	 */
-	static void SetEpicAccountId( const FString& AccountId );
+	static bool SetEpicAccountId( const FString& AccountId );
 
 	/** 
 	 * Get a string description of the mode the engine was running in.
