Remove changes to handle loading assets from 4.7 branch. 4.7 branch is currently at a higher object version than main, while we move assets over.

diff --git a/Engine/Source/Runtime/Core/Public/UObject/ObjectVersion.h b/Engine/Source/Runtime/Core/Public/UObject/ObjectVersion.h
--- a/Engine/Source/Runtime/Core/Public/UObject/ObjectVersion.h
+++ b/Engine/Source/Runtime/Core/Public/UObject/ObjectVersion.h
@@ -457,9 +457,9 @@ enum EUnrealEngineObjectUE4Version
 	// Tidied up material's handling of masked blend mode.
 	VER_UE4_MATERIAL_MASKED_BLENDMODE_TIDY,
 	// Merged VER_UE4_ADD_MODIFIERS_RUNTIME_GENERATION into 4.7 branch. 
-	VER_UE4_MERGED_ADD_MODIFIERS_RUNTIME_GENERATION_TO_4_7,
+	VER_UE4_MERGED_ADD_MODIFIERS_RUNTIME_GENERATION_TO_4_7_DEPRECATED,
 	// After adding VER_UE4_MERGED_ADD_MODIFIERS_RUNTIME_GENERATION_TO_4_7 into 4.7 branch. 
-	VER_UE4_AFTER_MERGED_ADD_MODIFIERS_RUNTIME_GENERATION_TO_4_7,
+	VER_UE4_AFTER_MERGED_ADD_MODIFIERS_RUNTIME_GENERATION_TO_4_7_DEPRECATED,
 
 	// -----<new versions can be added before this line>-------------------------------------------------
 	// - this needs to be the last line (see note below)
diff --git a/Engine/Source/Runtime/Engine/Private/AI/Navigation/RecastNavMesh.cpp b/Engine/Source/Runtime/Engine/Private/AI/Navigation/RecastNavMesh.cpp
--- a/Engine/Source/Runtime/Engine/Private/AI/Navigation/RecastNavMesh.cpp
+++ b/Engine/Source/Runtime/Engine/Private/AI/Navigation/RecastNavMesh.cpp
@@ -697,7 +697,7 @@ void ARecastNavMesh::Serialize( FArchive& Ar )
 
 	if (Ar.IsLoading())
 	{
-		if (NavMeshVersion < NAVMESHVER_MIN_COMPATIBLE || (Ar.UE4Ver() >= VER_UE4_MERGED_ADD_MODIFIERS_RUNTIME_GENERATION_TO_4_7 && Ar.UE4Ver() < VER_UE4_AFTER_MERGED_ADD_MODIFIERS_RUNTIME_GENERATION_TO_4_7))
+		if (NavMeshVersion < NAVMESHVER_MIN_COMPATIBLE)
 		{
 			// incompatible, just skip over this data.  navmesh needs rebuilt.
 			Ar.Seek( RecastNavMeshSizePos + RecastNavMeshSizeBytes );
