Fixes for Doxygen warnings.

diff --git a/Engine/Source/Programs/UnrealDocTool/APIDocTool/APIDocTool/Program.cs b/Engine/Source/Programs/UnrealDocTool/APIDocTool/APIDocTool/Program.cs
--- a/Engine/Source/Programs/UnrealDocTool/APIDocTool/APIDocTool/Program.cs
+++ b/Engine/Source/Programs/UnrealDocTool/APIDocTool/APIDocTool/Program.cs
@@ -215,6 +215,7 @@ namespace APIDocTool
 			"MSVC_PRAGMA(X)=",
 			"MS_ALIGN(X)= ",
 			"GCC_ALIGN(X)= ",
+			"SAFE_BOOL_OPERATORS(X)= ",
 			"VARARGS=",
 			"VARARG_DECL(FuncRet,StaticFuncRet,Return,FuncName,Pure,FmtType,ExtraDecl,ExtraCall)=FuncRet FuncName(ExtraDecl FmtType Fmt, ...)",
 			"VARARG_BODY(FuncRet,FuncName,FmtType,ExtraDecl)=FuncRet FuncName(ExtraDecl FmtType Fmt, ...)",
@@ -299,6 +300,7 @@ namespace APIDocTool
 			// Delegates
 			"FUNC_PARAM_MEMBERS=",
 			"DECLARE_DERIVED_EVENT(X,Y,Z)="
+			"DELEGATE_DEPRECATED(X)=",
 		};
 
 		static Program()
diff --git a/Engine/Source/Runtime/Core/Private/GenericPlatform/GenericPlatformString.cpp b/Engine/Source/Runtime/Core/Private/GenericPlatform/GenericPlatformString.cpp
--- a/Engine/Source/Runtime/Core/Private/GenericPlatform/GenericPlatformString.cpp
+++ b/Engine/Source/Runtime/Core/Private/GenericPlatform/GenericPlatformString.cpp
@@ -66,9 +66,11 @@ void FGenericPlatformString::LogBogusChars(const SourceEncoding* Src, int32 SrcS
 	}
 }
 
+#if !UE_BUILD_DOCS
 template CORE_API void FGenericPlatformString::LogBogusChars<ANSICHAR, WIDECHAR>(const WIDECHAR* Src, int32 SrcSize);
 template CORE_API void FGenericPlatformString::LogBogusChars<ANSICHAR, UCS2CHAR>(const UCS2CHAR* Src, int32 SrcSize);
 template CORE_API void FGenericPlatformString::LogBogusChars<WIDECHAR, ANSICHAR>(const ANSICHAR* Src, int32 SrcSize);
 template CORE_API void FGenericPlatformString::LogBogusChars<WIDECHAR, UCS2CHAR>(const UCS2CHAR* Src, int32 SrcSize);
 template CORE_API void FGenericPlatformString::LogBogusChars<UCS2CHAR, ANSICHAR>(const ANSICHAR* Src, int32 SrcSize);
 template CORE_API void FGenericPlatformString::LogBogusChars<UCS2CHAR, WIDECHAR>(const WIDECHAR* Src, int32 SrcSize);
+#endif
\ No newline at end of file
diff --git a/Engine/Source/Runtime/InputCore/Private/InputCoreTypes.cpp b/Engine/Source/Runtime/InputCore/Private/InputCoreTypes.cpp
--- a/Engine/Source/Runtime/InputCore/Private/InputCoreTypes.cpp
+++ b/Engine/Source/Runtime/InputCore/Private/InputCoreTypes.cpp
@@ -2,6 +2,7 @@
 
 #include "InputCorePrivatePCH.h"
 #include "PropertyTag.h"
+#include "InputCoreTypes.h"
 
 UInputCoreTypes::UInputCoreTypes(const FObjectInitializer& ObjectInitializer)
 : Super(ObjectInitializer)
