Re-enable thread affinity on Xbox One, and set some better defaults core mappings.

diff --git a/Engine/Source/Runtime/Core/Public/GenericPlatform/GenericPlatformAffinity.h b/Engine/Source/Runtime/Core/Public/GenericPlatform/GenericPlatformAffinity.h
--- a/Engine/Source/Runtime/Core/Public/GenericPlatform/GenericPlatformAffinity.h
+++ b/Engine/Source/Runtime/Core/Public/GenericPlatform/GenericPlatformAffinity.h
@@ -3,11 +3,13 @@
 #pragma once
 #include "HAL/Platform.h"
 
-#define		MAKEAFFINITYMASK1(x)			(1<<x)
-#define		MAKEAFFINITYMASK2(x,y)			((1<<x)+(1<<y))
-#define		MAKEAFFINITYMASK3(x,y,z)		((1<<x)+(1<<y)+(1<<z))
-#define		MAKEAFFINITYMASK4(w,x,y,z)		((1<<w)+(1<<x)+(1<<y)+(1<<z))
-#define		MAKEAFFINITYMASK5(v,w,x,y,z)	((1<<v)+(1<<w)+(1<<x)+(1<<y)+(1<<z))
+#define		MAKEAFFINITYMASK1(x)				(1<<x)
+#define		MAKEAFFINITYMASK2(x,y)				((1<<x)+(1<<y))
+#define		MAKEAFFINITYMASK3(x,y,z)			((1<<x)+(1<<y)+(1<<z))
+#define		MAKEAFFINITYMASK4(w,x,y,z)			((1<<w)+(1<<x)+(1<<y)+(1<<z))
+#define		MAKEAFFINITYMASK5(v,w,x,y,z)		((1<<v)+(1<<w)+(1<<x)+(1<<y)+(1<<z))
+#define		MAKEAFFINITYMASK6(u,v,w,x,y,z)		((1<<u)+(1<<v)+(1<<w)+(1<<x)+(1<<y)+(1<<z))
+#define		MAKEAFFINITYMASK7(t,u,v,w,x,y,z)	((1<<t)+(1<<u)+(1<<v)+(1<<w)+(1<<x)+(1<<y)+(1<<z))
 
 class FGenericPlatformAffinity
 {
@@ -22,6 +24,11 @@ public:
 		return 0xFFFFFFFFFFFFFFFF;
 	}
 
+	static const CORE_API uint64 GetRHIThreadMask()
+	{
+		return 0xFFFFFFFFFFFFFFFF;
+	}
+
 	static const CORE_API uint64 GetRTHeartBeatMask()
 	{
 		return 0xFFFFFFFFFFFFFFFF;
diff --git a/Engine/Source/Runtime/RenderCore/Private/RenderingThread.cpp b/Engine/Source/Runtime/RenderCore/Private/RenderingThread.cpp
--- a/Engine/Source/Runtime/RenderCore/Private/RenderingThread.cpp
+++ b/Engine/Source/Runtime/RenderCore/Private/RenderingThread.cpp
@@ -244,8 +244,7 @@ public:
 	void Start()
 	{
 		Thread = FRunnableThread::Create(this, TEXT("RHIThread"), 512 * 1024, TPri_Normal, 
-			//FPlatformAffinity::GetRHIThreadMask()
-			0xFFFFFFFFFFFFFFFF
+			FPlatformAffinity::GetRHIThreadMask()
 			);
 		check(Thread);
 	}
