Update ObjectVersion in main past the ObjectVersion in the 4.7 branch, so we can detect packages serialized with the cobbled 4.7 editor.

diff --git a/Engine/Source/Runtime/Core/Public/UObject/ObjectVersion.h b/Engine/Source/Runtime/Core/Public/UObject/ObjectVersion.h
--- a/Engine/Source/Runtime/Core/Public/UObject/ObjectVersion.h
+++ b/Engine/Source/Runtime/Core/Public/UObject/ObjectVersion.h
@@ -456,10 +456,14 @@ enum EUnrealEngineObjectUE4Version
 	VER_UE4_ADD_MODIFIERS_RUNTIME_GENERATION,
 	// Tidied up material's handling of masked blend mode.
 	VER_UE4_MATERIAL_MASKED_BLENDMODE_TIDY,
-	// Merged VER_UE4_ADD_MODIFIERS_RUNTIME_GENERATION into 4.7 branch. 
+	// Original version of VER_UE4_MERGED_ADD_MODIFIERS_RUNTIME_GENERATION_TO_4_7; renumbered to prevent blocking promotion in main.
 	VER_UE4_MERGED_ADD_MODIFIERS_RUNTIME_GENERATION_TO_4_7_DEPRECATED,
-	// After adding VER_UE4_MERGED_ADD_MODIFIERS_RUNTIME_GENERATION_TO_4_7 into 4.7 branch. 
+	// Original version of VER_UE4_AFTER_MERGED_ADD_MODIFIERS_RUNTIME_GENERATION_TO_4_7; renumbered to prevent blocking promotion in main.
 	VER_UE4_AFTER_MERGED_ADD_MODIFIERS_RUNTIME_GENERATION_TO_4_7_DEPRECATED,
+	// After merging VER_UE4_ADD_MODIFIERS_RUNTIME_GENERATION into 4.7 branch
+	VER_UE4_MERGED_ADD_MODIFIERS_RUNTIME_GENERATION_TO_4_7,
+	// After merging VER_UE4_ADD_MODIFIERS_RUNTIME_GENERATION into 4.7 branch
+	VER_UE4_AFTER_MERGING_ADD_MODIFIERS_RUNTIME_GENERATION_TO_4_7,
 
 	// -----<new versions can be added before this line>-------------------------------------------------
 	// - this needs to be the last line (see note below)
