Call MessageHandler->OnApplicationActivationChanged() on Mac when needed

diff --git a/Engine/Source/Runtime/Core/Private/Mac/MacApplication.cpp b/Engine/Source/Runtime/Core/Private/Mac/MacApplication.cpp
--- a/Engine/Source/Runtime/Core/Private/Mac/MacApplication.cpp
+++ b/Engine/Source/Runtime/Core/Private/Mac/MacApplication.cpp
@@ -174,7 +174,11 @@ FMacApplication::FMacApplication()
 
 								// app is active, allow sound
 								FApp::SetVolumeMultiplier( 1.0f );
-							}];
+
+								GameThreadCall(^{
+									MessageHandler->OnApplicationActivationChanged(true);
+								}, @[ NSDefaultRunLoopMode ], false);
+							 }];
 
 	AppDeactivationObserver = [[NSNotificationCenter defaultCenter] addObserverForName:NSApplicationWillResignActiveNotification object:[NSApplication sharedApplication] queue:[NSOperationQueue mainQueue] usingBlock:^(NSNotification* Notification)
 							{
@@ -219,6 +223,10 @@ FMacApplication::FMacApplication()
 
 								// app is inactive, apply multiplier
 								FApp::SetVolumeMultiplier(FApp::GetUnfocusedVolumeMultiplier());
+
+								GameThreadCall(^{
+									MessageHandler->OnApplicationActivationChanged(false);
+								}, @[ NSDefaultRunLoopMode ], false);
 							}];
 
 	WorkspaceActivationObserver = [[[NSWorkspace sharedWorkspace] notificationCenter] addObserverForName:NSWorkspaceSessionDidBecomeActiveNotification object:[NSWorkspace sharedWorkspace] queue:[NSOperationQueue mainQueue] usingBlock:^(NSNotification* Notification){
