Fixed up path in android causing launchon to fail because of missing file time stamps.

diff --git a/Engine/Source/Programs/AutomationTool/Scripts/CopyBuildToStagingDirectory.Automation.cs b/Engine/Source/Programs/AutomationTool/Scripts/CopyBuildToStagingDirectory.Automation.cs
--- a/Engine/Source/Programs/AutomationTool/Scripts/CopyBuildToStagingDirectory.Automation.cs
+++ b/Engine/Source/Programs/AutomationTool/Scripts/CopyBuildToStagingDirectory.Automation.cs
@@ -411,9 +411,11 @@ public partial class Project : CommandUtils
 
 	public static void DumpTargetManifest(Dictionary<string, string> Mapping, string Filename, string StageDir, List<string> CRCFiles)
 	{
-        // const string Iso8601DateTimeFormat = "%Y-%m-%dT%H:%M:%S.%sZ";
-        // const string Iso8601DateTimeFormat = "yyyy'-'MM'-'dd HH':'mm':'ss'Z'";
-        const string Iso8601DateTimeFormat = "u"; // predefined universal Iso standard format
+        // const string Iso8601DateTimeFormat = "yyyy-MM-ddTHH:mm:ssZ"; // probably should work
+		// const string Iso8601DateTimeFormat = "o"; // predefined universal Iso standard format (has too many millisecond spaces for our read code in FDateTime.ParseISO8601
+        const string Iso8601DateTimeFormat = "yyyy'-'MM'-'dd'T'HH':'mm':'ss'.'fffzzz";
+
+        
         if (Mapping.Count > 0)
 		{
 			var Lines = new List<string>();
diff --git a/Engine/Source/Runtime/Core/Private/Android/AndroidFile.cpp b/Engine/Source/Runtime/Core/Private/Android/AndroidFile.cpp
--- a/Engine/Source/Runtime/Core/Private/Android/AndroidFile.cpp
+++ b/Engine/Source/Runtime/Core/Private/Android/AndroidFile.cpp
@@ -396,6 +396,7 @@ public:
 #if LOG_ANDROID_FILE_MANIFEST
 					FPlatformMisc::LowLevelOutputDebugStringf(TEXT("Read time stamp '%s' %s"), *Filename, *ModifiedDate.ToString());
 #endif
+					Filename.ReplaceInline( TEXT("\\"), TEXT("/") );
 					ManifestEntries.Emplace( MoveTemp(Filename), ModifiedDate );
 				}
 				else
@@ -1084,6 +1085,11 @@ public:
 				return Result;
 			}
 
+			if ( UFSManifest.GetFileTimeStamp( AssetPath, Result ) )
+			{
+				return Result;
+			}
+
 #if LOG_ANDROID_FILE_MANIFEST
 			FPlatformMisc::LowLevelOutputDebugStringf(TEXT("Failed to find time stamp in NonUFSManifest for file '%s'"), Filename);
 #endif
@@ -1134,8 +1140,17 @@ public:
 			utime(TCHAR_TO_UTF8(*LocalPath), &Times);
 #else
 			// do something better as utime isn't supported on android very well...
-			NonUFSManifest.SetFileTimeStamp( AssetPath, DateTime );
-			NonUFSManifest.Write();
+			FDateTime TempDateTime;
+			if ( NonUFSManifest.GetFileTimeStamp( AssetPath, TempDateTime ) )
+			{
+				NonUFSManifest.SetFileTimeStamp( AssetPath, DateTime );
+				NonUFSManifest.Write();
+			}
+			else
+			{
+				UFSManifest.SetFileTimeStamp( AssetPath, DateTime );
+				UFSManifest.Write();
+			}
 #endif
 		}
 	}
