Linux: working around limitation of open file handles.

- Largely the same approach as on Mac.

diff --git a/Engine/Source/Runtime/Core/Private/Linux/LinuxPlatformFile.cpp b/Engine/Source/Runtime/Core/Private/Linux/LinuxPlatformFile.cpp
--- a/Engine/Source/Runtime/Core/Private/Linux/LinuxPlatformFile.cpp
+++ b/Engine/Source/Runtime/Core/Private/Linux/LinuxPlatformFile.cpp
@@ -9,13 +9,20 @@ DEFINE_LOG_CATEGORY_STATIC(LogLinuxPlatformFile, Log, All);
 // make an FTimeSpan object that represents the "epoch" for time_t (from a stat struct)
 const FDateTime UnixEpoch(1970, 1, 1);
 
+/** 
+ * Linux file handle implementation which limits number of open files per thread. This
+ * is to prevent running out of system file handles. Should not be neccessary when
+ * using pak file (e.g., SHIPPING?) so not particularly optimized. Only manages
+ * files which are opened READ_ONLY.
+ */
+#define MANAGE_FILE_HANDLES 	PLATFORM_LINUX // !UE_BUILD_SHIPPING
+
 /** 
  * Linux file handle implementation
-**/
+ */
 class CORE_API FFileHandleLinux : public IFileHandle
 {
 	enum {READWRITE_SIZE = 1024 * 1024};
-	int32 FileHandle;
 
 	FORCEINLINE bool IsValid()
 	{
@@ -23,53 +30,116 @@ class CORE_API FFileHandleLinux : public IFileHandle
 	}
 
 public:
-	FFileHandleLinux(int32 InFileHandle = -1)
+	FFileHandleLinux(int32 InFileHandle, const TCHAR* InFilename, bool bIsReadOnly)
 		: FileHandle(InFileHandle)
+#if MANAGE_FILE_HANDLES
+		, Filename(InFilename)
+		, HandleSlot(-1)
+		, FileOffset(0)
+		, FileSize(0)
+#endif // MANAGE_FILE_HANDLES
 	{
+		check(FileHandle > -1);
+		check(Filename.Len() > 0);
+		
+#if MANAGE_FILE_HANDLES
+		// Only files opened for read will be managed
+		if (bIsReadOnly)
+		{
+			ReserveSlot();
+			ActiveHandles[HandleSlot] = this;
+			struct stat FileInfo;
+			fstat(FileHandle, &FileInfo);
+			FileSize = FileInfo.st_size;
+		}
+#endif // MANAGE_FILE_HANDLES
 	}
 
 	virtual ~FFileHandleLinux()
 	{
-		close(FileHandle);
+#if MANAGE_FILE_HANDLES
+		if( IsManaged() )
+		{
+			if( ActiveHandles[ HandleSlot ] == this )
+			{
+				close(FileHandle);
+				ActiveHandles[ HandleSlot ] = NULL;
+			}
+		}
+		else
+#endif // MANAGE_FILE_HANDLES
+		{
+			close(FileHandle);
+		}
 		FileHandle = -1;
 	}
 
 	virtual int64 Tell() override
 	{
-		check(IsValid());
-		return lseek(FileHandle, 0, SEEK_CUR);
+#if MANAGE_FILE_HANDLES
+		if( IsManaged() )
+		{
+			return FileOffset;
+		}
+		else
+#endif // MANAGE_FILE_HANDLES
+		{
+			check(IsValid());
+			return lseek(FileHandle, 0, SEEK_CUR);
+		}
 	}
 
 	virtual bool Seek(int64 NewPosition) override
 	{
-		check(IsValid());
 		check(NewPosition >= 0);
-		return lseek(FileHandle, NewPosition, SEEK_SET) != -1;
+		
+#if MANAGE_FILE_HANDLES
+		if( IsManaged() )
+		{
+			FileOffset = NewPosition >= FileSize ? FileSize - 1 : NewPosition;
+			return IsValid() && ActiveHandles[ HandleSlot ] == this ? lseek(FileHandle, FileOffset, SEEK_SET) != -1 : true;
+		}
+		else
+#endif // MANAGE_FILE_HANDLES
+		{
+			check(IsValid());
+			return lseek(FileHandle, NewPosition, SEEK_SET) != -1;
+		}
 	}
 
 	virtual bool SeekFromEnd(int64 NewPositionRelativeToEnd = 0) override
 	{
-		check(IsValid());
 		check(NewPositionRelativeToEnd <= 0);
-		return lseek(FileHandle, NewPositionRelativeToEnd, SEEK_END) != -1;
+
+#if MANAGE_FILE_HANDLES
+		if( IsManaged() )
+		{
+			FileOffset = (NewPositionRelativeToEnd >= FileSize) ? 0 : ( FileSize + NewPositionRelativeToEnd - 1 );
+			return IsValid() && ActiveHandles[ HandleSlot ] == this ? lseek(FileHandle, FileOffset, SEEK_SET) != -1 : true;
+		}
+		else
+#endif // MANAGE_FILE_HANDLES
+		{
+			check(IsValid());
+			return lseek(FileHandle, NewPositionRelativeToEnd, SEEK_END) != -1;
+		}
 	}
 
 	virtual bool Read(uint8* Destination, int64 BytesToRead) override
 	{
-		check(IsValid());
-		while (BytesToRead)
+#if MANAGE_FILE_HANDLES
+		if( IsManaged() )
 		{
-			check(BytesToRead >= 0);
-			int64 ThisSize = FMath::Min<int64>(READWRITE_SIZE, BytesToRead);
-			check(Destination);
-			if (read(FileHandle, Destination, ThisSize) != ThisSize)
-			{
-				return false;
-			}
-			Destination += ThisSize;
-			BytesToRead -= ThisSize;
+			ActivateSlot();
+			int64 BytesRead = ReadInternal(Destination, BytesToRead);
+			FileOffset += BytesRead;
+			return BytesRead == BytesToRead;
+		}
+		else
+#endif // MANAGE_FILE_HANDLES
+		{
+			return ReadInternal(Destination, BytesToRead) == BytesToRead;
 		}
-		return true;
 	}
 
 	virtual bool Write(const uint8* Source, int64 BytesToWrite) override
@@ -89,8 +159,143 @@ public:
 		}
 		return true;
 	}
+
+	virtual int64 Size() override
+	{
+		#if MANAGE_FILE_HANDLES
+		if( IsManaged() )
+		{
+			return FileSize;
+		}
+		else
+			#endif
+		{
+			struct stat FileInfo;
+			fstat(FileHandle, &FileInfo);
+			return FileInfo.st_size;
+		}
+	}
+
+	
+private:
+
+#if MANAGE_FILE_HANDLES
+	FORCEINLINE bool IsManaged()
+	{
+		return HandleSlot != -1;
+	}
+
+	void ActivateSlot()
+	{
+		if( IsManaged() )
+		{
+			if( ActiveHandles[ HandleSlot ] != this || (ActiveHandles[ HandleSlot ] && ActiveHandles[ HandleSlot ]->FileHandle == -1) )
+			{
+				ReserveSlot();
+				
+				FileHandle = open(TCHAR_TO_UTF8(*Filename), O_RDONLY);
+				if( FileHandle != -1 )
+				{
+					lseek(FileHandle, FileOffset, SEEK_SET);
+					ActiveHandles[ HandleSlot ] = this;
+				}
+				else
+				{
+					UE_LOG(LogLinuxPlatformFile, Warning, TEXT("Could not (re)activate slot for file '%s'"), *Filename);
+				}
+			}
+			else
+			{
+				AccessTimes[ HandleSlot ] = FPlatformTime::Seconds();
+			}
+		}
+	}
+
+	void ReserveSlot()
+	{
+		HandleSlot = -1;
+		
+		// Look for non-reserved slot
+		for( int32 i = 0; i < ACTIVE_HANDLE_COUNT; ++i )
+		{
+			if( ActiveHandles[ i ] == NULL )
+			{
+				HandleSlot = i;
+				break;
+			}
+		}
+		
+		// Take the oldest handle
+		if( HandleSlot == -1 )
+		{
+			int32 Oldest = 0;
+			for( int32 i = 1; i < ACTIVE_HANDLE_COUNT; ++i )
+			{
+				if( AccessTimes[ Oldest ] > AccessTimes[ i ] )
+				{
+					Oldest = i;
+				}
+			}
+			
+			close( ActiveHandles[ Oldest ]->FileHandle );
+			ActiveHandles[ Oldest ]->FileHandle = -1;
+			HandleSlot = Oldest;
+		}
+		
+		ActiveHandles[ HandleSlot ] = NULL;
+		AccessTimes[ HandleSlot ] = FPlatformTime::Seconds();
+	}
+#endif // MANAGE_FILE_HANDLES
+
+	int64 ReadInternal(uint8* Destination, int64 BytesToRead)
+	{
+		check(IsValid());
+		int64 BytesRead = 0;
+		while (BytesToRead)
+		{
+			check(BytesToRead >= 0);
+			int64 ThisSize = FMath::Min<int64>(READWRITE_SIZE, BytesToRead);
+			check(Destination);
+			int64 ThisRead = read(FileHandle, Destination, ThisSize);
+			BytesRead += ThisRead;
+			if (ThisRead != ThisSize)
+			{
+				return BytesRead;
+			}
+			Destination += ThisSize;
+			BytesToRead -= ThisSize;
+		}
+		return BytesRead;
+	}
+
+	// Holds the internal file handle.
+	int32 FileHandle;
+
+#if MANAGE_FILE_HANDLES
+	// Holds the name of the file that this handle represents. Kept around for possible reopen of file.
+	FString Filename;
+	
+	// Most recent valid slot index for this handle; >=0 for handles which are managed.
+	int32 HandleSlot;
+	
+	// Current file offset; valid if a managed handle.
+	int64 FileOffset;
+	
+	// Cached file size; valid if a managed handle.
+	int64 FileSize;
+	
+	// Each thread keeps a collection of active handles with access times.
+	static const int32 ACTIVE_HANDLE_COUNT = 256;
+	static __thread FFileHandleLinux* ActiveHandles[ ACTIVE_HANDLE_COUNT ];
+	static __thread double AccessTimes[ ACTIVE_HANDLE_COUNT ];
+#endif // MANAGE_FILE_HANDLES
 };
 
+#if MANAGE_FILE_HANDLES
+__thread FFileHandleLinux* FFileHandleLinux::ActiveHandles[ FFileHandleLinux::ACTIVE_HANDLE_COUNT ];
+__thread double FFileHandleLinux::AccessTimes[ FFileHandleLinux::ACTIVE_HANDLE_COUNT ];
+#endif // MANAGE_FILE_HANDLES
+
 /**
  * Linux File I/O implementation
 **/
@@ -387,12 +592,17 @@ public:
 	 * Opens a file for reading, disregarding the case.
 	 * 
 	 * @param Filename absolute filename
+	 * @param MappedToFilename absolute filename that we mapped the Filename to (always filled out on success, even if the same as Filename)
 	 */
-	int32 OpenCaseInsensitiveRead(const FString & Filename)
+	int32 OpenCaseInsensitiveRead(const FString & Filename, FString & MappedToFilename)
 	{
 		// try opening right away
 		int32 Handle = open(TCHAR_TO_UTF8(*Filename), O_RDONLY);
-		if (Handle == -1)
+		if (Handle != -1)
+		{
+			MappedToFilename = Filename;
+		}
+		else
 		{
 			// log non-standard errors only
 			if (ENOENT != errno)
@@ -414,7 +624,11 @@ public:
 
 					if (Handle != -1)
 					{
-						UE_LOG(LogLinuxPlatformFile, Log, TEXT("Mapped '%s' to '%s'"), *Filename, *FoundFilename);
+						MappedToFilename = FoundFilename;
+						if (Filename != MappedToFilename)
+						{
+							UE_LOG(LogLinuxPlatformFile, Log, TEXT("Mapped '%s' to '%s'"), *Filename, *MappedToFilename);
+						}
 					}
 				}
 			}
@@ -427,10 +641,11 @@ public:
 IFileHandle* FLinuxPlatformFile::OpenRead(const TCHAR* Filename)
 {
 	FLinuxFileMapper CaseInsensMapper;
-	int32 Handle = CaseInsensMapper.OpenCaseInsensitiveRead(TCHAR_TO_UTF8(*NormalizeFilename(Filename)));
+	FString MappedToName;
+	int32 Handle = CaseInsensMapper.OpenCaseInsensitiveRead(TCHAR_TO_UTF8(*NormalizeFilename(Filename)), MappedToName);
 	if (Handle != -1)
 	{
-		return new FFileHandleLinux(Handle);
+		return new FFileHandleLinux(Handle, *MappedToName, true);
 	}
 	return nullptr;
 }
@@ -488,7 +703,12 @@ IFileHandle* FLinuxPlatformFile::OpenWrite(const TCHAR* Filename, bool bAppend,
 			}
 		}
 
-		FFileHandleLinux* FileHandleLinux = new FFileHandleLinux(Handle);
+#if MANAGE_FILE_HANDLES
+		FFileHandleLinux* FileHandleLinux = new FFileHandleLinux(Handle, *NormalizeDirectory(Filename), false);
+#else
+		FFileHandleLinux* FileHandleLinux = new FFileHandleLinux(Handle, Filename, false);
+#endif // MANAGE_FILE_HANDLES
+
 		if (bAppend)
 		{
 			FileHandleLinux->SeekFromEnd(0);
