Fix binding extra params to a lambda/functor style delegate

diff --git a/Engine/Source/Runtime/Core/Public/Delegates/DelegateSignatureImpl_Variadics.inl b/Engine/Source/Runtime/Core/Public/Delegates/DelegateSignatureImpl_Variadics.inl
--- a/Engine/Source/Runtime/Core/Public/Delegates/DelegateSignatureImpl_Variadics.inl
+++ b/Engine/Source/Runtime/Core/Public/Delegates/DelegateSignatureImpl_Variadics.inl
@@ -105,7 +105,7 @@ public:
 	template<typename FunctorType, typename... VarTypes>
 	inline static TBaseDelegate<RetValType, ParamTypes...> CreateLambda(FunctorType&& InFunctor, VarTypes... Vars)
 	{
-		return TBaseFunctorDelegateInstance<TFuncType, typename TRemoveReference<FunctorType>::Type>::Create(Forward<FunctorType>(InFunctor), Vars...);
+		return TBaseFunctorDelegateInstance<TFuncType, typename TRemoveReference<FunctorType>::Type, VarTypes...>::Create(Forward<FunctorType>(InFunctor), Vars...);
 	}
 
 	/**
