CanConvertPointerFromTo deprecated - TPointerIsConvertibleFromTo should be used instead.

diff --git a/Engine/Source/Runtime/Core/Public/Delegates/DelegateInstancesImpl.inl b/Engine/Source/Runtime/Core/Public/Delegates/DelegateInstancesImpl.inl
--- a/Engine/Source/Runtime/Core/Public/Delegates/DelegateInstancesImpl.inl
+++ b/Engine/Source/Runtime/Core/Public/Delegates/DelegateInstancesImpl.inl
@@ -276,7 +276,7 @@ template<class UserClass, FUNC_PAYLOAD_TEMPLATE_DECL_TYPENAME>
 class RAW_METHOD_DELEGATE_INSTANCE_CLASS
 	: public DELEGATE_INSTANCE_INTERFACE_CLASS<FUNC_TEMPLATE_ARGS>
 {
-	static_assert((!CanConvertPointerFromTo<UserClass, UObjectBase>::Result), "You cannot use raw method delegates with UObjects.");
+	static_assert(!TPointerIsConvertibleFromTo<UserClass, const UObjectBase>::Value, "You cannot use raw method delegates with UObjects.");
 
 public:
 
@@ -449,7 +449,7 @@ template<class UserClass, FUNC_PAYLOAD_TEMPLATE_DECL_TYPENAME>
 class UOBJECT_METHOD_DELEGATE_INSTANCE_CLASS
 	: public DELEGATE_INSTANCE_INTERFACE_CLASS<FUNC_TEMPLATE_ARGS>
 {
-	static_assert((CanConvertPointerFromTo<UserClass, UObjectBase>::Result), "You cannot use UObject method delegates with raw pointers.");
+	static_assert(TPointerIsConvertibleFromTo<UserClass, const UObjectBase>::Value, "You cannot use UObject method delegates with raw pointers.");
 
 public:
 
@@ -904,7 +904,7 @@ template<class UserClass, FUNC_PAYLOAD_TEMPLATE_DECL_TYPENAME>
 class UFUNCTION_DELEGATE_INSTANCE_CLASS
 	: public DELEGATE_INSTANCE_INTERFACE_CLASS<FUNC_TEMPLATE_ARGS>
 {
-	static_assert((CanConvertPointerFromTo<UserClass, UObjectBase>::Result), "You cannot use UFunction delegates with non UObject classes.");
+	static_assert(TPointerIsConvertibleFromTo<UserClass, const UObjectBase>::Value, "You cannot use UFunction delegates with non UObject classes.");
 
 	// Structure for UFunction call parameter marshaling.
 	struct FParmsWithPayload
diff --git a/Engine/Source/Runtime/Core/Public/Delegates/DelegateInstancesImpl_Variadics.inl b/Engine/Source/Runtime/Core/Public/Delegates/DelegateInstancesImpl_Variadics.inl
--- a/Engine/Source/Runtime/Core/Public/Delegates/DelegateInstancesImpl_Variadics.inl
+++ b/Engine/Source/Runtime/Core/Public/Delegates/DelegateInstancesImpl_Variadics.inl
@@ -37,7 +37,7 @@ private:
 	typedef IBaseDelegateInstance<RetValType (ParamTypes...)>                                  Super;
 	typedef TBaseUFunctionDelegateInstance<UserClass, RetValType (ParamTypes...), VarTypes...> UnwrappedThisType;
 
-	static_assert((CanConvertPointerFromTo<UserClass, const UObjectBase>::Result), "You cannot use UFunction delegates with non UObject classes.");
+	static_assert(TPointerIsConvertibleFromTo<UserClass, const UObjectBase>::Value, "You cannot use UFunction delegates with non UObject classes.");
 
 public:
 	TBaseUFunctionDelegateInstance(UserClass* InUserObject, const FName& InFunctionName, VarTypes... Vars)
@@ -408,7 +408,7 @@ public:
 	typedef typename TUnwrapType<WrappedRetValType>::Type RetValType;
 
 private:
-	static_assert((!CanConvertPointerFromTo<UserClass, const UObjectBase>::Result), "You cannot use raw method delegates with UObjects.");
+	static_assert(!TPointerIsConvertibleFromTo<UserClass, const UObjectBase>::Value, "You cannot use raw method delegates with UObjects.");
 
 	typedef IBaseDelegateInstance<RetValType (ParamTypes...)>                                          Super;
 	typedef TBaseRawMethodDelegateInstance<bConst, UserClass, RetValType (ParamTypes...), VarTypes...> UnwrappedThisType;
@@ -600,7 +600,7 @@ private:
 	typedef IBaseDelegateInstance<RetValType (ParamTypes...)>                                              Super;
 	typedef TBaseUObjectMethodDelegateInstance<bConst, UserClass, RetValType (ParamTypes...), VarTypes...> UnwrappedThisType;
 
-	static_assert((CanConvertPointerFromTo<UserClass, const UObjectBase>::Result), "You cannot use UObject method delegates with raw pointers.");
+	static_assert(TPointerIsConvertibleFromTo<UserClass, const UObjectBase>::Value, "You cannot use UObject method delegates with raw pointers.");
 
 public:
 	typedef typename TMemFunPtrType<bConst, UserClass, RetValType (ParamTypes..., VarTypes...)>::Type FMethodPtr;
diff --git a/Engine/Source/Runtime/Core/Public/HAL/Platform.h b/Engine/Source/Runtime/Core/Public/HAL/Platform.h
--- a/Engine/Source/Runtime/Core/Public/HAL/Platform.h
+++ b/Engine/Source/Runtime/Core/Public/HAL/Platform.h
@@ -233,10 +233,6 @@
 	#define PLATFORM_SUPPORTS_DRAW_MESH_EVENTS	1
 #endif
 
-#ifndef PLATFORM_SUPPORTS_CanConvertPointerFromTo
-	#define PLATFORM_SUPPORTS_CanConvertPointerFromTo		1
-#endif
-
 #ifndef PLATFORM_USES_ES2
 	#define PLATFORM_USES_ES2					0
 #endif
diff --git a/Engine/Source/Runtime/Core/Public/Templates/PointerIsConvertibleFromTo.h b/Engine/Source/Runtime/Core/Public/Templates/PointerIsConvertibleFromTo.h
--- a/Engine/Source/Runtime/Core/Public/Templates/PointerIsConvertibleFromTo.h
+++ b/Engine/Source/Runtime/Core/Public/Templates/PointerIsConvertibleFromTo.h
@@ -15,3 +15,33 @@ private:
 public:
 	enum { Value  = sizeof(Test((From*)nullptr)) - 1 };
 };
+
+
+class TPointerIsConvertibleFromTo_TestBase
+{
+};
+
+class TPointerIsConvertibleFromTo_TestDerived : public TPointerIsConvertibleFromTo_TestBase
+{
+};
+
+class TPointerIsConvertibleFromTo_Unrelated
+{
+};
+
+static_assert(TPointerIsConvertibleFromTo<bool, bool>::Value, "Platform TPointerIsConvertibleFromTo test failed.");
+static_assert(TPointerIsConvertibleFromTo<void, void>::Value, "Platform TPointerIsConvertibleFromTo test failed.");
+static_assert(TPointerIsConvertibleFromTo<bool, void>::Value, "Platform TPointerIsConvertibleFromTo test failed.");
+static_assert(TPointerIsConvertibleFromTo<const bool, const void>::Value, "Platform TPointerIsConvertibleFromTo test failed.");
+static_assert(TPointerIsConvertibleFromTo<TPointerIsConvertibleFromTo_TestDerived, TPointerIsConvertibleFromTo_TestBase>::Value, "Platform TPointerIsConvertibleFromTo test failed.");
+static_assert(TPointerIsConvertibleFromTo<TPointerIsConvertibleFromTo_TestDerived, const TPointerIsConvertibleFromTo_TestBase>::Value, "Platform TPointerIsConvertibleFromTo test failed.");
+static_assert(TPointerIsConvertibleFromTo<const TPointerIsConvertibleFromTo_TestDerived, const TPointerIsConvertibleFromTo_TestBase>::Value, "Platform TPointerIsConvertibleFromTo test failed.");
+static_assert(TPointerIsConvertibleFromTo<TPointerIsConvertibleFromTo_TestBase, TPointerIsConvertibleFromTo_TestBase>::Value, "Platform TPointerIsConvertibleFromTo test failed.");
+static_assert(TPointerIsConvertibleFromTo<TPointerIsConvertibleFromTo_TestBase, void>::Value, "Platform TPointerIsConvertibleFromTo test failed.");
+
+static_assert(!TPointerIsConvertibleFromTo<TPointerIsConvertibleFromTo_TestBase, TPointerIsConvertibleFromTo_TestDerived>::Value, "Platform TPointerIsConvertibleFromTo test failed.");
+static_assert(!TPointerIsConvertibleFromTo<TPointerIsConvertibleFromTo_Unrelated, TPointerIsConvertibleFromTo_TestBase>::Value, "Platform TPointerIsConvertibleFromTo test failed.");
+static_assert(!TPointerIsConvertibleFromTo<bool, TPointerIsConvertibleFromTo_TestBase>::Value, "Platform TPointerIsConvertibleFromTo test failed.");
+static_assert(!TPointerIsConvertibleFromTo<void, TPointerIsConvertibleFromTo_TestBase>::Value, "Platform TPointerIsConvertibleFromTo test failed.");
+static_assert(!TPointerIsConvertibleFromTo<TPointerIsConvertibleFromTo_TestBase, bool>::Value, "Platform TPointerIsConvertibleFromTo test failed.");
+static_assert(!TPointerIsConvertibleFromTo<void, bool>::Value, "Platform TPointerIsConvertibleFromTo test failed.");
diff --git a/Engine/Source/Runtime/Core/Public/Templates/SharedPointer.h b/Engine/Source/Runtime/Core/Public/Templates/SharedPointer.h
--- a/Engine/Source/Runtime/Core/Public/Templates/SharedPointer.h
+++ b/Engine/Source/Runtime/Core/Public/Templates/SharedPointer.h
@@ -133,7 +133,7 @@ template< class ObjectType, ESPMode Mode >
 class TSharedRef
 {
 	// TSharedRefs with UObjects are illegal.
-	static_assert((!CanConvertPointerFromTo<ObjectType, UObjectBase>::Result), "You cannot use TSharedRef with UObjects.");
+	static_assert(!TPointerIsConvertibleFromTo<ObjectType, const UObjectBase>::Value, "You cannot use TSharedRef with UObjects.");
 
 public:
 
@@ -425,7 +425,7 @@ template< class ObjectType, ESPMode Mode >
 class TSharedPtr
 {
 	// TSharedPtrs with UObjects are illegal.
-	static_assert((!CanConvertPointerFromTo<ObjectType, UObjectBase>::Result), "You cannot use TSharedPtr with UObjects.");
+	static_assert(!TPointerIsConvertibleFromTo<ObjectType, const UObjectBase>::Value, "You cannot use TSharedPtr with UObjects.");
 
 public:
 
diff --git a/Engine/Source/Runtime/Core/Public/Templates/UnrealTemplate.h b/Engine/Source/Runtime/Core/Public/Templates/UnrealTemplate.h
--- a/Engine/Source/Runtime/Core/Public/Templates/UnrealTemplate.h
+++ b/Engine/Source/Runtime/Core/Public/Templates/UnrealTemplate.h
@@ -14,49 +14,16 @@
 -----------------------------------------------------------------------------*/
 
 /** 
-* CanConvertPointerFromTo<From, To>::Result is an enum value equal to 1 if From* is automatically convertable to a To* (without regard to const!)
-**/
+ * CanConvertPointerFromTo<From, To>::Result is an enum value equal to 1 if From* is automatically convertable to a To* (without regard to const!)
+ **/
 template<class From, class To>
 class CanConvertPointerFromTo
 {
 public:
-	enum Type
-	{
-		Result = TPointerIsConvertibleFromTo<From, const To>::Value
-	};
+	DEPRECATED(4.8, "This trait is deprecated as it does not follow UE4 standards of syntax and hides a const conversion. Please use TPointerIsConvertibleFromTo instead.")
+	static const bool Result = TPointerIsConvertibleFromTo<From, const To>::Value;
 };
 
-class CanConvertPointerFromTo_TestBase
-{
-};
-
-class CanConvertPointerFromTo_TestDerived : public CanConvertPointerFromTo_TestBase
-{
-};
-
-class CanConvertPointerFromTo_Unrelated
-{
-};
-
-static_assert((CanConvertPointerFromTo<bool, bool>::Result), "Platform CanConvertPointerFromTo test failed.");
-static_assert((CanConvertPointerFromTo<void, void>::Result), "Platform CanConvertPointerFromTo test failed.");
-static_assert((CanConvertPointerFromTo<bool, void>::Result), "Platform CanConvertPointerFromTo test failed.");
-static_assert((CanConvertPointerFromTo<const bool, void>::Result), "Platform CanConvertPointerFromTo test failed.");
-static_assert((CanConvertPointerFromTo<CanConvertPointerFromTo_TestDerived, CanConvertPointerFromTo_TestBase>::Result), "Platform CanConvertPointerFromTo test failed.");
-static_assert((CanConvertPointerFromTo<CanConvertPointerFromTo_TestDerived, const CanConvertPointerFromTo_TestBase>::Result), "Platform CanConvertPointerFromTo test failed.");
-static_assert((CanConvertPointerFromTo<const CanConvertPointerFromTo_TestDerived, CanConvertPointerFromTo_TestBase>::Result), "Platform CanConvertPointerFromTo test failed.");
-static_assert((CanConvertPointerFromTo<const CanConvertPointerFromTo_TestDerived, const CanConvertPointerFromTo_TestBase>::Result), "Platform CanConvertPointerFromTo test failed.");
-static_assert((CanConvertPointerFromTo<CanConvertPointerFromTo_TestBase, CanConvertPointerFromTo_TestBase>::Result), "Platform CanConvertPointerFromTo test failed.");
-static_assert((CanConvertPointerFromTo<CanConvertPointerFromTo_TestBase, void>::Result), "Platform CanConvertPointerFromTo test failed.");
-
-static_assert(!(CanConvertPointerFromTo<CanConvertPointerFromTo_TestBase, CanConvertPointerFromTo_TestDerived>::Result), "Platform CanConvertPointerFromTo test failed.");
-static_assert(!(CanConvertPointerFromTo<CanConvertPointerFromTo_Unrelated, CanConvertPointerFromTo_TestBase>::Result), "Platform CanConvertPointerFromTo test failed.");
-static_assert(!(CanConvertPointerFromTo<bool, CanConvertPointerFromTo_TestBase>::Result), "Platform CanConvertPointerFromTo test failed.");
-static_assert(!(CanConvertPointerFromTo<void, CanConvertPointerFromTo_TestBase>::Result), "Platform CanConvertPointerFromTo test failed.");
-static_assert(!(CanConvertPointerFromTo<CanConvertPointerFromTo_TestBase, bool>::Result), "Platform CanConvertPointerFromTo test failed.");
-static_assert(!(CanConvertPointerFromTo<void, bool>::Result), "Platform CanConvertPointerFromTo test failed.");
-
-
 
 /**
  * Aligns a value to the nearest higher multiple of 'Alignment', which must be a power of two.
diff --git a/Engine/Source/Runtime/Core/Public/UObject/WeakObjectPtrTemplates.h b/Engine/Source/Runtime/Core/Public/UObject/WeakObjectPtrTemplates.h
--- a/Engine/Source/Runtime/Core/Public/UObject/WeakObjectPtrTemplates.h
+++ b/Engine/Source/Runtime/Core/Public/UObject/WeakObjectPtrTemplates.h
@@ -40,7 +40,7 @@ public:
 	{
 		// This static assert is in here rather than in the body of the class because we want
 		// to be able to define TWeakObjectPtr<UUndefinedClass>.
-		static_assert(CanConvertPointerFromTo<T, UObject>::Result, "TWeakObjectPtr can only be constructed with UObject types");
+		static_assert(TPointerIsConvertibleFromTo<T, const volatile UObject>::Value, "TWeakObjectPtr can only be constructed with UObject types");
 	}
 
 	/**  
@@ -61,7 +61,7 @@ public:
 	{
 		// It's also possible that this static_assert may fail for valid conversions because
 		// one or both of the types have only been forward-declared.
-		static_assert(CanConvertPointerFromTo<OtherT, T>::Result, "Unable to convert TWeakObjectPtr - types are incompatible");
+		static_assert(TPointerIsConvertibleFromTo<OtherT, T>::Value, "Unable to convert TWeakObjectPtr - types are incompatible");
 	}
 
 	/**
@@ -101,7 +101,7 @@ public:
 	{
 		// It's also possible that this static_assert may fail for valid conversions because
 		// one or both of the types have only been forward-declared.
-		static_assert(CanConvertPointerFromTo<OtherT, T>::Result, "Unable to convert TWeakObjectPtr - types are incompatible");
+		static_assert(TPointerIsConvertibleFromTo<OtherT, T>::Value, "Unable to convert TWeakObjectPtr - types are incompatible");
 
 		TWeakObjectPtrBase::operator=(Other);
 	}
@@ -193,7 +193,7 @@ FORCENOINLINE bool operator==(const TWeakObjectPtr<LhsT, OtherTWeakObjectPtrBase
 {
 	// It's also possible that this static_assert may fail for valid conversions because
 	// one or both of the types have only been forward-declared.
-	static_assert(CanConvertPointerFromTo<LhsT, RhsT>::Result || CanConvertPointerFromTo<RhsT, LhsT>::Result, "Unable to compare TWeakObjectPtrs - types are incompatible");
+	static_assert(TPointerIsConvertibleFromTo<LhsT, RhsT>::Value || TPointerIsConvertibleFromTo<RhsT, LhsT>::Value, "Unable to compare TWeakObjectPtrs - types are incompatible");
 
 	return (const OtherTWeakObjectPtrBase&)Lhs == (const OtherTWeakObjectPtrBase&)Rhs;
 }
@@ -203,8 +203,8 @@ FORCENOINLINE bool operator==(const TWeakObjectPtr<LhsT, OtherTWeakObjectPtrBase
 {
 	// It's also possible that these static_asserts may fail for valid conversions because
 	// one or both of the types have only been forward-declared.
-	static_assert(CanConvertPointerFromTo<RhsT, UObject>::Result, "TWeakObjectPtr can only be compared with UObject types");
-	static_assert(CanConvertPointerFromTo<LhsT, RhsT>::Result || CanConvertPointerFromTo<RhsT, LhsT>::Result, "Unable to compare TWeakObjectPtr with raw pointer - types are incompatible");
+	static_assert(TPointerIsConvertibleFromTo<RhsT, UObject>::Value, "TWeakObjectPtr can only be compared with UObject types");
+	static_assert(TPointerIsConvertibleFromTo<LhsT, RhsT>::Value || TPointerIsConvertibleFromTo<RhsT, LhsT>::Value, "Unable to compare TWeakObjectPtr with raw pointer - types are incompatible");
 
 	// NOTE: this constructs a TWeakObjectPtrBase, which has some amount of overhead, so this may not be an efficient operation
 	return (const OtherTWeakObjectPtrBase&)Lhs == OtherTWeakObjectPtrBase(Rhs);
@@ -215,8 +215,8 @@ FORCENOINLINE bool operator==(const LhsT* Lhs, const TWeakObjectPtr<RhsT, OtherT
 {
 	// It's also possible that these static_asserts may fail for valid conversions because
 	// one or both of the types have only been forward-declared.
-	static_assert(CanConvertPointerFromTo<LhsT, UObject>::Result, "TWeakObjectPtr can only be compared with UObject types");
-	static_assert(CanConvertPointerFromTo<LhsT, RhsT>::Result || CanConvertPointerFromTo<RhsT, LhsT>::Result, "Unable to compare TWeakObjectPtr with raw pointer - types are incompatible");
+	static_assert(TPointerIsConvertibleFromTo<LhsT, UObject>::Value, "TWeakObjectPtr can only be compared with UObject types");
+	static_assert(TPointerIsConvertibleFromTo<LhsT, RhsT>::Value || TPointerIsConvertibleFromTo<RhsT, LhsT>::Value, "Unable to compare TWeakObjectPtr with raw pointer - types are incompatible");
 
 	// NOTE: this constructs a TWeakObjectPtrBase, which has some amount of overhead, so this may not be an efficient operation
 	return OtherTWeakObjectPtrBase(Lhs) == (const OtherTWeakObjectPtrBase&)Rhs;
@@ -239,7 +239,7 @@ FORCENOINLINE bool operator!=(const TWeakObjectPtr<LhsT, OtherTWeakObjectPtrBase
 {
 	// It's also possible that this static_assert may fail for valid conversions because
 	// one or both of the types have only been forward-declared.
-	static_assert(CanConvertPointerFromTo<LhsT, RhsT>::Result || CanConvertPointerFromTo<RhsT, LhsT>::Result, "Unable to compare TWeakObjectPtrs - types are incompatible");
+	static_assert(TPointerIsConvertibleFromTo<LhsT, RhsT>::Value || TPointerIsConvertibleFromTo<RhsT, LhsT>::Value, "Unable to compare TWeakObjectPtrs - types are incompatible");
 
 	return (const OtherTWeakObjectPtrBase&)Lhs != (const OtherTWeakObjectPtrBase&)Rhs;
 }
@@ -249,20 +249,20 @@ FORCENOINLINE bool operator!=(const TWeakObjectPtr<LhsT, OtherTWeakObjectPtrBase
 {
 	// It's also possible that these static_asserts may fail for valid conversions because
 	// one or both of the types have only been forward-declared.
-	static_assert(CanConvertPointerFromTo<RhsT, UObject>::Result, "TWeakObjectPtr can only be compared with UObject types");
-	static_assert(CanConvertPointerFromTo<LhsT, RhsT>::Result || CanConvertPointerFromTo<RhsT, LhsT>::Result, "Unable to compare TWeakObjectPtr with raw pointer - types are incompatible");
+	static_assert(TPointerIsConvertibleFromTo<RhsT, UObject>::Value, "TWeakObjectPtr can only be compared with UObject types");
+	static_assert(TPointerIsConvertibleFromTo<LhsT, RhsT>::Value || TPointerIsConvertibleFromTo<RhsT, LhsT>::Value, "Unable to compare TWeakObjectPtr with raw pointer - types are incompatible");
 
 	// NOTE: this constructs a TWeakObjectPtrBase, which has some amount of overhead, so this may not be an efficient operation
 	return (const OtherTWeakObjectPtrBase&)Lhs != OtherTWeakObjectPtrBase(Rhs);
 }
 
 template <typename LhsT, typename RhsT, typename OtherTWeakObjectPtrBase, typename OtherTUObjectArray>
-FORCENOINLINE bool operator!=(const RhsT* Lhs, const TWeakObjectPtr<LhsT, OtherTWeakObjectPtrBase, OtherTUObjectArray> &Rhs)
+FORCENOINLINE bool operator!=(const LhsT* Lhs, const TWeakObjectPtr<RhsT, OtherTWeakObjectPtrBase, OtherTUObjectArray> &Rhs)
 {
 	// It's also possible that these static_asserts may fail for valid conversions because
 	// one or both of the types have only been forward-declared.
-	static_assert(CanConvertPointerFromTo<LhsT, UObject>::Result, "TWeakObjectPtr can only be compared with UObject types");
-	static_assert(CanConvertPointerFromTo<LhsT, RhsT>::Result || CanConvertPointerFromTo<RhsT, LhsT>::Result, "Unable to compare TWeakObjectPtr with raw pointer - types are incompatible");
+	static_assert(TPointerIsConvertibleFromTo<LhsT, UObject>::Value, "TWeakObjectPtr can only be compared with UObject types");
+	static_assert(TPointerIsConvertibleFromTo<LhsT, RhsT>::Value || TPointerIsConvertibleFromTo<RhsT, LhsT>::Value, "Unable to compare TWeakObjectPtr with raw pointer - types are incompatible");
 
 	// NOTE: this constructs a TWeakObjectPtrBase, which has some amount of overhead, so this may not be an efficient operation
 	return OtherTWeakObjectPtrBase(Lhs) != (const OtherTWeakObjectPtrBase&)Rhs;
diff --git a/Engine/Source/Runtime/CoreUObject/Public/Templates/Casts.h b/Engine/Source/Runtime/CoreUObject/Public/Templates/Casts.h
--- a/Engine/Source/Runtime/CoreUObject/Public/Templates/Casts.h
+++ b/Engine/Source/Runtime/CoreUObject/Public/Templates/Casts.h
@@ -13,7 +13,7 @@ COREUOBJECT_API void CastLogError(const TCHAR* FromType, const TCHAR* ToType);
  *
  * Otherwise, assume it's an IInterface.
  */
-template <typename T, bool bIsAUObject_IMPL = CanConvertPointerFromTo<T, UObject>::Result>
+template <typename T, bool bIsAUObject_IMPL = TPointerIsConvertibleFromTo<T, const volatile UObject>::Value>
 struct TIsIInterface
 {
 	enum { Value = false };
@@ -32,7 +32,7 @@ template <typename T>
 struct TIsCastable
 {
 	// It's from-castable if it's an interface or a UObject-derived type
-	enum { Value = TIsIInterface<T>::Value || CanConvertPointerFromTo<T, const volatile UObject>::Result };
+	enum { Value = TIsIInterface<T>::Value || TPointerIsConvertibleFromTo<T, const volatile UObject>::Value };
 };
 
 template <typename T>
diff --git a/Engine/Source/Runtime/CoreUObject/Public/UObject/UObjectGlobals.h b/Engine/Source/Runtime/CoreUObject/Public/UObject/UObjectGlobals.h
--- a/Engine/Source/Runtime/CoreUObject/Public/UObject/UObjectGlobals.h
+++ b/Engine/Source/Runtime/CoreUObject/Public/UObject/UObjectGlobals.h
@@ -576,7 +576,7 @@ public:
 	TSubobjectPtrDeprecated(TSubobjectPtrDeprecated<DerivedSubobjectType>& Other)
 		: FSubobjectPtr(Other.Object)
 	{
-		static_assert((CanConvertPointerFromTo<DerivedSubobjectType, SubobjectType>::Result), "Subobject pointers must be compatible.");
+		static_assert(TPointerIsConvertibleFromTo<DerivedSubobjectType, const SubobjectType>::Value, "Subobject pointers must be compatible.");
 	}
 	/** Gets the sub-object pointer. */
 	FORCEINLINE SubobjectType* Get() const
diff --git a/Engine/Source/Runtime/Engine/Classes/GameFramework/Actor.h b/Engine/Source/Runtime/Engine/Classes/GameFramework/Actor.h
--- a/Engine/Source/Runtime/Engine/Classes/GameFramework/Actor.h
+++ b/Engine/Source/Runtime/Engine/Classes/GameFramework/Actor.h
@@ -2111,7 +2111,7 @@ public:
 	template<class T>
 	T* FindComponentByClass() const
 	{
-		static_assert(CanConvertPointerFromTo<T, UActorComponent>::Result, "'T' template parameter to FindComponentByClass must be derived from UActorComponent");
+		static_assert(TPointerIsConvertibleFromTo<T, const UActorComponent>::Value, "'T' template parameter to FindComponentByClass must be derived from UActorComponent");
 
 		return (T*)FindComponentByClass(T::StaticClass());
 	}
@@ -2128,7 +2128,7 @@ public:
 	template<class T, class AllocatorType>
 	void GetComponents(TArray<T*, AllocatorType>& OutComponents) const
 	{
-		static_assert(CanConvertPointerFromTo<T, UActorComponent>::Result, "'T' template parameter to GetComponents must be derived from UActorComponent");
+		static_assert(TPointerIsConvertibleFromTo<T, const UActorComponent>::Value, "'T' template parameter to GetComponents must be derived from UActorComponent");
 		SCOPE_CYCLE_COUNTER(STAT_GetComponentsTime);
 
 		OutComponents.Reset(OwnedComponents.Num());
