fix for incorrect autorotate when using landscape right only on iOS device
UE-8175
#ios

diff --git a/Engine/Source/Runtime/Core/Private/IOS/IOSView.cpp b/Engine/Source/Runtime/Core/Private/IOS/IOSView.cpp
--- a/Engine/Source/Runtime/Core/Private/IOS/IOSView.cpp
+++ b/Engine/Source/Runtime/Core/Private/IOS/IOSView.cpp
@@ -498,12 +498,6 @@ id<MTLDevice> GMetalDevice = nil;
  */
 - (BOOL)shouldAutorotate
 {
-	NSUInteger orient = [[UIApplication sharedApplication] supportedInterfaceOrientationsForWindow: [IOSAppDelegate GetDelegate].Window];
-	if (orient == UIInterfaceOrientationMaskPortrait || orient == UIInterfaceOrientationMaskPortraitUpsideDown
-		|| orient == UIInterfaceOrientationMaskLandscapeLeft || orient == UIInterfaceOrientationMaskLandscapeRight)
-	{
-		return NO;
-	}
 	return YES;
 }
 
