Writes to default*.ini from project settings no longer writes properties twice.
Can and will be optimized in a later CL.

Fix for UE-4856 and UE-8252

diff --git a/Engine/Source/Runtime/Core/Private/Misc/ConfigCacheIni.cpp b/Engine/Source/Runtime/Core/Private/Misc/ConfigCacheIni.cpp
--- a/Engine/Source/Runtime/Core/Private/Misc/ConfigCacheIni.cpp
+++ b/Engine/Source/Runtime/Core/Private/Misc/ConfigCacheIni.cpp
@@ -923,8 +923,13 @@ bool FConfigFile::Write( const FString& Filename, bool bDoRemoteWrite/* = true*/
 				// check whether the option we are attempting to write out, came from the commandline as a temporary override.
 				const bool bOptionIsFromCommandline = PropertySetFromCommandlineOption(this, SectionName, PropertyName, PropertyValue);
 
+				// If we are writing to a default config file and this property is an array, we need to be careful to remove those from higher up the hierarchy
+				const bool bIsADefaultIniWrite = !Filename.Contains( FPaths::GameSavedDir() ) 
+					&& !Filename.Contains( FPaths::EngineSavedDir() ) 
+					&& FPaths::GetBaseFilename(Filename).StartsWith( TEXT( "Default" ) );
+
 				// Check if the property matches the source configs. We do not wanna write it out if so.
-				if( (bDifferentNumberOfElements || !DoesConfigPropertyValueMatch( SourceConfigFile, SectionName, PropertyName, PropertyValue )) && !bOptionIsFromCommandline )
+				if ((bIsADefaultIniWrite || bDifferentNumberOfElements || !DoesConfigPropertyValueMatch(SourceConfigFile, SectionName, PropertyName, PropertyValue)) && !bOptionIsFromCommandline)
 				{
 					// If this is the first property we are writing of this section, then print the section name
 					if( !bWroteASectionProperty )
@@ -944,11 +949,6 @@ bool FConfigFile::Write( const FString& Filename, bool bDoRemoteWrite/* = true*/
 					TArray< FString > CompletePropertyToWrite;
 					Section.MultiFind( PropertyName, CompletePropertyToWrite, true );
 
-					// If we are writing to a default config file and this property is an array, we need to be careful to remove those from higher up the hierarchy
-					bool bIsADefaultIniWrite = !Filename.Contains( FPaths::GameSavedDir() ) 
-						&& !Filename.Contains( FPaths::EngineSavedDir() ) 
-						&& FPaths::GetBaseFilename(Filename).StartsWith( TEXT( "Default" ) );
-
 					if( bIsADefaultIniWrite )
 					{
 						ProcessPropertyAndWriteForDefaults(CompletePropertyToWrite, Text, SectionName, PropertyName.ToString());
@@ -1250,70 +1250,31 @@ void FConfigFile::ProcessSourceAndCheckAgainstBackup()
 
 void FConfigFile::ProcessPropertyAndWriteForDefaults( const TArray< FString >& InCompletePropertyToProcess, FString& OutText, const FString& SectionName, const FString& PropertyName )
 {
-	// Get a list of the array elements we have processed, 
-	// we will remove any which are in the base hierarchy
-	// any which are left will be written to the default config
-	TArray<FString> UnprocessedPropertyValues = InCompletePropertyToProcess;
+	FConfigCacheIni Hierarchy;
+	const FString& LastFileInHierarchy = SourceIniHierarchy.Last().Filename;
 
-	// Iterate over the base hierarchy of the config file and check that our element did not originate from a further up the tree file
-	// if it did we need to ensure it is not added again.
-	for( TArray<FIniFilename>::TIterator HierarchyFileIt(SourceIniHierarchy); HierarchyFileIt; ++HierarchyFileIt )
+	// Build a config file out of this default configs hierarchy.
+	FConfigFile& DefaultConfigFile = Hierarchy.Add(LastFileInHierarchy, FConfigFile());
+	for (TArray<FIniFilename>::TIterator HierarchyFileIt(SourceIniHierarchy); HierarchyFileIt; ++HierarchyFileIt)
+	{
+		DefaultConfigFile.Combine(HierarchyFileIt->Filename);
+	}
+	
+	// Handle array elements from the configs hierarchy.
+	if (PropertyName.StartsWith(TEXT("+")) || InCompletePropertyToProcess.Num() > 1)
 	{
-		FConfigFile HierarchyFile;
-		ProcessIniContents(*HierarchyFileIt->Filename, *HierarchyFileIt->Filename, &HierarchyFile, false, false, false);
+		TArray<FString> ArrayProperties;
+		Hierarchy.GetArray(*SectionName, *PropertyName.Replace(TEXT("+"), TEXT("")), ArrayProperties, *LastFileInHierarchy);
 
-		FConfigSection* FoundSection = HierarchyFile.Find( SectionName );
-		if( FoundSection != NULL )
+		for (const FString& NextElement : ArrayProperties)
 		{
-			// The non-default config hierarchy array property names are prefixed with a . and not a +, so change
-			// that before searching.
-			FString HierarchyPropertyName;
-			if (PropertyName.StartsWith(TEXT("+")))
-			{
-				HierarchyPropertyName = TEXT(".") + PropertyName.RightChop(1);
-			}
-			else
-			{
-				HierarchyPropertyName = PropertyName;
-			}
-			TArray< FString > HierearchysArrayContribution;
-			FoundSection->MultiFind( *HierarchyPropertyName, HierearchysArrayContribution, true );
-
-			// Find array elements which should be removed.
-			for( TArray<FString>::TIterator PropertyIt(HierearchysArrayContribution); PropertyIt; ++PropertyIt )
-			{
-				FString PropertyValue = *PropertyIt;
-
-				// Check if the string formatted array element needs quotation marks
-				if (ShouldExportQuotedString(PropertyValue))
-				{
-					PropertyValue = FString::Printf( TEXT( "\"%s\"" ), *PropertyValue );
-				}
-				
-				// If the array we are processing doesn't have this value, there is a good chance it has been deleted and
-				// the default ini should add the removed element to the stream with the appropriate array syntax
-				//
-				// NOTE:	This also allows us to fix an issue with string comparison issue with elements in the base
-				//			configs having slightly different formats of properties to the ones serialized by the config load system.
-				//			I.e.	ArrEl=1.000f and ArrEl=1.f would not be a match.
-				//			This logic allows is to remove the ArrEl=1.000f and instead add 1.f in the default ini. It would look like this:
-				//			[MySection]
-				//			-ArrEl=1.000f
-				//			+ArrEl=1.f
-				if( UnprocessedPropertyValues.Contains( PropertyValue ) == false )
-				{
-					FString PropertyNameWithRemoveOp = PropertyName.Replace(TEXT("+"), TEXT("-"));
-					OutText += FString::Printf(TEXT("%s=%s") LINE_TERMINATOR, *PropertyNameWithRemoveOp, *PropertyValue);
-				}
-
-				// We need to remove this element from unprocessed if it exists on the list.
-				UnprocessedPropertyValues.RemoveSingle( PropertyValue );
-			}
+			FString PropertyNameWithRemoveOp = PropertyName.Replace(TEXT("+"), TEXT("-"));
+			OutText += FString::Printf(TEXT("%s=%s") LINE_TERMINATOR, *PropertyNameWithRemoveOp, *NextElement);
 		}
 	}
 
-	// Add all elements which were not in the hierarchy to the default with the correct '+' syntax
-	for (TArray<FString>::TIterator PropertyIt(UnprocessedPropertyValues); PropertyIt; ++PropertyIt)
+	// Write the properties out to a file.
+	for ( auto& PropertyIt : InCompletePropertyToProcess)
 	{
 		FString PropertyValue = *PropertyIt;
 
