Token concatenation fix for the variable name generation in ON_SCOPE_EXIT.

diff --git a/Engine/Source/Runtime/Core/Public/Misc/ScopeExit.h b/Engine/Source/Runtime/Core/Public/Misc/ScopeExit.h
--- a/Engine/Source/Runtime/Core/Public/Misc/ScopeExit.h
+++ b/Engine/Source/Runtime/Core/Public/Misc/ScopeExit.h
@@ -3,6 +3,7 @@
 #pragma once
 
 #include "HAL/Platform.h"
+#include "CoreMiscDefines.h"
 #include "Optional.h"
 
 namespace ScopeExitSupport
@@ -57,11 +58,11 @@ namespace ScopeExitSupport
 
 #ifdef __COUNTER__
 	// Created a variable with a unique name
-	#define ANONYMOUS_VARIABLE( Name ) Name##__COUNTER__
+	#define ANONYMOUS_VARIABLE( Name ) PREPROCESSOR_JOIN(Name, __COUNTER__)
 #else
 	// Created a variable with a unique name.
 	// Less reliable than the __COUNTER__ version.
-	#define ANONYMOUS_VARIABLE( Name ) Name##__LINE__
+	#define ANONYMOUS_VARIABLE( Name ) PREPROCESSOR_JOIN(Name, __LINE__)
 #endif
 
 
@@ -85,4 +86,4 @@ namespace ScopeExitSupport
  *      // * Execution reaches the end of the block.
  *    }
  */
-#define ON_SCOPE_EXIT const auto ANONYMOUS_VARIABLE(ScopeGuard_) = ::ScopeExitSupport::FScopeGuardSyntaxSupport() + [&]()
\ No newline at end of file
+#define ON_SCOPE_EXIT const auto ANONYMOUS_VARIABLE(ScopeGuard_) = ::ScopeExitSupport::FScopeGuardSyntaxSupport() + [&]()
