Make crash report functions explicitly visible.

- We need those symbols in another shared object, and otherwise they would be hidden in Test/Shipping builds (UE-8168).

(Merged CL 2418865 from UT branch).

diff --git a/Engine/Source/Runtime/Core/Private/Linux/LinuxPlatformCrashContext.cpp b/Engine/Source/Runtime/Core/Private/Linux/LinuxPlatformCrashContext.cpp
--- a/Engine/Source/Runtime/Core/Private/Linux/LinuxPlatformCrashContext.cpp
+++ b/Engine/Source/Runtime/Core/Private/Linux/LinuxPlatformCrashContext.cpp
@@ -646,7 +646,7 @@ void GenerateMinidump(const FString & Path)
 }
 
 
-int32 ReportCrash(const FLinuxCrashContext & Context)
+int32 DLLEXPORT ReportCrash(const FLinuxCrashContext & Context)
 {
 	static bool GAlreadyCreatedMinidump = false;
 	// Only create a minidump the first time this function is called.
@@ -673,7 +673,7 @@ int32 ReportCrash(const FLinuxCrashContext & Context)
 /**
  * Generates information for crash reporter
  */
-void GenerateCrashInfoAndLaunchReporter(const FLinuxCrashContext & Context)
+void DLLEXPORT GenerateCrashInfoAndLaunchReporter(const FLinuxCrashContext & Context)
 {
 	// do not report crashes for tools (particularly for crash reporter itself)
 #if !IS_PROGRAM
diff --git a/Engine/Source/Runtime/Launch/Private/Linux/LaunchLinux.cpp b/Engine/Source/Runtime/Launch/Private/Linux/LaunchLinux.cpp
--- a/Engine/Source/Runtime/Launch/Private/Linux/LaunchLinux.cpp
+++ b/Engine/Source/Runtime/Launch/Private/Linux/LaunchLinux.cpp
@@ -16,8 +16,8 @@ extern int32 GuardedMain( const TCHAR* CmdLine );
 extern void LaunchStaticShutdownAfterError();
 
 // FIXME: handle expose it someplace else?
-extern int32 ReportCrash(const FLinuxCrashContext& Context);
-extern void GenerateCrashInfoAndLaunchReporter(const FLinuxCrashContext& Context);
+extern int32 DLLIMPORT ReportCrash(const FLinuxCrashContext& Context);
+extern void DLLIMPORT GenerateCrashInfoAndLaunchReporter(const FLinuxCrashContext& Context);
 
 /**
  * Game-specific crash reporter
