[UE-7296] Enabling deferred dependency loading for the linker (to help handle circular dependency issues).

diff --git a/Engine/Config/BaseEngine.ini b/Engine/Config/BaseEngine.ini
--- a/Engine/Config/BaseEngine.ini
+++ b/Engine/Config/BaseEngine.ini
@@ -1485,7 +1485,7 @@ bNativeCodeGenerationTool=false
 bOptimizeExecutionFlowStack=true
 bOptimizeAdjacentStates=true
 bEnableInheritableComponents=true
-bDeferDependencyLoads=false
+bDeferDependencyLoads=true
 
 [/Script/Engine.Blueprint]
 bRecompileOnLoad=true
diff --git a/Engine/Source/Runtime/Core/Public/Misc/Build.h b/Engine/Source/Runtime/Core/Public/Misc/Build.h
--- a/Engine/Source/Runtime/Core/Public/Misc/Build.h
+++ b/Engine/Source/Runtime/Core/Public/Misc/Build.h
@@ -241,5 +241,5 @@
 /** Enable UberGraphPersistentFrame feature. It can speed up BP compilation (re-instancing) in editor, but introduce an unnecessary overhead in runtime. */
 #define USE_UBER_GRAPH_PERSISTENT_FRAME 1
 
-#define USE_CIRCULAR_DEPENDENCY_LOAD_DEFERRING (WITH_EDITOR && 0)
+#define USE_CIRCULAR_DEPENDENCY_LOAD_DEFERRING (WITH_EDITOR && 1)
 #define USE_DEFERRED_DEPENDENCY_CHECK_VERIFICATION_TESTS (USE_CIRCULAR_DEPENDENCY_LOAD_DEFERRING && 0)
