Linux: Always use SDL_WINDOW_FULLSCREEN_DESKTOP for fullscreen mode.
This removes modeswitching from the equation entirely, avoiding related driver/Xserver bugs. In addition, it solves several other UX problems (resizing desktop windows, etc).

diff --git a/Engine/Source/Runtime/Core/Private/Linux/LinuxWindow.cpp b/Engine/Source/Runtime/Core/Private/Linux/LinuxWindow.cpp
--- a/Engine/Source/Runtime/Core/Private/Linux/LinuxWindow.cpp
+++ b/Engine/Source/Runtime/Core/Private/Linux/LinuxWindow.cpp
@@ -279,53 +279,21 @@ static void _GetBestFullscreenResolution( SDL_HWindow hWnd, int32 *pWidth, int32
 	*pHeight = BestHeight;
 }
 
-
-static void _SetBestFullscreenDisplayMode( SDL_HWindow hWnd, int32 *pWidth, int32 *pHeight )
-{
-	SDL_DisplayMode dsp_mode;
-
-	dsp_mode.w = *pWidth;
-	dsp_mode.h = *pHeight;
-	dsp_mode.format = SDL_PIXELFORMAT_ARGB8888;
-	dsp_mode.refresh_rate = 60;
-	dsp_mode.driverdata = NULL;
-
-	_GetBestFullscreenResolution( hWnd, &dsp_mode.w, &dsp_mode.h );
-	SDL_SetWindowDisplayMode( hWnd, &dsp_mode );
-
-	*pWidth  = dsp_mode.w;
-	*pHeight = dsp_mode.h;
-}
-
 void FLinuxWindow::ReshapeWindow( int32 NewX, int32 NewY, int32 NewWidth, int32 NewHeight )
 {
-	int32 closest_w = NewWidth;
-	int32 closest_h = NewHeight;
-
 	switch( WindowMode )
 	{
+		// Fullscreen and WindowedFullscreen both use SDL_WINDOW_FULLSCREEN_DESKTOP now
+		//  and code elsewhere handles the backbufer blit properly. This solves several
+		//  problems that actual mode switches cause, and a GPU scales better than your
+		//  cheap LCD display anyhow.
 		case EWindowMode::Fullscreen:
-		{
-			SDL_SetWindowFullscreen( HWnd, 0 );
-			_GetBestFullscreenResolution( HWnd, &closest_w, &closest_h );
-			SDL_SetWindowSize( HWnd, closest_w, closest_h );
-
-			_SetBestFullscreenDisplayMode( HWnd, &closest_w, &closest_h );
-			SDL_SetWindowFullscreen( HWnd, SDL_WINDOW_FULLSCREEN );
-
-			bWasFullscreen = true;
-
-		}	break;
-
 		case EWindowMode::WindowedFullscreen:
 		{
 			SDL_SetWindowFullscreen( HWnd, 0 );
-			SDL_SetWindowPosition( HWnd, 0, 0 );
-			_SetBestFullscreenDisplayMode( HWnd, &closest_w, &closest_h );
+			SDL_SetWindowSize( HWnd, NewWidth, NewHeight );
 			SDL_SetWindowFullscreen( HWnd, SDL_WINDOW_FULLSCREEN_DESKTOP );
-
 			bWasFullscreen = true;
-
 		}	break;
 
 		case EWindowMode::Windowed:
@@ -359,42 +327,21 @@ void FLinuxWindow::SetWindowMode( EWindowMode::Type NewWindowMode )
 {
 	if( NewWindowMode != WindowMode )
 	{
-		SDL_DisplayMode dsp_mode;
-		int32 closest_w;
-		int32 closest_h;
-
-		closest_w = VirtualWidth;
-		closest_h = VirtualHeight;
-
 		switch( NewWindowMode )
 		{
+			// Fullscreen and WindowedFullscreen both use SDL_WINDOW_FULLSCREEN_DESKTOP now
+			//  and code elsewhere handles the backbufer blit properly. This solves several
+			//  problems that actual mode switches cause, and a GPU scales better than your
+			//  cheap LCD display anyhow.
 			case EWindowMode::Fullscreen:
-			{
-				if ( bWasFullscreen != true )
-				{
-					SDL_SetWindowPosition( HWnd, 0, 0 );
-					_SetBestFullscreenDisplayMode( HWnd, &closest_w, &closest_h );
-
-					SDL_SetWindowFullscreen( HWnd, SDL_WINDOW_FULLSCREEN );
-					SDL_SetWindowGrab( HWnd, SDL_TRUE );
-
-					bWasFullscreen = true;
-				}
-
-			}	break;
-
 			case EWindowMode::WindowedFullscreen:
 			{
 				if ( bWasFullscreen != true )
 				{
-					SDL_SetWindowPosition( HWnd, 0, 0 );
-
-					_SetBestFullscreenDisplayMode( HWnd, &closest_w, &closest_h );
+					SDL_SetWindowSize( HWnd, VirtualWidth, VirtualHeight );
 					SDL_SetWindowFullscreen( HWnd, SDL_WINDOW_FULLSCREEN_DESKTOP );
-
 					bWasFullscreen = true;
 				}
-
 			}	break;
 
 			case EWindowMode::Windowed:
diff --git a/Engine/Source/Runtime/OpenGLDrv/Private/Linux/OpenGLLinux.cpp b/Engine/Source/Runtime/OpenGLDrv/Private/Linux/OpenGLLinux.cpp
--- a/Engine/Source/Runtime/OpenGLDrv/Private/Linux/OpenGLLinux.cpp
+++ b/Engine/Source/Runtime/OpenGLDrv/Private/Linux/OpenGLLinux.cpp
@@ -384,10 +384,77 @@ bool PlatformBlitToViewport(FPlatformOpenGLDevice* Device,
 		glReadBuffer( GL_COLOR_ATTACHMENT0 );
 		glDisable(GL_FRAMEBUFFER_SRGB);
 
+		int WinW, WinH;
+		SDL_GetWindowSize(Context->hWnd, &WinW, &WinH);
+		GLenum BlitFilter;
+		GLint DestX0, DestY0, DestX1, DestY1;
+
+		if ( ( WinW == BackbufferSizeX ) && ( WinH == BackbufferSizeY ) )
+		{
+			// We match up. We're probably in windowed mode, or an exact
+			//  match for FULLSCREEN_DESKTOP mode. Use a NEAREST blit and
+			//  don't clear the window system's framebuffer first.
+			BlitFilter = GL_NEAREST;
+			DestX0 = 0;
+			DestY0 = BackbufferSizeY;  // flip vertically.
+			DestX1 = BackbufferSizeX;
+			DestY1 = 0;
+		}
+		else
+		{
+			// we need to scale to match the size of the window system's
+			// framebuffer, so scale linearly, and adjust for letterboxing.
+			BlitFilter = GL_LINEAR;
+
+			const uint32 w = BackbufferSizeX;
+			const uint32 h = BackbufferSizeY;
+			const float WantedAspect = (w > h) ? (((float) w) / ((float) h)) : (((float) h) / ((float) w));
+			const float PhysicalAspect = (((float) WinW) / ((float) WinH));
+
+			bool bMustClear;  // have to clear the window framebuffer if letterboxing.
+			if ( PhysicalAspect == WantedAspect )  // Perfect aspect ratio; no letterboxing needed?
+			{
+				bMustClear = false;
+				DestX0 = 0;
+				DestY0 = WinH;  // flip vertically.
+				DestX1 = WinW;
+				DestY1 = 0;
+			}
+			else if ( PhysicalAspect > WantedAspect )  // view is wider than wanted aspect?
+			{  
+				bMustClear = true;
+				const float ScaledW = WinH * WantedAspect;
+				const float ScaledX = (((float)WinW) - ScaledW) / 2.0f;
+				DestX0 = ScaledX;
+				DestY0 = WinH;  // flip vertically.
+				DestX1 = ScaledX + ScaledW;
+				DestY1 = 0;
+			}
+			else  // view is taller than wanted aspect?
+			{
+				bMustClear = true;
+				const float ScaledH = WinW / WantedAspect;
+				const float ScaledY = (((float)WinH) - ScaledH) / 2.0f;
+				DestX0 = 0;
+				DestY0 = ScaledY + ScaledH;  // flip vertically.
+				DestX1 = WinW;
+				DestY1 = ScaledY;
+			}
+
+			// if the Steam Overlay is running, it might write garbage into our
+			//  letterbox area, so if we have a letterbox, clear the framebuffer.
+			if ( bMustClear )
+			{
+				glClearColor( 0.0f, 0.0f, 0.0f, 1.0f );
+				glClear( GL_COLOR_BUFFER_BIT );
+			}
+		}
+
+		// Get it to the window system's framebuffer.
 		glBlitFramebuffer(	0, 0, BackbufferSizeX, BackbufferSizeY,
-							0, BackbufferSizeY, BackbufferSizeX, 0,
+							DestX0, DestY0, DestX1, DestY1,
 							GL_COLOR_BUFFER_BIT,
-							GL_NEAREST	);
+							BlitFilter	);
 
 		if ( bPresent )
 		{
