fixed bottom right of bloom darker then left top

diff --git a/Engine/Source/Runtime/Core/Public/Math/IntPoint.h b/Engine/Source/Runtime/Core/Public/Math/IntPoint.h
--- a/Engine/Source/Runtime/Core/Public/Math/IntPoint.h
+++ b/Engine/Source/Runtime/Core/Public/Math/IntPoint.h
@@ -243,6 +243,7 @@ public:
 	 * @see DivideAndRoundDown
 	 */
 	static FIntPoint DivideAndRoundUp( FIntPoint lhs, int32 Divisor );
+	static FIntPoint DivideAndRoundUp( FIntPoint lhs, FIntPoint Divisor );
 
 	/**
 	 * Divide an int point and round down the result.
@@ -415,12 +416,15 @@ FORCEINLINE FIntPoint FIntPoint::ComponentMax(const FIntPoint& Other) const
 	return FIntPoint(FMath::Max(X, Other.X), FMath::Max(Y, Other.Y));
 }
 
-
 FORCEINLINE FIntPoint FIntPoint::DivideAndRoundUp( FIntPoint lhs, int32 Divisor )
 {
 	return FIntPoint(FMath::DivideAndRoundUp(lhs.X, Divisor), FMath::DivideAndRoundUp(lhs.Y, Divisor));
-}	
+}
 
+FORCEINLINE FIntPoint FIntPoint::DivideAndRoundUp( FIntPoint lhs, FIntPoint Divisor )
+{
+	return FIntPoint(FMath::DivideAndRoundUp(lhs.X, Divisor.X), FMath::DivideAndRoundUp(lhs.Y, Divisor.Y));
+}	
 
 FORCEINLINE FIntPoint FIntPoint::DivideAndRoundDown( FIntPoint lhs, int32 Divisor )
 {
diff --git a/Engine/Source/Runtime/Core/Public/Math/IntRect.h b/Engine/Source/Runtime/Core/Public/Math/IntRect.h
--- a/Engine/Source/Runtime/Core/Public/Math/IntRect.h
+++ b/Engine/Source/Runtime/Core/Public/Math/IntRect.h
@@ -276,6 +276,7 @@ public:
 	 * @return New divided rectangle.
 	 */
 	static FIntRect DivideAndRoundUp( FIntRect lhs, int32 Div );
+	static FIntRect DivideAndRoundUp( FIntRect lhs, FIntPoint Div );
 
 	/**
 	 * Gets number of points in the Rectangle.
@@ -456,9 +457,13 @@ FORCEINLINE bool FIntRect::Contains( FIntPoint P ) const
 
 FORCEINLINE FIntRect FIntRect::DivideAndRoundUp( FIntRect lhs, int32 Div )
 {
-	return FIntRect(lhs.Min / Div, FIntPoint::DivideAndRoundUp(lhs.Max, Div));
+	return DivideAndRoundUp(lhs, FIntPoint(Div, Div));
 }
 
+FORCEINLINE FIntRect FIntRect::DivideAndRoundUp( FIntRect lhs, FIntPoint Div )
+{
+	return FIntRect(lhs.Min / Div, FIntPoint::DivideAndRoundUp(lhs.Max, Div));
+}
 
 FORCEINLINE void FIntRect::GetCenterAndExtents( FIntPoint& OutCenter, FIntPoint& OutExtent ) const
 {
diff --git a/Engine/Source/Runtime/Renderer/Private/PostProcess/PostProcessWeightedSampleSum.cpp b/Engine/Source/Runtime/Renderer/Private/PostProcess/PostProcessWeightedSampleSum.cpp
--- a/Engine/Source/Runtime/Renderer/Private/PostProcess/PostProcessWeightedSampleSum.cpp
+++ b/Engine/Source/Runtime/Renderer/Private/PostProcess/PostProcessWeightedSampleSum.cpp
@@ -466,8 +466,8 @@ void FRCPassPostProcessWeightedSampleSum::Process(FRenderingCompositePassContext
 		bRequiresClear = false;
 	}
 
-	FIntRect SrcRect =  View.ViewRect / SrcScaleFactor;
-	FIntRect DestRect = View.ViewRect / DstScaleFactor;
+	FIntRect SrcRect =  FIntRect::DivideAndRoundUp(View.ViewRect, SrcScaleFactor);
+	FIntRect DestRect = FIntRect::DivideAndRoundUp(View.ViewRect, DstScaleFactor);
 
 	DrawQuad(Context.RHICmdList, bDoFastBlur, SrcRect, DestRect, bRequiresClear, DestSize, SrcSize, VertexShader);
 
