Removed obsolete (non-plugin) SteamController support.

diff --git a/Engine/Source/Runtime/Core/Private/Linux/LinuxApplication.cpp b/Engine/Source/Runtime/Core/Private/Linux/LinuxApplication.cpp
--- a/Engine/Source/Runtime/Core/Private/Linux/LinuxApplication.cpp
+++ b/Engine/Source/Runtime/Core/Private/Linux/LinuxApplication.cpp
@@ -8,9 +8,6 @@
 #include "IInputInterface.h"
 #include "IInputDeviceModule.h"
 #include "IInputDevice.h"
-#if STEAM_CONTROLLER_SUPPORT
-	#include "SteamControllerInterface.h"
-#endif // STEAM_CONTROLLER_SUPPORT
 
 #if WITH_EDITOR
 #include "ModuleManager.h"
@@ -63,14 +60,12 @@ FLinuxApplication* FLinuxApplication::CreateLinuxApplication()
 }
 
 
-FLinuxApplication::FLinuxApplication() : GenericApplication( MakeShareable( new FLinuxCursor() ) )
-#if STEAM_CONTROLLER_SUPPORT
-	, SteamInput( SteamControllerInterface::Create(MessageHandler) )
-#endif // STEAM_CONTROLLER_SUPPORT
-	, bIsMouseCursorLocked(false)
-	, bIsMouseCaptureEnabled(false)
-	, bHasLoadedInputPlugins(false)
-	, bInsideOwnWindow(false)
+FLinuxApplication::FLinuxApplication() 
+	:	GenericApplication( MakeShareable( new FLinuxCursor() ) )
+	,	bIsMouseCursorLocked(false)
+	,	bIsMouseCaptureEnabled(false)
+	,	bHasLoadedInputPlugins(false)
+	,	bInsideOwnWindow(false)
 {
 	bUsingHighPrecisionMouseInput = false;
 	bAllowedToDeferMessageProcessing = true;
@@ -125,9 +120,6 @@ void FLinuxApplication::InitializeWindow(	const TSharedRef< FGenericWindow >& In
 void FLinuxApplication::SetMessageHandler( const TSharedRef< FGenericApplicationMessageHandler >& InMessageHandler )
 {
 	GenericApplication::SetMessageHandler(InMessageHandler);
-#if STEAM_CONTROLLER_SUPPORT
-	SteamInput->SetMessageHandler(InMessageHandler);
-#endif // STEAM_CONTROLLER_SUPPORT
 }
 
 namespace
@@ -795,6 +787,7 @@ void FLinuxApplication::PollGameDeviceState( const float TimeDelta )
 			TSharedPtr<IInputDevice> Device = (*InputPluginIt)->CreateInputDevice(MessageHandler);
 			if (Device.IsValid())
 			{
+				UE_LOG(LogInit, Log, TEXT("Adding external input plugin."));
 				ExternalInputDevices.Add(Device);
 			}
 		}
@@ -802,11 +795,6 @@ void FLinuxApplication::PollGameDeviceState( const float TimeDelta )
 		bHasLoadedInputPlugins = true;
 	}
 	
-#if STEAM_CONTROLLER_SUPPORT
-	// Poll game device states and send new events
-	SteamInput->SendControllerEvents();
-#endif // STEAM_CONTROLLER_SUPPORT
-
 	// Poll externally-implemented devices
 	for (auto DeviceIt = ExternalInputDevices.CreateIterator(); DeviceIt; ++DeviceIt)
 	{
diff --git a/Engine/Source/Runtime/Core/Public/Linux/LinuxApplication.h b/Engine/Source/Runtime/Core/Public/Linux/LinuxApplication.h
--- a/Engine/Source/Runtime/Core/Public/Linux/LinuxApplication.h
+++ b/Engine/Source/Runtime/Core/Public/Linux/LinuxApplication.h
@@ -5,8 +5,6 @@
 #include "GenericApplication.h"
 #include "LinuxWindow.h"
 
-#define STEAM_CONTROLLER_SUPPORT				(WITH_STEAMWORKS && WITH_ENGINE && WITH_STEAMCONTROLLER)
-
 typedef SDL_GameController* SDL_HController;
 
 class FLinuxWindow;
@@ -131,10 +129,6 @@ private:
 
 	float fMouseWheelScrollAccel;
 
-#if STEAM_CONTROLLER_SUPPORT
-	TSharedPtr< class SteamControllerInterface > SteamInput;
-#endif // STEAM_CONTROLLER_SUPPORT
-
 	/** List of input devices implemented in external modules. */
 	TArray< TSharedPtr<class IInputDevice> > ExternalInputDevices;
 
