Implemented FMacPlatformMisc::GetOperatingSystemId()

diff --git a/Engine/Source/Runtime/Core/Private/Mac/MacPlatformMisc.cpp b/Engine/Source/Runtime/Core/Private/Mac/MacPlatformMisc.cpp
--- a/Engine/Source/Runtime/Core/Private/Mac/MacPlatformMisc.cpp
+++ b/Engine/Source/Runtime/Core/Private/Mac/MacPlatformMisc.cpp
@@ -1098,6 +1098,24 @@ int32 FMacPlatformMisc::MacOSXVersionCompare(uint8 Major, uint8 Minor, uint8 Rev
 	return 0;
 }
 
+FString FMacPlatformMisc::GetOperatingSystemId()
+{
+	FString Result;
+	io_service_t Entry = IOServiceGetMatchingService(kIOMasterPortDefault, IOServiceMatching("IOPlatformExpertDevice"));
+	if (Entry)
+	{
+		CFTypeRef UUID = IORegistryEntryCreateCFProperty(Entry, CFSTR(kIOPlatformUUIDKey), kCFAllocatorDefault, 0);
+		Result = FString((__bridge NSString*)UUID);
+		IOObjectRelease(Entry);
+		CFRelease(UUID);
+	}
+	else
+	{
+		UE_LOG(LogMac, Warning, TEXT("GetOperatingSystemId() failed"));
+	}
+	return Result;
+}
+
 /** Global pointer to crash handler */
 void (* GCrashHandlerPointer)(const FGenericCrashContext& Context) = NULL;
 FMacMallocCrashHandler* GCrashMalloc = nullptr;
diff --git a/Engine/Source/Runtime/Core/Public/Mac/MacPlatformMisc.h b/Engine/Source/Runtime/Core/Public/Mac/MacPlatformMisc.h
--- a/Engine/Source/Runtime/Core/Public/Mac/MacPlatformMisc.h
+++ b/Engine/Source/Runtime/Core/Public/Mac/MacPlatformMisc.h
@@ -152,12 +152,17 @@ struct CORE_API FMacPlatformMisc : public FGenericPlatformMisc
 	 * Returns whether the Mac OS X version is 10.9.x or not.
 	 */
 	static bool IsRunningOnMavericks();
-	
+
 	/**
 	 * Returns if current < target returns 1, if current > target returns 1, else current == target and returns 0.
 	 */
 	static int32 MacOSXVersionCompare(uint8 Major, uint8 Minor, uint8 Revision);
 
+	/**
+	 * Gets a globally unique ID the represents a particular operating system install.
+	 */
+	static FString GetOperatingSystemId();
+
 	static bool bChachedMacMenuStateNeedsUpdate;
 };
 
