Defer setting of the new key window when current one is closed to next PumpMessages call. This way we avoid a situation where bringing window to front triggers redrawing of all windows before Slate gets information that a window was closed.

diff --git a/Engine/Source/Runtime/Core/Private/Mac/CocoaWindow.cpp b/Engine/Source/Runtime/Core/Private/Mac/CocoaWindow.cpp
--- a/Engine/Source/Runtime/Core/Private/Mac/CocoaWindow.cpp
+++ b/Engine/Source/Runtime/Core/Private/Mac/CocoaWindow.cpp
@@ -549,7 +549,7 @@ NSString* NSPerformDragOperation = @"NSPerformDragOperation";
 	// Really we shouldn't be doing this - on OS X only left-click changes focus,
 	// but for the moment it is easier than changing Slate.
 	SCOPED_AUTORELEASE_POOL;
-	if([self canBecomeKeyWindow])
+	if([self canBecomeKeyWindow] && self != [NSApp keyWindow])
 	{
 		[self makeKeyWindow];
 	}
diff --git a/Engine/Source/Runtime/Core/Private/Mac/MacApplication.cpp b/Engine/Source/Runtime/Core/Private/Mac/MacApplication.cpp
--- a/Engine/Source/Runtime/Core/Private/Mac/MacApplication.cpp
+++ b/Engine/Source/Runtime/Core/Private/Mac/MacApplication.cpp
@@ -107,6 +107,7 @@ FMacApplication::FMacApplication()
 	, bIsProcessingNSEvent( false )
 	, ModifierKeysFlags( 0 )
 	, CurrentModifierFlags( 0 )
+	, bKeyWindowUpdateRequested( false )
 	, bIsWorkspaceSessionActive( true )
 {
 	CGDisplayRegisterReconfigurationCallback(FMacApplication::OnDisplayReconfiguration, this);
@@ -422,9 +423,20 @@ void FMacApplication::ProcessEvent(FMacEvent const* const Event)
 	delete Event;
 }
 
-void FMacApplication::PumpMessages( const float TimeDelta )
+void FMacApplication::PumpMessages(const float TimeDelta)
 {
-	FPlatformMisc::PumpMessages( true );
+	FPlatformMisc::PumpMessages(true);
+
+	if (bKeyWindowUpdateRequested)
+	{
+		TSharedPtr<FMacWindow> KeyWindow = GetKeyWindow();
+		if (KeyWindow.IsValid() && KeyWindow->GetOSWindowHandle() && ![(FCocoaWindow*)KeyWindow->GetOSWindowHandle() isMiniaturized])
+		{
+			// Activate specified previous window if still present, provided it isn't minimized
+			KeyWindow->SetWindowFocus();
+		}
+		bKeyWindowUpdateRequested = false;
+	}
 }
 
 void FMacApplication::OnWindowDraggingFinished()
diff --git a/Engine/Source/Runtime/Core/Private/Mac/MacWindow.cpp b/Engine/Source/Runtime/Core/Private/Mac/MacWindow.cpp
--- a/Engine/Source/Runtime/Core/Private/Mac/MacWindow.cpp
+++ b/Engine/Source/Runtime/Core/Private/Mac/MacWindow.cpp
@@ -305,26 +305,23 @@ void FMacWindow::BringToFront( bool bForce )
 
 void FMacWindow::Destroy()
 {
-	if( WindowHandle )
+	if (WindowHandle)
 	{
 		SCOPED_AUTORELEASE_POOL;
-		FCocoaWindow* Window = WindowHandle;
-
 		bIsClosed = true;
 
-		if( MacApplication->OnWindowDestroyed( Window ) )
+		FCocoaWindow* Window = WindowHandle;
+		const bool bIsKey = [Window isKeyWindow];
+
+		if (MacApplication->OnWindowDestroyed(Window) )
 		{
 			// This FMacWindow may have been destructed by now & so the WindowHandle will probably be invalid memory.
-			bool const bIsKey = [Window isKeyWindow];
-			
-			// Then change the focus to something useful, either the previous in the stack
-			TSharedPtr<FMacWindow> KeyWindow = MacApplication->GetKeyWindow();
-			if( KeyWindow.IsValid() && bIsKey && KeyWindow->GetOSWindowHandle() && ![(FCocoaWindow*)KeyWindow->GetOSWindowHandle() isMiniaturized] )
+
+			if (bIsKey)
 			{
-				// Activate specified previous window if still present, provided it isn't minimized
-				KeyWindow->SetWindowFocus();
+				MacApplication->RequestKeyWindowUpdate();
 			}
-			
+
 			// Close the window
 			MainThreadCall(^{
 				SCOPED_AUTORELEASE_POOL;
diff --git a/Engine/Source/Runtime/Core/Public/Mac/MacApplication.h b/Engine/Source/Runtime/Core/Public/Mac/MacApplication.h
--- a/Engine/Source/Runtime/Core/Public/Mac/MacApplication.h
+++ b/Engine/Source/Runtime/Core/Public/Mac/MacApplication.h
@@ -71,8 +71,6 @@ public:
 
 	void ResetModifierKeys() { ModifierKeysFlags = 0; }
 
-	TSharedPtr<FMacWindow> GetKeyWindow();
-
 	uint32 GetModifierKeysFlags() { return ModifierKeysFlags; }
 
 	bool IsProcessingNSEvent() const { return bIsProcessingNSEvent; }
@@ -132,6 +130,9 @@ private:
 	FCocoaWindow* FindEventWindow( NSEvent* CocoaEvent );
 	TSharedPtr<FMacWindow> LocateWindowUnderCursor( const NSPoint Position );
 
+	TSharedPtr<FMacWindow> GetKeyWindow();
+	void RequestKeyWindowUpdate() { bKeyWindowUpdateRequested = true; }
+
 	NSScreen* FindScreenByPoint( int32 X, int32 Y ) const;
 
 	void HandleModifierChange(NSUInteger NewModifierFlags, NSUInteger FlagsShift, NSUInteger UE4Shift, EMacModifierKeys TranslatedCode);
@@ -179,6 +180,7 @@ private:
 	NSUInteger CurrentModifierFlags;
 
 	TArray< TSharedRef< FMacWindow > > KeyWindows;
+	bool bKeyWindowUpdateRequested;
 
 	TSharedPtr<FMacTextInputMethodSystem> TextInputMethodSystem;
 
