Core: Added type container

diff --git a/Engine/Source/Runtime/Core/Private/Tests/Misc/TypeContainerTest.cpp b/Engine/Source/Runtime/Core/Private/Tests/Misc/TypeContainerTest.cpp
new file mode 100644
--- /dev/null
+++ b/Engine/Source/Runtime/Core/Private/Tests/Misc/TypeContainerTest.cpp
@@ -0,0 +1,184 @@
+// Copyright 1998-2015 Epic Games, Inc. All Rights Reserved.
+
+#if PLATFORM_COMPILER_HAS_VARIADIC_TEMPLATES
+
+
+#include "CorePrivatePCH.h"
+#include "Async.h"
+#include "Function.h"
+#include "Future.h"
+#include "TypeContainer.h"
+#include "AutomationTest.h"
+
+
+IMPLEMENT_SIMPLE_AUTOMATION_TEST(FTypeContainerTest, "Core.Misc.TypeContainer", EAutomationTestFlags::ATF_SmokeTest)
+
+
+/* Helpers
+ *****************************************************************************/
+
+struct IFruit
+{
+	DECLARE_REGISTRABLE_TYPE(IFruit)
+	virtual FString Name() = 0;
+};
+
+struct IBerry : public IFruit
+{
+	DECLARE_REGISTRABLE_TYPE(IBerry)
+};
+
+struct FBanana : public IFruit
+{
+	virtual FString Name() override { return TEXT("Cherry"); }
+};
+
+struct FStrawberry : public IBerry
+{
+	virtual FString Name() override { return TEXT("Strawberry"); }
+};
+
+struct ISmoothie
+{
+	DECLARE_REGISTRABLE_TYPE(ISmoothie)
+	virtual TSharedRef<IBerry> GetBerry() = 0;
+	virtual TSharedRef<IFruit> GetFruit() = 0;
+};
+
+struct FSmoothie : public ISmoothie
+{
+	FSmoothie(TSharedRef<IFruit> InFruit, TSharedRef<IBerry> InBerry) : Berry(InBerry), Fruit(InFruit) { }
+	virtual TSharedRef<IBerry> GetBerry() override { return Berry; }
+	virtual TSharedRef<IFruit> GetFruit() override { return Fruit; }
+	TSharedRef<IBerry> Berry;
+	TSharedRef<IFruit> Fruit;
+};
+
+struct FTwoSmoothies
+{
+	TSharedPtr<ISmoothie> One;
+	TSharedPtr<ISmoothie> Two;
+};
+
+
+/* Tests
+ *****************************************************************************/
+
+bool FTypeContainerTest::RunTest(const FString& Parameters)
+{
+	// existing instance test
+	{
+		TSharedRef<IFruit> Fruit = MakeShareable(new FBanana());
+		TSharedRef<IBerry> Berry = MakeShareable(new FStrawberry());
+
+		FTypeContainer Container;
+		{
+			Container.RegisterInstance<IFruit>(Fruit);
+			Container.RegisterInstance<IBerry>(Berry);
+		}
+
+		auto Instance = Container.GetInstance<IFruit>();
+
+		TestTrue(TEXT("Existing registered instance must be returned"), Instance.IsValid());
+
+		if (Instance.IsValid())
+		{
+			TestEqual(TEXT("Correct instance must be returned"), Instance.ToSharedRef(), Fruit);
+			TestNotEqual(TEXT("Incorrect instance must not be returned"), Instance.ToSharedRef(), StaticCastSharedRef<IFruit>(Berry));
+		}
+	}
+
+	// per instance test
+	{
+		FTypeContainer Container;
+		{
+			Container.RegisterClass<IFruit, FBanana>(ETypeContainerScope::Instance);
+			Container.RegisterClass<IBerry, FStrawberry>(ETypeContainerScope::Instance);
+			Container.RegisterClass<ISmoothie, FSmoothie, IFruit, IBerry>(ETypeContainerScope::Instance); // !!!
+		}
+
+		auto Smoothie1 = Container.GetInstance<ISmoothie>();
+		auto Smoothie2 = Container.GetInstance<ISmoothie>();
+
+		TestTrue(TEXT("For per-instance classes, an instance must be returned [1]"), Smoothie1.IsValid());
+		TestTrue(TEXT("For per-instance classes, an instance must be returned [2]"), Smoothie2.IsValid());
+		TestNotEqual(TEXT("For per-instances classes, a unique instance must be returned each time"), Smoothie1, Smoothie2);
+		TestNotEqual(TEXT("For per-instances dependencies, a unique instance must be returned each time [1]"), Smoothie1->GetBerry(), Smoothie2->GetBerry());
+		TestNotEqual(TEXT("For per-instances dependencies, a unique instance must be returned each time [2]"), Smoothie1->GetFruit(), Smoothie2->GetFruit());
+	}
+
+	// per thread test
+	{
+		FTypeContainer Container;
+		{
+			Container.RegisterClass<IFruit, FBanana>(ETypeContainerScope::Thread);
+			Container.RegisterClass<IBerry, FStrawberry>(ETypeContainerScope::Instance);
+			Container.RegisterClass<ISmoothie, FSmoothie, IFruit, IBerry>(ETypeContainerScope::Thread); // !!!
+		}
+
+		TFunction<FTwoSmoothies()> MakeSmoothies = [&]()
+		{
+			FTwoSmoothies Smoothies;
+			Smoothies.One = Container.GetInstance<ISmoothie>();
+			Smoothies.Two = Container.GetInstance<ISmoothie>();
+			return Smoothies;
+		};
+
+		auto Smoothies1 = Async(EAsyncExecution::Thread, MakeSmoothies);
+		auto Smoothies2 = Async(EAsyncExecution::Thread, MakeSmoothies);
+
+		TSharedPtr<ISmoothie> One1 = Smoothies1.Get().One;
+		TSharedPtr<ISmoothie> Two1 = Smoothies1.Get().Two;
+		TSharedPtr<ISmoothie> One2 = Smoothies2.Get().One;
+		TSharedPtr<ISmoothie> Two2 = Smoothies2.Get().Two;
+
+		TestTrue(TEXT("For per-thread classes, an instance must be returned [1]"), One1.IsValid());
+		TestTrue(TEXT("For per-thread classes, an instance must be returned [2]"), Two1.IsValid());
+		TestTrue(TEXT("For per-thread classes, an instance must be returned [3]"), One2.IsValid());
+		TestTrue(TEXT("For per-thread classes, an instance must be returned [4]"), Two2.IsValid());
+		TestEqual(TEXT("For per-thread classes, the same instance must be returned from the same thread [1]"), One1, Two1);
+		TestEqual(TEXT("For per-thread classes, the same instance must be returned from the same thread [2]"), One2, Two2);
+		TestNotEqual(TEXT("For per-thread classes, different instances must be returned from different threads [1]"), One1, One2);
+		TestNotEqual(TEXT("For per-thread classes, different instances must be returned from different threads [2]"), Two1, Two2);
+	}
+
+	// per process test
+	{
+		FTypeContainer Container;
+		{
+			Container.RegisterClass<IFruit, FBanana>(ETypeContainerScope::Thread);
+			Container.RegisterClass<IBerry, FStrawberry>(ETypeContainerScope::Instance);
+			Container.RegisterClass<ISmoothie, FSmoothie, IFruit, IBerry>(ETypeContainerScope::Process); //!!!
+		}
+
+		TFunction<FTwoSmoothies()> MakeSmoothies = [&]()
+		{
+			FTwoSmoothies Smoothies;
+			Smoothies.One = Container.GetInstance<ISmoothie>();
+			Smoothies.Two = Container.GetInstance<ISmoothie>();
+			return Smoothies;
+		};
+
+		auto Smoothies1 = Async(EAsyncExecution::Thread, MakeSmoothies);
+		auto Smoothies2 = Async(EAsyncExecution::Thread, MakeSmoothies);
+
+		TSharedPtr<ISmoothie> One1 = Smoothies1.Get().One;
+		TSharedPtr<ISmoothie> Two1 = Smoothies1.Get().Two;
+		TSharedPtr<ISmoothie> One2 = Smoothies2.Get().One;
+		TSharedPtr<ISmoothie> Two2 = Smoothies2.Get().Two;
+
+		TestTrue(TEXT("For per-process classes, an instance must be returned [1]"), One1.IsValid());
+		TestTrue(TEXT("For per-process classes, an instance must be returned [2]"), Two1.IsValid());
+		TestTrue(TEXT("For per-process classes, an instance must be returned [3]"), One2.IsValid());
+		TestTrue(TEXT("For per-process classes, an instance must be returned [4]"), Two2.IsValid());
+		TestEqual(TEXT("For per-process classes, the same instance must be returned from the same thread [1]"), One1, Two1);
+		TestEqual(TEXT("For per-process classes, the same instance must be returned from the same thread [2]"), One2, Two2);
+		TestEqual(TEXT("For per-process classes, the same instance must be returned from different threads [1]"), One1, One2);
+		TestEqual(TEXT("For per-process classes, the same instance must be returned from different threads [1]"), Two1, Two2);
+	}
+
+	return true;
+}
+
+
+#endif //PLATFORM_COMPILER_HAS_VARIADIC_TEMPLATES
diff --git a/Engine/Source/Runtime/Core/Public/Misc/TypeContainer.h b/Engine/Source/Runtime/Core/Public/Misc/TypeContainer.h
new file mode 100644
--- /dev/null
+++ b/Engine/Source/Runtime/Core/Public/Misc/TypeContainer.h
@@ -0,0 +1,342 @@
+// Copyright 1998-2015 Epic Games, Inc. All Rights Reserved.
+
+#pragma once
+
+
+#if PLATFORM_COMPILER_HAS_VARIADIC_TEMPLATES
+
+
+/**
+ * Convenience macro for declaring type traits for FTypeContainer compatible classes.
+ */
+#define DECLARE_REGISTRABLE_TYPE(Type) \
+	public: \
+		static FName GetTypeName() \
+		{ \
+			static FName TypeName = TEXT(#Type); \
+			return TypeName; \
+		}
+
+
+/**
+ * Enumerates the scopes for instance creation in type containers.
+ */
+enum class ETypeContainerScope
+{
+	/** Create a new instance each time. */
+	Instance,
+
+	/** One singleton for the entire process. */
+	Process,
+
+	/** One singleton per thread. */
+	Thread,
+};
+
+
+/**
+ * Implements a type container.
+ *
+ * Type containers allow for configuring objects and their dependencies without actually
+ * creating these objects. They fulfill a role similar to class factories, but without
+ * being locked to a particular type of class. When passed into class constructors or
+ * methods, type containers can facilitate the Inversion of Control (IoC) pattern.
+ *
+ * Since UE4 neither uses run-time type information nor pre-processes plain old C++ classes,
+ * a special macro - DECLARE_REGISTRABLE_TYPE - needs to be inserted into a class declaration
+ * in order to make it registrable with type containers.
+ *
+ * Once a type is registered with a container, instances of that type can be retrieved from it.
+ * There are currently three life time scopes available for instance creation:
+ *
+ *   1. Unique instance per process (aka. singleton),
+ *      using RegisterClass(ETypeContainerScope::Process) or RegisterInstance()
+ *   2. Unique instance per thread (aka. thread singleton),
+ *      using RegisterClass(ETypeContainerScope::Thread)
+ *   3. Unique instance per call (aka. class factory),
+ *      using RegisterClass(ETypeContainerScope::Instance)
+ *
+ * Note: Type containers depend on variadic templates and is therefore not available
+ * on XboxOne at this time. It should therefore only be used for desktop applications.
+ *
+ * @todo gmp: better documentation
+ */
+class FTypeContainer
+{
+	/** Interface for object instance providers. */
+	class IInstanceProvider
+	{
+	public:
+
+		/** Virtual destructor. */
+		virtual ~IInstanceProvider() { }
+
+		/**
+		 * Gets a pointer to the stored instance.
+		 *
+		 * The return value is not the instance itself, but a untyped pointer to a shared
+		 * pointer to the instance. The caller must cast this value to TSharedPtr<T> in
+		 * order to access the actual instance.
+		 * 
+		 */
+		virtual void* GetInstance() = 0;
+	};
+
+	/** Implements an instance provider that returns the same instance for all threads. */
+	template<class T>
+	class TSharedInstance
+		: public IInstanceProvider
+	{
+	public:
+
+		/** Creates and initializes a new instance. */
+		TSharedInstance(const TSharedRef<T>& InInstance)
+			: Instance(InInstance)
+		{ }
+
+		virtual ~TSharedInstance() { }
+
+		virtual void* GetInstance() override
+		{
+			return &Instance;
+		}
+
+	private:
+
+		/** The stored instance. */
+		TSharedPtr<T> Instance;
+	};
+
+
+	/** Base class for per-thread instances providers. */
+	class FThreadInstance
+		: public IInstanceProvider
+	{
+	public:
+
+		/** Constructor. */
+		FThreadInstance(TFunction<void*()>&& InCreateFunc)
+			: CreateFunc(MoveTemp(InCreateFunc))
+			, TlsSlot(FPlatformTLS::AllocTlsSlot())
+		{ }
+
+		virtual ~FThreadInstance() override
+		{
+			FPlatformTLS::FreeTlsSlot(TlsSlot);
+		}
+
+	protected:
+
+		/** Factory function for creating instances. */
+		TFunction<void*()> CreateFunc;
+
+		/** Slot ID for thread-local storage of the instance. */
+		uint32 TlsSlot;
+	};
+
+
+	/** Implements an instance provider that returns the same instance per thread. */
+	template<class T>
+	class TThreadInstance
+		: public FThreadInstance
+	{
+	public:
+
+		/** Constructor. */
+		TThreadInstance(TFunction<void*()>&& InCreateFunc)
+			: FThreadInstance(MoveTemp(InCreateFunc))
+		{ }
+
+		virtual ~TThreadInstance() override { }
+
+		virtual void* GetInstance() override
+		{
+			void* Instance = FPlatformTLS::GetTlsValue(TlsSlot);
+
+			if (Instance != nullptr)
+			{
+				return Instance;
+			}
+
+			Instance = new TSharedPtr<T>((T*)CreateFunc());
+			/** @todo gmp: this cannot possibly work, because the FRunnableThread is not yet registered.
+			    A better mechanism for executing code on thread shutdown is already being worked on.
+				In the meantime we leak some memory here.
+			FRunnableThread::GetThreadRegistry().Lock();
+			{
+				const uint32 ThreadId = FPlatformTLS::GetCurrentThreadId();
+				FRunnableThread* RunnableThread = FRunnableThread::GetThreadRegistry().GetThread(ThreadId);
+
+				if (RunnableThread != nullptr)
+				{
+					Instance = new TSharedPtr<T>((T*)CreateFunc());
+					RunnableThread->OnThreadDestroyed().AddStatic(&TThreadInstance::HandleDestroyInstance, Instance);
+				}
+			}
+			FRunnableThread::GetThreadRegistry().Unlock();*/
+			FPlatformTLS::SetTlsValue(TlsSlot, Instance);
+
+			return Instance;
+		}
+
+	private:
+
+		/** Callback for destroying per-thread instances. */
+		static void HandleDestroyInstance(void* Instance)
+		{
+			auto InstancePtr = (TSharedPtr<T>*)Instance;
+			delete InstancePtr;
+		}
+	};
+
+public:
+
+	/**
+	 * Gets a shared pointer to an instance of the specified class.
+	 *
+	 * @param T The type of class to get an instance for.
+	 * @param A shared pointer to the instance, or nullptr if no instance was registered.
+	 * @see GetInstanceRef, RegisterClass, RegisterInstance
+	 */
+	template<class R>
+	TSharedPtr<R> GetInstance()
+	{
+		const FName ClassName = R::GetTypeName();
+
+		FScopeLock Lock(&CriticalSection);
+		{
+			const TFunction<void*()>* Factory = Factories.Find(ClassName);
+
+			if (Factory != nullptr)
+			{
+				return TSharedPtr<R>((R*)(*Factory)());
+			}
+
+			const TSharedPtr<IInstanceProvider>& Provider = Providers.FindRef(ClassName);
+
+			if (Provider.IsValid())
+			{
+				void* Instance = Provider->GetInstance();
+				
+				if (Instance != nullptr)
+				{
+					return *((TSharedPtr<R>*)Instance);
+				}
+			}
+		}
+
+		return TSharedPtr<R>();
+	}
+
+	/**
+	 * Gets a shared reference to an instance of the specified class.
+	 *
+	 * Unlike GetInstance(), this function will assert if no instance was registered for
+	 * the requested type of class using either RegisterClass() or RegisterInstance().
+	 *
+	 * @param R The type of class that an instance is being requested for.
+	 * @param A shared pointer to the instance.
+	 * @see GetInstance, RegisterClass, RegisterInstance
+	 */
+	template<class R>
+	TSharedRef<R> GetInstanceRef()
+	{
+		return GetInstance<R>().ToSharedRef();
+	}
+
+	/**
+	 * Registers a class for instances of the specified class.
+	 *
+	 * @param R The type of class to register an instance class for.
+	 * @param T Tye type of class to register (must be the same as or derived from R).
+	 * @param P The constructor parameters that the class requires.
+	 * @param Scope The scope at which instances must be unique.
+	 */
+	template<class R, class T, typename... P>
+	void RegisterClass(ETypeContainerScope Scope)
+	{
+		static_assert(TPointerIsConvertibleFromTo<T, R>::Value, "T must implement R");
+
+		const FName ClassName = R::GetTypeName();
+
+		if (Scope == ETypeContainerScope::Instance)
+		{
+			FScopeLock Lock(&CriticalSection);
+			{
+				Providers.Remove(ClassName);
+				Factories.Add(ClassName, [this]() -> void* { return new T(GetInstanceRef<P>()...); });
+			}
+		}
+		else
+		{
+			TSharedPtr<IInstanceProvider> Provider;
+			
+			if (Scope == ETypeContainerScope::Thread)
+			{
+				Provider = MakeShareable(new TThreadInstance<T>([this]() -> void* { return new T(GetInstanceRef<P>()...); }));
+			}
+			else
+			{
+				Provider = MakeShareable(new TSharedInstance<T>(MakeShareable(new T(GetInstanceRef<P>()...))));
+			}
+
+			FScopeLock Lock(&CriticalSection);
+			{
+				Factories.Remove(ClassName);
+				Providers.Add(ClassName, Provider);
+			}
+		}
+	}
+
+	/**
+	 * Registers an existing instance for the specified class.
+	 *
+	 * @param R The type of class to register the instance for.
+	 * @param T The type of the instance to register (must be the same as or derived from R).
+	 * @param Instance The instance of the class to register.
+	 */
+	template<class R, class T>
+	void RegisterInstance(const TSharedRef<T>& Instance)
+	{
+		static_assert(TPointerIsConvertibleFromTo<T, R>::Value, "T must implement R");
+
+		FScopeLock Lock(&CriticalSection);
+		{
+			Providers.Add(
+				R::GetTypeName(),
+				MakeShareable(new TSharedInstance<R>(Instance))
+			);
+		}
+	}
+
+	/**
+	 * Unregisters the instance or class for a previously registered class.
+	 *
+	 * @param R The type of class to unregister.
+	 */
+	template<class R>
+	void Unregister()
+	{
+		const FName ClassName = R::GetTypeName();
+
+		FScopeLock Lock(&CriticalSection);
+		{
+			Factories.Remove(ClassName);
+			Providers.Remove(ClassName);
+		}
+	}
+
+private:
+
+	/** Critical section for synchronizing access. */
+	FCriticalSection CriticalSection;
+
+	/** Maps class names to factory functions. */
+	TMap<FName, TFunction<void*()>> Factories;
+
+	/** Maps class names to instance providers. */
+	TMap<FName, TSharedPtr<IInstanceProvider>> Providers;
+};
+
+
+#endif //PLATFORM_COMPILER_HAS_VARIADIC_TEMPLATES
