Added missing virtual destructors

diff --git a/Engine/Source/Runtime/Core/Private/Tests/Misc/TypeContainerTest.cpp b/Engine/Source/Runtime/Core/Private/Tests/Misc/TypeContainerTest.cpp
--- a/Engine/Source/Runtime/Core/Private/Tests/Misc/TypeContainerTest.cpp
+++ b/Engine/Source/Runtime/Core/Private/Tests/Misc/TypeContainerTest.cpp
@@ -20,27 +20,32 @@ IMPLEMENT_SIMPLE_AUTOMATION_TEST(FTypeContainerTest, "Core.Misc.TypeContainer",
 struct IFruit
 {
 	DECLARE_REGISTRABLE_TYPE(IFruit)
+	virtual ~IFruit() { }
 	virtual FString Name() = 0;
 };
 
 struct IBerry : public IFruit
 {
 	DECLARE_REGISTRABLE_TYPE(IBerry)
+	virtual ~IBerry() { }
 };
 
 struct FBanana : public IFruit
 {
+	virtual ~FBanana() { }
 	virtual FString Name() override { return TEXT("Cherry"); }
 };
 
 struct FStrawberry : public IBerry
 {
+	virtual ~FStrawberry() { }
 	virtual FString Name() override { return TEXT("Strawberry"); }
 };
 
 struct ISmoothie
 {
 	DECLARE_REGISTRABLE_TYPE(ISmoothie)
+	virtual ~ISmoothie() { }
 	virtual TSharedRef<IBerry> GetBerry() = 0;
 	virtual TSharedRef<IFruit> GetFruit() = 0;
 };
@@ -48,6 +53,7 @@ struct ISmoothie
 struct FSmoothie : public ISmoothie
 {
 	FSmoothie(TSharedRef<IFruit> InFruit, TSharedRef<IBerry> InBerry) : Berry(InBerry), Fruit(InFruit) { }
+	virtual ~FSmoothie() { }
 	virtual TSharedRef<IBerry> GetBerry() override { return Berry; }
 	virtual TSharedRef<IFruit> GetFruit() override { return Fruit; }
 	TSharedRef<IBerry> Berry;
