Fixed Clang compile error

diff --git a/Engine/Source/Runtime/Core/Private/Internationalization/Text.cpp b/Engine/Source/Runtime/Core/Private/Internationalization/Text.cpp
--- a/Engine/Source/Runtime/Core/Private/Internationalization/Text.cpp
+++ b/Engine/Source/Runtime/Core/Private/Internationalization/Text.cpp
@@ -849,8 +849,8 @@ bool FTextSnapshot::IsDisplayStringEqualTo(const FText& InText) const
 }
 
 FScopedTextIdentityPreserver::FScopedTextIdentityPreserver(FText& InTextToPersist)
-	: TextToPersist(InTextToPersist)
-	, Flags(InTextToPersist.Flags)
+	: Flags(InTextToPersist.Flags)
+	, TextToPersist(InTextToPersist)
 {
 	// Empty display strings can't have a namespace or key.
 	if (GIsEditor && !InTextToPersist.DisplayString->IsEmpty())
