Core: default parameter for FTypeContainer::RegisterClass()

diff --git a/Engine/Source/Runtime/Core/Public/Misc/TypeContainer.h b/Engine/Source/Runtime/Core/Public/Misc/TypeContainer.h
--- a/Engine/Source/Runtime/Core/Public/Misc/TypeContainer.h
+++ b/Engine/Source/Runtime/Core/Public/Misc/TypeContainer.h
@@ -250,10 +250,10 @@ public:
 	 * @param R The type of class to register an instance class for.
 	 * @param T Tye type of class to register (must be the same as or derived from R).
 	 * @param P The constructor parameters that the class requires.
-	 * @param Scope The scope at which instances must be unique.
+	 * @param Scope The scope at which instances must be unique (default = Process).
 	 */
 	template<class R, class T, typename... P>
-	void RegisterClass(ETypeContainerScope Scope)
+	void RegisterClass(ETypeContainerScope Scope = ETypeContainerScope::Process)
 	{
 		static_assert(TPointerIsConvertibleFromTo<T, R>::Value, "T must implement R");
 
