Added FApplePlatformFile::CopyFile() that differs from the generic implementation in that it also copies the file permission bits. Also, did a little cleanup in ApplePlatformFile.

diff --git a/Engine/Source/Runtime/Core/Private/Apple/ApplePlatformFile.cpp b/Engine/Source/Runtime/Core/Private/Apple/ApplePlatformFile.cpp
--- a/Engine/Source/Runtime/Core/Private/Apple/ApplePlatformFile.cpp
+++ b/Engine/Source/Runtime/Core/Private/Apple/ApplePlatformFile.cpp
@@ -355,8 +355,7 @@ bool FApplePlatformFile::MoveFile(const TCHAR* To, const TCHAR* From)
 bool FApplePlatformFile::SetReadOnly(const TCHAR* Filename, bool bNewReadOnlyValue)
 {
 	struct stat FileInfo;
-	ANSICHAR FixedFilePath[MAX_PATH] = { 0 };
-	if (Stat(Filename, &FileInfo, FixedFilePath, MAX_PATH) != -1)
+	if (Stat(Filename, &FileInfo) == 0)
 	{
 		if (bNewReadOnlyValue)
 		{
@@ -366,7 +365,7 @@ bool FApplePlatformFile::SetReadOnly(const TCHAR* Filename, bool bNewReadOnlyVal
 		{
 			FileInfo.st_mode |= S_IWUSR;
 		}
-		return chmod(*FixedFilePath ? FixedFilePath : TCHAR_TO_UTF8(*NormalizeFilename(Filename)), FileInfo.st_mode);
+		return chmod(TCHAR_TO_UTF8(*NormalizeFilename(Filename)), FileInfo.st_mode) == 0;
 	}
 	return false;
 }
@@ -391,8 +390,7 @@ void FApplePlatformFile::SetTimeStamp(const TCHAR* Filename, const FDateTime Dat
 {
 	// get file times
 	struct stat FileInfo;
-	ANSICHAR FixedFilePath[MAX_PATH] = { 0 };
-	if (Stat(Filename, &FileInfo, FixedFilePath, MAX_PATH) != -1)
+	if (Stat(Filename, &FileInfo) == 0)
 	{
 		return;
 	}
@@ -401,7 +399,7 @@ void FApplePlatformFile::SetTimeStamp(const TCHAR* Filename, const FDateTime Dat
 	struct utimbuf Times;
 	Times.actime = FileInfo.st_atime;
 	Times.modtime = (DateTime - MacEpoch).GetTotalSeconds();
-	utime(*FixedFilePath ? FixedFilePath : TCHAR_TO_UTF8(*NormalizeFilename(Filename)), &Times);
+	utime(TCHAR_TO_UTF8(*NormalizeFilename(Filename)), &Times);
 }
 FDateTime FApplePlatformFile::GetAccessTimeStamp(const TCHAR* Filename)
 {
@@ -530,7 +528,21 @@ bool FApplePlatformFile::IterateDirectory(const TCHAR* Directory, FDirectoryVisi
 	}
 }
 
-int32 FApplePlatformFile::Stat(const TCHAR* Filename, struct stat* OutFileInfo, ANSICHAR* OutFixedPath, SIZE_T FixedPathSize)
+bool FApplePlatformFile::CopyFile(const TCHAR* To, const TCHAR* From)
+{
+	bool Result = IPlatformFile::CopyFile(To, From);
+	if (Result)
+	{
+		struct stat FileInfo;
+		if (Stat(From, &FileInfo) == 0)
+		{
+			chmod(TCHAR_TO_UTF8(*NormalizeFilename(To)), FileInfo.st_mode);
+		}
+	}
+	return Result;
+}
+
+int32 FApplePlatformFile::Stat(const TCHAR* Filename, struct stat* OutFileInfo)
 {
 	return stat(TCHAR_TO_UTF8(*NormalizeFilename(Filename)), OutFileInfo);
 }
diff --git a/Engine/Source/Runtime/Core/Public/Apple/ApplePlatformFile.h b/Engine/Source/Runtime/Core/Public/Apple/ApplePlatformFile.h
--- a/Engine/Source/Runtime/Core/Public/Apple/ApplePlatformFile.h
+++ b/Engine/Source/Runtime/Core/Public/Apple/ApplePlatformFile.h
@@ -36,7 +36,8 @@ public:
 	virtual bool DirectoryExists(const TCHAR* Directory) override;
 	virtual bool CreateDirectory(const TCHAR* Directory) override;
 	virtual bool DeleteDirectory(const TCHAR* Directory) override;
-	bool IterateDirectory(const TCHAR* Directory, FDirectoryVisitor& Visitor);
+	virtual bool IterateDirectory(const TCHAR* Directory, FDirectoryVisitor& Visitor);
+	virtual bool CopyFile(const TCHAR* To, const TCHAR* From);
 private:
-	int32 Stat(const TCHAR* Filename, struct stat* OutFileInfo, ANSICHAR* OutFixedPath = nullptr, SIZE_T FixedPathSize = 0);
+	int32 Stat(const TCHAR* Filename, struct stat* OutFileInfo);
 };
