Stats - UECORE-79 - Add better mechanism to identifying Event's stats (each Event has an unique name now)

diff --git a/Engine/Source/Runtime/Core/Private/GenericPlatform/GenericPlatformProcess.cpp b/Engine/Source/Runtime/Core/Private/GenericPlatform/GenericPlatformProcess.cpp
--- a/Engine/Source/Runtime/Core/Private/GenericPlatform/GenericPlatformProcess.cpp
+++ b/Engine/Source/Runtime/Core/Private/GenericPlatform/GenericPlatformProcess.cpp
@@ -306,7 +306,7 @@ DECLARE_CYCLE_STAT(TEXT("CPU Stall - Wait For Event"),STAT_EventWait,STATGROUP_C
 
 bool FPThreadEvent::Wait(uint32 WaitTime, const bool bIgnoreThreadIdleStats /*= false*/)
 {
-	SCOPE_CYCLE_COUNTER(STAT_EventWait);
+	FScopeCycleCounter Counter(StatID);
 	FThreadIdleStats::FScopeIdle Scope(bIgnoreThreadIdleStats);
 
 	check(bInitialized);
@@ -417,6 +417,11 @@ FEvent* FGenericPlatformProcess::GetSynchEventFromPool(bool bIsManualReset)
 
 void FGenericPlatformProcess::ReturnSynchEventToPool(FEvent* Event)
 {
+	if( !Event )
+	{
+		return;
+	}
+
 	if (Event->IsManualReset())
 	{
 		FEventPool<EEventPoolTypes::ManualReset>::Get().ReturnToPool(Event);
diff --git a/Engine/Source/Runtime/Core/Private/Misc/EventPool.cpp b/Engine/Source/Runtime/Core/Private/Misc/EventPool.cpp
new file mode 100644
--- /dev/null
+++ b/Engine/Source/Runtime/Core/Private/Misc/EventPool.cpp
@@ -0,0 +1,29 @@
+// Copyright 1998-2015 Epic Games, Inc. All Rights Reserved.
+
+
+#include "CorePrivatePCH.h"
+#include "EventPool.h"
+
+TStatId FEventStats::CreateStatID()
+{
+#if	STATS
+	static FThreadSafeCounter UniqueId;
+	const int32 PreviousId = UniqueId.Add(1);
+	const FString UniqueEventName = FString::Printf( TEXT("CPU Stall - Wait for event %3i"), PreviousId );
+
+	const FName StatName = FName( *UniqueEventName );
+	FStartupMessages::Get().AddMetadata( StatName, *UniqueEventName, 
+		STAT_GROUP_TO_FStatGroup( STATGROUP_CPUStalls )::GetGroupName(), 
+		STAT_GROUP_TO_FStatGroup( STATGROUP_CPUStalls )::GetGroupCategory(), 
+		STAT_GROUP_TO_FStatGroup( STATGROUP_CPUStalls )::GetDescription(),
+		true, EStatDataType::ST_int64, true );
+
+	return IStatGroupEnableManager::Get().GetHighPerformanceEnableForStat(StatName, 
+		STAT_GROUP_TO_FStatGroup( STATGROUP_CPUStalls )::GetGroupName(), 
+		STAT_GROUP_TO_FStatGroup( STATGROUP_CPUStalls )::GetGroupCategory(), 
+		STAT_GROUP_TO_FStatGroup( STATGROUP_CPUStalls )::DefaultEnable,
+		true, EStatDataType::ST_int64, *UniqueEventName, true);
+#else
+	return TStatId();
+#endif // STATS
+}
diff --git a/Engine/Source/Runtime/Core/Private/WinRT/WinRTEvent.h b/Engine/Source/Runtime/Core/Private/WinRT/WinRTEvent.h
--- a/Engine/Source/Runtime/Core/Private/WinRT/WinRTEvent.h
+++ b/Engine/Source/Runtime/Core/Private/WinRT/WinRTEvent.h
@@ -37,7 +37,7 @@ public:
 //		Event = CreateEvent(NULL, bIsManualReset, 0, NULL);
 		Event = CreateEventEx(NULL, NULL, bIsManualReset ? CREATE_EVENT_MANUAL_RESET : 0, EVENT_ALL_ACCESS);
 		ManualReset = bIsManualReset;
-
+		StatID = FEventStats::CreateStatID();
 		return Event != NULL;
 	}
 
@@ -69,6 +69,9 @@ private:
 
 	/** Whether the signaled state of the event needs to be reset manually. */
 	bool ManualReset;
+
+	/** Stat ID of this event. */
+	TStatId StatID;
 };
 
 
diff --git a/Engine/Source/Runtime/Core/Private/WinRT/WinRTProcess.cpp b/Engine/Source/Runtime/Core/Private/WinRT/WinRTProcess.cpp
--- a/Engine/Source/Runtime/Core/Private/WinRT/WinRTProcess.cpp
+++ b/Engine/Source/Runtime/Core/Private/WinRT/WinRTProcess.cpp
@@ -176,13 +176,11 @@ FEvent* FWinRTProcess::CreateSynchEvent(bool bIsManualReset /*= false*/)
 	return Event;
 }
 
-DECLARE_CYCLE_STAT(TEXT("CPU Stall - Wait For Event"),STAT_EventWait,STATGROUP_CPUStalls);
-
 #include "AllowWinRTPlatformTypes.h"
 
 bool FEventWinRT::Wait(uint32 WaitTime, const bool bIgnoreThreadIdleStats /*= false*/)
 {
-	SCOPE_CYCLE_COUNTER(STAT_EventWait);
+	FScopeCycleCounter Counter(StatID);
 	FThreadIdleStats::FScopeIdle Scope(bIgnoreThreadIdleStats);
 	check(Event);
 
diff --git a/Engine/Source/Runtime/Core/Private/Windows/WindowsEvent.h b/Engine/Source/Runtime/Core/Private/Windows/WindowsEvent.h
--- a/Engine/Source/Runtime/Core/Private/Windows/WindowsEvent.h
+++ b/Engine/Source/Runtime/Core/Private/Windows/WindowsEvent.h
@@ -27,8 +27,6 @@ public:
 		}
 	}
 
-public:
-
 	// FEvent interface
 
 	virtual bool Create( bool bIsManualReset = false ) override
@@ -36,7 +34,7 @@ public:
 		// Create the event and default it to non-signaled
 		Event = CreateEvent(nullptr, bIsManualReset, 0, nullptr);
 		ManualReset = bIsManualReset;
-
+		StatID = FEventStats::CreateStatID();
 		return Event != nullptr;
 	}
 
@@ -68,6 +66,9 @@ private:
 
 	/** Whether the signaled state of the event needs to be reset manually. */
 	bool ManualReset;
+
+	/** Stat ID of this event. */
+	TStatId StatID;
 };
 
 
diff --git a/Engine/Source/Runtime/Core/Private/Windows/WindowsPlatformProcess.cpp b/Engine/Source/Runtime/Core/Private/Windows/WindowsPlatformProcess.cpp
--- a/Engine/Source/Runtime/Core/Private/Windows/WindowsPlatformProcess.cpp
+++ b/Engine/Source/Runtime/Core/Private/Windows/WindowsPlatformProcess.cpp
@@ -21,6 +21,7 @@
 
 #include "HideWindowsPlatformTypes.h"
 #include "WindowsPlatformMisc.h"
+#include "EventPool.h"
 
 #pragma comment(lib, "psapi.lib")
 
@@ -1008,14 +1009,11 @@ FEvent* FWindowsPlatformProcess::CreateSynchEvent(bool bIsManualReset)
 
 #include "AllowWindowsPlatformTypes.h"
 
-DECLARE_CYCLE_STAT(TEXT("CPU Stall - Wait For Event"),STAT_EventWait,STATGROUP_CPUStalls);
-
 bool FEventWin::Wait(uint32 WaitTime, const bool bIgnoreThreadIdleStats /*= false*/)
 {
-	SCOPE_CYCLE_COUNTER(STAT_EventWait);
-	check(Event);
-
+	FScopeCycleCounter Counter(StatID);
 	FThreadIdleStats::FScopeIdle Scope(bIgnoreThreadIdleStats);
+	check(Event);
 	return (WaitForSingleObject(Event, WaitTime) == WAIT_OBJECT_0);
 }
 
diff --git a/Engine/Source/Runtime/Core/Public/HAL/PThreadEvent.h b/Engine/Source/Runtime/Core/Public/HAL/PThreadEvent.h
--- a/Engine/Source/Runtime/Core/Public/HAL/PThreadEvent.h
+++ b/Engine/Source/Runtime/Core/Public/HAL/PThreadEvent.h
@@ -24,6 +24,9 @@ class FPThreadEvent
     pthread_mutex_t Mutex;
     pthread_cond_t Condition;
 
+	/** Stat ID of this event. */
+	TStatId StatID;
+
 	inline void LockEventMutex()
 	{
 		const int rc = pthread_mutex_lock(&Mutex);
@@ -93,8 +96,6 @@ public:
 		}
 	}
 
-public:
-
 	// FEvent interface
 
 	virtual bool Create( bool _bIsManualReset = false ) override
@@ -116,6 +117,7 @@ public:
 				pthread_mutex_destroy(&Mutex);
 			}
 		}
+		StatID = FEventStats::CreateStatID();
 		return RetVal;
 	}
 
diff --git a/Engine/Source/Runtime/Core/Public/HAL/ThreadingBase.h b/Engine/Source/Runtime/Core/Public/HAL/ThreadingBase.h
--- a/Engine/Source/Runtime/Core/Public/HAL/ThreadingBase.h
+++ b/Engine/Source/Runtime/Core/Public/HAL/ThreadingBase.h
@@ -99,8 +99,6 @@ public:
 		return Wait(WaitTime.GetTotalMilliseconds(), bIgnoreThreadIdleStats);
 	}
 
-public:
-
 	/** Virtual destructor. */
 	virtual ~FEvent() { }
 };
