Stats - Removed unneeded CORE_API

diff --git a/Engine/Source/Runtime/Core/Public/Stats/Stats2.h b/Engine/Source/Runtime/Core/Public/Stats/Stats2.h
--- a/Engine/Source/Runtime/Core/Public/Stats/Stats2.h
+++ b/Engine/Source/Runtime/Core/Public/Stats/Stats2.h
@@ -224,8 +224,6 @@ struct EStatOperation
 		Add,
 		/** Subtract operation. */
 		Subtract,
-		/** This is a memory operation. @see EMemoryOperation. */
-		Memory,
 
 		// these are special ones for processed data
 		ChildrenStart,
@@ -233,6 +231,9 @@ struct EStatOperation
 		Leaf,
 		MaxVal,
 
+		/** This is a memory operation. @see EMemoryOperation. */
+		Memory,
+
 		Num,
 		Mask = 0xf,
 		Shift = EStatDataType::Shift + EStatDataType::NumBits,
@@ -1082,8 +1083,9 @@ protected:
 };
 
 /** Preallocates FThreadStats to avoid dynamic memory allocation. */
-struct CORE_API FThreadStatsPool
+struct FThreadStatsPool
 {
+private:
 	enum
 	{
 		/**
@@ -1102,17 +1104,17 @@ public:
 	FThreadStatsPool();
 
 	/** Singleton accessor. */
-	static FThreadStatsPool& Get()
+	CORE_API static FThreadStatsPool& Get()
 	{
 		static FThreadStatsPool Singleton;
 		return Singleton;
 	}
 
 	/** Gets an instance from the pool and call the default constructor on it. */
-	FThreadStats* GetFromPool();
+	CORE_API FThreadStats* GetFromPool();
 
 	/** Return an instance to the pool. */
-	void ReturnToPool(FThreadStats* Instance);
+	CORE_API void ReturnToPool(FThreadStats* Instance);
 };
 
 /**
@@ -1408,20 +1410,22 @@ public:
 
 /** Manages startup messages, usually to update the metadata. */
 // @TODO yrx 2014-11-18 Eventually add startup messages for memory profiling?
-class CORE_API FStartupMessages
+class FStartupMessages
 {
-public:
+	friend class FStatsThread;
+
 	TArray<FStatMessage> DelayedMessages;
 	FCriticalSection CriticalSection;
 
+public:
 	/** Adds a thread metadata. */
-	void AddThreadMetadata( const FName InThreadFName, uint32 InThreadID );
+	CORE_API void AddThreadMetadata( const FName InThreadFName, uint32 InThreadID );
 
 	/** Adds a regular metadata. */
-	void AddMetadata( FName InStatName, const TCHAR* InStatDesc, const char* InGroupName, const char* InGroupCategory, const TCHAR* InGroupDesc, bool bShouldClearEveryFrame, EStatDataType::Type InStatType, bool bCycleStat, FPlatformMemory::EMemoryCounterRegion InMemoryRegion = FPlatformMemory::MCR_Invalid );
+	CORE_API void AddMetadata( FName InStatName, const TCHAR* InStatDesc, const char* InGroupName, const char* InGroupCategory, const TCHAR* InGroupDesc, bool bShouldClearEveryFrame, EStatDataType::Type InStatType, bool bCycleStat, FPlatformMemory::EMemoryCounterRegion InMemoryRegion = FPlatformMemory::MCR_Invalid );
 
 	/** Access the singleton. */
-	static FStartupMessages& Get();
+	CORE_API static FStartupMessages& Get();
 };
 
 /**
