Stats - Added STATGROUP_LoadTimeVerbose, mostly for memory profiling

diff --git a/Engine/Source/Runtime/Core/Public/Stats/Stats2.h b/Engine/Source/Runtime/Core/Public/Stats/Stats2.h
--- a/Engine/Source/Runtime/Core/Public/Stats/Stats2.h
+++ b/Engine/Source/Runtime/Core/Public/Stats/Stats2.h
@@ -1830,6 +1830,7 @@ DECLARE_STATS_GROUP(TEXT("Init Views"),STATGROUP_InitViews, STATCAT_Advanced);
 DECLARE_STATS_GROUP(TEXT("Landscape"),STATGROUP_Landscape, STATCAT_Advanced);
 DECLARE_STATS_GROUP(TEXT("Light Rendering"),STATGROUP_LightRendering, STATCAT_Advanced);
 DECLARE_STATS_GROUP(TEXT("Load Time"), STATGROUP_LoadTime, STATCAT_Advanced);
+DECLARE_STATS_GROUP_VERBOSE(TEXT("Load Time (Verbose)"), STATGROUP_LoadTimeVerbose, STATCAT_Advanced);
 DECLARE_STATS_GROUP(TEXT("Math"), STATGROUP_Math, STATCAT_Advanced);
 DECLARE_STATS_GROUP(TEXT("Memory Allocator"),STATGROUP_MemoryAllocator, STATCAT_Advanced);
 DECLARE_STATS_GROUP(TEXT("Memory Platform"),STATGROUP_MemoryPlatform, STATCAT_Advanced);
