Fix transition in/out of fullscreen mode on OS X sometimes causing faulty mouse coordinates to be reported, fixes [UE-8347].

diff --git a/Engine/Source/Runtime/Core/Private/Mac/CocoaWindow.cpp b/Engine/Source/Runtime/Core/Private/Mac/CocoaWindow.cpp
--- a/Engine/Source/Runtime/Core/Private/Mac/CocoaWindow.cpp
+++ b/Engine/Source/Runtime/Core/Private/Mac/CocoaWindow.cpp
@@ -397,6 +397,7 @@ NSString* NSPerformDragOperation = @"NSPerformDragOperation";
 	{
 		FMacEvent::SendToGameRunLoop(Notification, EMacEventSendMethod::Async, @[ NSDefaultRunLoopMode, UE4FullscreenEventMode ]);
 	}
+	[self setFrame:self.PreFullScreenRect display:YES];
 	FMacCursor* MacCursor = (FMacCursor*)MacApplication->Cursor.Get();
 	if ( MacCursor )
 	{
diff --git a/Engine/Source/Runtime/Core/Private/Mac/MacWindow.cpp b/Engine/Source/Runtime/Core/Private/Mac/MacWindow.cpp
--- a/Engine/Source/Runtime/Core/Private/Mac/MacWindow.cpp
+++ b/Engine/Source/Runtime/Core/Private/Mac/MacWindow.cpp
@@ -263,6 +263,17 @@ void FMacWindow::ReshapeWindow( int32 X, int32 Y, int32 Width, int32 Height )
 			if ( !NSEqualRects(WindowHandle.PreFullScreenRect, NewRect) )
 			{
 				WindowHandle.PreFullScreenRect = NewRect;
+				MainThreadCall(^{
+					FMacCursor* MacCursor = (FMacCursor*)MacApplication->Cursor.Get();
+					if ( MacCursor )
+					{
+						NSSize WindowSize = [WindowHandle frame].size;
+						NSSize ViewSize = [WindowHandle openGLFrame].size;
+						float WidthScale = ViewSize.width / WindowSize.width;
+						float HeightScale = ViewSize.height / WindowSize.height;
+						MacCursor->SetMouseScaling(FVector2D(WidthScale, HeightScale));
+					}
+				}, UE4ResizeEventMode, true);
 				FMacEvent::SendToGameRunLoop([NSNotification notificationWithName:NSWindowDidResizeNotification object:WindowHandle], EMacEventSendMethod::Sync, @[ NSDefaultRunLoopMode, UE4ResizeEventMode, UE4ShowEventMode, UE4FullscreenEventMode ]);
 			}
 		}
