Defer closing and releasing of FCocoaWindow objects till the next PumpMessages, so that it's done when MacApplication and Slate no longer know about them.

diff --git a/Engine/Source/Runtime/Core/Private/Mac/MacApplication.cpp b/Engine/Source/Runtime/Core/Private/Mac/MacApplication.cpp
--- a/Engine/Source/Runtime/Core/Private/Mac/MacApplication.cpp
+++ b/Engine/Source/Runtime/Core/Private/Mac/MacApplication.cpp
@@ -107,7 +107,6 @@ FMacApplication::FMacApplication()
 	, bIsProcessingNSEvent( false )
 	, ModifierKeysFlags( 0 )
 	, CurrentModifierFlags( 0 )
-	, bKeyWindowUpdateRequested( false )
 	, bIsWorkspaceSessionActive( true )
 {
 	CGDisplayRegisterReconfigurationCallback(FMacApplication::OnDisplayReconfiguration, this);
@@ -426,16 +425,38 @@ void FMacApplication::ProcessEvent(FMacEvent const* const Event)
 void FMacApplication::PumpMessages(const float TimeDelta)
 {
 	FPlatformMisc::PumpMessages(true);
+}
 
-	if (bKeyWindowUpdateRequested)
+void FMacApplication::CloseQueuedWindows()
+{
+	if (WindowsToClose.Num() > 0)
 	{
-		TSharedPtr<FMacWindow> KeyWindow = GetKeyWindow();
-		if (KeyWindow.IsValid() && KeyWindow->GetOSWindowHandle() && ![(FCocoaWindow*)KeyWindow->GetOSWindowHandle() isMiniaturized])
-		{
-			// Activate specified previous window if still present, provided it isn't minimized
-			KeyWindow->SetWindowFocus();
-		}
-		bKeyWindowUpdateRequested = false;
+		MainThreadCall(^{
+			SCOPED_AUTORELEASE_POOL;
+
+			// Set the new key window, if needed. We cannot trust Cocoa to set the key window to the actual top most window. It will prefer windows with title bars so,
+			// for example, will choose the main window over a context menu window, when closing a submenu.
+			NSArray* AllWindows = [NSApp orderedWindows];
+			for (NSWindow* Window : AllWindows)
+			{
+				if ([Window isKindOfClass:[FCocoaWindow class]] && !WindowsToClose.Contains((FCocoaWindow*)Window) && [Window canBecomeKeyWindow])
+				{
+					if (Window != [NSApp keyWindow])
+					{
+						[Window makeKeyWindow];
+					}
+					break;
+				}
+			}
+
+			for (FCocoaWindow* Window : WindowsToClose)
+			{
+				[Window destroy];
+				[Window release];
+			}
+		}, UE4CloseEventMode, true);
+
+		WindowsToClose.Empty();
 	}
 }
 
@@ -1149,7 +1170,6 @@ void FMacApplication::OnWindowDidBecomeKey( FCocoaWindow* Window )
 	if( EventWindow.IsValid() )
 	{
 		MessageHandler->OnWindowActivationChanged( EventWindow.ToSharedRef(), EWindowActivation::Activate );
-		KeyWindows.Add( EventWindow.ToSharedRef() );
 	}
 }
 
@@ -1233,7 +1253,6 @@ void FMacApplication::OnWindowDidClose( FCocoaWindow* Window )
 			MessageHandler->OnWindowActivationChanged( EventWindow.ToSharedRef(), EWindowActivation::Deactivate );
 		}
 		Windows.Remove( EventWindow.ToSharedRef() );
-		KeyWindows.Remove( EventWindow.ToSharedRef() );
 		MessageHandler->OnWindowClose( EventWindow.ToSharedRef() );
 	}
 }
@@ -1248,7 +1267,10 @@ bool FMacApplication::OnWindowDestroyed( FCocoaWindow* Window )
 			MessageHandler->OnWindowActivationChanged( EventWindow.ToSharedRef(), EWindowActivation::Deactivate );
 		}
 		Windows.Remove( EventWindow.ToSharedRef() );
-		KeyWindows.Remove( EventWindow.ToSharedRef() );
+		if (!WindowsToClose.Contains(Window))
+		{
+			WindowsToClose.Add(Window);
+		}
 		return true;
 	}
 	return false;
@@ -1355,16 +1377,6 @@ TCHAR FMacApplication::TranslateCharCode(TCHAR CharCode, uint32 KeyCode)
 	return CharCode;
 }
 
-TSharedPtr<FMacWindow> FMacApplication::GetKeyWindow()
-{
-	TSharedPtr<FMacWindow> KeyWindow;
-	if(KeyWindows.Num() > 0)
-	{
-		KeyWindow = KeyWindows.Top();
-	}
-	return KeyWindow;
-}
-
 #if WITH_EDITOR
 
 void FMacApplication::RecordUsage(EGestureEvent::Type Gesture)
diff --git a/Engine/Source/Runtime/Core/Private/Mac/MacPlatformMisc.cpp b/Engine/Source/Runtime/Core/Private/Mac/MacPlatformMisc.cpp
--- a/Engine/Source/Runtime/Core/Private/Mac/MacPlatformMisc.cpp
+++ b/Engine/Source/Runtime/Core/Private/Mac/MacPlatformMisc.cpp
@@ -489,10 +489,15 @@ void FMacPlatformMisc::PumpMessages( bool bFromMainLoop )
 	{
 		ProcessGameThreadEvents();
 
-		if (UpdateCachedMacMenuState && bChachedMacMenuStateNeedsUpdate && MacApplication && !MacApplication->IsProcessingNSEvent() && IsInGameThread())
+		if (MacApplication && !MacApplication->IsProcessingNSEvent() && IsInGameThread())
 		{
-			UpdateCachedMacMenuState();
-			bChachedMacMenuStateNeedsUpdate = false;
+			if (UpdateCachedMacMenuState && bChachedMacMenuStateNeedsUpdate)
+			{
+				UpdateCachedMacMenuState();
+				bChachedMacMenuStateNeedsUpdate = false;
+			}
+
+			MacApplication->CloseQueuedWindows();
 		}
 	}
 }
diff --git a/Engine/Source/Runtime/Core/Private/Mac/MacWindow.cpp b/Engine/Source/Runtime/Core/Private/Mac/MacWindow.cpp
--- a/Engine/Source/Runtime/Core/Private/Mac/MacWindow.cpp
+++ b/Engine/Source/Runtime/Core/Private/Mac/MacWindow.cpp
@@ -10,17 +10,6 @@
 
 FMacWindow::~FMacWindow()
 {
-	// While on Windows invalid HWNDs fail silently, accessing an invalid NSWindow is fatal.
-	// So instead we release the window here.
-	if(WindowHandle)
-	{
-		NSWindow* Window = WindowHandle;
-		dispatch_async(dispatch_get_main_queue(), ^{
-			SCOPED_AUTORELEASE_POOL;
-			[Window release];
-		});
-		WindowHandle = nil;
-	}
 }
 
 TSharedRef<FMacWindow> FMacWindow::Make()
@@ -320,25 +309,9 @@ void FMacWindow::Destroy()
 	{
 		SCOPED_AUTORELEASE_POOL;
 		bIsClosed = true;
-
-		FCocoaWindow* Window = WindowHandle;
-		const bool bIsKey = [Window isKeyWindow];
-
-		if (MacApplication->OnWindowDestroyed(Window) )
-		{
-			// This FMacWindow may have been destructed by now & so the WindowHandle will probably be invalid memory.
-
-			if (bIsKey)
-			{
-				MacApplication->RequestKeyWindowUpdate();
-			}
-
-			// Close the window
-			MainThreadCall(^{
-				SCOPED_AUTORELEASE_POOL;
-				[Window destroy];
-			}, UE4CloseEventMode, true);
-		}
+		[WindowHandle setAlphaValue:0.0f];
+		[WindowHandle setBackgroundColor:[NSColor clearColor]];
+		MacApplication->OnWindowDestroyed(WindowHandle);
 	}
 }
 
diff --git a/Engine/Source/Runtime/Core/Public/Mac/MacApplication.h b/Engine/Source/Runtime/Core/Public/Mac/MacApplication.h
--- a/Engine/Source/Runtime/Core/Public/Mac/MacApplication.h
+++ b/Engine/Source/Runtime/Core/Public/Mac/MacApplication.h
@@ -111,6 +111,8 @@ public:
 
 	const TArray<TSharedRef<FMacWindow>>& GetAllWindows() const { return Windows; }
 
+	void CloseQueuedWindows();
+
 private:
 	enum FMacApplicationEventTypes
 	{
@@ -130,9 +132,6 @@ private:
 	FCocoaWindow* FindEventWindow( NSEvent* CocoaEvent );
 	TSharedPtr<FMacWindow> LocateWindowUnderCursor( const NSPoint Position );
 
-	TSharedPtr<FMacWindow> GetKeyWindow();
-	void RequestKeyWindowUpdate() { bKeyWindowUpdateRequested = true; }
-
 	NSScreen* FindScreenByPoint( int32 X, int32 Y ) const;
 
 	void HandleModifierChange(NSUInteger NewModifierFlags, NSUInteger FlagsShift, NSUInteger UE4Shift, EMacModifierKeys TranslatedCode);
@@ -179,8 +178,7 @@ private:
 	/** The current set of Cocoa modifier flags, used to detect when Mission Control has been invoked & returned so that we can synthesis the modifier events it steals */
 	NSUInteger CurrentModifierFlags;
 
-	TArray< TSharedRef< FMacWindow > > KeyWindows;
-	bool bKeyWindowUpdateRequested;
+	TArray<FCocoaWindow*> WindowsToClose;
 
 	TSharedPtr<FMacTextInputMethodSystem> TextInputMethodSystem;
 
