- Changed FConfigCacheIni files to have a "type" that determines if Flush is called on them when they go out of scope (temporary objects were writing bad files to disk) [UE-8542]
- Temporarily (until 4.8) deleting the bad Engine/Config/NoRedist/BaseEngine.ini that was being created because of the above problem

diff --git a/Engine/Source/Editor/UnrealEd/Private/Commandlets/InternationalizationConditioningCommandlet.cpp b/Engine/Source/Editor/UnrealEd/Private/Commandlets/InternationalizationConditioningCommandlet.cpp
--- a/Engine/Source/Editor/UnrealEd/Private/Commandlets/InternationalizationConditioningCommandlet.cpp
+++ b/Engine/Source/Editor/UnrealEd/Private/Commandlets/InternationalizationConditioningCommandlet.cpp
@@ -207,6 +207,7 @@ bool UInternationalizationConditioningCommandlet::FLocalizationFilePair::HasFore
 
 UInternationalizationConditioningCommandlet::UInternationalizationConditioningCommandlet(const FObjectInitializer& ObjectInitializer)
 	: Super(ObjectInitializer)
+	, LegacyLocalizationCacheIni(EConfigCacheType::Temporary)
 {
 }
 
diff --git a/Engine/Source/Editor/UnrealEd/Private/Factories/PackFactory.cpp b/Engine/Source/Editor/UnrealEd/Private/Factories/PackFactory.cpp
--- a/Engine/Source/Editor/UnrealEd/Private/Factories/PackFactory.cpp
+++ b/Engine/Source/Editor/UnrealEd/Private/Factories/PackFactory.cpp
@@ -353,7 +353,7 @@ UObject* UPackFactory::FactoryCreateBinary
 						const FString LongOldGameName = FString::Printf(TEXT("/Script/%s"), **GameName);
 						const FString LongNewGameName = FString::Printf(TEXT("/Script/%s"), *ModuleInfo.ModuleName);
 						
-						FConfigCacheIni Config;
+						FConfigCacheIni Config(EConfigCacheType::Temporary);
 						FConfigFile& NewFile = Config.Add(EngineIniFilename, FConfigFile());
 						FConfigCacheIni::LoadLocalIniFile(NewFile, TEXT("DefaultEngine"), false);
 						FConfigSection* PackageRedirects = Config.GetSectionPrivate(*RedirectsSection, true, false, EngineIniFilename);
diff --git a/Engine/Source/Programs/UnrealBuildTool/Configuration/EngineConfiguration.cs b/Engine/Source/Programs/UnrealBuildTool/Configuration/EngineConfiguration.cs
--- a/Engine/Source/Programs/UnrealBuildTool/Configuration/EngineConfiguration.cs
+++ b/Engine/Source/Programs/UnrealBuildTool/Configuration/EngineConfiguration.cs
@@ -516,8 +516,9 @@ namespace UnrealBuildTool
 			// Engine/Config/NotForLicensees/Base* ini
 			yield return Path.Combine(EngineDirectory, "Config", "NotForLicensees", "Base" + BaseIniName + ".ini");
 
+			// NOTE: 4.7: See comment in GetSourceIniHierarchyFilenames()
 			// Engine/Config/NoRedist/Base* ini
-			yield return Path.Combine(EngineDirectory, "Config", "NoRedist", "Base" + BaseIniName + ".ini");
+			// yield return Path.Combine(EngineDirectory, "Config", "NoRedist", "Base" + BaseIniName + ".ini");
 
 			if(!String.IsNullOrEmpty(ProjectDirectory))
 			{
diff --git a/Engine/Source/Runtime/Core/Private/Misc/ConfigCacheIni.cpp b/Engine/Source/Runtime/Core/Private/Misc/ConfigCacheIni.cpp
--- a/Engine/Source/Runtime/Core/Private/Misc/ConfigCacheIni.cpp
+++ b/Engine/Source/Runtime/Core/Private/Misc/ConfigCacheIni.cpp
@@ -1250,7 +1250,7 @@ void FConfigFile::ProcessSourceAndCheckAgainstBackup()
 
 void FConfigFile::ProcessPropertyAndWriteForDefaults( const TArray< FString >& InCompletePropertyToProcess, FString& OutText, const FString& SectionName, const FString& PropertyName )
 {
-	FConfigCacheIni Hierarchy;
+	FConfigCacheIni Hierarchy(EConfigCacheType::Temporary);
 	const FString& LastFileInHierarchy = SourceIniHierarchy.Last().Filename;
 
 	// Build a config file out of this default configs hierarchy.
@@ -1292,9 +1292,10 @@ void FConfigFile::ProcessPropertyAndWriteForDefaults( const TArray< FString >& I
 	FConfigCacheIni
 -----------------------------------------------------------------------------*/
 
-FConfigCacheIni::FConfigCacheIni()
+FConfigCacheIni::FConfigCacheIni(EConfigCacheType InType)
 	: bAreFileOperationsDisabled(false)
 	, bIsReadyForUse(false)
+	, Type(InType)
 {
 }
 
@@ -1330,6 +1331,12 @@ FConfigFile* FConfigCacheIni::Find( const FString& Filename, bool CreateIfNotFou
 
 void FConfigCacheIni::Flush( bool Read, const FString& Filename )
 {
+	// never Flush temporary cache objects
+	if (Type == EConfigCacheType::Temporary)
+	{
+		return;
+	}
+
 	// write out the files if we can
 	if (!bAreFileOperationsDisabled)
 	{
@@ -2700,7 +2707,17 @@ static void GetSourceIniHierarchyFilenames(const TCHAR* InBaseIniName, const TCH
 	// Engine/Config/NotForLicensees/Base* ini
 	OutHierarchy.Add( FIniFilename(FString::Printf(TEXT("%sNotForLicensees/Base%s.ini"), EngineConfigDir, InBaseIniName), false) );
 	// Engine/Config/NoRedist/Base* ini
-	OutHierarchy.Add( FIniFilename(FString::Printf(TEXT("%sNoRedist/Base%s.ini"), EngineConfigDir, InBaseIniName), false) );
+	/////
+	// NOTE: 4.7
+	// There was a bug that was causing this file to be written out for all users and for all projects, with bad values
+	// that would break all projects (external reference errors)
+	// Since this file has not been used yet, we will delete it if it exists, and then for 4.8, we will put this back
+	// into the hierarchy and remove the delete operation
+	// - Remember to fixup EngineConfiguration.cs
+	/////
+	IFileManager::Get().Delete(*FString::Printf(TEXT("%sNoRedist/Base%s.ini"), EngineConfigDir, InBaseIniName), false, true, true);
+	IFileManager::Get().DeleteDirectory(*FString::Printf(TEXT("%sNoRedist"), EngineConfigDir), false, false);
+	// OutHierarchy.Add( FIniFilename(FString::Printf(TEXT("%sNoRedist/Base%s.ini"), EngineConfigDir, InBaseIniName), false) );
 
 	// [[[[ PROJECT SETTINGS ]]]]
 	// Game/Config/Default* ini
@@ -2716,11 +2733,14 @@ static void GetSourceIniHierarchyFilenames(const TCHAR* InBaseIniName, const TCH
 	FString HierarchyCheckpointPath = FString::Printf(TEXT("%sNoRedist/Default%s.ini"), SourceConfigDir, InBaseIniName);
 	OutHierarchy.Add(FIniFilename(HierarchyCheckpointPath, false, GenerateHierarchyCacheKey(OutHierarchy, HierarchyCheckpointPath)));
 	
-	// [[[[ PLATFORM DEFAULTS AND PROJECT SETTINGS ]]]]
-	// Engine/Config/Platform/Platform* ini
-	OutHierarchy.Add( FIniFilename(FString::Printf(TEXT("%s%s/%s%s.ini"), EngineConfigDir, *PlatformName, *PlatformName, InBaseIniName), false) );
-	// Game/Config/Platform/Platform* ini
-	OutHierarchy.Add( FIniFilename(FString::Printf(TEXT("%s%s/%s%s.ini"), SourceConfigDir, *PlatformName, *PlatformName, InBaseIniName), false) );
+	if (InPlatformName && InPlatformName[0])
+	{
+		// [[[[ PLATFORM DEFAULTS AND PROJECT SETTINGS ]]]]
+		// Engine/Config/Platform/Platform* ini
+		OutHierarchy.Add( FIniFilename(FString::Printf(TEXT("%s%s/%s%s.ini"), EngineConfigDir, *PlatformName, *PlatformName, InBaseIniName), false) );
+		// Game/Config/Platform/Platform* ini
+		OutHierarchy.Add( FIniFilename(FString::Printf(TEXT("%s%s/%s%s.ini"), SourceConfigDir, *PlatformName, *PlatformName, InBaseIniName), false) );
+	}
 
 	// [[[[ GLOBAL USER OVERRIDES ]]]]
 	// <AppData>/UE4/EngineConfig/User* ini
@@ -2769,7 +2789,7 @@ static FString GetDestIniFilename(const TCHAR* BaseIniName, const TCHAR* Platfor
 void FConfigCacheIni::InitializeConfigSystem()
 {
 	// create GConfig
-	GConfig = new FConfigCacheIni;
+	GConfig = new FConfigCacheIni(EConfigCacheType::DiskBacked);
 
 	// load the main .ini files (unless we're running a program or a gameless UE4Editor.exe, DefaultEngine.ini is required).
 	const bool bIsGamelessExe = !FApp::HasGameName();
diff --git a/Engine/Source/Runtime/Core/Public/Misc/ConfigCacheIni.h b/Engine/Source/Runtime/Core/Public/Misc/ConfigCacheIni.h
--- a/Engine/Source/Runtime/Core/Public/Misc/ConfigCacheIni.h
+++ b/Engine/Source/Runtime/Core/Public/Misc/ConfigCacheIni.h
@@ -52,6 +52,7 @@ struct FConfigCommandlineOverride
 
 
 // One config file.
+
 class FConfigFile : public TMap<FString,FConfigSection>
 {
 public:
@@ -157,12 +158,20 @@ private:
  */
 DECLARE_DELEGATE_TwoParams(FKeyValueSink, const TCHAR*, const TCHAR*);
 
+enum class EConfigCacheType : uint8
+{
+	// this type of config cache will write its files to disk during Flush (i.e. GConfig)
+	DiskBacked,
+	// this type of config cache is temporary and will never write to disk (only load from disk)
+	Temporary,
+};
+
 // Set of all cached config files.
 class CORE_API FConfigCacheIni : public TMap<FString,FConfigFile>
 {
 public:
 	// Basic functions.
-	FConfigCacheIni();
+	FConfigCacheIni(EConfigCacheType Type);
 	virtual ~FConfigCacheIni();
 
 	/**
@@ -550,6 +559,9 @@ private:
 
 	/** true after the base .ini files have been loaded, and GConfig is generally "ready for use" */
 	bool bIsReadyForUse;
+	
+	/** The type of the cache (basically, do we call Flush in the destructor) */
+	EConfigCacheType Type;
 };
 
 /**
diff --git a/Engine/Source/Runtime/CoreUObject/Private/UObject/Obj.cpp b/Engine/Source/Runtime/CoreUObject/Private/UObject/Obj.cpp
--- a/Engine/Source/Runtime/CoreUObject/Private/UObject/Obj.cpp
+++ b/Engine/Source/Runtime/CoreUObject/Private/UObject/Obj.cpp
@@ -1789,7 +1789,7 @@ FString UObject::GetGlobalUserConfigFilename() const
 void UObject::UpdateSingleSectionOfConfigFile(const FString& ConfigIniName)
 {
 	// create a sandbox FConfigCache
-	FConfigCacheIni Config;
+	FConfigCacheIni Config(EConfigCacheType::Temporary);
 
 	// add an empty file to the config so it doesn't read in the original file (see FConfigCacheIni.Find())
 	FConfigFile& NewFile = Config.Add(ConfigIniName, FConfigFile());
