Add write permission to files copied with FApplePlatformFile::CopyFile() as this is what engine expects in some places

diff --git a/Engine/Source/Runtime/Core/Private/Apple/ApplePlatformFile.cpp b/Engine/Source/Runtime/Core/Private/Apple/ApplePlatformFile.cpp
--- a/Engine/Source/Runtime/Core/Private/Apple/ApplePlatformFile.cpp
+++ b/Engine/Source/Runtime/Core/Private/Apple/ApplePlatformFile.cpp
@@ -536,6 +536,7 @@ bool FApplePlatformFile::CopyFile(const TCHAR* To, const TCHAR* From)
 		struct stat FileInfo;
 		if (Stat(From, &FileInfo) == 0)
 		{
+			FileInfo.st_mode |= S_IWUSR;
 			chmod(TCHAR_TO_UTF8(*NormalizeFilename(To)), FileInfo.st_mode);
 		}
 	}
