Fixed issue with config cache ini not respecting bForceRebuild on loadglobalIniFile.
Made the LoadLocalIniFile and LoadExternalIniFile not use the heirarchy cache.
Removed duplicate LoadConfigIniHierarchy calls.

diff --git a/Engine/Source/Runtime/Core/Private/Misc/ConfigCacheIni.cpp b/Engine/Source/Runtime/Core/Private/Misc/ConfigCacheIni.cpp
--- a/Engine/Source/Runtime/Core/Private/Misc/ConfigCacheIni.cpp
+++ b/Engine/Source/Runtime/Core/Private/Misc/ConfigCacheIni.cpp
@@ -19,7 +19,7 @@
 DEFINE_LOG_CATEGORY(LogConfig);
 namespace 
 {
-	FString GenerateHierarchyCacheKey(const TArray<FIniFilename>& IniHierarchy, const FString& IniPath)
+	FString GenerateHierarchyCacheKey(const TArray<FIniFilename>& IniHierarchy, const FString& IniPath, const FString& BaseIniName)
 	{
 #if !INI_CACHE
 		return TEXT("");
@@ -28,11 +28,13 @@ namespace
 		FString HierKey;
 		//
 		auto KeyLen = IniPath.Len();
+		KeyLen += BaseIniName.Len();
 		for (const auto& Ini : IniHierarchy)
 		{
 			KeyLen += Ini.Filename.Len();
 		}
 		HierKey.Reserve(KeyLen);
+		HierKey += BaseIniName;
 		for (const auto& Ini : IniHierarchy)
 		{
 			HierKey += Ini.Filename;
@@ -666,7 +668,7 @@ static void OverrideFromCommandline(FConfigFile* File, const FString& Filename)
  * @param ConfigFile - This is the FConfigFile which will have the contents of the .ini loaded into and Combined()
  *
  **/
-static bool LoadIniFileHierarchy(const TArray<FIniFilename>& HierarchyToLoad, FConfigFile& ConfigFile)
+static bool LoadIniFileHierarchy(const TArray<FIniFilename>& HierarchyToLoad, FConfigFile& ConfigFile, const bool bUseCache)
 {
 	// This shouldn't be getting called if seekfree is enabled on console.
 	check(!GUseSeekFreeLoading || !FPlatformProperties::RequiresCookedData());
@@ -699,12 +701,15 @@ static bool LoadIniFileHierarchy(const TArray<FIniFilename>& HierarchyToLoad, FC
 
 	int32 FirstCacheIndex = 0;
 #if INI_CACHE
-	// Find the last value in the hierarchy that is cached. We can start the load from there
-	for( int32 IniIndex = 0; IniIndex < HierarchyToLoad.Num(); IniIndex++ )
+	if ( bUseCache )
 	{
-		if (HierarchyCache.Find(HierarchyToLoad[IniIndex].CacheKey)) 
+		// Find the last value in the hierarchy that is cached. We can start the load from there
+		for( int32 IniIndex = 0; IniIndex < HierarchyToLoad.Num(); IniIndex++ )
 		{
-			FirstCacheIndex = IniIndex;
+			if (HierarchyCache.Find(HierarchyToLoad[IniIndex].CacheKey)) 
+			{
+				FirstCacheIndex = IniIndex;
+			}
 		}
 	}
 #endif
@@ -718,7 +723,9 @@ static bool LoadIniFileHierarchy(const TArray<FIniFilename>& HierarchyToLoad, FC
 		const FString& IniFileName = IniToLoad.Filename;
 		bool bDoProcess = true;
 #if INI_CACHE
-		if (IniToLoad.CacheKey.Len() > 0)
+		bool bShouldCache = IniToLoad.CacheKey.Len() > 0;
+		bShouldCache &= bUseCache;
+		if ( bShouldCache ) // if we are forcing a load don't mess with the cache
 		{
 			auto* CachedConfigFile = HierarchyCache.Find(IniToLoad.CacheKey);
 			if (CachedConfigFile) 
@@ -735,45 +742,41 @@ static bool LoadIniFileHierarchy(const TArray<FIniFilename>& HierarchyToLoad, FC
 #endif
 		if (bDoProcess)
 		{
-		// Spit out friendly error if there was a problem locating .inis (e.g. bad command line parameter or missing folder, ...).
-		if (IsUsingLocalIniFile(*IniFileName, NULL) && (IFileManager::Get().FileSize(*IniFileName) < 0))
-		{
-			if (IniToLoad.bRequired)
-			{
-				//UE_LOG(LogConfig, Error, TEXT("Couldn't locate '%s' which is required to run '%s'"), *IniToLoad.Filename, FApp::GetGameName() );
-				return false;
-			}
-			else
+			// Spit out friendly error if there was a problem locating .inis (e.g. bad command line parameter or missing folder, ...).
+			if (IsUsingLocalIniFile(*IniFileName, NULL) && (IFileManager::Get().FileSize(*IniFileName) < 0))
 			{
+				if (IniToLoad.bRequired)
+				{
+					//UE_LOG(LogConfig, Error, TEXT("Couldn't locate '%s' which is required to run '%s'"), *IniToLoad.Filename, FApp::GetGameName() );
+					return false;
+				}
+				else
+				{
 #if INI_CACHE
-					if (IniToLoad.CacheKey.Len() > 0)
+					// missing file just add the current config file to the cache
+					if ( bShouldCache )
 					{
-						if (bDoProcess)
-						{
 							HierarchyCache.Add(IniToLoad.CacheKey, ConfigFile);
-						}
 					}
 #endif
-				continue;
+					continue;
+				}
 			}
-		}
-
-		bool bDoEmptyConfig = false;
-		bool bDoCombine = (IniIndex != 0);
-		bool bDoWrite = false;
-		//UE_LOG(LogConfig, Log,  TEXT( "Combining configFile: %s" ), *IniList(IniIndex) );
-		ProcessIniContents(*HierarchyToLoad.Last().Filename, *IniFileName, &ConfigFile, bDoEmptyConfig, bDoCombine, bDoWrite);
-	}
 
+			bool bDoEmptyConfig = false;
+			bool bDoCombine = (IniIndex != 0);
+			bool bDoWrite = false;
+			//UE_LOG(LogConfig, Log,  TEXT( "Combining configFile: %s" ), *IniList(IniIndex) );
+			ProcessIniContents(*HierarchyToLoad.Last().Filename, *IniFileName, &ConfigFile, bDoEmptyConfig, bDoCombine, bDoWrite);
 #if INI_CACHE
-		if (IniToLoad.CacheKey.Len() > 0)
-		{
-			if (bDoProcess)
+			if ( bShouldCache )
 			{
 				HierarchyCache.Add(IniToLoad.CacheKey, ConfigFile);
 			}
-		}
 #endif
+		}
+
+
 	}
 
 	// Set this configs files source ini hierarchy to show where it was loaded from.
@@ -874,6 +877,25 @@ bool PropertySetFromCommandlineOption(const FConfigFile* InConfigFile, const FSt
 	return bFromCommandline;
 }
 
+/**
+ * Clear the hierarchy cache
+ * cos nobody want dat junk no more! bro
+ *
+ * @param Base ini name of the file hierarchy that we want to clear the cache for
+ */
+static void ClearHierarchyCache( const TCHAR* BaseIniName )
+{
+#if INI_CACHE
+	// if we are forcing reload from disk then clear the cached hierarchy files
+	for ( TMap<FString, FConfigFile>::TIterator It(HierarchyCache); It; ++It )
+	{
+		if ( It.Key().StartsWith( BaseIniName ) )
+		{
+			It.RemoveCurrent();
+		}
+	}
+#endif
+}
 
 bool FConfigFile::Write( const FString& Filename, bool bDoRemoteWrite/* = true*/, const FString& InitialText/*=FString()*/ )
 {
@@ -987,10 +1009,16 @@ bool FConfigFile::Write( const FString& Filename, bool bDoRemoteWrite/* = true*/
 
 #if INI_CACHE
 	// if we wrote the config successfully
-	if (bResult && CacheKey.Len() > 0)
+	if ( bResult && CacheKey.Len() > 0 )
 	{
-		HierarchyCache.Add(CacheKey, *this);
+		check( Name != NAME_None );
+		ClearHierarchyCache(*Name.ToString());
 	}
+	
+	/*if (bResult && CacheKey.Len() > 0)
+	{
+		HierarchyCache.Add(CacheKey, *this);
+	}*/
 #endif
 
 	// File is still dirty if it didn't save.
@@ -2531,10 +2559,10 @@ static void LoadAnIniFile(const FString& FilenameToLoad, FConfigFile& ConfigFile
  * @param DestIniFilename			The name of the destination .ini file - it's loaded and checked against source (e.g. Engine.ini)
  * @param SourceIniFilename			The name of the source .ini file (e.g. DefaultEngine.ini)
  */
-static bool GenerateDestIniFile(FConfigFile& DestConfigFile, const FString& DestIniFilename, const TArray<FIniFilename>& SourceIniHierarchy, bool bAllowGeneratedINIs)
+static bool GenerateDestIniFile(FConfigFile& DestConfigFile, const FString& DestIniFilename, const TArray<FIniFilename>& SourceIniHierarchy, bool bAllowGeneratedINIs, const bool bUseHierarchyCache)
 {
-	bool bResult = LoadIniFileHierarchy(SourceIniHierarchy, *DestConfigFile.SourceConfigFile);
-	if( !bResult )
+	bool bResult = LoadIniFileHierarchy(SourceIniHierarchy, *DestConfigFile.SourceConfigFile, bUseHierarchyCache);
+	if ( bResult == false )
 	{
 		return false;
 	}
@@ -2608,7 +2636,7 @@ static bool GenerateDestIniFile(FConfigFile& DestConfigFile, const FString& Dest
 	// Regenerate the file.
 	if( bForceRegenerate )
 	{
-		bResult = LoadIniFileHierarchy(SourceIniHierarchy, DestConfigFile);
+		bResult = LoadIniFileHierarchy(SourceIniHierarchy, DestConfigFile, bUseHierarchyCache);
 		DestConfigFile.SourceConfigFile = new FConfigFile( DestConfigFile );
 
 		// mark it as dirty (caller may want to save)
@@ -2731,7 +2759,7 @@ static void GetSourceIniHierarchyFilenames(const TCHAR* InBaseIniName, const TCH
 	OutHierarchy.Add(FIniFilename(FString::Printf(TEXT("%sNotForLicensees/Default%s.ini"), SourceConfigDir, InBaseIniName), false));
 	// Game/Config/NoRedist/Default* ini (Checkpointed here before the platform ini divergence)
 	FString HierarchyCheckpointPath = FString::Printf(TEXT("%sNoRedist/Default%s.ini"), SourceConfigDir, InBaseIniName);
-	OutHierarchy.Add(FIniFilename(HierarchyCheckpointPath, false, GenerateHierarchyCacheKey(OutHierarchy, HierarchyCheckpointPath)));
+	OutHierarchy.Add(FIniFilename(HierarchyCheckpointPath, false, GenerateHierarchyCacheKey(OutHierarchy, HierarchyCheckpointPath, InBaseIniName)));
 	
 	if (InPlatformName && InPlatformName[0])
 	{
@@ -2751,7 +2779,7 @@ static void GetSourceIniHierarchyFilenames(const TCHAR* InBaseIniName, const TCH
 	// [[[[ PROJECT USER OVERRIDES ]]]]
 	// Game/Config/User* ini (Checkpointed here at the end)
 	HierarchyCheckpointPath = FString::Printf(TEXT("%s/User%s.ini"), SourceConfigDir, InBaseIniName);
-	OutHierarchy.Add(FIniFilename(HierarchyCheckpointPath, false, GenerateHierarchyCacheKey(OutHierarchy, HierarchyCheckpointPath)));
+	OutHierarchy.Add(FIniFilename(HierarchyCheckpointPath, false, GenerateHierarchyCacheKey(OutHierarchy, HierarchyCheckpointPath, InBaseIniName)));
 	
 }
 
@@ -2878,15 +2906,20 @@ bool FConfigCacheIni::LoadGlobalIniFile(FString& FinalIniFilename, const TCHAR*
 	// make a new entry in GConfig (overwriting what's already there)
 	FConfigFile& NewConfigFile = GConfig->Add(FinalIniFilename, FConfigFile());
 
+	if ( bForceReload )
+	{
+		ClearHierarchyCache(BaseIniName);
+	}
+
+
 	// calculate the source ini file name,
 	GetSourceIniHierarchyFilenames( BaseIniName, Platform, GameName, *FPaths::EngineConfigDir(), *FPaths::SourceConfigDir(), NewConfigFile.SourceIniHierarchy, bRequireDefaultIni );
 
 	// Keep a record of the original settings
 	NewConfigFile.SourceConfigFile = new FConfigFile();
-	LoadIniFileHierarchy(NewConfigFile.SourceIniHierarchy, *NewConfigFile.SourceConfigFile );
 
 	// now generate and make sure it's up to date
-	bool bResult = GenerateDestIniFile(NewConfigFile, FinalIniFilename, NewConfigFile.SourceIniHierarchy, bAllowGeneratedIniWhenCooked);
+	bool bResult = GenerateDestIniFile(NewConfigFile, FinalIniFilename, NewConfigFile.SourceIniHierarchy, bAllowGeneratedIniWhenCooked, true);
 	NewConfigFile.Name = BaseIniName;
 
 	// don't write anything to disk in cooked builds - we will always use re-generated INI files anyway.
@@ -2908,10 +2941,14 @@ bool FConfigCacheIni::LoadGlobalIniFile(FString& FinalIniFilename, const TCHAR*
 	return bResult;
 }
 
-void FConfigCacheIni::LoadLocalIniFile(FConfigFile& ConfigFile, const TCHAR* IniName, bool bGenerateDestIni, const TCHAR* Platform, const TCHAR* GameName)
+void FConfigCacheIni::LoadLocalIniFile(FConfigFile& ConfigFile, const TCHAR* IniName, bool bGenerateDestIni, const TCHAR* Platform, const TCHAR* GameName, const bool bForceReload )
 {
+
+	LoadExternalIniFile( ConfigFile, IniName, *FPaths::EngineConfigDir(), *FPaths::SourceConfigDir(), bGenerateDestIni, Platform, GameName, bForceReload );
+
+
 	// if bGenerateDestIni is false, that means the .ini is a ready-to-go .ini file, and just needs to be loaded into the FConfigFile
-	if (!bGenerateDestIni)
+	/*if (!bGenerateDestIni)
 	{
 		// generate path to the .ini file (not a Default ini, IniName is the complete name of the file, without path)
 		FString SourceIniFilename = FString::Printf(TEXT("%s%s.ini"), *FPaths::SourceConfigDir(), IniName);
@@ -2925,16 +2962,15 @@ void FConfigCacheIni::LoadLocalIniFile(FConfigFile& ConfigFile, const TCHAR* Ini
 		
 		// Keep a record of the original settings
 		ConfigFile.SourceConfigFile = new FConfigFile();
-		LoadIniFileHierarchy(ConfigFile.SourceIniHierarchy, *ConfigFile.SourceConfigFile );
 
 		// now generate and make sure it's up to date (using IniName as a Base for an ini filename)
 		const bool bAllowGeneratedINIs = true;
-		GenerateDestIniFile(ConfigFile, GetDestIniFilename(IniName, Platform, GameName, *FPaths::GeneratedConfigDir()), ConfigFile.SourceIniHierarchy, bAllowGeneratedINIs);
+		GenerateDestIniFile(ConfigFile, GetDestIniFilename(IniName, Platform, GameName, *FPaths::GeneratedConfigDir()), ConfigFile.SourceIniHierarchy, bAllowGeneratedINIs, true);
 	}
-	ConfigFile.Name = IniName;
+	ConfigFile.Name = IniName;*/
 }
 
-void FConfigCacheIni::LoadExternalIniFile(FConfigFile& ConfigFile, const TCHAR* IniName, const TCHAR* EngineConfigDir, const TCHAR* SourceConfigDir, bool bGenerateDestIni, const TCHAR* Platform, const TCHAR* GameName)
+void FConfigCacheIni::LoadExternalIniFile(FConfigFile& ConfigFile, const TCHAR* IniName, const TCHAR* EngineConfigDir, const TCHAR* SourceConfigDir, bool bGenerateDestIni, const TCHAR* Platform, const TCHAR* GameName, const bool bForceReload)
 {
 	// if bGenerateDestIni is false, that means the .ini is a ready-to-go .ini file, and just needs to be loaded into the FConfigFile
 	if (!bGenerateDestIni)
@@ -2949,13 +2985,18 @@ void FConfigCacheIni::LoadExternalIniFile(FConfigFile& ConfigFile, const TCHAR*
 	{
 		GetSourceIniHierarchyFilenames( IniName, Platform, GameName, EngineConfigDir, SourceConfigDir, ConfigFile.SourceIniHierarchy, false );
 
+		if ( bForceReload )
+		{
+			ClearHierarchyCache( IniName );
+		}
+
+
 		// Keep a record of the original settings
 		ConfigFile.SourceConfigFile = new FConfigFile();
-		LoadIniFileHierarchy(ConfigFile.SourceIniHierarchy, *ConfigFile.SourceConfigFile );
 
 		// now generate and make sure it's up to date (using IniName as a Base for an ini filename)
 		const bool bAllowGeneratedINIs = true;
-		GenerateDestIniFile(ConfigFile, GetDestIniFilename(IniName, Platform, GameName, *FPaths::GeneratedConfigDir()), ConfigFile.SourceIniHierarchy, bAllowGeneratedINIs);
+		GenerateDestIniFile(ConfigFile, GetDestIniFilename(IniName, Platform, GameName, *FPaths::GeneratedConfigDir()), ConfigFile.SourceIniHierarchy, bAllowGeneratedINIs, true);
 	}
 	ConfigFile.Name = IniName;
 }
@@ -3075,11 +3116,13 @@ void FConfigFile::UpdateSections(const TCHAR* DiskFilename, const TCHAR* IniRoot
 			}
 		}
 
+		ClearHierarchyCache(IniRootName);
+
 		// Get a collection of the source hierachy properties
 		SourceConfigFile = new FConfigFile();
-		
+
 		// now when Write it called below, it will diff against SourceIniHierarchy
-		LoadIniFileHierarchy(SourceIniHierarchy, *SourceConfigFile);
+		LoadIniFileHierarchy(SourceIniHierarchy, *SourceConfigFile, true );
 
 	}
 
diff --git a/Engine/Source/Runtime/Core/Public/Misc/ConfigCacheIni.h b/Engine/Source/Runtime/Core/Public/Misc/ConfigCacheIni.h
--- a/Engine/Source/Runtime/Core/Public/Misc/ConfigCacheIni.h
+++ b/Engine/Source/Runtime/Core/Public/Misc/ConfigCacheIni.h
@@ -528,8 +528,9 @@ public:
 	 * @param bIsBaseIniName true if IniName is a Base name, which can be overridden on commandline, etc.
 	 * @param Platform The platform to use for Base ini names, NULL means to use the current platform
 	 * @param GameName The game to use for Base ini names, NULL means to use the current
+	 * @param bForceReload force reload the ini file from disk this is required if you make changes to the ini file not using the config system as the hierarchy cache will not be updated in this case
 	 */
-	static void LoadLocalIniFile(FConfigFile& ConfigFile, const TCHAR* IniName, bool bIsBaseIniName, const TCHAR* Platform=NULL, const TCHAR* GameName=NULL);
+	static void LoadLocalIniFile(FConfigFile& ConfigFile, const TCHAR* IniName, bool bIsBaseIniName, const TCHAR* Platform=NULL, const TCHAR* GameName=NULL, const bool bForceReload=false);
 
 	/**
 	 * Load an ini file directly into an FConfigFile from the specified config folders, and nothing is written to GConfig or disk. 
@@ -543,8 +544,9 @@ public:
 	 * @param bGenerateDestIni true if IniName is a Base name, which can be overridden on commandline, etc.
 	 * @param Platform The platform to use for Base ini names
 	 * @param GameName The game to use for Base ini names
+	 * @param bForceReload force reload the ini file from disk this is required if you make changes to the ini file not using the config system as the hierarchy cache will not be updated in this case
 	 */
-	static void LoadExternalIniFile(FConfigFile& ConfigFile, const TCHAR* IniName, const TCHAR* EngineConfigDir, const TCHAR* SourceConfigDir, bool bGenerateDestIni, const TCHAR* Platform=NULL, const TCHAR* GameName=NULL);
+	static void LoadExternalIniFile(FConfigFile& ConfigFile, const TCHAR* IniName, const TCHAR* EngineConfigDir, const TCHAR* SourceConfigDir, bool bGenerateDestIni, const TCHAR* Platform=NULL, const TCHAR* GameName=NULL, const bool bForceReload=false);
 
 	/**
 	 * Needs to be called after GConfig is set and LoadCoalescedFile was called.
