Fix for UT rendering just a portion of screen (UE-5169).

diff --git a/Engine/Source/Runtime/Core/Public/Linux/LinuxPlatformProperties.h b/Engine/Source/Runtime/Core/Public/Linux/LinuxPlatformProperties.h
--- a/Engine/Source/Runtime/Core/Public/Linux/LinuxPlatformProperties.h
+++ b/Engine/Source/Runtime/Core/Public/Linux/LinuxPlatformProperties.h
@@ -91,7 +91,7 @@ struct FLinuxPlatformProperties
 
 	static FORCEINLINE bool HasFixedResolution()
 	{
-		return false;
+		return true; // Fullscreen on Linux is always WindowedFullscreen (to avoid badness happening when switching to various resolutions)
 	}
 
 	static FORCEINLINE bool SupportsWindowedMode()
