Core - Fix for useless deprecated warning

diff --git a/Engine/Source/Runtime/Core/Private/Misc/EventPool.h b/Engine/Source/Runtime/Core/Private/Misc/EventPool.h
--- a/Engine/Source/Runtime/Core/Private/Misc/EventPool.h
+++ b/Engine/Source/Runtime/Core/Private/Misc/EventPool.h
@@ -61,7 +61,10 @@ public:
 
 		if (!Result)
 		{
+			// FEventPool is allowed to create synchronization events.
+			PRAGMA_DISABLE_DEPRECATION_WARNINGS
 			Result = FPlatformProcess::CreateSynchEvent((PoolType == EEventPoolTypes::ManualReset));
+			PRAGMA_ENABLE_DEPRECATION_WARNINGS
 		}
 
 		check(Result);
