Core - Fixed a typo in the comment

diff --git a/Engine/Source/Runtime/Core/Public/GenericPlatform/GenericPlatformProcess.h b/Engine/Source/Runtime/Core/Public/GenericPlatform/GenericPlatformProcess.h
--- a/Engine/Source/Runtime/Core/Public/GenericPlatform/GenericPlatformProcess.h
+++ b/Engine/Source/Runtime/Core/Public/GenericPlatform/GenericPlatformProcess.h
@@ -436,7 +436,7 @@ struct CORE_API FGenericPlatformProcess
 	 * @return A new event, or nullptr none could be created.
 	 * @see GetSynchEventFromPool, ReturnSynchEventToPool
 	 */
-	DEPRECATED(4.8, "Please use GetSynchEventFromPool to create a new event, and GetSynchEventFromPool to release the event.")
+	DEPRECATED(4.8, "Please use GetSynchEventFromPool to create a new event, and ReturnSynchEventToPool to release the event.")
 	static class FEvent* CreateSynchEvent(bool bIsManualReset = false);
 
 	/**
