Make check/ensure trigger debug break in current stack frame, not in AssertFailed internals.

diff --git a/Engine/Source/Runtime/Core/Public/HTML5/HTML5PlatformMisc.h b/Engine/Source/Runtime/Core/Public/HTML5/HTML5PlatformMisc.h
--- a/Engine/Source/Runtime/Core/Public/HTML5/HTML5PlatformMisc.h
+++ b/Engine/Source/Runtime/Core/Public/HTML5/HTML5PlatformMisc.h
@@ -53,6 +53,13 @@ struct CORE_API FHTML5Misc : public FGenericPlatformMisc
 		}
 	}
 
+	/** Break into debugger. Returning false allows this function to be used in conditionals. */
+	FORCEINLINE static bool DebugBreakReturningFalse()
+	{
+		DebugBreak();
+		return false;
+	}
+
 	FORCEINLINE static void LocalPrint( const TCHAR* Str )
 	{
 		wprintf(TEXT("%ls"), Str);
diff --git a/Engine/Source/Runtime/Core/Public/IOS/IOSPlatformMisc.h b/Engine/Source/Runtime/Core/Public/IOS/IOSPlatformMisc.h
--- a/Engine/Source/Runtime/Core/Public/IOS/IOSPlatformMisc.h
+++ b/Engine/Source/Runtime/Core/Public/IOS/IOSPlatformMisc.h
@@ -52,6 +52,13 @@ struct CORE_API FIOSPlatformMisc : public FGenericPlatformMisc
 	}
 #endif
 
+	/** Break into debugger. Returning false allows this function to be used in conditionals. */
+	FORCEINLINE static bool DebugBreakReturningFalse()
+	{
+		DebugBreak();
+		return false;
+	}
+
 	FORCEINLINE static void MemoryBarrier()
 	{
 		OSMemoryBarrier();
diff --git a/Engine/Source/Runtime/Core/Public/Linux/LinuxPlatformMisc.h b/Engine/Source/Runtime/Core/Public/Linux/LinuxPlatformMisc.h
--- a/Engine/Source/Runtime/Core/Public/Linux/LinuxPlatformMisc.h
+++ b/Engine/Source/Runtime/Core/Public/Linux/LinuxPlatformMisc.h
@@ -32,6 +32,13 @@ struct CORE_API FLinuxPlatformMisc : public FGenericPlatformMisc
 			raise(SIGTRAP);
 		}
 	}
+
+	/** Break into debugger. Returning false allows this function to be used in conditionals. */
+	FORCEINLINE static bool DebugBreakReturningFalse()
+	{
+		DebugBreak();
+		return false;
+	}
 #endif // !UE_BUILD_SHIPPING
 
 	static void PumpMessages(bool bFromMainLoop);
diff --git a/Engine/Source/Runtime/Core/Public/Mac/MacPlatformMisc.h b/Engine/Source/Runtime/Core/Public/Mac/MacPlatformMisc.h
--- a/Engine/Source/Runtime/Core/Public/Mac/MacPlatformMisc.h
+++ b/Engine/Source/Runtime/Core/Public/Mac/MacPlatformMisc.h
@@ -50,6 +50,13 @@ struct CORE_API FMacPlatformMisc : public FGenericPlatformMisc
 			__asm__ ( "int $3" );
 		}
 	}
+
+	/** Break into debugger. Returning false allows this function to be used in conditionals. */
+	FORCEINLINE static bool DebugBreakReturningFalse()
+	{
+		DebugBreak();
+		return false;
+	}
 #endif
 
 	FORCEINLINE static void MemoryBarrier()
diff --git a/Engine/Source/Runtime/Core/Public/Misc/AssertionMacros.h b/Engine/Source/Runtime/Core/Public/Misc/AssertionMacros.h
--- a/Engine/Source/Runtime/Core/Public/Misc/AssertionMacros.h
+++ b/Engine/Source/Runtime/Core/Public/Misc/AssertionMacros.h
@@ -17,19 +17,19 @@
 
 #if DO_CHECK
 	#define checkCode( Code )		do { Code } while ( false );
-	#define verify(expr)			{ if(!(expr)) FDebug::AssertFailed( #expr, __FILE__, __LINE__ ); CA_ASSUME(expr); }
-	#define check(expr)				{ if(!(expr)) FDebug::AssertFailed( #expr, __FILE__, __LINE__ ); CA_ASSUME(expr); }
+	#define verify(expr)			{ if(!(expr)) { FPlatformMisc::DebugBreak(); FDebug::AssertFailed( #expr, __FILE__, __LINE__ ); CA_ASSUME(expr); } }
+	#define check(expr)				{ if(!(expr)) { FPlatformMisc::DebugBreak(); FDebug::AssertFailed( #expr, __FILE__, __LINE__ ); CA_ASSUME(expr); } }
 	
 	/**
 	 * verifyf, checkf: Same as verify, check but with printf style additional parameters
 	 * Read about __VA_ARGS__ (variadic macros) on http://gcc.gnu.org/onlinedocs/gcc-3.4.4/cpp.pdf.
 	 */
-	#define verifyf(expr, ...)				{ if(!(expr)) FDebug::AssertFailed( #expr, __FILE__, __LINE__, ##__VA_ARGS__ ); CA_ASSUME(expr); }
-	#define checkf(expr, ...)				{ if(!(expr)) FDebug::AssertFailed( #expr, __FILE__, __LINE__, ##__VA_ARGS__ ); CA_ASSUME(expr); }
+	#define verifyf(expr, ...)				{ if(!(expr)) { FPlatformMisc::DebugBreak(); FDebug::AssertFailed( #expr, __FILE__, __LINE__, ##__VA_ARGS__ ); CA_ASSUME(expr); } }
+	#define checkf(expr, ...)				{ if(!(expr)) { FPlatformMisc::DebugBreak(); FDebug::AssertFailed( #expr, __FILE__, __LINE__, ##__VA_ARGS__ ); CA_ASSUME(expr); } }
 	/**
 	 * Denotes code paths that should never be reached.
 	 */
-	#define checkNoEntry()       { FDebug::AssertFailed( "Enclosing block should never be called", __FILE__, __LINE__ ); }
+	#define checkNoEntry()       { FPlatformMisc::DebugBreak(); FDebug::AssertFailed( "Enclosing block should never be called", __FILE__, __LINE__ ); }
 
 	/**
 	 * Denotes code paths that should not be executed more than once.
@@ -54,7 +54,7 @@
 	                            checkf( RecursionCounter##__LINE__ == 0, TEXT("Enclosing block was entered recursively") );  \
 	                            const FRecursionScopeMarker ScopeMarker##__LINE__( RecursionCounter##__LINE__ )
 
-	#define unimplemented()       { FDebug::AssertFailed( "Unimplemented function called", __FILE__, __LINE__ ); }
+	#define unimplemented()       { FPlatformMisc::DebugBreak(); FDebug::AssertFailed( "Unimplemented function called", __FILE__, __LINE__ ); }
 
 #else
 	#define checkCode(...)
@@ -72,9 +72,9 @@
 // Check for development only.
 //
 #if DO_GUARD_SLOW
-	#define checkSlow(expr, ...)   {if(!(expr)) FDebug::AssertFailed( #expr, __FILE__, __LINE__ ); CA_ASSUME(expr); }
-	#define checkfSlow(expr, ...)	{ if(!(expr)) FDebug::AssertFailed( #expr, __FILE__, __LINE__, __VA_ARGS__ ); CA_ASSUME(expr); }
-	#define verifySlow(expr)  {if(!(expr)) FDebug::AssertFailed( #expr, __FILE__, __LINE__ );}
+	#define checkSlow(expr, ...)	{ if(!(expr)) { FPlatformMisc::DebugBreak(); FDebug::AssertFailed( #expr, __FILE__, __LINE__ ); CA_ASSUME(expr); } }
+	#define checkfSlow(expr, ...)	{ if(!(expr)) { FPlatformMisc::DebugBreak(); FDebug::AssertFailed( #expr, __FILE__, __LINE__, __VA_ARGS__ ); CA_ASSUME(expr); } }
+	#define verifySlow(expr)		{ if(!(expr)) { FPlatformMisc::DebugBreak(); FDebug::AssertFailed( #expr, __FILE__, __LINE__ );} }
 #else
 	#define checkSlow(...)
 	#define checkfSlow(...)
@@ -113,7 +113,7 @@
 		(((InExpression) == 0) ? FDebug::EnsureNotFalse(false, #InExpression, __FILE__, __LINE__, InMsg ) : true)
 
 	#define ensureMsgf( InExpression, ... ) \
-		(((InExpression) == 0) ? FDebug::EnsureNotFalseFormatted(false, #InExpression, __FILE__, __LINE__, ##__VA_ARGS__ ) : true)
+		(((InExpression) == 0) ? (FPlatformMisc::DebugBreakReturningFalse() || FDebug::EnsureNotFalseFormatted(false, #InExpression, __FILE__, __LINE__, ##__VA_ARGS__ )) : true)
 
 #else	// DO_CHECK
 
@@ -178,6 +178,7 @@ struct FTCharArrayTester
 #define LowLevelFatalError(Format, ...) \
 	{ \
 		static_assert(IS_TCHAR_ARRAY(Format), "Formatting string must be a TCHAR array."); \
+		FPlatformMisc::DebugBreak(); \
 		FError::LowLevelFatal(__FILE__, __LINE__, Format, ##__VA_ARGS__); \
 	}
 
@@ -204,6 +205,7 @@ struct FTCharArrayTester
 		static_assert(IS_TCHAR_ARRAY(Format), "Formatting string must be a TCHAR array."); \
 		if (ELogVerbosity::Verbosity == ELogVerbosity::Fatal) \
 		{ \
+			FPlatformMisc::DebugBreak(); \
 			FError::LowLevelFatal(__FILE__, __LINE__, Format, ##__VA_ARGS__); \
 		} \
 	}
@@ -215,6 +217,7 @@ struct FTCharArrayTester
 		{ \
 			if (Condition) \
 			{ \
+				FPlatformMisc::DebugBreak(); \
 				FError::LowLevelFatal(__FILE__, __LINE__, Format, ##__VA_ARGS__); \
 			} \
 		} \
@@ -274,6 +277,10 @@ struct FTCharArrayTester
 			{ \
 				if (!CategoryName.IsSuppressed(ELogVerbosity::Verbosity)) \
 				{ \
+					if (ELogVerbosity::Verbosity == ELogVerbosity::Fatal) \
+					{\
+						FPlatformMisc::DebugBreak();\
+					}\
 					FMsg::Logf(__FILE__, __LINE__, CategoryName.GetCategoryName(), ELogVerbosity::Verbosity, Format, ##__VA_ARGS__); \
 				} \
 			} \
@@ -289,6 +296,10 @@ struct FTCharArrayTester
 				{ \
 					if (Condition) \
 					{ \
+						if (ELogVerbosity::Verbosity == ELogVerbosity::Fatal) \
+						{\
+							FPlatformMisc::DebugBreak();\
+						}\
 						FMsg::Logf(__FILE__, __LINE__, CategoryName.GetCategoryName(), ELogVerbosity::Verbosity, Format, ##__VA_ARGS__); \
 					} \
 				} \
diff --git a/Engine/Source/Runtime/Core/Public/Misc/OutputDevice.h b/Engine/Source/Runtime/Core/Public/Misc/OutputDevice.h
--- a/Engine/Source/Runtime/Core/Public/Misc/OutputDevice.h
+++ b/Engine/Source/Runtime/Core/Public/Misc/OutputDevice.h
@@ -417,6 +417,7 @@ struct CORE_API FDebug
 	{
 		if( bExpressionResult == 0 )
 		{
+			FPlatformMisc::DebugBreak();
 			EnsureFailed( Expr, File, Line, Msg );
 		}
 	
diff --git a/Engine/Source/Runtime/Core/Public/Windows/WindowsPlatformMisc.h b/Engine/Source/Runtime/Core/Public/Windows/WindowsPlatformMisc.h
--- a/Engine/Source/Runtime/Core/Public/Windows/WindowsPlatformMisc.h
+++ b/Engine/Source/Runtime/Core/Public/Windows/WindowsPlatformMisc.h
@@ -49,6 +49,14 @@ struct CORE_API FWindowsPlatformMisc
 			::DebugBreak();
 		}
 	}
+
+	/** Break into debugger. Returning false allows this function to be used in conditionals. */
+	FORCEINLINE static bool DebugBreakReturningFalse()
+	{
+		DebugBreak();
+		return false;
+	}
+
 #endif
 
 	static void PumpMessages(bool bFromMainLoop);
