CIS fix.

diff --git a/Engine/Source/Runtime/Core/Public/IOS/IOSPlatformMisc.h b/Engine/Source/Runtime/Core/Public/IOS/IOSPlatformMisc.h
--- a/Engine/Source/Runtime/Core/Public/IOS/IOSPlatformMisc.h
+++ b/Engine/Source/Runtime/Core/Public/IOS/IOSPlatformMisc.h
@@ -55,7 +55,9 @@ struct CORE_API FIOSPlatformMisc : public FGenericPlatformMisc
 	/** Break into debugger. Returning false allows this function to be used in conditionals. */
 	FORCEINLINE static bool DebugBreakReturningFalse()
 	{
+#if !UE_BUILD_SHIPPING
 		DebugBreak();
+#endif
 		return false;
 	}
 
diff --git a/Engine/Source/Runtime/Core/Public/Linux/LinuxPlatformMisc.h b/Engine/Source/Runtime/Core/Public/Linux/LinuxPlatformMisc.h
--- a/Engine/Source/Runtime/Core/Public/Linux/LinuxPlatformMisc.h
+++ b/Engine/Source/Runtime/Core/Public/Linux/LinuxPlatformMisc.h
@@ -32,14 +32,16 @@ struct CORE_API FLinuxPlatformMisc : public FGenericPlatformMisc
 			raise(SIGTRAP);
 		}
 	}
+#endif // !UE_BUILD_SHIPPING
 
 	/** Break into debugger. Returning false allows this function to be used in conditionals. */
 	FORCEINLINE static bool DebugBreakReturningFalse()
 	{
+#if !UE_BUILD_SHIPPING
 		DebugBreak();
+#endif
 		return false;
 	}
-#endif // !UE_BUILD_SHIPPING
 
 	static void PumpMessages(bool bFromMainLoop);
 	static uint32 GetKeyMap( uint16* KeyCodes, FString* KeyNames, uint32 MaxMappings );
diff --git a/Engine/Source/Runtime/Core/Public/Mac/MacPlatformMisc.h b/Engine/Source/Runtime/Core/Public/Mac/MacPlatformMisc.h
--- a/Engine/Source/Runtime/Core/Public/Mac/MacPlatformMisc.h
+++ b/Engine/Source/Runtime/Core/Public/Mac/MacPlatformMisc.h
@@ -50,14 +50,16 @@ struct CORE_API FMacPlatformMisc : public FGenericPlatformMisc
 			__asm__ ( "int $3" );
 		}
 	}
+#endif
 
 	/** Break into debugger. Returning false allows this function to be used in conditionals. */
 	FORCEINLINE static bool DebugBreakReturningFalse()
 	{
+#if !UE_BUILD_SHIPPING
 		DebugBreak();
+#endif
 		return false;
 	}
-#endif
 
 	FORCEINLINE static void MemoryBarrier()
 	{
diff --git a/Engine/Source/Runtime/Core/Public/Windows/WindowsPlatformMisc.h b/Engine/Source/Runtime/Core/Public/Windows/WindowsPlatformMisc.h
--- a/Engine/Source/Runtime/Core/Public/Windows/WindowsPlatformMisc.h
+++ b/Engine/Source/Runtime/Core/Public/Windows/WindowsPlatformMisc.h
@@ -49,16 +49,18 @@ struct CORE_API FWindowsPlatformMisc
 			::DebugBreak();
 		}
 	}
+#endif
+
 
 	/** Break into debugger. Returning false allows this function to be used in conditionals. */
 	FORCEINLINE static bool DebugBreakReturningFalse()
 	{
+#if !UE_BUILD_SHIPPING
 		DebugBreak();
+#endif
 		return false;
 	}
 
-#endif
-
 	static void PumpMessages(bool bFromMainLoop);
 	static uint32 GetKeyMap( uint16* KeyCodes, FString* KeyNames, uint32 MaxMappings );
 	static uint32 GetCharKeyMap(uint16* KeyCodes, FString* KeyNames, uint32 MaxMappings);
