CIS fix.

diff --git a/Engine/Source/Runtime/Core/Public/Android/AndroidMisc.h b/Engine/Source/Runtime/Core/Public/Android/AndroidMisc.h
--- a/Engine/Source/Runtime/Core/Public/Android/AndroidMisc.h
+++ b/Engine/Source/Runtime/Core/Public/Android/AndroidMisc.h
@@ -72,6 +72,15 @@ struct CORE_API FAndroidMisc : public FGenericPlatformMisc
 	}
 #endif
 
+	/** Break into debugger. Returning false allows this function to be used in conditionals. */
+	FORCEINLINE static bool DebugBreakReturningFalse()
+	{
+#if !UE_BUILD_SHIPPING
+		DebugBreak();
+#endif
+		return false;
+	}
+
 	FORCEINLINE static void MemoryBarrier()
 	{
 		__sync_synchronize();
