Fixed new crash at exit on Mac

diff --git a/Engine/Source/Runtime/Core/Private/Mac/MacWindow.cpp b/Engine/Source/Runtime/Core/Private/Mac/MacWindow.cpp
--- a/Engine/Source/Runtime/Core/Private/Mac/MacWindow.cpp
+++ b/Engine/Source/Runtime/Core/Private/Mac/MacWindow.cpp
@@ -312,6 +312,7 @@ void FMacWindow::Destroy()
 		[WindowHandle setAlphaValue:0.0f];
 		[WindowHandle setBackgroundColor:[NSColor clearColor]];
 		MacApplication->OnWindowDestroyed(WindowHandle);
+		WindowHandle = nullptr;
 	}
 }
 
