[INTEGRATE] Change 2420313 by [REDACTED_PERSON][REDACTED_HANDLE] on 2015/01/27 09:53:31

	UE-8228 Assert in ConstructObject occurred after hitting Play in viewport in tracetest

	Pending-Kill objects are not duplicated.

diff --git a/Engine/Source/Runtime/Core/Public/UObject/ScriptDelegates.h b/Engine/Source/Runtime/Core/Public/UObject/ScriptDelegates.h
--- a/Engine/Source/Runtime/Core/Public/UObject/ScriptDelegates.h
+++ b/Engine/Source/Runtime/Core/Public/UObject/ScriptDelegates.h
@@ -349,9 +349,9 @@ public:
 	/** Multi-cast delegate serialization */
 	friend FArchive& operator<<( FArchive& Ar, TMulticastScriptDelegate<TWeakPtr>& D )
 	{
-		if( Ar.IsSaving() )
+		if (Ar.IsSaving() || Ar.HasAnyPortFlags(PPF_Duplicate))
 		{
-			// When saving the delegate, clean up the list to make sure there are no bad object references
+			// When saving or duplicating the delegate, clean up the list to make sure there are no bad object references
 			D.CompactInvocationList();
 		}
 
diff --git a/Engine/Source/Runtime/CoreUObject/Private/Serialization/DuplicateDataWriter.cpp b/Engine/Source/Runtime/CoreUObject/Private/Serialization/DuplicateDataWriter.cpp
--- a/Engine/Source/Runtime/CoreUObject/Private/Serialization/DuplicateDataWriter.cpp
+++ b/Engine/Source/Runtime/CoreUObject/Private/Serialization/DuplicateDataWriter.cpp
@@ -117,7 +117,7 @@ UObject* FDuplicateDataWriter::AddDuplicate(UObject* SourceObject,UObject* DupOb
 UObject* FDuplicateDataWriter::GetDuplicatedObject(UObject* Object, bool bCreateIfMissing)
 {
 	UObject* Result = NULL;
-	if( Object != NULL )
+	if (IsValid(Object))
 	{
 		// Check for an existing duplicate of the object.
 		FDuplicatedObject DupObjectInfo = DuplicatedObjectAnnotation.GetAnnotation( Object );
