[INTEGRATE] Change 2420515 by [REDACTED_PERSON][REDACTED_HANDLE] on 2015/01/27 11:46:05

	CIS Fix - Back out revision 2 from //depot/UE4-Releases/4.7/Engine/Source/Runtime/Core/Public/UObject/ScriptDelegates.h

diff --git a/Engine/Source/Runtime/Core/Public/UObject/ScriptDelegates.h b/Engine/Source/Runtime/Core/Public/UObject/ScriptDelegates.h
--- a/Engine/Source/Runtime/Core/Public/UObject/ScriptDelegates.h
+++ b/Engine/Source/Runtime/Core/Public/UObject/ScriptDelegates.h
@@ -349,9 +349,9 @@ public:
 	/** Multi-cast delegate serialization */
 	friend FArchive& operator<<( FArchive& Ar, TMulticastScriptDelegate<TWeakPtr>& D )
 	{
-		if (Ar.IsSaving() || Ar.HasAnyPortFlags(PPF_Duplicate))
+		if( Ar.IsSaving() )
 		{
-			// When saving or duplicating the delegate, clean up the list to make sure there are no bad object references
+			// When saving the delegate, clean up the list to make sure there are no bad object references
 			D.CompactInvocationList();
 		}
 
