Change many string replace/compare/find operations which use a fixed case search pattern to be case sensitive. This avoids unnecessary expensive string case operations when searching for things like dots and slashes.

diff --git a/Engine/Source/Developer/PackageDependencyInfo/Private/PackageDependencyInfo.cpp b/Engine/Source/Developer/PackageDependencyInfo/Private/PackageDependencyInfo.cpp
--- a/Engine/Source/Developer/PackageDependencyInfo/Private/PackageDependencyInfo.cpp
+++ b/Engine/Source/Developer/PackageDependencyInfo/Private/PackageDependencyInfo.cpp
@@ -45,7 +45,7 @@ public:
 		for (int32 DirIndex = 0; DirIndex < InDirectoriesWildcards.Num(); DirIndex++)
 		{
 			FString Dir = InDirectoriesWildcards[DirIndex];
-			Dir.ReplaceInline(TEXT("\\"), TEXT("/"));
+			Dir.ReplaceInline(TEXT("\\"), TEXT("/"), ESearchCase::CaseSensitive);
 			DirectoriesWildcards.Add(Dir);
 		}
 
diff --git a/Engine/Source/Developer/TargetPlatform/Public/Interfaces/TargetDeviceId.h b/Engine/Source/Developer/TargetPlatform/Public/Interfaces/TargetDeviceId.h
--- a/Engine/Source/Developer/TargetPlatform/Public/Interfaces/TargetDeviceId.h
+++ b/Engine/Source/Developer/TargetPlatform/Public/Interfaces/TargetDeviceId.h
@@ -126,7 +126,7 @@ public:
 	 */
 	static bool Parse( const FString& IdString, FTargetDeviceId& OutId )
 	{
-		return IdString.Split(TEXT("@"), &OutId.PlatformName, &OutId.DeviceName);
+		return IdString.Split(TEXT("@"), &OutId.PlatformName, &OutId.DeviceName, ESearchCase::CaseSensitive);
 	}
 
 private:
diff --git a/Engine/Source/Editor/AnimGraph/Private/AnimGraphNode_Base.cpp b/Engine/Source/Editor/AnimGraph/Private/AnimGraphNode_Base.cpp
--- a/Engine/Source/Editor/AnimGraph/Private/AnimGraphNode_Base.cpp
+++ b/Engine/Source/Editor/AnimGraph/Private/AnimGraphNode_Base.cpp
@@ -273,7 +273,7 @@ void UAnimGraphNode_Base::GetPinAssociatedProperty(const UScriptStruct* NodeType
 	OutIndex = INDEX_NONE;
 
 	//@TODO: Name-based hackery, avoid the roundtrip and better indicate when it's an array pose pin
-	int32 UnderscoreIndex = InputPin->PinName.Find(TEXT("_"));
+	int32 UnderscoreIndex = InputPin->PinName.Find(TEXT("_"), ESearchCase::CaseSensitive);
 	if (UnderscoreIndex != INDEX_NONE)
 	{
 		FString ArrayName = InputPin->PinName.Left(UnderscoreIndex);
@@ -302,7 +302,7 @@ FPoseLinkMappingRecord UAnimGraphNode_Base::GetLinkIDLocation(const UScriptStruc
 		if (UAnimGraphNode_Base* LinkedNode = Cast<UAnimGraphNode_Base>(SourcePin->LinkedTo[0]->GetOwningNode()))
 		{
 			//@TODO: Name-based hackery, avoid the roundtrip and better indicate when it's an array pose pin
-			int32 UnderscoreIndex = SourcePin->PinName.Find(TEXT("_"));
+			int32 UnderscoreIndex = SourcePin->PinName.Find(TEXT("_"), ESearchCase::CaseSensitive);
 			if (UnderscoreIndex != INDEX_NONE)
 			{
 				FString ArrayName = SourcePin->PinName.Left(UnderscoreIndex);
diff --git a/Engine/Source/Editor/BlueprintGraph/Private/BlueprintNodeSignature.cpp b/Engine/Source/Editor/BlueprintGraph/Private/BlueprintNodeSignature.cpp
--- a/Engine/Source/Editor/BlueprintGraph/Private/BlueprintNodeSignature.cpp
+++ b/Engine/Source/Editor/BlueprintGraph/Private/BlueprintNodeSignature.cpp
@@ -64,8 +64,8 @@ FBlueprintNodeSignature::FBlueprintNodeSignature(FString const& UserString)
 		SignatureElement.Split(SignatureKeyDelim, &SignatureKey, &SignatureValue);
 		// @TODO: look for UObject redirects with SignatureValue
 
-		SignatureValue.RemoveFromStart(TEXT("\""));
-		SignatureValue.RemoveFromEnd(TEXT("\""));
+		SignatureValue.RemoveFromStart(TEXT("\""), ESearchCase::CaseSensitive);
+		SignatureValue.RemoveFromEnd(TEXT("\""), ESearchCase::CaseSensitive);
 		AddNamedValue(FName(*SignatureKey), SignatureValue);
 	}
 }
diff --git a/Engine/Source/Editor/UnrealEd/Private/Commandlets/CookCommandlet.cpp b/Engine/Source/Editor/UnrealEd/Private/Commandlets/CookCommandlet.cpp
--- a/Engine/Source/Editor/UnrealEd/Private/Commandlets/CookCommandlet.cpp
+++ b/Engine/Source/Editor/UnrealEd/Private/Commandlets/CookCommandlet.cpp
@@ -1105,7 +1105,7 @@ bool UCookCommandlet::NewCook( const TArray<ITargetPlatform*>& Platforms, TArray
 				FString ValuesList = Switch.Right(Switch.Len() - (SwitchKey + TEXT("=")).Len());
 
 				// Allow support for -KEY=Value1+Value2+Value3 as well as -KEY=Value1 -KEY=Value2
-				for (int32 PlusIdx = ValuesList.Find(TEXT("+")); PlusIdx != INDEX_NONE; PlusIdx = ValuesList.Find(TEXT("+")))
+				for (int32 PlusIdx = ValuesList.Find(TEXT("+"), ESearchCase::CaseSensitive); PlusIdx != INDEX_NONE; PlusIdx = ValuesList.Find(TEXT("+"), ESearchCase::CaseSensitive))
 				{
 					const FString ValueElement = ValuesList.Left(PlusIdx);
 					ValueElements.Add(ValueElement);
diff --git a/Engine/Source/Editor/UnrealEd/Private/EditorCategoryUtils.cpp b/Engine/Source/Editor/UnrealEd/Private/EditorCategoryUtils.cpp
--- a/Engine/Source/Editor/UnrealEd/Private/EditorCategoryUtils.cpp
+++ b/Engine/Source/Editor/UnrealEd/Private/EditorCategoryUtils.cpp
@@ -279,7 +279,7 @@ FString FEditorCategoryUtils::GetCategoryDisplayString(FString const& Unsanitize
 	} while (KeyIndex != INDEX_NONE);
 
 	DisplayString = FName::NameToDisplayString(DisplayString, /*bIsBool =*/false);
-	DisplayString.ReplaceInline(TEXT("| "), TEXT("|"));
+	DisplayString.ReplaceInline(TEXT("| "), TEXT("|"), ESearchCase::CaseSensitive);
 
 	return DisplayString;
 }
diff --git a/Engine/Source/Runtime/AssetRegistry/Private/AssetRegistry.cpp b/Engine/Source/Runtime/AssetRegistry/Private/AssetRegistry.cpp
--- a/Engine/Source/Runtime/AssetRegistry/Private/AssetRegistry.cpp
+++ b/Engine/Source/Runtime/AssetRegistry/Private/AssetRegistry.cpp
@@ -1488,7 +1488,7 @@ void FAssetRegistry::DependencyDataGathered(const double TickStartTime, TArray<F
 		for (int32 StringAssetRefIdx = 0; StringAssetRefIdx < Result.StringAssetReferencesMap.Num(); ++StringAssetRefIdx)
 		{
 			FString PackageName, ObjName;
-			Result.StringAssetReferencesMap[StringAssetRefIdx].Split(".", &PackageName, &ObjName);
+			Result.StringAssetReferencesMap[StringAssetRefIdx].Split(".", &PackageName, &ObjName, ESearchCase::CaseSensitive);
 			PackageDependencies.Add(*PackageName);
 		}
 
diff --git a/Engine/Source/Runtime/Core/Private/Containers/String.cpp b/Engine/Source/Runtime/Core/Private/Containers/String.cpp
--- a/Engine/Source/Runtime/Core/Private/Containers/String.cpp
+++ b/Engine/Source/Runtime/Core/Private/Containers/String.cpp
@@ -20,7 +20,7 @@ void FString::TrimToNullTerminator()
 }
 
 
-int32 FString::Find( const TCHAR* SubStr, ESearchCase::Type SearchCase, ESearchDir::Type SearchDir, int32 StartPosition ) const
+int32 FString::Find(const TCHAR* SubStr, ESearchCase::Type SearchCase, ESearchDir::Type SearchDir, int32 StartPosition) const
 {
 	if ( SubStr == nullptr )
 	{
@@ -643,8 +643,8 @@ bool FString::MatchesWildcard(const FString& InWildcard, ESearchCase::Type Searc
 			Wildcard = Wildcard.Left(Suffix + 1);
 			Target = Target.Left(Target.Len() - SuffixString.Len());
 		}
-		int32 PrefixIndexOfStar = Wildcard.Find(TEXT("*")); 
-		int32 PrefixIndexOfQuestion = Wildcard.Find( TEXT("?"));
+		int32 PrefixIndexOfStar = Wildcard.Find(TEXT("*"), ESearchCase::CaseSensitive); 
+		int32 PrefixIndexOfQuestion = Wildcard.Find(TEXT("?"), ESearchCase::CaseSensitive);
 		int32 Prefix = FMath::Min<int32>(PrefixIndexOfStar < 0 ? MAX_int32 : PrefixIndexOfStar, PrefixIndexOfQuestion < 0 ? MAX_int32 : PrefixIndexOfQuestion);
 		check(Prefix >= 0 && Prefix < Wildcard.Len());
 		if (Prefix > 0)
@@ -833,31 +833,31 @@ FString FString::Replace(const TCHAR* From, const TCHAR* To, ESearchCase::Type S
 	return Result;
 }
 
-int32 FString::ReplaceInline( const TCHAR* SearchText, const TCHAR* ReplacementText, ESearchCase::Type SearchCase )
+int32 FString::ReplaceInline(const TCHAR* SearchText, const TCHAR* ReplacementText, ESearchCase::Type SearchCase)
 {
 	int32 ReplacementCount = 0;
 
 	if (Len() > 0
-	&&	SearchText != nullptr && *SearchText != 0
-	&&	ReplacementText != nullptr && (SearchCase == ESearchCase::IgnoreCase || FCString::Strcmp(SearchText, ReplacementText) != 0) )
+		&& SearchText != nullptr && *SearchText != 0
+		&& ReplacementText != nullptr && (SearchCase == ESearchCase::IgnoreCase || FCString::Strcmp(SearchText, ReplacementText) != 0))
 	{
-		const int32 NumCharsToReplace=FCString::Strlen(SearchText);
-		const int32 NumCharsToInsert=FCString::Strlen(ReplacementText);
+		const int32 NumCharsToReplace = FCString::Strlen(SearchText);
+		const int32 NumCharsToInsert = FCString::Strlen(ReplacementText);
 
-		if ( NumCharsToInsert == NumCharsToReplace )
+		if (NumCharsToInsert == NumCharsToReplace)
 		{
 			TCHAR* Pos = SearchCase == ESearchCase::IgnoreCase ? FCString::Stristr(&(*this)[0], SearchText) : FCString::Strstr(&(*this)[0], SearchText);
-			while ( Pos != nullptr )
+			while (Pos != nullptr)
 			{
 				ReplacementCount++;
 
 				// FCString::Strcpy now inserts a terminating zero so can't use that
-				for ( int32 i = 0; i < NumCharsToInsert; i++ )
+				for (int32 i = 0; i < NumCharsToInsert; i++)
 				{
 					Pos[i] = ReplacementText[i];
 				}
 
-				if ( Pos + NumCharsToReplace - **this < Len() )
+				if (Pos + NumCharsToReplace - **this < Len())
 				{
 					Pos = SearchCase == ESearchCase::IgnoreCase ? FCString::Stristr(Pos + NumCharsToReplace, SearchText) : FCString::Strstr(Pos + NumCharsToReplace, SearchText);
 				}
@@ -867,7 +867,7 @@ int32 FString::ReplaceInline( const TCHAR* SearchText, const TCHAR* ReplacementT
 				}
 			}
 		}
-		else if ( Contains(SearchText, SearchCase) )
+		else if (Contains(SearchText, SearchCase))
 		{
 			FString Copy(*this);
 			Empty(Len());
@@ -876,7 +876,7 @@ int32 FString::ReplaceInline( const TCHAR* SearchText, const TCHAR* ReplacementT
 			TCHAR* WritePosition = (TCHAR*)Copy.Data.GetData();
 			// look for From in the remaining string
 			TCHAR* SearchPosition = SearchCase == ESearchCase::IgnoreCase ? FCString::Stristr(WritePosition, SearchText) : FCString::Strstr(WritePosition, SearchText);
-			while ( SearchPosition != nullptr )
+			while (SearchPosition != nullptr)
 			{
 				ReplacementCount++;
 
@@ -910,7 +910,7 @@ int32 FString::ReplaceInline( const TCHAR* SearchText, const TCHAR* ReplacementT
  */
 FString FString::ReplaceQuotesWithEscapedQuotes() const
 {
-	if ( Contains(TEXT("\"")) )
+	if (Contains(TEXT("\""), ESearchCase::CaseSensitive))
 	{
 		FString Result;
 
diff --git a/Engine/Source/Runtime/Core/Private/Misc/AutomationTest.cpp b/Engine/Source/Runtime/Core/Private/Misc/AutomationTest.cpp
--- a/Engine/Source/Runtime/Core/Private/Misc/AutomationTest.cpp
+++ b/Engine/Source/Runtime/Core/Private/Misc/AutomationTest.cpp
@@ -208,7 +208,7 @@ void FAutomationTestFramework::StartTestByName( const FString& InTestToRun, cons
 
 	FString TestName;
 	FString Params;
-	if (!InTestToRun.Split(TEXT(" "), &TestName, &Params))
+	if (!InTestToRun.Split(TEXT(" "), &TestName, &Params, ESearchCase::CaseSensitive))
 	{
 		TestName = InTestToRun;
 	}
@@ -576,7 +576,7 @@ void FAutomationTestFramework::DumpAutomationTestExecutionInfo( const TMap<FStri
 void FAutomationTestFramework::InternalStartTest( const FString& InTestToRun )
 {
 	FString TestName;
-	if (!InTestToRun.Split(TEXT(" "), &TestName, &Parameters))
+	if (!InTestToRun.Split(TEXT(" "), &TestName, &Parameters, ESearchCase::CaseSensitive))
 	{
 		TestName = InTestToRun;
 	}
diff --git a/Engine/Source/Runtime/Core/Private/Misc/ConfigCacheIni.cpp b/Engine/Source/Runtime/Core/Private/Misc/ConfigCacheIni.cpp
--- a/Engine/Source/Runtime/Core/Private/Misc/ConfigCacheIni.cpp
+++ b/Engine/Source/Runtime/Core/Private/Misc/ConfigCacheIni.cpp
@@ -3146,7 +3146,7 @@ void ApplyCVarSettingsFromIni(const TCHAR* InSectionName, const TCHAR* InIniFile
 		{
 			FString CVarName;
 			FString CVarValue;
-			SectionStrings[VarIndex].Split(TEXT("="), &CVarName, &CVarValue);
+			SectionStrings[VarIndex].Split(TEXT("="), &CVarName, &CVarValue, ESearchCase::CaseSensitive);
 
 			IConsoleVariable* CVar = IConsoleManager::Get().FindConsoleVariable(*CVarName);
 			if (CVar)
diff --git a/Engine/Source/Runtime/Core/Private/Misc/DateTime.cpp b/Engine/Source/Runtime/Core/Private/Misc/DateTime.cpp
--- a/Engine/Source/Runtime/Core/Private/Misc/DateTime.cpp
+++ b/Engine/Source/Runtime/Core/Private/Misc/DateTime.cpp
@@ -277,8 +277,8 @@ bool FDateTime::Parse( const FString& DateTimeString, FDateTime& OutDateTime )
 	// first replace -, : and . with space
 	FString FixedString = DateTimeString.Replace(TEXT("-"), TEXT(" "));
 
-	FixedString.ReplaceInline(TEXT(":"), TEXT(" "));
-	FixedString.ReplaceInline(TEXT("."), TEXT(" "));
+	FixedString.ReplaceInline(TEXT(":"), TEXT(" "), ESearchCase::CaseSensitive);
+	FixedString.ReplaceInline(TEXT("."), TEXT(" "), ESearchCase::CaseSensitive);
 
 	TArray<FString> Tokens;
 
diff --git a/Engine/Source/Runtime/Core/Private/Misc/Paths.cpp b/Engine/Source/Runtime/Core/Private/Misc/Paths.cpp
--- a/Engine/Source/Runtime/Core/Private/Misc/Paths.cpp
+++ b/Engine/Source/Runtime/Core/Private/Misc/Paths.cpp
@@ -501,7 +501,7 @@ bool FPaths::IsDrive(const FString& InPath)
 		}
 
 		FString CheckPath = TEXT("");
-		int32 ColonSlashIndex = TempPath.Find(TEXT(":\\"));
+		int32 ColonSlashIndex = TempPath.Find(TEXT(":\\"), ESearchCase::CaseSensitive);
 		if (ColonSlashIndex != INDEX_NONE)
 		{
 			// Remove the 'X:\' from the start
@@ -510,11 +510,11 @@ bool FPaths::IsDrive(const FString& InPath)
 		else
 		{
 			// See if the first two characters are '\\' to handle \\Server\Foo\Bar cases
-			if (TempPath.StartsWith(TEXT("\\\\")) == true)
+			if (TempPath.StartsWith(TEXT("\\\\"), ESearchCase::CaseSensitive) == true)
 			{
 				CheckPath = TempPath.Right(TempPath.Len() - 2);
 				// Find the next slash
-				int32 SlashIndex = CheckPath.Find(TEXT("\\"));
+				int32 SlashIndex = CheckPath.Find(TEXT("\\"), ESearchCase::CaseSensitive);
 				if (SlashIndex != INDEX_NONE)
 				{
 					CheckPath = CheckPath.Right(CheckPath.Len() - SlashIndex  - 1);
@@ -529,10 +529,10 @@ bool FPaths::IsDrive(const FString& InPath)
 		if (CheckPath.Len() > 0)
 		{
 			// Replace any remaining '\\' instances with '\'
-			CheckPath.Replace(TEXT("\\\\"), TEXT("\\"));
+			CheckPath.Replace(TEXT("\\\\"), TEXT("\\"), ESearchCase::CaseSensitive);
 
 			int32 CheckCount = 0;
-			int32 SlashIndex = CheckPath.Find(TEXT("\\"));
+			int32 SlashIndex = CheckPath.Find(TEXT("\\"), ESearchCase::CaseSensitive);
 			while (SlashIndex != INDEX_NONE)
 			{
 				FString FolderName = CheckPath.Left(SlashIndex);
@@ -547,7 +547,7 @@ bool FPaths::IsDrive(const FString& InPath)
 					CheckCount++;
 				}
 				CheckPath = CheckPath.Right(CheckPath.Len() - SlashIndex  - 1);
-				SlashIndex = CheckPath.Find(TEXT("\\"));
+				SlashIndex = CheckPath.Find(TEXT("\\"), ESearchCase::CaseSensitive);
 			}
 
 			if (CheckCount <= 0)
@@ -578,15 +578,15 @@ bool FPaths::IsRelative(const FString& InPath)
 
 void FPaths::NormalizeFilename(FString& InPath)
 {
-	InPath.ReplaceInline(TEXT("\\"), TEXT("/"));
+	InPath.ReplaceInline(TEXT("\\"), TEXT("/"), ESearchCase::CaseSensitive);
 
 	FPlatformMisc::NormalizePath(InPath);
 }
 
 void FPaths::NormalizeDirectoryName(FString& InPath)
 {
-	InPath.ReplaceInline(TEXT("\\"), TEXT("/"));
-	if (InPath.EndsWith(TEXT("/")) && !InPath.EndsWith(TEXT("//")) && !InPath.EndsWith(TEXT(":/")))
+	InPath.ReplaceInline(TEXT("\\"), TEXT("/"), ESearchCase::CaseSensitive);
+	if (InPath.EndsWith(TEXT("/"), ESearchCase::CaseSensitive) && !InPath.EndsWith(TEXT("//"), ESearchCase::CaseSensitive) && !InPath.EndsWith(TEXT(":/"), ESearchCase::CaseSensitive))
 	{
 		// overwrite trailing slash with terminator
 		InPath.GetCharArray()[InPath.Len() - 1] = 0;
@@ -609,11 +609,11 @@ bool FPaths::CollapseRelativeDirectories(FString& InPath)
 			break;
 
 		// Consider empty paths or paths which start with .. or /.. as invalid
-		if (InPath.StartsWith(TEXT("..")) || InPath.StartsWith(ParentDir))
+		if (InPath.StartsWith(TEXT(".."), ESearchCase::CaseSensitive) || InPath.StartsWith(ParentDir))
 			return false;
 
 		// If there are no "/.."s left then we're done
-		const int32 Index = InPath.Find(ParentDir);
+		const int32 Index = InPath.Find(ParentDir, ESearchCase::CaseSensitive);
 		if (Index == -1)
 			break;
 
@@ -640,7 +640,7 @@ bool FPaths::CollapseRelativeDirectories(FString& InPath)
 		InPath.RemoveAt(PreviousSeparatorIndex, Index - PreviousSeparatorIndex + ParentDirLength, false);
 	}
 
-	InPath.ReplaceInline(TEXT("./"), TEXT(""));
+	InPath.ReplaceInline(TEXT("./"), TEXT(""), ESearchCase::CaseSensitive);
 
 	InPath.TrimToNullTerminator();
 
@@ -649,9 +649,9 @@ bool FPaths::CollapseRelativeDirectories(FString& InPath)
 
 void FPaths::RemoveDuplicateSlashes(FString& InPath)
 {
-	while (InPath.Contains(TEXT("//")))
+	while (InPath.Contains(TEXT("//"), ESearchCase::CaseSensitive))
 	{
-		InPath = InPath.Replace(TEXT("//"), TEXT("/"));
+		InPath = InPath.Replace(TEXT("//"), TEXT("/"), ESearchCase::CaseSensitive);
 	}
 }
 
@@ -670,13 +670,13 @@ void FPaths::MakeStandardFilename(FString& InPath)
 		return;
 	}
 
-	FString WithSlashes = InPath.Replace(TEXT("\\"), TEXT("/"));
+	FString WithSlashes = InPath.Replace(TEXT("\\"), TEXT("/"), ESearchCase::CaseSensitive);
 
 	FString RootDirectory = FPaths::ConvertRelativePathToFull(FPaths::RootDir());
 
 	// look for paths that cannot be made relative, and are therefore left alone
 	// UNC (windows) network path
-	bool bCannotBeStandardized = InPath.StartsWith(TEXT("\\\\"));
+	bool bCannotBeStandardized = InPath.StartsWith(TEXT("\\\\"), ESearchCase::CaseSensitive);
 	// windows drive letter path that doesn't start with base dir
 	bCannotBeStandardized |= (InPath[1] == ':' && !WithSlashes.StartsWith(RootDirectory));
 	// Unix style absolute path that doesn't start with base dir
@@ -707,8 +707,8 @@ void FPaths::MakeStandardFilename(FString& InPath)
 
 void FPaths::MakePlatformFilename( FString& InPath )
 {
-	InPath.ReplaceInline( TEXT( "\\" ), FPlatformMisc::GetDefaultPathSeparator() );
-	InPath.ReplaceInline( TEXT( "/" ), FPlatformMisc::GetDefaultPathSeparator() );
+	InPath.ReplaceInline(TEXT("\\"), FPlatformMisc::GetDefaultPathSeparator(), ESearchCase::CaseSensitive);
+	InPath.ReplaceInline(TEXT("/"), FPlatformMisc::GetDefaultPathSeparator(), ESearchCase::CaseSensitive);
 }
 
 bool FPaths::MakePathRelativeTo( FString& InPath, const TCHAR* InRelativeTo )
@@ -717,8 +717,8 @@ bool FPaths::MakePathRelativeTo( FString& InPath, const TCHAR* InRelativeTo )
 	FString Source = FPaths::ConvertRelativePathToFull(InRelativeTo);
 	
 	Source = FPaths::GetPath(Source);
-	Source.ReplaceInline(TEXT("\\"), TEXT("/"));
-	Target.ReplaceInline(TEXT("\\"), TEXT("/"));
+	Source.ReplaceInline(TEXT("\\"), TEXT("/"), ESearchCase::CaseSensitive);
+	Target.ReplaceInline(TEXT("\\"), TEXT("/"), ESearchCase::CaseSensitive);
 
 	TArray<FString> TargetArray;
 	Target.ParseIntoArray(&TargetArray, TEXT("/"), true);
diff --git a/Engine/Source/Runtime/Core/Private/Misc/Timespan.cpp b/Engine/Source/Runtime/Core/Private/Misc/Timespan.cpp
--- a/Engine/Source/Runtime/Core/Private/Misc/Timespan.cpp
+++ b/Engine/Source/Runtime/Core/Private/Misc/Timespan.cpp
@@ -126,7 +126,7 @@ bool FTimespan::Parse( const FString& TimespanString, FTimespan& OutTimespan )
 
 	bool Negative = TokenString.StartsWith(TEXT("-"));
 
-	TokenString.ReplaceInline(TEXT("-"), TEXT(":"));
+	TokenString.ReplaceInline(TEXT("-"), TEXT(":"), ESearchCase::CaseSensitive);
 
 	TArray<FString> Tokens;
 	TokenString.ParseIntoArray(&Tokens, TEXT(":"), true);
diff --git a/Engine/Source/Runtime/Core/Private/Modules/ModuleManager.cpp b/Engine/Source/Runtime/Core/Private/Modules/ModuleManager.cpp
--- a/Engine/Source/Runtime/Core/Private/Modules/ModuleManager.cpp
+++ b/Engine/Source/Runtime/Core/Private/Modules/ModuleManager.cpp
@@ -792,9 +792,9 @@ void FModuleManager::GetModuleFilenameFormat(bool bGameModule, FString& OutPrefi
 
 	// Get the base name for modules of this application
 	OutPrefix = FPlatformProcess::GetModulePrefix() + FPaths::GetBaseFilename(FPlatformProcess::ExecutableName());
-	if (OutPrefix.Contains(TEXT("-")))
+	if (OutPrefix.Contains(TEXT("-"), ESearchCase::CaseSensitive))
 	{
-		OutPrefix = OutPrefix.Left(OutPrefix.Find(TEXT("-")) + 1);
+		OutPrefix = OutPrefix.Left(OutPrefix.Find(TEXT("-"), ESearchCase::CaseSensitive) + 1);
 	}
 	else
 	{
diff --git a/Engine/Source/Runtime/Core/Private/Stats/Stats2.cpp b/Engine/Source/Runtime/Core/Private/Stats/Stats2.cpp
--- a/Engine/Source/Runtime/Core/Private/Stats/Stats2.cpp
+++ b/Engine/Source/Runtime/Core/Private/Stats/Stats2.cpp
@@ -557,7 +557,7 @@ FName FStatNameAndInfo::GetShortNameFrom(FName InLongName)
 	if (Input.StartsWith(TEXT("//")))
 	{
 		Input = Input.RightChop(2);
-		const int32 IndexEnd = Input.Find(TEXT("//"));
+		const int32 IndexEnd = Input.Find(TEXT("//"), ESearchCase::CaseSensitive);
 		if (IndexEnd == INDEX_NONE)
 		{
 			checkStats(0);
@@ -565,7 +565,7 @@ FName FStatNameAndInfo::GetShortNameFrom(FName InLongName)
 		}
 		Input = Input.RightChop(IndexEnd + 2);
 	}
-	const int32 IndexEnd = Input.Find(TEXT("///"));
+	const int32 IndexEnd = Input.Find(TEXT("///"), ESearchCase::CaseSensitive);
 	if (IndexEnd != INDEX_NONE)
 	{
 		Input = Input.Left(IndexEnd);
@@ -580,7 +580,7 @@ FName FStatNameAndInfo::GetGroupNameFrom(FName InLongName)
 	if (Input.StartsWith(TEXT("//")))
 	{
 		Input = Input.RightChop(2);
-		const int32 IndexEnd = Input.Find(TEXT("//"));
+		const int32 IndexEnd = Input.Find(TEXT("//"), ESearchCase::CaseSensitive);
 		if (IndexEnd != INDEX_NONE)
 		{
 			return FName(*Input.Left(IndexEnd));
@@ -594,11 +594,11 @@ FString FStatNameAndInfo::GetDescriptionFrom(FName InLongName)
 {
 	FString Input(InLongName.ToString());
 
-	const int32 IndexStart = Input.Find(TEXT("///"));
+	const int32 IndexStart = Input.Find(TEXT("///"), ESearchCase::CaseSensitive);
 	if (IndexStart != INDEX_NONE)
 	{
 		Input = Input.RightChop(IndexStart + 3);
-		const int32 IndexEnd = Input.Find(TEXT("///"));
+		const int32 IndexEnd = Input.Find(TEXT("///"), ESearchCase::CaseSensitive);
 		if (IndexEnd != INDEX_NONE)
 		{
 			return FStatsUtils::FromEscapedFString(*Input.Left(IndexEnd));
@@ -615,7 +615,7 @@ FName FStatNameAndInfo::GetGroupCategoryFrom(FName InLongName)
 	if (IndexStart != INDEX_NONE)
 	{
 		Input = Input.RightChop(IndexStart + 7);
-		const int32 IndexEnd = Input.Find(TEXT("////"));
+		const int32 IndexEnd = Input.Find(TEXT("////"), ESearchCase::CaseSensitive);
 		if (IndexEnd != INDEX_NONE)
 		{
 			return FName(*Input.Left(IndexEnd));
diff --git a/Engine/Source/Runtime/Core/Private/Stats/StatsData.cpp b/Engine/Source/Runtime/Core/Private/Stats/StatsData.cpp
--- a/Engine/Source/Runtime/Core/Private/Stats/StatsData.cpp
+++ b/Engine/Source/Runtime/Core/Private/Stats/StatsData.cpp
@@ -1657,7 +1657,7 @@ FString FStatsUtils::FromEscapedFString(const TCHAR* Escaped)
 	while (Input.Len())
 	{
 		{
-			int32 Index = Input.Find(TEXT("$"));
+			int32 Index = Input.Find(TEXT("$"), ESearchCase::CaseSensitive);
 			if (Index == INDEX_NONE)
 			{
 				Result += Input;
@@ -1668,7 +1668,7 @@ FString FStatsUtils::FromEscapedFString(const TCHAR* Escaped)
 
 		}
 		{
-			int32 IndexEnd = Input.Find(TEXT("$"));
+			int32 IndexEnd = Input.Find(TEXT("$"), ESearchCase::CaseSensitive);
 			if (IndexEnd == INDEX_NONE)
 			{
 				checkStats(0); // malformed escaped fname
diff --git a/Engine/Source/Runtime/Core/Private/Windows/WindowsApplication.cpp b/Engine/Source/Runtime/Core/Private/Windows/WindowsApplication.cpp
--- a/Engine/Source/Runtime/Core/Private/Windows/WindowsApplication.cpp
+++ b/Engine/Source/Runtime/Core/Private/Windows/WindowsApplication.cpp
@@ -527,7 +527,7 @@ void GetMonitorInfo(TArray<FMonitorInfo>& OutMonitorInfo)
 					FMonitorInfo Info;
 
 					Info.ID = FString::Printf(TEXT("%s"), Monitor.DeviceID);
-					Info.Name = Info.ID.Mid (8, Info.ID.Find (TEXT("\\"), ESearchCase::IgnoreCase, ESearchDir::FromStart, 9) - 8);
+					Info.Name = Info.ID.Mid (8, Info.ID.Find (TEXT("\\"), ESearchCase::CaseSensitive, ESearchDir::FromStart, 9) - 8);
 
 					if (GetSizeForDevID(Info.Name, Info.NativeWidth, Info.NativeHeight))
 					{
@@ -1916,7 +1916,7 @@ void FWindowsApplication::QueryConnectedMice()
 
 		Name[NameLen] = 0;
 		FString WName = ANSI_TO_TCHAR(Name);
-		WName.ReplaceInline(TEXT("#"), TEXT("\\"));
+		WName.ReplaceInline(TEXT("#"), TEXT("\\"), ESearchCase::CaseSensitive);
 		/*
 		 * Name XP starts with \??\, vista+ starts \\?\ 
 		 * In the device list exists a fake Mouse device with the device name of RDP_MOU
diff --git a/Engine/Source/Runtime/CoreUObject/Private/Misc/PackageName.cpp b/Engine/Source/Runtime/CoreUObject/Private/Misc/PackageName.cpp
--- a/Engine/Source/Runtime/CoreUObject/Private/Misc/PackageName.cpp
+++ b/Engine/Source/Runtime/CoreUObject/Private/Misc/PackageName.cpp
@@ -241,7 +241,7 @@ private:
 FString FPackageName::InternalFilenameToLongPackageName(const FString& InFilename)
 {
 	const FLongPackagePathsSingleton& Paths = FLongPackagePathsSingleton::Get();
-	FString Filename = InFilename.Replace(TEXT("\\"), TEXT("/"));
+	FString Filename = InFilename.Replace(TEXT("\\"), TEXT("/"), ESearchCase::CaseSensitive);
 
 	// Convert to relative path if it's not already a long package name
 	bool bIsValidLongPackageName = false;
@@ -262,7 +262,7 @@ FString FPackageName::InternalFilenameToLongPackageName(const FString& InFilenam
 	FString PackageName         = FPaths::GetBaseFilename(Filename);
 	int32   PackageNameStartsAt = Filename.Len() - FPaths::GetCleanFilename(Filename).Len();
 	FString Result              = Filename.Mid(0, PackageNameStartsAt + PackageName.Len());
-	Result.ReplaceInline(TEXT("\\"), TEXT("/"));
+	Result.ReplaceInline(TEXT("\\"), TEXT("/"), ESearchCase::CaseSensitive);
 
 	for (const auto& Pair : Paths.ContentPathToRoot)
 	{
@@ -281,7 +281,7 @@ bool FPackageName::TryConvertFilenameToLongPackageName(const FString& InFilename
 	FString Result = InternalFilenameToLongPackageName(InFilename);
 
 	// we don't support loading packages from outside of well defined places
-	if (Result.Contains(TEXT(".")) || Result.Contains(TEXT("\\")) || Result.Contains(TEXT(":")))
+	if (Result.Contains(TEXT("."), ESearchCase::CaseSensitive) || Result.Contains(TEXT("\\"), ESearchCase::CaseSensitive) || Result.Contains(TEXT(":"), ESearchCase::CaseSensitive))
 	{
 		return false;
 	}
@@ -295,9 +295,9 @@ FString FPackageName::FilenameToLongPackageName(const FString& InFilename)
 	FString Result = InternalFilenameToLongPackageName(InFilename);
 
 	// we don't support loading packages from outside of well defined places
-	const bool bContainsDot       = Result.Contains(TEXT("."));
-	const bool bContainsBackslash = Result.Contains(TEXT("\\"));
-	const bool bContainsColon     = Result.Contains(TEXT(":"));
+	const bool bContainsDot = Result.Contains(TEXT("."), ESearchCase::CaseSensitive);
+	const bool bContainsBackslash = Result.Contains(TEXT("\\"), ESearchCase::CaseSensitive);
+	const bool bContainsColon = Result.Contains(TEXT(":"), ESearchCase::CaseSensitive);
 
 	if (bContainsDot || bContainsBackslash || bContainsColon)
 	{
@@ -694,7 +694,7 @@ bool FPackageName::SearchForPackageOnDisk(const FString& PackageName, FString* O
 			}
 		}
 
-		const FString PackageWildcard = (PackageName.Find(TEXT(".")) != INDEX_NONE ? PackageName : PackageName + TEXT(".*"));
+		const FString PackageWildcard = (PackageName.Find(TEXT("."), ESearchCase::CaseSensitive) != INDEX_NONE ? PackageName : PackageName + TEXT(".*"));
 		TArray<FString> Results;
 
 		for (int32 PathIndex = 0; PathIndex < Paths.Num() && !bResult; ++PathIndex)
@@ -805,7 +805,7 @@ void FPackageName::EnsureContentPathsAreRegistered()
 
 bool FPackageName::ParseExportTextPath(const FString& InExportTextPath, FString* OutClassName, FString* OutObjectPath)
 {
-	if ( InExportTextPath.Split(TEXT("'"), OutClassName, OutObjectPath) )
+	if (InExportTextPath.Split(TEXT("'"), OutClassName, OutObjectPath, ESearchCase::CaseSensitive))
 	{
 		if ( OutObjectPath )
 		{
diff --git a/Engine/Source/Runtime/CoreUObject/Private/UObject/LinkerLoad.cpp b/Engine/Source/Runtime/CoreUObject/Private/UObject/LinkerLoad.cpp
--- a/Engine/Source/Runtime/CoreUObject/Private/UObject/LinkerLoad.cpp
+++ b/Engine/Source/Runtime/CoreUObject/Private/UObject/LinkerLoad.cpp
@@ -128,7 +128,7 @@ void ULinkerLoad::CreateActiveRedirectsMap(const FString& GEngineIniName)
 				//full redirect
 				else
 				{
-					if (NewClassName.ToString().Find(TEXT(".")) != NewClassName.ToString().Find(TEXT("."), ESearchCase::CaseSensitive, ESearchDir::FromEnd))
+					if (NewClassName.ToString().Find(TEXT("."), ESearchCase::CaseSensitive) != NewClassName.ToString().Find(TEXT("."), ESearchCase::CaseSensitive, ESearchDir::FromEnd))
 					{
 						UE_LOG(LogLinker, Error, TEXT("Currently we cannot rename nested objects for '%s'; if you want to leave the outer alone, just specify the name with no path"), *NewClassName.ToString());
 					}
@@ -1183,7 +1183,7 @@ ULinkerLoad::ELinkerStatus ULinkerLoad::FixupImportMap()
 						}
 
 						// Accepts either "PackageName.ClassName" or just "ClassName"
-						int32 Offset = RedirectName.Find(TEXT("."));
+						int32 Offset = RedirectName.Find(TEXT("."), ESearchCase::CaseSensitive);
 						if ( Offset >= 0 )
 						{
 							// A package class name redirect
diff --git a/Engine/Source/Runtime/CoreUObject/Private/UObject/UObjectGlobals.cpp b/Engine/Source/Runtime/CoreUObject/Private/UObject/UObjectGlobals.cpp
--- a/Engine/Source/Runtime/CoreUObject/Private/UObject/UObjectGlobals.cpp
+++ b/Engine/Source/Runtime/CoreUObject/Private/UObject/UObjectGlobals.cpp
@@ -467,7 +467,7 @@ UPackage* CreatePackage( UObject* InOuter, const TCHAR* PackageName )
 		InName = PackageName;
 	}
 
-	if (InName.Contains(TEXT("//")))
+	if (InName.Contains(TEXT("//"), ESearchCase::CaseSensitive))
 	{
 		UE_LOG(LogUObjectGlobals, Fatal, TEXT("Attempted to create a package with name containing double slashes. PackageName: %s"), PackageName);
 	}
@@ -542,7 +542,7 @@ bool ResolveName( UObject*& InPackage, FString& InOutName, bool Create, bool Thr
 	}
 
 
-	if( IniFilename && InOutName.Contains(TEXT(".")) )
+	if( IniFilename && InOutName.Contains(TEXT("."), ESearchCase::CaseSensitive) )
 	{
 		// Get .ini key and section.
 		FString Section = InOutName.Mid(1+InOutName.Find(TEXT(":"), ESearchCase::CaseSensitive));
@@ -583,8 +583,8 @@ bool ResolveName( UObject*& InPackage, FString& InOutName, bool Create, bool Thr
 	bool bSubobjectPath = false;
 
 	// to make parsing the name easier, replace the subobject delimiter with an extra dot
-	InOutName.ReplaceInline(SUBOBJECT_DELIMITER, TEXT(".."));
-	while( (i=InOutName.Find(TEXT("."))) != -1 )
+	InOutName.ReplaceInline(SUBOBJECT_DELIMITER, TEXT(".."), ESearchCase::CaseSensitive);
+	while ((i = InOutName.Find(TEXT("."), ESearchCase::CaseSensitive)) != -1)
 	{
 		FString PartialName = InOutName.Left(i);
 
diff --git a/Engine/Source/Runtime/CoreUObject/Public/UObject/Class.h b/Engine/Source/Runtime/CoreUObject/Public/UObject/Class.h
--- a/Engine/Source/Runtime/CoreUObject/Public/UObject/Class.h
+++ b/Engine/Source/Runtime/CoreUObject/Public/UObject/Class.h
@@ -1477,7 +1477,7 @@ public:
 
 			// Strip the namespace from the name.
 			FString EnumName(Names[InIndex].ToString());
-			int32 ScopeIndex = EnumName.Find(TEXT("::"));
+			int32 ScopeIndex = EnumName.Find(TEXT("::"), ESearchCase::CaseSensitive);
 			if (ScopeIndex != INDEX_NONE)
 			{
 				return EnumName.Mid(ScopeIndex + 2);
diff --git a/Engine/Source/Runtime/Engine/Private/GameMode.cpp b/Engine/Source/Runtime/Engine/Private/GameMode.cpp
--- a/Engine/Source/Runtime/Engine/Private/GameMode.cpp
+++ b/Engine/Source/Runtime/Engine/Private/GameMode.cpp
@@ -1153,16 +1153,16 @@ bool AGameMode::GrabOption( FString& Options, FString& Result )
 	{
 		// Get result.
 		Result = Options.Mid(1, MAX_int32);
-		if( Result.Contains(TEXT("?")) )
+		if (Result.Contains(TEXT("?"), ESearchCase::CaseSensitive))
 		{
-			Result =  Result.Left( Result.Find(TEXT("?")) );
+			Result = Result.Left(Result.Find(TEXT("?"), ESearchCase::CaseSensitive));
 		}
 
 		// Update options.
 		Options = Options.Mid(1, MAX_int32);
-		if( Options.Contains(TEXT("?")) )
+		if (Options.Contains(TEXT("?"), ESearchCase::CaseSensitive))
 		{
-			Options =  Options.Mid( Options.Find(TEXT("?")) , MAX_int32);
+			Options = Options.Mid(Options.Find(TEXT("?"), ESearchCase::CaseSensitive), MAX_int32);
 		}
 		else
 		{
@@ -1176,7 +1176,7 @@ bool AGameMode::GrabOption( FString& Options, FString& Result )
 
 void AGameMode::GetKeyValue( const FString& Pair, FString& Key, FString& Value )
 {
-	const int32 EqualSignIndex = Pair.Find(TEXT("="));
+	const int32 EqualSignIndex = Pair.Find(TEXT("="), ESearchCase::CaseSensitive);
 	if( EqualSignIndex != INDEX_NONE )
 	{
 		Key = Pair.Left(EqualSignIndex);
diff --git a/Engine/Source/Runtime/Engine/Private/Materials/HLSLMaterialTranslator.h b/Engine/Source/Runtime/Engine/Private/Materials/HLSLMaterialTranslator.h
--- a/Engine/Source/Runtime/Engine/Private/Materials/HLSLMaterialTranslator.h
+++ b/Engine/Source/Runtime/Engine/Private/Materials/HLSLMaterialTranslator.h
@@ -3614,7 +3614,7 @@ protected:
 		{
 			Code = FString(TEXT("return "))+Code+TEXT(";");
 		}
-		Code.ReplaceInline(TEXT("\n"),TEXT("\r\n"));
+		Code.ReplaceInline(TEXT("\n"),TEXT("\r\n"), ESearchCase::CaseSensitive);
 		FString ImplementationCode = FString::Printf(TEXT("%s CustomExpression%d(FMaterial%sParameters Parameters%s)\r\n{\r\n%s\r\n}\r\n"), *OutputTypeString, CustomExpressionIndex, ShaderFrequency==SF_Vertex?TEXT("Vertex"):TEXT("Pixel"), *InputParamDecl, *Code);
 		CustomExpressionImplementations.Add( ImplementationCode );
 
diff --git a/Engine/Source/Runtime/Engine/Private/UserInterface/Console.cpp b/Engine/Source/Runtime/Engine/Private/UserInterface/Console.cpp
--- a/Engine/Source/Runtime/Engine/Private/UserInterface/Console.cpp
+++ b/Engine/Source/Runtime/Engine/Private/UserInterface/Console.cpp
@@ -433,7 +433,7 @@ void UConsole::OutputText(const FString& Text)
 	while(StringLength > 0)
 	{
 		// Find the number of characters in the next line of text.
-		int32 LineLength = RemainingText.Find(TEXT("\n"));
+		int32 LineLength = RemainingText.Find(TEXT("\n"), ESearchCase::CaseSensitive);
 		if(LineLength == -1)
 		{
 			// There aren't any more newlines in the string, assume there's a newline at the end of the string.
diff --git a/Engine/Source/Runtime/GameplayAbilities/Private/AbilitySystemGlobals.cpp b/Engine/Source/Runtime/GameplayAbilities/Private/AbilitySystemGlobals.cpp
--- a/Engine/Source/Runtime/GameplayAbilities/Private/AbilitySystemGlobals.cpp
+++ b/Engine/Source/Runtime/GameplayAbilities/Private/AbilitySystemGlobals.cpp
@@ -76,7 +76,7 @@ void UAbilitySystemGlobals::DeriveGameplayCueTagFromAssetName(FString AssetName,
 			AssetName.RemoveFromStart(TEXT("GC_"));		// allow GC_ prefix in asset name
 			AssetName.RemoveFromEnd(TEXT("_c"));
 
-			AssetName.ReplaceInline(TEXT("_"), TEXT("."));
+			AssetName.ReplaceInline(TEXT("_"), TEXT("."), ESearchCase::CaseSensitive);
 
 			if (!AssetName.Contains(TEXT("GameplayCue")))
 			{
diff --git a/Engine/Source/Runtime/Online/HTTP/Private/Windows/HttpWinInet.cpp b/Engine/Source/Runtime/Online/HTTP/Private/Windows/HttpWinInet.cpp
--- a/Engine/Source/Runtime/Online/HTTP/Private/Windows/HttpWinInet.cpp
+++ b/Engine/Source/Runtime/Online/HTTP/Private/Windows/HttpWinInet.cpp
@@ -1027,7 +1027,7 @@ void FHttpResponseWinInet::ProcessResponseHeaders()
 			}
 			FString HeaderLine(DelimiterPtr-HeaderPtr, HeaderPtr);
 			FString HeaderKey,HeaderValue;
-			if (HeaderLine.Split(TEXT(":"), &HeaderKey, &HeaderValue))
+			if (HeaderLine.Split(TEXT(":"), &HeaderKey, &HeaderValue, ESearchCase::CaseSensitive))
 			{
 				if (!HeaderKey.IsEmpty())
 				{
diff --git a/Engine/Source/Runtime/SandboxFile/Private/IPlatformFileSandboxWrapper.cpp b/Engine/Source/Runtime/SandboxFile/Private/IPlatformFileSandboxWrapper.cpp
--- a/Engine/Source/Runtime/SandboxFile/Private/IPlatformFileSandboxWrapper.cpp
+++ b/Engine/Source/Runtime/SandboxFile/Private/IPlatformFileSandboxWrapper.cpp
@@ -54,7 +54,7 @@ bool FSandboxPlatformFile::Initialize(IPlatformFile* Inner, const TCHAR* CmdLine
 		// Commandline syntax
 		bool bWipeSandbox = false;
 		FPaths::NormalizeFilename(CommandLineDirectory);
-		int32 CommandIndex = CommandLineDirectory.Find( TEXT(":") );
+		int32 CommandIndex = CommandLineDirectory.Find(TEXT(":"), ESearchCase::CaseSensitive);
 		if( CommandIndex != INDEX_NONE )
 		{
 			// Check if absolute path was specified and the ':' refers to drive name
