fixed warning message

diff --git a/Engine/Source/Runtime/Core/Private/HAL/ConsoleManager.cpp b/Engine/Source/Runtime/Core/Private/HAL/ConsoleManager.cpp
--- a/Engine/Source/Runtime/Core/Private/HAL/ConsoleManager.cpp
+++ b/Engine/Source/Runtime/Core/Private/HAL/ConsoleManager.cpp
@@ -122,10 +122,10 @@ public:
 			FConsoleManager& ConsoleManager = (FConsoleManager&)IConsoleManager::Get();
 			FString CVarName = ConsoleManager.FindConsoleObjectName(this);
 			UE_LOG(LogConsoleManager, Warning,
-				TEXT("Console variable '%s' wasn't set (Priority %s < %s)"),
+				TEXT("Console variable '%s' wasn't set ('%s' has a lower priority than '%s')"),
 				CVarName.IsEmpty() ? TEXT("unknown?") : *CVarName,
-				GetSetByTCHAR((EConsoleVariableFlags)OldPri),
-				GetSetByTCHAR((EConsoleVariableFlags)NewPri)
+				GetSetByTCHAR((EConsoleVariableFlags)NewPri),
+				GetSetByTCHAR((EConsoleVariableFlags)OldPri)
 				);
 		}
 
