Added support for TSharedPtr<void>

diff --git a/Engine/Source/Runtime/Core/Public/Templates/SharedPointer.h b/Engine/Source/Runtime/Core/Public/Templates/SharedPointer.h
--- a/Engine/Source/Runtime/Core/Public/Templates/SharedPointer.h
+++ b/Engine/Source/Runtime/Core/Public/Templates/SharedPointer.h
@@ -417,6 +417,26 @@ private:
 };
 
 
+/**
+ * Wrapper for a type that yields a reference to that type.
+ */
+template<class T>
+struct FMakeReferenceTo
+{
+	typedef T& Type;
+};
+
+
+/**
+ * Specialization for FMakeReferenceTo<void>.
+ */
+template<>
+struct FMakeReferenceTo<void>
+{
+	typedef void Type;
+};
+
+
 /**
  * TSharedPtr is a non-intrusive reference-counted authoritative object pointer.  This shared pointer
  * will be conditionally thread-safe when the optional Mode template argument is set to ThreadSafe.
@@ -642,7 +662,7 @@ public:
 	 *
 	 * @return  Reference to the object
 	 */
-	FORCEINLINE ObjectType& operator*() const
+	FORCEINLINE typename FMakeReferenceTo<ObjectType>::Type operator*() const
 	{
 		check( IsValid() );
 		return *Object;
