#ue4 - Fix for FastAsin() not returning 0 for input of 0, caused by mismatch of constants.

diff --git a/Engine/Source/Runtime/Core/Public/Math/UnrealMathUtility.h b/Engine/Source/Runtime/Core/Public/Math/UnrealMathUtility.h
--- a/Engine/Source/Runtime/Core/Public/Math/UnrealMathUtility.h
+++ b/Engine/Source/Runtime/Core/Public/Math/UnrealMathUtility.h
@@ -344,13 +344,20 @@ struct FMath : public FPlatformMath
 		*ScalarCos = sign*p;
 	}
 
+
+	// Note:  We use FASTASIN_HALF_PI instead of HALF_PI inside of FastASin(), since it was the value that accompanied the minimax coefficients below.
+	// It is important to use exactly the same value in all places inside this function to ensure that FastASin(0.0f) == 0.0f.
+	// For comparison:
+	//		HALF_PI				== 1.57079632679f == 0x3fC90FDB
+	//		FASTASIN_HALF_PI	== 1.5707963050f  == 0x3fC90FDA
+#define FASTASIN_HALF_PI (1.5707963050f)
 	/**
 	* Computes the ASin of a scalar float.
 	*
-	* @param Value  input angle 
+	* @param Value  input angle
 	* @return ASin of Value
 	*/
-	static FORCEINLINE float FastAsin( float Value )
+	static FORCEINLINE float FastAsin(float Value)
 	{
 		// Clamp input to [-1,1].
 		bool nonnegative = (Value >= 0.0f);
@@ -361,14 +368,13 @@ struct FMath : public FPlatformMath
 			omx = 0.0f;
 		}
 		float root = FMath::Sqrt(omx);
-
 		// 7-degree minimax approximation
-		float result = ( ( ( ( ( ( -0.0012624911f * x + 0.0066700901f ) * x - 0.0170881256f ) * x + 0.0308918810f ) * x - 0.0501743046f ) * x + 0.0889789874f ) * x - 0.2145988016f ) * x + 1.5707963050f;
+		float result = ((((((-0.0012624911f * x + 0.0066700901f) * x - 0.0170881256f) * x + 0.0308918810f) * x - 0.0501743046f) * x + 0.0889789874f) * x - 0.2145988016f) * x + FASTASIN_HALF_PI;
 		result *= root;  // acos(|x|)
-
 		// acos(x) = pi - acos(-x) when x < 0, asin(x) = pi/2 - acos(x)
-		return (nonnegative ? HALF_PI - result : result - HALF_PI);
+		return (nonnegative ? FASTASIN_HALF_PI - result : result - FASTASIN_HALF_PI);
 	}
+#undef FASTASIN_HALF_PI
 
 
 	// Conversion Functions
