Core: TypeContainer cleanup pass; added support for registering factory functions

diff --git a/Engine/Source/Runtime/Core/Private/Tests/Misc/TypeContainerTest.cpp b/Engine/Source/Runtime/Core/Private/Tests/Misc/TypeContainerTest.cpp
--- a/Engine/Source/Runtime/Core/Private/Tests/Misc/TypeContainerTest.cpp
+++ b/Engine/Source/Runtime/Core/Private/Tests/Misc/TypeContainerTest.cpp
@@ -183,6 +183,22 @@ bool FTypeContainerTest::RunTest(const FString& Parameters)
 		TestEqual(TEXT("For per-process classes, the same instance must be returned from different threads [1]"), Two1, Two2);
 	}
 
+	// factory test
+	{
+		FTypeContainer Container;
+		{
+			Container.RegisterFactory<IFruit>(
+				[]() -> TSharedPtr<IFruit> {
+					return MakeShareable(new FBanana());
+				}
+			);
+		}
+
+		auto Fruit = Container.GetInstance<IFruit>();
+
+		TestTrue(TEXT("For delegate classes, an instance must be returned"), Fruit.IsValid());
+	}
+
 	return true;
 }
 
diff --git a/Engine/Source/Runtime/Core/Public/Misc/TypeContainer.h b/Engine/Source/Runtime/Core/Public/Misc/TypeContainer.h
--- a/Engine/Source/Runtime/Core/Public/Misc/TypeContainer.h
+++ b/Engine/Source/Runtime/Core/Public/Misc/TypeContainer.h
@@ -54,7 +54,7 @@ enum class ETypeContainerScope
  *   2. Unique instance per thread (aka. thread singleton),
  *      using RegisterClass(ETypeContainerScope::Thread)
  *   3. Unique instance per call (aka. class factory),
- *      using RegisterClass(ETypeContainerScope::Instance)
+ *      using RegisterClass(ETypeContainerScope::Instance) or RegisterFactory()
  *
  * Note: Type containers depend on variadic templates and is therefore not available
  * on XboxOne at this time. It should therefore only be used for desktop applications.
@@ -72,63 +72,76 @@ class FTypeContainer
 		virtual ~IInstanceProvider() { }
 
 		/**
-		 * Gets a pointer to the stored instance.
+		 * Gets an instance of a class.
 		 *
-		 * The return value is not the instance itself, but a untyped pointer to a shared
-		 * pointer to the instance. The caller must cast this value to TSharedPtr<T> in
-		 * order to access the actual instance.
-		 * 
+		 * @return Shared pointer to the instance (must be cast to TSharedPtr<R>).
 		 */
-		virtual void* GetInstance() = 0;
+		virtual TSharedPtr<void> GetInstance() = 0;
 	};
 
-	/** Implements an instance provider that returns the same instance for all threads. */
+
+	/** Implements an instance provider that forwards instance requests to a factory delegate. */
 	template<class T>
-	class TSharedInstance
+	struct TFactoryInstanceProvider
 		: public IInstanceProvider
 	{
-	public:
+		/** Constructor. */
+		TFactoryInstanceProvider(TFunction<TSharedPtr<T>()> InCreateFunc)
+			: CreateFunc(InCreateFunc)
+		{ }
+
+		virtual ~TFactoryInstanceProvider() override { }
 
+		virtual TSharedPtr<void> GetInstance() override
+		{
+			return CreateFunc();
+		}
+
+		/** Factory function that creates the instances. */
+		TFunction<TSharedPtr<T>()> CreateFunc;
+	};
+
+
+	/** Implements an instance provider that returns the same instance for all threads. */
+	template<class T>
+	struct TSharedInstanceProvider
+		: public IInstanceProvider
+	{
 		/** Creates and initializes a new instance. */
-		TSharedInstance(const TSharedRef<T>& InInstance)
+		TSharedInstanceProvider(const TSharedRef<T>& InInstance)
 			: Instance(InInstance)
 		{ }
 
-		virtual ~TSharedInstance() { }
+		/** Virtual destructor. */
+		virtual ~TSharedInstanceProvider() { }
 
-		virtual void* GetInstance() override
+		virtual TSharedPtr<void> GetInstance() override
 		{
-			return &Instance;
+			return Instance;
 		}
 
-	private:
-
 		/** The stored instance. */
-		TSharedPtr<T> Instance;
+		TSharedRef<T> Instance;
 	};
 
 
 	/** Base class for per-thread instances providers. */
-	class FThreadInstance
+	struct FThreadInstanceProvider
 		: public IInstanceProvider
 	{
-	public:
-
 		/** Constructor. */
-		FThreadInstance(TFunction<void*()>&& InCreateFunc)
+		FThreadInstanceProvider(TFunction<TSharedPtr<void>()>&& InCreateFunc)
 			: CreateFunc(MoveTemp(InCreateFunc))
 			, TlsSlot(FPlatformTLS::AllocTlsSlot())
 		{ }
 
-		virtual ~FThreadInstance() override
+		virtual ~FThreadInstanceProvider() override
 		{
 			FPlatformTLS::FreeTlsSlot(TlsSlot);
 		}
 
-	protected:
-
 		/** Factory function for creating instances. */
-		TFunction<void*()> CreateFunc;
+		TFunction<TSharedPtr<void>()> CreateFunc;
 
 		/** Slot ID for thread-local storage of the instance. */
 		uint32 TlsSlot;
@@ -137,50 +150,46 @@ class FTypeContainer
 
 	/** Implements an instance provider that returns the same instance per thread. */
 	template<class T>
-	class TThreadInstance
-		: public FThreadInstance
+	struct TThreadInstanceProvider
+		: public FThreadInstanceProvider
 	{
-	public:
-
 		/** Constructor. */
-		TThreadInstance(TFunction<void*()>&& InCreateFunc)
-			: FThreadInstance(MoveTemp(InCreateFunc))
+		TThreadInstanceProvider(TFunction<TSharedPtr<void>()>&& InCreateFunc)
+			: FThreadInstanceProvider(MoveTemp(InCreateFunc))
 		{ }
 
-		virtual ~TThreadInstance() override { }
+		virtual ~TThreadInstanceProvider() override { }
 
-		virtual void* GetInstance() override
+		virtual TSharedPtr<void> GetInstance() override
 		{
 			void* Instance = FPlatformTLS::GetTlsValue(TlsSlot);
 
-			if (Instance != nullptr)
+			if (Instance == nullptr)
 			{
-				return Instance;
-			}
+				Instance = new TSharedPtr<T>(StaticCastSharedPtr<T>(CreateFunc()));
 
-			Instance = new TSharedPtr<T>((T*)CreateFunc());
-			/** @todo gmp: this cannot possibly work, because the FRunnableThread is not yet registered.
-			    A better mechanism for executing code on thread shutdown is already being worked on.
-				In the meantime we leak some memory here.
-			FRunnableThread::GetThreadRegistry().Lock();
-			{
-				const uint32 ThreadId = FPlatformTLS::GetCurrentThreadId();
-				FRunnableThread* RunnableThread = FRunnableThread::GetThreadRegistry().GetThread(ThreadId);
-
-				if (RunnableThread != nullptr)
+				/** @todo gmp: this cannot possibly work, because the FRunnableThread is not yet registered.
+					A better mechanism for executing code on thread shutdown is already being worked on.
+					In the meantime we leak some memory here.
+				FRunnableThread::GetThreadRegistry().Lock();
 				{
-					Instance = new TSharedPtr<T>((T*)CreateFunc());
-					RunnableThread->OnThreadDestroyed().AddStatic(&TThreadInstance::HandleDestroyInstance, Instance);
+					const uint32 ThreadId = FPlatformTLS::GetCurrentThreadId();
+					FRunnableThread* RunnableThread = FRunnableThread::GetThreadRegistry().GetThread(ThreadId);
+
+					if (RunnableThread != nullptr)
+					{
+						Instance = new TSharedPtr<T>((T*)CreateFunc());
+						RunnableThread->OnThreadDestroyed().AddStatic(&TThreadInstance::HandleDestroyInstance, Instance);
+					}
 				}
+				FRunnableThread::GetThreadRegistry().Unlock();*/
+
+				FPlatformTLS::SetTlsValue(TlsSlot, Instance);
 			}
-			FRunnableThread::GetThreadRegistry().Unlock();*/
-			FPlatformTLS::SetTlsValue(TlsSlot, Instance);
 
-			return Instance;
+			return *((TSharedPtr<void>*)Instance);
 		}
 
-	private:
-
 		/** Callback for destroying per-thread instances. */
 		static void HandleDestroyInstance(void* Instance)
 		{
@@ -205,23 +214,11 @@ public:
 
 		FScopeLock Lock(&CriticalSection);
 		{
-			const TFunction<void*()>* Factory = Factories.Find(ClassName);
-
-			if (Factory != nullptr)
-			{
-				return TSharedPtr<R>((R*)(*Factory)());
-			}
-
 			const TSharedPtr<IInstanceProvider>& Provider = Providers.FindRef(ClassName);
 
 			if (Provider.IsValid())
 			{
-				void* Instance = Provider->GetInstance();
-				
-				if (Instance != nullptr)
-				{
-					return *((TSharedPtr<R>*)Instance);
-				}
+				return StaticCastSharedPtr<R>(Provider->GetInstance());
 			}
 		}
 
@@ -251,21 +248,20 @@ public:
 	 * @param T Tye type of class to register (must be the same as or derived from R).
 	 * @param P The constructor parameters that the class requires.
 	 * @param Scope The scope at which instances must be unique (default = Process).
+	 * @see RegisterDelegate, RegisterInstance, Unregister
 	 */
 	template<class R, class T, typename... P>
 	void RegisterClass(ETypeContainerScope Scope = ETypeContainerScope::Process)
 	{
 		static_assert(TPointerIsConvertibleFromTo<T, R>::Value, "T must implement R");
 
-		const FName ClassName = R::GetTypeName();
-
 		if (Scope == ETypeContainerScope::Instance)
 		{
-			FScopeLock Lock(&CriticalSection);
-			{
-				Providers.Remove(ClassName);
-				Factories.Add(ClassName, [this]() -> void* { return new T(GetInstanceRef<P>()...); });
-			}
+			RegisterFactory<R>(
+				[this]() -> TSharedPtr<R> {
+					return MakeShareable(new T(GetInstanceRef<P>()...));
+				}
+			);
 		}
 		else
 		{
@@ -273,27 +269,56 @@ public:
 			
 			if (Scope == ETypeContainerScope::Thread)
 			{
-				Provider = MakeShareable(new TThreadInstance<T>([this]() -> void* { return new T(GetInstanceRef<P>()...); }));
+				Provider = MakeShareable(
+					new TThreadInstanceProvider<T>(
+						[this]() -> TSharedPtr<void> {
+							return MakeShareable(new T(GetInstanceRef<P>()...));
+						}
+					)
+				);
 			}
 			else
 			{
-				Provider = MakeShareable(new TSharedInstance<T>(MakeShareable(new T(GetInstanceRef<P>()...))));
+				Provider = MakeShareable(
+					new TSharedInstanceProvider<T>(
+						MakeShareable(new T(GetInstanceRef<P>()...))
+					)
+				);
 			}
 
 			FScopeLock Lock(&CriticalSection);
 			{
-				Factories.Remove(ClassName);
-				Providers.Add(ClassName, Provider);
+				Providers.Add(R::GetTypeName(), Provider);
 			}
 		}
 	}
 
+	/**
+	 * Register a factory delegate for the specified class.
+	 *
+	 * @param R The type of class to register the instance for.
+	 * @param CreateFunc A factory function that will create instances.
+	 * @see RegisterClass, RegisterInstance, Unregister
+	 */
+	template<class R>
+	void RegisterFactory(TFunction<TSharedPtr<R>()> CreateFunc)
+	{
+		FScopeLock Lock(&CriticalSection);
+		{
+			Providers.Add(
+				R::GetTypeName(),
+				MakeShareable(new TFactoryInstanceProvider<R>(CreateFunc))
+			);
+		}
+	}
+
 	/**
 	 * Registers an existing instance for the specified class.
 	 *
 	 * @param R The type of class to register the instance for.
 	 * @param T The type of the instance to register (must be the same as or derived from R).
 	 * @param Instance The instance of the class to register.
+	 * @see RegisterClass, RegisterDelegate, Unregister
 	 */
 	template<class R, class T>
 	void RegisterInstance(const TSharedRef<T>& Instance)
@@ -304,7 +329,7 @@ public:
 		{
 			Providers.Add(
 				R::GetTypeName(),
-				MakeShareable(new TSharedInstance<R>(Instance))
+				MakeShareable(new TSharedInstanceProvider<R>(Instance))
 			);
 		}
 	}
@@ -313,6 +338,7 @@ public:
 	 * Unregisters the instance or class for a previously registered class.
 	 *
 	 * @param R The type of class to unregister.
+	 * @see RegisterClass, RegisterDelegate, RegisterInstance
 	 */
 	template<class R>
 	void Unregister()
@@ -321,7 +347,6 @@ public:
 
 		FScopeLock Lock(&CriticalSection);
 		{
-			Factories.Remove(ClassName);
 			Providers.Remove(ClassName);
 		}
 	}
@@ -331,9 +356,6 @@ private:
 	/** Critical section for synchronizing access. */
 	FCriticalSection CriticalSection;
 
-	/** Maps class names to factory functions. */
-	TMap<FName, TFunction<void*()>> Factories;
-
 	/** Maps class names to instance providers. */
 	TMap<FName, TSharedPtr<IInstanceProvider>> Providers;
 };
