Allow TFunctions to be bound to TFunctions of another type (compatible by co-/contra-variance).

diff --git a/Engine/Source/Runtime/Core/Public/Templates/Function.h b/Engine/Source/Runtime/Core/Public/Templates/Function.h
--- a/Engine/Source/Runtime/Core/Public/Templates/Function.h
+++ b/Engine/Source/Runtime/Core/Public/Templates/Function.h
@@ -4,6 +4,7 @@
 
 #include <new>
 
+#include "AreTypesEqual.h"
 #include "UnrealMemory.h"
 
 // Disable visualization hack for shipping or test builds.
@@ -326,9 +327,9 @@ public:
 	 */
 #if ENABLE_TFUNCTIONREF_WORKAROUND
 	template <typename FunctorType>
-	TFunctionRef(FunctorType& Functor, typename TEnableIf<!TIsFunction<FunctorType>::Value && !TIsATFunctionRef<FunctorType>::Value>::Type* = nullptr)
+	TFunctionRef(FunctorType& Functor, typename TEnableIf<!TIsFunction<FunctorType>::Value && !TAreTypesEqual<TFunctionRef, FunctorType>::Value>::Type* = nullptr)
 #else
-	template <typename FunctorType, typename = typename TEnableIf<!TIsFunction<FunctorType>::Value && !TIsATFunctionRef<FunctorType>::Value>::Type>
+	template <typename FunctorType, typename = typename TEnableIf<!TIsFunction<FunctorType>::Value && !TAreTypesEqual<TFunctionRef, FunctorType>::Value>::Type>
 	TFunctionRef(FunctorType& Functor)
 #endif
 	{
@@ -343,9 +344,9 @@ public:
 	 */
 #if ENABLE_TFUNCTIONREF_WORKAROUND
 	template <typename FunctorType>
-	TFunctionRef(const FunctorType& Functor, typename TEnableIf<!TIsFunction<FunctorType>::Value && !TIsATFunctionRef<FunctorType>::Value>::Type* = nullptr)
+	TFunctionRef(const FunctorType& Functor, typename TEnableIf<!TIsFunction<FunctorType>::Value && !TAreTypesEqual<TFunctionRef, FunctorType>::Value>::Type* = nullptr)
 #else
-	template <typename FunctorType, typename = typename TEnableIf<!TIsFunction<FunctorType>::Value && !TIsATFunctionRef<FunctorType>::Value>::Type>
+	template <typename FunctorType, typename = typename TEnableIf<!TIsFunction<FunctorType>::Value && !TAreTypesEqual<TFunctionRef, FunctorType>::Value>::Type>
 	TFunctionRef(const FunctorType& Functor)
 #endif
 	{
@@ -542,9 +543,9 @@ public:
 	 */
 #if ENABLE_TFUNCTIONREF_WORKAROUND
 	template <typename FunctorType>
-	TFunction(FunctorType&& InFunc, typename TEnableIf<!TIsATFunction<typename TDecay<FunctorType>::Type>::Value>::Type* = nullptr)
+	TFunction(FunctorType&& InFunc, typename TEnableIf<!TAreTypesEqual<TFunction, typename TDecay<FunctorType>::Type>::Value>::Type* = nullptr)
 #else
-	template <typename FunctorType, typename = typename TEnableIf<!TIsATFunction<typename TDecay<FunctorType>::Type>::Value>::Type>
+	template <typename FunctorType, typename = typename TEnableIf<!TAreTypesEqual<TFunction, typename TDecay<FunctorType>::Type>::Value>::Type>
 	TFunction(FunctorType&& InFunc)
 #endif
 		: Super(NoInit)
