- Removed bad NULL platform check [UE-8875]

diff --git a/Engine/Source/Runtime/Core/Private/Misc/ConfigCacheIni.cpp b/Engine/Source/Runtime/Core/Private/Misc/ConfigCacheIni.cpp
--- a/Engine/Source/Runtime/Core/Private/Misc/ConfigCacheIni.cpp
+++ b/Engine/Source/Runtime/Core/Private/Misc/ConfigCacheIni.cpp
@@ -2761,7 +2761,7 @@ static void GetSourceIniHierarchyFilenames(const TCHAR* InBaseIniName, const TCH
 	FString HierarchyCheckpointPath = FString::Printf(TEXT("%sNoRedist/Default%s.ini"), SourceConfigDir, InBaseIniName);
 	OutHierarchy.Add(FIniFilename(HierarchyCheckpointPath, false, GenerateHierarchyCacheKey(OutHierarchy, HierarchyCheckpointPath, InBaseIniName)));
 	
-	if (InPlatformName && InPlatformName[0])
+	if (PlatformName.Len() > 0)
 	{
 		// [[[[ PLATFORM DEFAULTS AND PROJECT SETTINGS ]]]]
 		// Engine/Config/Platform/Platform* ini
