Adding guards to prevent crashes when log files fail to open.

diff --git a/Engine/Source/Runtime/Core/Public/HAL/IPlatformFileOpenLogWrapper.h b/Engine/Source/Runtime/Core/Public/HAL/IPlatformFileOpenLogWrapper.h
--- a/Engine/Source/Runtime/Core/Public/HAL/IPlatformFileOpenLogWrapper.h
+++ b/Engine/Source/Runtime/Core/Public/HAL/IPlatformFileOpenLogWrapper.h
@@ -57,7 +57,11 @@ public:
 				LogFilePath = FPaths::Combine( *LogFileDirectory, TEXT("GameOpenOrder.log"));
 #endif
 				Inner->CreateDirectoryTree(*LogFileDirectory);
-				LogOutput.Add(Inner->OpenWrite(*LogFilePath, false, false));
+				auto* FileHandle = Inner->OpenWrite(*LogFilePath, false, false);
+				if (FileHandle) 
+				{
+					LogOutput.Add(FileHandle);
+				}
 			}
 		}
 		else
@@ -69,7 +73,11 @@ public:
 			LogFilePath = FPaths::Combine( *LogFileDirectory, TEXT("GameOpenOrder.log"));
 #endif
 			Inner->CreateDirectoryTree(*LogFileDirectory);
-			LogOutput.Add(Inner->OpenWrite(*LogFilePath, false, false));
+			auto* FileHandle = Inner->OpenWrite(*LogFilePath, false, false);
+			if (FileHandle)
+			{
+				LogOutput.Add(FileHandle);
+			}
 		}
 		return true;
 	}
