Improved handling content with case-sensitivity problems.

- Makes uncooked UT work on Linux.
- Still uncached and suboptimal.

(Merging CL 2431047 from UT branch to main)

diff --git a/Engine/Source/Runtime/Core/Private/Linux/LinuxPlatformFile.cpp b/Engine/Source/Runtime/Core/Private/Linux/LinuxPlatformFile.cpp
--- a/Engine/Source/Runtime/Core/Private/Linux/LinuxPlatformFile.cpp
+++ b/Engine/Source/Runtime/Core/Private/Linux/LinuxPlatformFile.cpp
@@ -297,158 +297,15 @@
 #endif // MANAGE_FILE_HANDLES
 
 /**
- * Linux File I/O implementation
-**/
-FString FLinuxPlatformFile::NormalizeFilename(const TCHAR* Filename)
-{
-	FString Result(Filename);
-	FPaths::NormalizeFilename(Result);
-	return FPaths::ConvertRelativePathToFull(Result);
-}
-
-FString FLinuxPlatformFile::NormalizeDirectory(const TCHAR* Directory)
-{
-	FString Result(Directory);
-	FPaths::NormalizeDirectoryName(Result);
-	return FPaths::ConvertRelativePathToFull(Result);
-}
-
-bool FLinuxPlatformFile::FileExists(const TCHAR* Filename)
-{
-	struct stat FileInfo;
-	if (stat(TCHAR_TO_UTF8(*NormalizeFilename(Filename)), &FileInfo) != -1)
-	{
-		return S_ISREG(FileInfo.st_mode);
-	}
-	return false;
-}
-
-int64 FLinuxPlatformFile::FileSize(const TCHAR* Filename)
-{
-	struct stat FileInfo;
-	FileInfo.st_size = -1;
-	if (stat(TCHAR_TO_UTF8(*NormalizeFilename(Filename)), &FileInfo) != -1)
-	{
-		// make sure to return -1 for directories
-		if (S_ISDIR(FileInfo.st_mode))
-		{
-			FileInfo.st_size = -1;
-		}
-	}
-	return FileInfo.st_size;
-}
-
-bool FLinuxPlatformFile::DeleteFile(const TCHAR* Filename)
-{
-	return unlink(TCHAR_TO_UTF8(*NormalizeFilename(Filename))) == 0;
-}
-
-bool FLinuxPlatformFile::IsReadOnly(const TCHAR* Filename)
-{
-	if (access(TCHAR_TO_UTF8(*NormalizeFilename(Filename)), F_OK) == -1)
-	{
-		return false; // file doesn't exist
-	}
-	if (access(TCHAR_TO_UTF8(*NormalizeFilename(Filename)), W_OK) == -1)
-	{
-		return errno == EACCES;
-	}
-	return false;
-}
-
-bool FLinuxPlatformFile::MoveFile(const TCHAR* To, const TCHAR* From)
-{
-	return rename(TCHAR_TO_UTF8(*NormalizeFilename(From)), TCHAR_TO_UTF8(*NormalizeFilename(To))) != -1;
-}
-
-bool FLinuxPlatformFile::SetReadOnly(const TCHAR* Filename, bool bNewReadOnlyValue)
-{
-	struct stat FileInfo;
-	if (stat(TCHAR_TO_UTF8(*NormalizeFilename(Filename)), &FileInfo) != -1)
-	{
-		if (bNewReadOnlyValue)
-		{
-			FileInfo.st_mode &= ~S_IWUSR;
-		}
-		else
-		{
-			FileInfo.st_mode |= S_IWUSR;
-		}
-		return chmod(TCHAR_TO_UTF8(*NormalizeFilename(Filename)), FileInfo.st_mode);
-	}
-	return false;
-}
-
-FDateTime FLinuxPlatformFile::GetTimeStamp(const TCHAR* Filename)
-{
-	// get file times
-	struct stat FileInfo;
-	if(stat(TCHAR_TO_UTF8(*NormalizeFilename(Filename)), &FileInfo) == -1)
-	{
-		if (errno == EOVERFLOW)
-		{
-			// hacky workaround for files mounted on Samba (see https://bugzilla.samba.org/show_bug.cgi?id=7707)
-			return FDateTime::Now();
-		}
-		else
-		{
-			return FDateTime::MinValue();
-		}
-	}
-
-	// convert _stat time to FDateTime
-	FTimespan TimeSinceEpoch(0, 0, FileInfo.st_mtime);
-	return UnixEpoch + TimeSinceEpoch;
-}
-
-void FLinuxPlatformFile::SetTimeStamp(const TCHAR* Filename, const FDateTime DateTime)
-{
-	// get file times
-	struct stat FileInfo;
-	if(stat(TCHAR_TO_UTF8(*NormalizeFilename(Filename)), &FileInfo) == -1)
-	{
-		return;
-	}
-
-	// change the modification time only
-	struct utimbuf Times;
-	Times.actime = FileInfo.st_atime;
-	Times.modtime = (DateTime - UnixEpoch).GetTotalSeconds();
-	utime(TCHAR_TO_UTF8(*NormalizeFilename(Filename)), &Times);
-}
-
-FDateTime FLinuxPlatformFile::GetAccessTimeStamp(const TCHAR* Filename)
-{
-	// get file times
-	struct stat FileInfo;
-	if(stat(TCHAR_TO_UTF8(*NormalizeFilename(Filename)), &FileInfo) == -1)
-	{
-		return FDateTime::MinValue();
-	}
-
-	// convert _stat time to FDateTime
-	FTimespan TimeSinceEpoch(0, 0, FileInfo.st_atime);
-	return UnixEpoch + TimeSinceEpoch;
-}
-
-FString FLinuxPlatformFile::GetFilenameOnDisk(const TCHAR* Filename)
-{
-	return Filename;
-}
-
-/**
  * A class to handle case insensitive file opening. This is a band-aid, non-performant approach,
  * without any caching.
  */
 class FLinuxFileMapper
 {
-	/** Max path component */
-	int MaxPathComponents;
 
 public:
 
 	FLinuxFileMapper()
-		:	MaxPathComponents(-1)
 	{
 	}
 
@@ -515,11 +372,12 @@
 	 * 
 	 * @param Filename Original file path as requested (absolute)
 	 * @param PathComponentToLookFor Part of path we are currently trying to find.
+	 * @param MaxPathComponents Maximum number of path components (directories), i.e. how deep the path is.
 	 * @param ConstructedPath The real (absolute) path that we have found so far
 	 * 
 	 * @return a handle opened with open()
 	 */
-	int32 TryOpenRecursively(const FString & Filename, int PathComponentToLookFor, FString & ConstructedPath)
+	bool MapFileRecursively(const FString & Filename, int PathComponentToLookFor, int MaxPathComponents, FString & ConstructedPath)
 	{
 		// get the directory without the last path component
 		FString BaseDir = ConstructedPath;
@@ -528,7 +386,7 @@
 		FString PathComponent = GetPathComponent(Filename, PathComponentToLookFor);
 		FString PathComponentLower = PathComponent.ToLower();
 
-		int32 Handle = -1;
+		bool bFound = false;
 
 		// see if we can open this (we should)
 		DIR* DirHandle = opendir(TCHAR_TO_UTF8(*BaseDir));
@@ -559,8 +417,8 @@
 							FString NewConstructedPath = ConstructedPath;
 							NewConstructedPath /= DirEntry;
 
-							Handle = TryOpenRecursively(Filename, PathComponentToLookFor + 1, NewConstructedPath);
-							if (Handle != -1)
+							bFound = MapFileRecursively(Filename, PathComponentToLookFor + 1, MaxPathComponents, NewConstructedPath);
+							if (bFound)
 							{
 								ConstructedPath = NewConstructedPath;
 								break;
@@ -573,8 +431,9 @@
 						FString ConstructedFilename = ConstructedPath;
 						ConstructedFilename /= DirEntry;
 
-						Handle = open(TCHAR_TO_UTF8(*ConstructedFilename), O_RDONLY);
-						if (Handle != -1)
+						struct stat StatInfo;
+						bFound = (stat(TCHAR_TO_UTF8(*ConstructedFilename), &StatInfo) == 0);
+						if (bFound)
 						{
 							ConstructedPath = ConstructedFilename;
 							break;
@@ -584,10 +443,49 @@
 			}
 			closedir(DirHandle);
 		}
+		
+		return bFound;
+	}
 
-		return Handle;
+	/**
+	 * Tries to map a filename (one with a possibly wrong case) to one that exists.
+	 * 
+	 * @param PossiblyWrongFilename absolute filename (that has possibly a wrong case)
+	 * @param ExistingFilename filename that exists (only valid to use if the function returned success).
+	 */
+	bool MapCaseInsensitiveFile(const FString & PossiblyWrongFilename, FString & ExistingFilename)
+	{
+		// Cannot log anything here, as this may result in infinite recursion when this function is called on log file itself
+
+		// try the filename as given first
+		struct stat StatInfo;
+		bool bFound = stat(TCHAR_TO_UTF8(*PossiblyWrongFilename), &StatInfo) == 0;
+
+		if (bFound)
+		{
+			ExistingFilename = PossiblyWrongFilename;
+		}
+		else
+		{
+			// perform a case-insensitive search
+			// make sure we get the absolute filename
+			checkf(PossiblyWrongFilename[0] == TEXT('/'), TEXT("Filename '%s' given to OpenCaseInsensitiveRead is not absolute!"), *PossiblyWrongFilename);
+
+			int MaxPathComponents = CountPathComponents(PossiblyWrongFilename);
+			if (MaxPathComponents > 0)
+			{
+				FString FoundFilename(TEXT("/"));	// start with root
+				bFound = MapFileRecursively(PossiblyWrongFilename, 0, MaxPathComponents, FoundFilename);
+				if (bFound)
+				{
+					ExistingFilename = FoundFilename;
+				}
+			}
+		}
+
+		return bFound;
 	}
-	
+
 	/**
 	 * Opens a file for reading, disregarding the case.
 	 * 
@@ -616,18 +514,20 @@
 				// make sure we get the absolute filename
 				checkf(Filename[0] == TEXT('/'), TEXT("Filename '%s' given to OpenCaseInsensitiveRead is not absolute!"), *Filename);
 				
-				MaxPathComponents = CountPathComponents(Filename);
+				int MaxPathComponents = CountPathComponents(Filename);
 				if (MaxPathComponents > 0)
 				{
 					FString FoundFilename(TEXT("/"));	// start with root
-					Handle = TryOpenRecursively(Filename, 0, FoundFilename);
-
-					if (Handle != -1)
+					if (MapFileRecursively(Filename, 0, MaxPathComponents, FoundFilename))
 					{
-						MappedToFilename = FoundFilename;
-						if (Filename != MappedToFilename)
+						Handle = open(TCHAR_TO_UTF8(*FoundFilename), O_RDONLY);
+						if (Handle != -1)
 						{
-							UE_LOG(LogLinuxPlatformFile, Log, TEXT("Mapped '%s' to '%s'"), *Filename, *MappedToFilename);
+							MappedToFilename = FoundFilename;
+							if (Filename != MappedToFilename)
+							{
+								UE_LOG(LogLinuxPlatformFile, Log, TEXT("Mapped '%s' to '%s'"), *Filename, *MappedToFilename);
+							}
 						}
 					}
 				}
@@ -637,12 +537,228 @@
 	}
 };
 
+FLinuxFileMapper GCaseInsensMapper;
 
+/**
+ * Linux File I/O implementation
+**/
+FString FLinuxPlatformFile::NormalizeFilename(const TCHAR* Filename)
+{
+	FString Result(Filename);
+	FPaths::NormalizeFilename(Result);
+	return FPaths::ConvertRelativePathToFull(Result);
+}
+
+FString FLinuxPlatformFile::NormalizeDirectory(const TCHAR* Directory)
+{
+	FString Result(Directory);
+	FPaths::NormalizeDirectoryName(Result);
+	return FPaths::ConvertRelativePathToFull(Result);
+}
+
+bool FLinuxPlatformFile::FileExists(const TCHAR* Filename)
+{
+	FString CaseSensitiveFilename;
+	if (!GCaseInsensMapper.MapCaseInsensitiveFile(NormalizeFilename(Filename), CaseSensitiveFilename))
+	{
+		// could not find the file
+		return false;
+	}
+
+	struct stat FileInfo;
+	if(stat(TCHAR_TO_UTF8(*CaseSensitiveFilename), &FileInfo) != -1)
+	{
+		return S_ISREG(FileInfo.st_mode);
+	}
+	return false;
+}
+
+int64 FLinuxPlatformFile::FileSize(const TCHAR* Filename)
+{
+	FString CaseSensitiveFilename;
+	if (!GCaseInsensMapper.MapCaseInsensitiveFile(NormalizeFilename(Filename), CaseSensitiveFilename))
+	{
+		// could not find the file
+		return -1;
+	}
+
+	struct stat FileInfo;
+	FileInfo.st_size = -1;
+	if (stat(TCHAR_TO_UTF8(*CaseSensitiveFilename), &FileInfo) != -1)
+	{
+		// make sure to return -1 for directories
+		if (S_ISDIR(FileInfo.st_mode))
+		{
+			FileInfo.st_size = -1;
+		}
+	}
+	return FileInfo.st_size;
+}
+
+bool FLinuxPlatformFile::DeleteFile(const TCHAR* Filename)
+{
+	FString CaseSensitiveFilename;
+	FString IntendedFilename(NormalizeFilename(Filename));
+	if (!GCaseInsensMapper.MapCaseInsensitiveFile(IntendedFilename, CaseSensitiveFilename))
+	{
+		// could not find the file
+		return false;
+	}
+
+	// removing mapped file is too dangerous
+	if (IntendedFilename != CaseSensitiveFilename)
+	{
+		UE_LOG(LogLinuxPlatformFile, Warning, TEXT("Could not find file '%s', deleting file '%s' instead (for consistency with the rest of file ops)"), *IntendedFilename, *CaseSensitiveFilename);
+	}
+	return unlink(TCHAR_TO_UTF8(*CaseSensitiveFilename)) == 0;
+}
+
+bool FLinuxPlatformFile::IsReadOnly(const TCHAR* Filename)
+{
+	FString CaseSensitiveFilename;
+	if (!GCaseInsensMapper.MapCaseInsensitiveFile(NormalizeFilename(Filename), CaseSensitiveFilename))
+	{
+		// could not find the file
+		return false;
+	}
+
+	// skipping checking F_OK since this is already taken care of by case mapper
+
+	if (access(TCHAR_TO_UTF8(*CaseSensitiveFilename), W_OK) == -1)
+	{
+		return errno == EACCES;
+	}
+	return false;
+}
+
+bool FLinuxPlatformFile::MoveFile(const TCHAR* To, const TCHAR* From)
+{
+	FString CaseSensitiveFilename;
+	if (!GCaseInsensMapper.MapCaseInsensitiveFile(NormalizeFilename(From), CaseSensitiveFilename))
+	{
+		// could not find the file
+		return false;
+	}
+
+	return rename(TCHAR_TO_UTF8(*CaseSensitiveFilename), TCHAR_TO_UTF8(*NormalizeFilename(To))) != -1;
+}
+
+bool FLinuxPlatformFile::SetReadOnly(const TCHAR* Filename, bool bNewReadOnlyValue)
+{
+	FString CaseSensitiveFilename;
+	if (!GCaseInsensMapper.MapCaseInsensitiveFile(NormalizeFilename(Filename), CaseSensitiveFilename))
+	{
+		// could not find the file
+		return false;
+	}
+
+	struct stat FileInfo;
+	if (stat(TCHAR_TO_UTF8(*CaseSensitiveFilename), &FileInfo) != -1)
+	{
+		if (bNewReadOnlyValue)
+		{
+			FileInfo.st_mode &= ~S_IWUSR;
+		}
+		else
+		{
+			FileInfo.st_mode |= S_IWUSR;
+		}
+		return chmod(TCHAR_TO_UTF8(*CaseSensitiveFilename), FileInfo.st_mode);
+	}
+	return false;
+}
+
+FDateTime FLinuxPlatformFile::GetTimeStamp(const TCHAR* Filename)
+{
+	FString CaseSensitiveFilename;
+	if (!GCaseInsensMapper.MapCaseInsensitiveFile(NormalizeFilename(Filename), CaseSensitiveFilename))
+	{
+		// could not find the file
+		return FDateTime::MinValue();
+	}
+	
+	// get file times
+	struct stat FileInfo;
+	if(stat(TCHAR_TO_UTF8(*CaseSensitiveFilename), &FileInfo) == -1)
+	{
+		if (errno == EOVERFLOW)
+		{
+			// hacky workaround for files mounted on Samba (see https://bugzilla.samba.org/show_bug.cgi?id=7707)
+			return FDateTime::Now();
+		}
+		else
+		{
+			return FDateTime::MinValue();
+		}
+	}
+
+	// convert _stat time to FDateTime
+	FTimespan TimeSinceEpoch(0, 0, FileInfo.st_mtime);
+	return UnixEpoch + TimeSinceEpoch;
+}
+
+void FLinuxPlatformFile::SetTimeStamp(const TCHAR* Filename, const FDateTime DateTime)
+{
+	FString CaseSensitiveFilename;
+	if (!GCaseInsensMapper.MapCaseInsensitiveFile(NormalizeFilename(Filename), CaseSensitiveFilename))
+	{
+		// could not find the file
+		return;
+	}
+
+	// get file times
+	struct stat FileInfo;
+	if(stat(TCHAR_TO_UTF8(*CaseSensitiveFilename), &FileInfo) == -1)
+	{
+		return;
+	}
+
+	// change the modification time only
+	struct utimbuf Times;
+	Times.actime = FileInfo.st_atime;
+	Times.modtime = (DateTime - UnixEpoch).GetTotalSeconds();
+	utime(TCHAR_TO_UTF8(*CaseSensitiveFilename), &Times);
+}
+
+FDateTime FLinuxPlatformFile::GetAccessTimeStamp(const TCHAR* Filename)
+{
+	FString CaseSensitiveFilename;
+	if (!GCaseInsensMapper.MapCaseInsensitiveFile(NormalizeFilename(Filename), CaseSensitiveFilename))
+	{
+		// could not find the file
+		return FDateTime::MinValue();
+	}
+
+	// get file times
+	struct stat FileInfo;
+	if(stat(TCHAR_TO_UTF8(*CaseSensitiveFilename), &FileInfo) == -1)
+	{
+		return FDateTime::MinValue();
+	}
+
+	// convert _stat time to FDateTime
+	FTimespan TimeSinceEpoch(0, 0, FileInfo.st_atime);
+	return UnixEpoch + TimeSinceEpoch;
+}
+
+FString FLinuxPlatformFile::GetFilenameOnDisk(const TCHAR* Filename)
+{
+	return Filename;
+/*
+	FString CaseSensitiveFilename;
+	if (!GCaseInsensMapper.MapCaseInsensitiveFile(NormalizeFilename(Filename), CaseSensitiveFilename))
+	{
+		return Filename;
+	}
+
+	return CaseSensitiveFilename;
+*/
+}
+
 IFileHandle* FLinuxPlatformFile::OpenRead(const TCHAR* Filename)
 {
-	FLinuxFileMapper CaseInsensMapper;
 	FString MappedToName;
-	int32 Handle = CaseInsensMapper.OpenCaseInsensitiveRead(TCHAR_TO_UTF8(*NormalizeFilename(Filename)), MappedToName);
+	int32 Handle = GCaseInsensMapper.OpenCaseInsensitiveRead(TCHAR_TO_UTF8(*NormalizeFilename(Filename)), MappedToName);
 	if (Handle != -1)
 	{
 		return new FFileHandleLinux(Handle, *MappedToName, true);
@@ -723,8 +839,15 @@
 
 bool FLinuxPlatformFile::DirectoryExists(const TCHAR* Directory)
 {
+	FString CaseSensitiveFilename;
+	if (!GCaseInsensMapper.MapCaseInsensitiveFile(NormalizeFilename(Directory), CaseSensitiveFilename))
+	{
+		// could not find the file
+		return false;
+	}
+
 	struct stat FileInfo;
-	if (stat(TCHAR_TO_UTF8(*NormalizeFilename(Directory)), &FileInfo) != -1)
+	if (stat(TCHAR_TO_UTF8(*CaseSensitiveFilename), &FileInfo) != -1)
 	{
 		return S_ISDIR(FileInfo.st_mode);
 	}
@@ -738,7 +861,20 @@
 
 bool FLinuxPlatformFile::DeleteDirectory(const TCHAR* Directory)
 {
-	return rmdir(TCHAR_TO_UTF8(*NormalizeFilename(Directory))) == 0;
+	FString CaseSensitiveFilename;
+	FString IntendedFilename(NormalizeFilename(Directory));
+	if (!GCaseInsensMapper.MapCaseInsensitiveFile(IntendedFilename, CaseSensitiveFilename))
+	{
+		// could not find the directory
+		return false;
+	}
+
+	// removing mapped directory is too dangerous
+	if (IntendedFilename != CaseSensitiveFilename)
+	{
+		UE_LOG(LogLinuxPlatformFile, Warning, TEXT("Could not find directory '%s', deleting '%s' instead (for consistency with the rest of file ops)"), *IntendedFilename, *CaseSensitiveFilename);
+	}
+	return rmdir(TCHAR_TO_UTF8(*CaseSensitiveFilename)) == 0;
 }
 
 bool FLinuxPlatformFile::IterateDirectory(const TCHAR* Directory, FDirectoryVisitor& Visitor)
