Fixing Mac compile error

diff --git a/Engine/Source/Runtime/Core/Public/Misc/TypeContainer.h b/Engine/Source/Runtime/Core/Public/Misc/TypeContainer.h
--- a/Engine/Source/Runtime/Core/Public/Misc/TypeContainer.h
+++ b/Engine/Source/Runtime/Core/Public/Misc/TypeContainer.h
@@ -319,7 +319,7 @@ public:
 	template<class R, class D>
 	void RegisterDelegate(D Delegate)
 	{
-		static_assert(TAreTypesEqual<typename TSharedPtr<R>, D::RetValType>::Value, "Delegate return type must be TSharedPtr<R>");
+		static_assert(TAreTypesEqual<TSharedPtr<R>, typename D::RetValType>::Value, "Delegate return type must be TSharedPtr<R>");
 
 		FScopeLock Lock(&CriticalSection);
 		{
