Fix a bunch of VS2012 compile errors.

diff --git a/Engine/Source/Runtime/Core/Public/Delegates/DelegateCombinations.h b/Engine/Source/Runtime/Core/Public/Delegates/DelegateCombinations.h
--- a/Engine/Source/Runtime/Core/Public/Delegates/DelegateCombinations.h
+++ b/Engine/Source/Runtime/Core/Public/Delegates/DelegateCombinations.h
@@ -215,7 +215,7 @@
 #undef FUNC_IS_VOID
 
 #define FUNC_SUFFIX RetVal_NoParams
-#define FUNC_RETVAL_TYPEDEF  typedef RetValType RetValType
+#define FUNC_RETVAL_TYPEDEF  typedef RetValType RetValType;
 #define FUNC_TEMPLATE_DECL RetValType
 #define FUNC_TEMPLATE_DECL_TYPENAME typename RetValType
 #define FUNC_TEMPLATE_DECL_NO_SHADOW typename RetValTypeNoShadow
@@ -622,7 +622,7 @@
 #undef FUNC_IS_VOID
 
 #define FUNC_SUFFIX RetVal_OneParam
-#define FUNC_RETVAL_TYPEDEF  typedef RetValType RetValType
+#define FUNC_RETVAL_TYPEDEF  typedef RetValType RetValType;
 #define FUNC_TEMPLATE_DECL RetValType, Param1Type
 #define FUNC_TEMPLATE_DECL_TYPENAME typename RetValType, typename Param1Type
 #define FUNC_TEMPLATE_DECL_NO_SHADOW typename RetValTypeNoShadow, typename Param1TypeNoShadow
@@ -1029,7 +1029,7 @@
 #undef FUNC_IS_VOID
 
 #define FUNC_SUFFIX RetVal_TwoParams
-#define FUNC_RETVAL_TYPEDEF  typedef RetValType RetValType
+#define FUNC_RETVAL_TYPEDEF  typedef RetValType RetValType;
 #define FUNC_TEMPLATE_DECL RetValType, Param1Type, Param2Type
 #define FUNC_TEMPLATE_DECL_TYPENAME typename RetValType, typename Param1Type, typename Param2Type
 #define FUNC_TEMPLATE_DECL_NO_SHADOW typename RetValTypeNoShadow, typename Param1TypeNoShadow, typename Param2TypeNoShadow
@@ -1436,7 +1436,7 @@
 #undef FUNC_IS_VOID
 
 #define FUNC_SUFFIX RetVal_ThreeParams
-#define FUNC_RETVAL_TYPEDEF  typedef RetValType RetValType
+#define FUNC_RETVAL_TYPEDEF  typedef RetValType RetValType;
 #define FUNC_TEMPLATE_DECL RetValType, Param1Type, Param2Type, Param3Type
 #define FUNC_TEMPLATE_DECL_TYPENAME typename RetValType, typename Param1Type, typename Param2Type, typename Param3Type
 #define FUNC_TEMPLATE_DECL_NO_SHADOW typename RetValTypeNoShadow, typename Param1TypeNoShadow, typename Param2TypeNoShadow, typename Param3TypeNoShadow
@@ -1843,7 +1843,7 @@
 #undef FUNC_IS_VOID
 
 #define FUNC_SUFFIX RetVal_FourParams
-#define FUNC_RETVAL_TYPEDEF  typedef RetValType RetValType
+#define FUNC_RETVAL_TYPEDEF  typedef RetValType RetValType;
 #define FUNC_TEMPLATE_DECL RetValType, Param1Type, Param2Type, Param3Type, Param4Type
 #define FUNC_TEMPLATE_DECL_TYPENAME typename RetValType, typename Param1Type, typename Param2Type, typename Param3Type, typename Param4Type
 #define FUNC_TEMPLATE_DECL_NO_SHADOW typename RetValTypeNoShadow, typename Param1TypeNoShadow, typename Param2TypeNoShadow, typename Param3TypeNoShadow, typename Param4TypeNoShadow
@@ -2250,7 +2250,7 @@
 #undef FUNC_IS_VOID
 
 #define FUNC_SUFFIX RetVal_FiveParams
-#define FUNC_RETVAL_TYPEDEF  typedef RetValType RetValType
+#define FUNC_RETVAL_TYPEDEF  typedef RetValType RetValType;
 #define FUNC_TEMPLATE_DECL RetValType, Param1Type, Param2Type, Param3Type, Param4Type, Param5Type
 #define FUNC_TEMPLATE_DECL_TYPENAME typename RetValType, typename Param1Type, typename Param2Type, typename Param3Type, typename Param4Type, typename Param5Type
 #define FUNC_TEMPLATE_DECL_NO_SHADOW typename RetValTypeNoShadow, typename Param1TypeNoShadow, typename Param2TypeNoShadow, typename Param3TypeNoShadow, typename Param4TypeNoShadow, typename Param5TypeNoShadow
@@ -2657,7 +2657,7 @@
 #undef FUNC_IS_VOID
 
 #define FUNC_SUFFIX RetVal_SixParams
-#define FUNC_RETVAL_TYPEDEF  typedef RetValType RetValType
+#define FUNC_RETVAL_TYPEDEF  typedef RetValType RetValType;
 #define FUNC_TEMPLATE_DECL RetValType, Param1Type, Param2Type, Param3Type, Param4Type, Param5Type, Param6Type
 #define FUNC_TEMPLATE_DECL_TYPENAME typename RetValType, typename Param1Type, typename Param2Type, typename Param3Type, typename Param4Type, typename Param5Type, typename Param6Type
 #define FUNC_TEMPLATE_DECL_NO_SHADOW typename RetValTypeNoShadow, typename Param1TypeNoShadow, typename Param2TypeNoShadow, typename Param3TypeNoShadow, typename Param4TypeNoShadow, typename Param5TypeNoShadow, typename Param6TypeNoShadow
@@ -3064,7 +3064,7 @@
 #undef FUNC_IS_VOID
 
 #define FUNC_SUFFIX RetVal_SevenParams
-#define FUNC_RETVAL_TYPEDEF  typedef RetValType RetValType
+#define FUNC_RETVAL_TYPEDEF  typedef RetValType RetValType;
 #define FUNC_TEMPLATE_DECL RetValType, Param1Type, Param2Type, Param3Type, Param4Type, Param5Type, Param6Type, Param7Type
 #define FUNC_TEMPLATE_DECL_TYPENAME typename RetValType, typename Param1Type, typename Param2Type, typename Param3Type, typename Param4Type, typename Param5Type, typename Param6Type, typename Param7Type
 #define FUNC_TEMPLATE_DECL_NO_SHADOW typename RetValTypeNoShadow, typename Param1TypeNoShadow, typename Param2TypeNoShadow, typename Param3TypeNoShadow, typename Param4TypeNoShadow, typename Param5TypeNoShadow, typename Param6TypeNoShadow, typename Param7TypeNoShadow
@@ -3471,7 +3471,7 @@
 #undef FUNC_IS_VOID
 
 #define FUNC_SUFFIX RetVal_EightParams
-#define FUNC_RETVAL_TYPEDEF  typedef RetValType RetValType
+#define FUNC_RETVAL_TYPEDEF  typedef RetValType RetValType;
 #define FUNC_TEMPLATE_DECL RetValType, Param1Type, Param2Type, Param3Type, Param4Type, Param5Type, Param6Type, Param7Type, Param8Type
 #define FUNC_TEMPLATE_DECL_TYPENAME typename RetValType, typename Param1Type, typename Param2Type, typename Param3Type, typename Param4Type, typename Param5Type, typename Param6Type, typename Param7Type, typename Param8Type
 #define FUNC_TEMPLATE_DECL_NO_SHADOW typename RetValTypeNoShadow, typename Param1TypeNoShadow, typename Param2TypeNoShadow, typename Param3TypeNoShadow, typename Param4TypeNoShadow, typename Param5TypeNoShadow, typename Param6TypeNoShadow, typename Param7TypeNoShadow, typename Param8TypeNoShadow
