Missing update to DelegateHeaderTool after CL 2412156.

diff --git a/Engine/Source/Programs/DelegateHeaderTool/Windows/Program.cs b/Engine/Source/Programs/DelegateHeaderTool/Windows/Program.cs
--- a/Engine/Source/Programs/DelegateHeaderTool/Windows/Program.cs
+++ b/Engine/Source/Programs/DelegateHeaderTool/Windows/Program.cs
@@ -239,13 +239,13 @@ namespace DelegateHeaderTool
             DefineDeclareEvent.AppendFormat(" ) FUNC_DECLARE_EVENT( OwningType, EventName, {0}", FuncSuffix);
 			if( bSupportsRetVal )
 			{
-				DefineDeclareDynamicDelegate.AppendFormat( " ) FUNC_DECLARE_DYNAMIC_DELEGATE_RETVAL( {0}, {1}, DelegateName, DelegateName##_DelegateWrapper", DynamicDelegateDefaultPtr, FuncSuffix );
+				DefineDeclareDynamicDelegate.AppendFormat(" ) BODY_MACRO_COMBINE(CURRENT_FILE_ID,_,__LINE__,_DELEGATE) FUNC_DECLARE_DYNAMIC_DELEGATE_RETVAL( {0}, {1}, DelegateName, DelegateName##_DelegateWrapper", DynamicDelegateDefaultPtr, FuncSuffix);
 			}
 			else
 			{
-				DefineDeclareDynamicDelegate.AppendFormat( " ) FUNC_DECLARE_DYNAMIC_DELEGATE( {0}, {1}, DelegateName, DelegateName##_DelegateWrapper", DynamicDelegateDefaultPtr, FuncSuffix );
+				DefineDeclareDynamicDelegate.AppendFormat(" ) BODY_MACRO_COMBINE(CURRENT_FILE_ID,_,__LINE__,_DELEGATE) FUNC_DECLARE_DYNAMIC_DELEGATE( {0}, {1}, DelegateName, DelegateName##_DelegateWrapper", DynamicDelegateDefaultPtr, FuncSuffix);
 			}
-			DefineDeclareDynamicMulticastDelegate.AppendFormat( " ) FUNC_DECLARE_DYNAMIC_MULTICAST_DELEGATE( {0}, {1}, DelegateName, DelegateName##_DelegateWrapper", DynamicDelegateDefaultPtr, FuncSuffix );
+			DefineDeclareDynamicMulticastDelegate.AppendFormat(" ) BODY_MACRO_COMBINE(CURRENT_FILE_ID,_,__LINE__,_DELEGATE) FUNC_DECLARE_DYNAMIC_MULTICAST_DELEGATE( {0}, {1}, DelegateName, DelegateName##_DelegateWrapper", DynamicDelegateDefaultPtr, FuncSuffix);
 
 			if( bSupportsRetVal )
 			{
@@ -292,7 +292,7 @@ namespace DelegateHeaderTool
 			{
 				DefineDeclareDelegate.Append( ", RetValType" );
 				DefineDeclareDynamicDelegate.Append( ", RetValType" );
-                DefineRetvalTypedef.Append( " typedef RetValType RetValType" );
+                DefineRetvalTypedef.Append( " typedef RetValType RetValType;" );
 			}
 			else
 			{
diff --git a/Engine/Source/Runtime/Core/Public/Delegates/DelegateCombinations.h b/Engine/Source/Runtime/Core/Public/Delegates/DelegateCombinations.h
--- a/Engine/Source/Runtime/Core/Public/Delegates/DelegateCombinations.h
+++ b/Engine/Source/Runtime/Core/Public/Delegates/DelegateCombinations.h
@@ -199,8 +199,8 @@
 #define DECLARE_DELEGATE( DelegateName ) FUNC_DECLARE_DELEGATE( NoParams, DelegateName, void )
 #define DECLARE_MULTICAST_DELEGATE( DelegateName ) FUNC_DECLARE_MULTICAST_DELEGATE( NoParams, DelegateName, void )
 #define DECLARE_EVENT( OwningType, EventName ) FUNC_DECLARE_EVENT( OwningType, EventName, NoParams, void )
-#define DECLARE_DYNAMIC_DELEGATE( DelegateName ) FUNC_DECLARE_DYNAMIC_DELEGATE( FWeakObjectPtr, NoParams, DelegateName, DelegateName##_DelegateWrapper, , FUNC_CONCAT( *this ), void )
-#define DECLARE_DYNAMIC_MULTICAST_DELEGATE( DelegateName ) FUNC_DECLARE_DYNAMIC_MULTICAST_DELEGATE( FWeakObjectPtr, NoParams, DelegateName, DelegateName##_DelegateWrapper, , FUNC_CONCAT( *this ), void )
+#define DECLARE_DYNAMIC_DELEGATE( DelegateName ) BODY_MACRO_COMBINE(CURRENT_FILE_ID,_,__LINE__,_DELEGATE) FUNC_DECLARE_DYNAMIC_DELEGATE( FWeakObjectPtr, NoParams, DelegateName, DelegateName##_DelegateWrapper, , FUNC_CONCAT( *this ), void )
+#define DECLARE_DYNAMIC_MULTICAST_DELEGATE( DelegateName ) BODY_MACRO_COMBINE(CURRENT_FILE_ID,_,__LINE__,_DELEGATE) FUNC_DECLARE_DYNAMIC_MULTICAST_DELEGATE( FWeakObjectPtr, NoParams, DelegateName, DelegateName##_DelegateWrapper, , FUNC_CONCAT( *this ), void )
 #undef FUNC_SUFFIX
 #undef FUNC_TEMPLATE_DECL
 #undef FUNC_TEMPLATE_DECL_TYPENAME
@@ -402,7 +402,7 @@
 #include "DelegateSignatureImpl.inl"
 
 #define DECLARE_DELEGATE_RetVal( RetValType, DelegateName ) FUNC_DECLARE_DELEGATE( RetVal_NoParams, DelegateName, RetValType )
-#define DECLARE_DYNAMIC_DELEGATE_RetVal( RetValType, DelegateName ) FUNC_DECLARE_DYNAMIC_DELEGATE_RETVAL( FWeakObjectPtr, RetVal_NoParams, DelegateName, DelegateName##_DelegateWrapper, RetValType, , FUNC_CONCAT( *this ), RetValType )
+#define DECLARE_DYNAMIC_DELEGATE_RetVal( RetValType, DelegateName ) BODY_MACRO_COMBINE(CURRENT_FILE_ID,_,__LINE__,_DELEGATE) FUNC_DECLARE_DYNAMIC_DELEGATE_RETVAL( FWeakObjectPtr, RetVal_NoParams, DelegateName, DelegateName##_DelegateWrapper, RetValType, , FUNC_CONCAT( *this ), RetValType )
 #undef FUNC_SUFFIX
 #undef FUNC_TEMPLATE_DECL
 #undef FUNC_TEMPLATE_DECL_TYPENAME
@@ -606,8 +606,8 @@
 #define DECLARE_DELEGATE_OneParam( DelegateName, Param1Type ) FUNC_DECLARE_DELEGATE( OneParam, DelegateName, void, Param1Type )
 #define DECLARE_MULTICAST_DELEGATE_OneParam( DelegateName, Param1Type ) FUNC_DECLARE_MULTICAST_DELEGATE( OneParam, DelegateName, void, Param1Type )
 #define DECLARE_EVENT_OneParam( OwningType, EventName, Param1Type ) FUNC_DECLARE_EVENT( OwningType, EventName, OneParam, void, Param1Type )
-#define DECLARE_DYNAMIC_DELEGATE_OneParam( DelegateName, Param1Type, Param1Name ) FUNC_DECLARE_DYNAMIC_DELEGATE( FWeakObjectPtr, OneParam, DelegateName, DelegateName##_DelegateWrapper, FUNC_CONCAT( Param1Type InParam1 ), FUNC_CONCAT( *this, InParam1 ), void, Param1Type )
-#define DECLARE_DYNAMIC_MULTICAST_DELEGATE_OneParam( DelegateName, Param1Type, Param1Name ) FUNC_DECLARE_DYNAMIC_MULTICAST_DELEGATE( FWeakObjectPtr, OneParam, DelegateName, DelegateName##_DelegateWrapper, FUNC_CONCAT( Param1Type InParam1 ), FUNC_CONCAT( *this, InParam1 ), void, Param1Type )
+#define DECLARE_DYNAMIC_DELEGATE_OneParam( DelegateName, Param1Type, Param1Name ) BODY_MACRO_COMBINE(CURRENT_FILE_ID,_,__LINE__,_DELEGATE) FUNC_DECLARE_DYNAMIC_DELEGATE( FWeakObjectPtr, OneParam, DelegateName, DelegateName##_DelegateWrapper, FUNC_CONCAT( Param1Type InParam1 ), FUNC_CONCAT( *this, InParam1 ), void, Param1Type )
+#define DECLARE_DYNAMIC_MULTICAST_DELEGATE_OneParam( DelegateName, Param1Type, Param1Name ) BODY_MACRO_COMBINE(CURRENT_FILE_ID,_,__LINE__,_DELEGATE) FUNC_DECLARE_DYNAMIC_MULTICAST_DELEGATE( FWeakObjectPtr, OneParam, DelegateName, DelegateName##_DelegateWrapper, FUNC_CONCAT( Param1Type InParam1 ), FUNC_CONCAT( *this, InParam1 ), void, Param1Type )
 #undef FUNC_SUFFIX
 #undef FUNC_TEMPLATE_DECL
 #undef FUNC_TEMPLATE_DECL_TYPENAME
@@ -809,7 +809,7 @@
 #include "DelegateSignatureImpl.inl"
 
 #define DECLARE_DELEGATE_RetVal_OneParam( RetValType, DelegateName, Param1Type ) FUNC_DECLARE_DELEGATE( RetVal_OneParam, DelegateName, RetValType, Param1Type )
-#define DECLARE_DYNAMIC_DELEGATE_RetVal_OneParam( RetValType, DelegateName, Param1Type, Param1Name ) FUNC_DECLARE_DYNAMIC_DELEGATE_RETVAL( FWeakObjectPtr, RetVal_OneParam, DelegateName, DelegateName##_DelegateWrapper, RetValType, FUNC_CONCAT( Param1Type InParam1 ), FUNC_CONCAT( *this, InParam1 ), RetValType, Param1Type )
+#define DECLARE_DYNAMIC_DELEGATE_RetVal_OneParam( RetValType, DelegateName, Param1Type, Param1Name ) BODY_MACRO_COMBINE(CURRENT_FILE_ID,_,__LINE__,_DELEGATE) FUNC_DECLARE_DYNAMIC_DELEGATE_RETVAL( FWeakObjectPtr, RetVal_OneParam, DelegateName, DelegateName##_DelegateWrapper, RetValType, FUNC_CONCAT( Param1Type InParam1 ), FUNC_CONCAT( *this, InParam1 ), RetValType, Param1Type )
 #undef FUNC_SUFFIX
 #undef FUNC_TEMPLATE_DECL
 #undef FUNC_TEMPLATE_DECL_TYPENAME
@@ -1013,8 +1013,8 @@
 #define DECLARE_DELEGATE_TwoParams( DelegateName, Param1Type, Param2Type ) FUNC_DECLARE_DELEGATE( TwoParams, DelegateName, void, Param1Type, Param2Type )
 #define DECLARE_MULTICAST_DELEGATE_TwoParams( DelegateName, Param1Type, Param2Type ) FUNC_DECLARE_MULTICAST_DELEGATE( TwoParams, DelegateName, void, Param1Type, Param2Type )
 #define DECLARE_EVENT_TwoParams( OwningType, EventName, Param1Type, Param2Type ) FUNC_DECLARE_EVENT( OwningType, EventName, TwoParams, void, Param1Type, Param2Type )
-#define DECLARE_DYNAMIC_DELEGATE_TwoParams( DelegateName, Param1Type, Param1Name, Param2Type, Param2Name ) FUNC_DECLARE_DYNAMIC_DELEGATE( FWeakObjectPtr, TwoParams, DelegateName, DelegateName##_DelegateWrapper, FUNC_CONCAT( Param1Type InParam1, Param2Type InParam2 ), FUNC_CONCAT( *this, InParam1, InParam2 ), void, Param1Type, Param2Type )
-#define DECLARE_DYNAMIC_MULTICAST_DELEGATE_TwoParams( DelegateName, Param1Type, Param1Name, Param2Type, Param2Name ) FUNC_DECLARE_DYNAMIC_MULTICAST_DELEGATE( FWeakObjectPtr, TwoParams, DelegateName, DelegateName##_DelegateWrapper, FUNC_CONCAT( Param1Type InParam1, Param2Type InParam2 ), FUNC_CONCAT( *this, InParam1, InParam2 ), void, Param1Type, Param2Type )
+#define DECLARE_DYNAMIC_DELEGATE_TwoParams( DelegateName, Param1Type, Param1Name, Param2Type, Param2Name ) BODY_MACRO_COMBINE(CURRENT_FILE_ID,_,__LINE__,_DELEGATE) FUNC_DECLARE_DYNAMIC_DELEGATE( FWeakObjectPtr, TwoParams, DelegateName, DelegateName##_DelegateWrapper, FUNC_CONCAT( Param1Type InParam1, Param2Type InParam2 ), FUNC_CONCAT( *this, InParam1, InParam2 ), void, Param1Type, Param2Type )
+#define DECLARE_DYNAMIC_MULTICAST_DELEGATE_TwoParams( DelegateName, Param1Type, Param1Name, Param2Type, Param2Name ) BODY_MACRO_COMBINE(CURRENT_FILE_ID,_,__LINE__,_DELEGATE) FUNC_DECLARE_DYNAMIC_MULTICAST_DELEGATE( FWeakObjectPtr, TwoParams, DelegateName, DelegateName##_DelegateWrapper, FUNC_CONCAT( Param1Type InParam1, Param2Type InParam2 ), FUNC_CONCAT( *this, InParam1, InParam2 ), void, Param1Type, Param2Type )
 #undef FUNC_SUFFIX
 #undef FUNC_TEMPLATE_DECL
 #undef FUNC_TEMPLATE_DECL_TYPENAME
@@ -1216,7 +1216,7 @@
 #include "DelegateSignatureImpl.inl"
 
 #define DECLARE_DELEGATE_RetVal_TwoParams( RetValType, DelegateName, Param1Type, Param2Type ) FUNC_DECLARE_DELEGATE( RetVal_TwoParams, DelegateName, RetValType, Param1Type, Param2Type )
-#define DECLARE_DYNAMIC_DELEGATE_RetVal_TwoParams( RetValType, DelegateName, Param1Type, Param1Name, Param2Type, Param2Name ) FUNC_DECLARE_DYNAMIC_DELEGATE_RETVAL( FWeakObjectPtr, RetVal_TwoParams, DelegateName, DelegateName##_DelegateWrapper, RetValType, FUNC_CONCAT( Param1Type InParam1, Param2Type InParam2 ), FUNC_CONCAT( *this, InParam1, InParam2 ), RetValType, Param1Type, Param2Type )
+#define DECLARE_DYNAMIC_DELEGATE_RetVal_TwoParams( RetValType, DelegateName, Param1Type, Param1Name, Param2Type, Param2Name ) BODY_MACRO_COMBINE(CURRENT_FILE_ID,_,__LINE__,_DELEGATE) FUNC_DECLARE_DYNAMIC_DELEGATE_RETVAL( FWeakObjectPtr, RetVal_TwoParams, DelegateName, DelegateName##_DelegateWrapper, RetValType, FUNC_CONCAT( Param1Type InParam1, Param2Type InParam2 ), FUNC_CONCAT( *this, InParam1, InParam2 ), RetValType, Param1Type, Param2Type )
 #undef FUNC_SUFFIX
 #undef FUNC_TEMPLATE_DECL
 #undef FUNC_TEMPLATE_DECL_TYPENAME
@@ -1420,8 +1420,8 @@
 #define DECLARE_DELEGATE_ThreeParams( DelegateName, Param1Type, Param2Type, Param3Type ) FUNC_DECLARE_DELEGATE( ThreeParams, DelegateName, void, Param1Type, Param2Type, Param3Type )
 #define DECLARE_MULTICAST_DELEGATE_ThreeParams( DelegateName, Param1Type, Param2Type, Param3Type ) FUNC_DECLARE_MULTICAST_DELEGATE( ThreeParams, DelegateName, void, Param1Type, Param2Type, Param3Type )
 #define DECLARE_EVENT_ThreeParams( OwningType, EventName, Param1Type, Param2Type, Param3Type ) FUNC_DECLARE_EVENT( OwningType, EventName, ThreeParams, void, Param1Type, Param2Type, Param3Type )
-#define DECLARE_DYNAMIC_DELEGATE_ThreeParams( DelegateName, Param1Type, Param1Name, Param2Type, Param2Name, Param3Type, Param3Name ) FUNC_DECLARE_DYNAMIC_DELEGATE( FWeakObjectPtr, ThreeParams, DelegateName, DelegateName##_DelegateWrapper, FUNC_CONCAT( Param1Type InParam1, Param2Type InParam2, Param3Type InParam3 ), FUNC_CONCAT( *this, InParam1, InParam2, InParam3 ), void, Param1Type, Param2Type, Param3Type )
-#define DECLARE_DYNAMIC_MULTICAST_DELEGATE_ThreeParams( DelegateName, Param1Type, Param1Name, Param2Type, Param2Name, Param3Type, Param3Name ) FUNC_DECLARE_DYNAMIC_MULTICAST_DELEGATE( FWeakObjectPtr, ThreeParams, DelegateName, DelegateName##_DelegateWrapper, FUNC_CONCAT( Param1Type InParam1, Param2Type InParam2, Param3Type InParam3 ), FUNC_CONCAT( *this, InParam1, InParam2, InParam3 ), void, Param1Type, Param2Type, Param3Type )
+#define DECLARE_DYNAMIC_DELEGATE_ThreeParams( DelegateName, Param1Type, Param1Name, Param2Type, Param2Name, Param3Type, Param3Name ) BODY_MACRO_COMBINE(CURRENT_FILE_ID,_,__LINE__,_DELEGATE) FUNC_DECLARE_DYNAMIC_DELEGATE( FWeakObjectPtr, ThreeParams, DelegateName, DelegateName##_DelegateWrapper, FUNC_CONCAT( Param1Type InParam1, Param2Type InParam2, Param3Type InParam3 ), FUNC_CONCAT( *this, InParam1, InParam2, InParam3 ), void, Param1Type, Param2Type, Param3Type )
+#define DECLARE_DYNAMIC_MULTICAST_DELEGATE_ThreeParams( DelegateName, Param1Type, Param1Name, Param2Type, Param2Name, Param3Type, Param3Name ) BODY_MACRO_COMBINE(CURRENT_FILE_ID,_,__LINE__,_DELEGATE) FUNC_DECLARE_DYNAMIC_MULTICAST_DELEGATE( FWeakObjectPtr, ThreeParams, DelegateName, DelegateName##_DelegateWrapper, FUNC_CONCAT( Param1Type InParam1, Param2Type InParam2, Param3Type InParam3 ), FUNC_CONCAT( *this, InParam1, InParam2, InParam3 ), void, Param1Type, Param2Type, Param3Type )
 #undef FUNC_SUFFIX
 #undef FUNC_TEMPLATE_DECL
 #undef FUNC_TEMPLATE_DECL_TYPENAME
@@ -1623,7 +1623,7 @@
 #include "DelegateSignatureImpl.inl"
 
 #define DECLARE_DELEGATE_RetVal_ThreeParams( RetValType, DelegateName, Param1Type, Param2Type, Param3Type ) FUNC_DECLARE_DELEGATE( RetVal_ThreeParams, DelegateName, RetValType, Param1Type, Param2Type, Param3Type )
-#define DECLARE_DYNAMIC_DELEGATE_RetVal_ThreeParams( RetValType, DelegateName, Param1Type, Param1Name, Param2Type, Param2Name, Param3Type, Param3Name ) FUNC_DECLARE_DYNAMIC_DELEGATE_RETVAL( FWeakObjectPtr, RetVal_ThreeParams, DelegateName, DelegateName##_DelegateWrapper, RetValType, FUNC_CONCAT( Param1Type InParam1, Param2Type InParam2, Param3Type InParam3 ), FUNC_CONCAT( *this, InParam1, InParam2, InParam3 ), RetValType, Param1Type, Param2Type, Param3Type )
+#define DECLARE_DYNAMIC_DELEGATE_RetVal_ThreeParams( RetValType, DelegateName, Param1Type, Param1Name, Param2Type, Param2Name, Param3Type, Param3Name ) BODY_MACRO_COMBINE(CURRENT_FILE_ID,_,__LINE__,_DELEGATE) FUNC_DECLARE_DYNAMIC_DELEGATE_RETVAL( FWeakObjectPtr, RetVal_ThreeParams, DelegateName, DelegateName##_DelegateWrapper, RetValType, FUNC_CONCAT( Param1Type InParam1, Param2Type InParam2, Param3Type InParam3 ), FUNC_CONCAT( *this, InParam1, InParam2, InParam3 ), RetValType, Param1Type, Param2Type, Param3Type )
 #undef FUNC_SUFFIX
 #undef FUNC_TEMPLATE_DECL
 #undef FUNC_TEMPLATE_DECL_TYPENAME
@@ -1827,8 +1827,8 @@
 #define DECLARE_DELEGATE_FourParams( DelegateName, Param1Type, Param2Type, Param3Type, Param4Type ) FUNC_DECLARE_DELEGATE( FourParams, DelegateName, void, Param1Type, Param2Type, Param3Type, Param4Type )
 #define DECLARE_MULTICAST_DELEGATE_FourParams( DelegateName, Param1Type, Param2Type, Param3Type, Param4Type ) FUNC_DECLARE_MULTICAST_DELEGATE( FourParams, DelegateName, void, Param1Type, Param2Type, Param3Type, Param4Type )
 #define DECLARE_EVENT_FourParams( OwningType, EventName, Param1Type, Param2Type, Param3Type, Param4Type ) FUNC_DECLARE_EVENT( OwningType, EventName, FourParams, void, Param1Type, Param2Type, Param3Type, Param4Type )
-#define DECLARE_DYNAMIC_DELEGATE_FourParams( DelegateName, Param1Type, Param1Name, Param2Type, Param2Name, Param3Type, Param3Name, Param4Type, Param4Name ) FUNC_DECLARE_DYNAMIC_DELEGATE( FWeakObjectPtr, FourParams, DelegateName, DelegateName##_DelegateWrapper, FUNC_CONCAT( Param1Type InParam1, Param2Type InParam2, Param3Type InParam3, Param4Type InParam4 ), FUNC_CONCAT( *this, InParam1, InParam2, InParam3, InParam4 ), void, Param1Type, Param2Type, Param3Type, Param4Type )
-#define DECLARE_DYNAMIC_MULTICAST_DELEGATE_FourParams( DelegateName, Param1Type, Param1Name, Param2Type, Param2Name, Param3Type, Param3Name, Param4Type, Param4Name ) FUNC_DECLARE_DYNAMIC_MULTICAST_DELEGATE( FWeakObjectPtr, FourParams, DelegateName, DelegateName##_DelegateWrapper, FUNC_CONCAT( Param1Type InParam1, Param2Type InParam2, Param3Type InParam3, Param4Type InParam4 ), FUNC_CONCAT( *this, InParam1, InParam2, InParam3, InParam4 ), void, Param1Type, Param2Type, Param3Type, Param4Type )
+#define DECLARE_DYNAMIC_DELEGATE_FourParams( DelegateName, Param1Type, Param1Name, Param2Type, Param2Name, Param3Type, Param3Name, Param4Type, Param4Name ) BODY_MACRO_COMBINE(CURRENT_FILE_ID,_,__LINE__,_DELEGATE) FUNC_DECLARE_DYNAMIC_DELEGATE( FWeakObjectPtr, FourParams, DelegateName, DelegateName##_DelegateWrapper, FUNC_CONCAT( Param1Type InParam1, Param2Type InParam2, Param3Type InParam3, Param4Type InParam4 ), FUNC_CONCAT( *this, InParam1, InParam2, InParam3, InParam4 ), void, Param1Type, Param2Type, Param3Type, Param4Type )
+#define DECLARE_DYNAMIC_MULTICAST_DELEGATE_FourParams( DelegateName, Param1Type, Param1Name, Param2Type, Param2Name, Param3Type, Param3Name, Param4Type, Param4Name ) BODY_MACRO_COMBINE(CURRENT_FILE_ID,_,__LINE__,_DELEGATE) FUNC_DECLARE_DYNAMIC_MULTICAST_DELEGATE( FWeakObjectPtr, FourParams, DelegateName, DelegateName##_DelegateWrapper, FUNC_CONCAT( Param1Type InParam1, Param2Type InParam2, Param3Type InParam3, Param4Type InParam4 ), FUNC_CONCAT( *this, InParam1, InParam2, InParam3, InParam4 ), void, Param1Type, Param2Type, Param3Type, Param4Type )
 #undef FUNC_SUFFIX
 #undef FUNC_TEMPLATE_DECL
 #undef FUNC_TEMPLATE_DECL_TYPENAME
@@ -2030,7 +2030,7 @@
 #include "DelegateSignatureImpl.inl"
 
 #define DECLARE_DELEGATE_RetVal_FourParams( RetValType, DelegateName, Param1Type, Param2Type, Param3Type, Param4Type ) FUNC_DECLARE_DELEGATE( RetVal_FourParams, DelegateName, RetValType, Param1Type, Param2Type, Param3Type, Param4Type )
-#define DECLARE_DYNAMIC_DELEGATE_RetVal_FourParams( RetValType, DelegateName, Param1Type, Param1Name, Param2Type, Param2Name, Param3Type, Param3Name, Param4Type, Param4Name ) FUNC_DECLARE_DYNAMIC_DELEGATE_RETVAL( FWeakObjectPtr, RetVal_FourParams, DelegateName, DelegateName##_DelegateWrapper, RetValType, FUNC_CONCAT( Param1Type InParam1, Param2Type InParam2, Param3Type InParam3, Param4Type InParam4 ), FUNC_CONCAT( *this, InParam1, InParam2, InParam3, InParam4 ), RetValType, Param1Type, Param2Type, Param3Type, Param4Type )
+#define DECLARE_DYNAMIC_DELEGATE_RetVal_FourParams( RetValType, DelegateName, Param1Type, Param1Name, Param2Type, Param2Name, Param3Type, Param3Name, Param4Type, Param4Name ) BODY_MACRO_COMBINE(CURRENT_FILE_ID,_,__LINE__,_DELEGATE) FUNC_DECLARE_DYNAMIC_DELEGATE_RETVAL( FWeakObjectPtr, RetVal_FourParams, DelegateName, DelegateName##_DelegateWrapper, RetValType, FUNC_CONCAT( Param1Type InParam1, Param2Type InParam2, Param3Type InParam3, Param4Type InParam4 ), FUNC_CONCAT( *this, InParam1, InParam2, InParam3, InParam4 ), RetValType, Param1Type, Param2Type, Param3Type, Param4Type )
 #undef FUNC_SUFFIX
 #undef FUNC_TEMPLATE_DECL
 #undef FUNC_TEMPLATE_DECL_TYPENAME
@@ -2234,8 +2234,8 @@
 #define DECLARE_DELEGATE_FiveParams( DelegateName, Param1Type, Param2Type, Param3Type, Param4Type, Param5Type ) FUNC_DECLARE_DELEGATE( FiveParams, DelegateName, void, Param1Type, Param2Type, Param3Type, Param4Type, Param5Type )
 #define DECLARE_MULTICAST_DELEGATE_FiveParams( DelegateName, Param1Type, Param2Type, Param3Type, Param4Type, Param5Type ) FUNC_DECLARE_MULTICAST_DELEGATE( FiveParams, DelegateName, void, Param1Type, Param2Type, Param3Type, Param4Type, Param5Type )
 #define DECLARE_EVENT_FiveParams( OwningType, EventName, Param1Type, Param2Type, Param3Type, Param4Type, Param5Type ) FUNC_DECLARE_EVENT( OwningType, EventName, FiveParams, void, Param1Type, Param2Type, Param3Type, Param4Type, Param5Type )
-#define DECLARE_DYNAMIC_DELEGATE_FiveParams( DelegateName, Param1Type, Param1Name, Param2Type, Param2Name, Param3Type, Param3Name, Param4Type, Param4Name, Param5Type, Param5Name ) FUNC_DECLARE_DYNAMIC_DELEGATE( FWeakObjectPtr, FiveParams, DelegateName, DelegateName##_DelegateWrapper, FUNC_CONCAT( Param1Type InParam1, Param2Type InParam2, Param3Type InParam3, Param4Type InParam4, Param5Type InParam5 ), FUNC_CONCAT( *this, InParam1, InParam2, InParam3, InParam4, InParam5 ), void, Param1Type, Param2Type, Param3Type, Param4Type, Param5Type )
-#define DECLARE_DYNAMIC_MULTICAST_DELEGATE_FiveParams( DelegateName, Param1Type, Param1Name, Param2Type, Param2Name, Param3Type, Param3Name, Param4Type, Param4Name, Param5Type, Param5Name ) FUNC_DECLARE_DYNAMIC_MULTICAST_DELEGATE( FWeakObjectPtr, FiveParams, DelegateName, DelegateName##_DelegateWrapper, FUNC_CONCAT( Param1Type InParam1, Param2Type InParam2, Param3Type InParam3, Param4Type InParam4, Param5Type InParam5 ), FUNC_CONCAT( *this, InParam1, InParam2, InParam3, InParam4, InParam5 ), void, Param1Type, Param2Type, Param3Type, Param4Type, Param5Type )
+#define DECLARE_DYNAMIC_DELEGATE_FiveParams( DelegateName, Param1Type, Param1Name, Param2Type, Param2Name, Param3Type, Param3Name, Param4Type, Param4Name, Param5Type, Param5Name ) BODY_MACRO_COMBINE(CURRENT_FILE_ID,_,__LINE__,_DELEGATE) FUNC_DECLARE_DYNAMIC_DELEGATE( FWeakObjectPtr, FiveParams, DelegateName, DelegateName##_DelegateWrapper, FUNC_CONCAT( Param1Type InParam1, Param2Type InParam2, Param3Type InParam3, Param4Type InParam4, Param5Type InParam5 ), FUNC_CONCAT( *this, InParam1, InParam2, InParam3, InParam4, InParam5 ), void, Param1Type, Param2Type, Param3Type, Param4Type, Param5Type )
+#define DECLARE_DYNAMIC_MULTICAST_DELEGATE_FiveParams( DelegateName, Param1Type, Param1Name, Param2Type, Param2Name, Param3Type, Param3Name, Param4Type, Param4Name, Param5Type, Param5Name ) BODY_MACRO_COMBINE(CURRENT_FILE_ID,_,__LINE__,_DELEGATE) FUNC_DECLARE_DYNAMIC_MULTICAST_DELEGATE( FWeakObjectPtr, FiveParams, DelegateName, DelegateName##_DelegateWrapper, FUNC_CONCAT( Param1Type InParam1, Param2Type InParam2, Param3Type InParam3, Param4Type InParam4, Param5Type InParam5 ), FUNC_CONCAT( *this, InParam1, InParam2, InParam3, InParam4, InParam5 ), void, Param1Type, Param2Type, Param3Type, Param4Type, Param5Type )
 #undef FUNC_SUFFIX
 #undef FUNC_TEMPLATE_DECL
 #undef FUNC_TEMPLATE_DECL_TYPENAME
@@ -2437,7 +2437,7 @@
 #include "DelegateSignatureImpl.inl"
 
 #define DECLARE_DELEGATE_RetVal_FiveParams( RetValType, DelegateName, Param1Type, Param2Type, Param3Type, Param4Type, Param5Type ) FUNC_DECLARE_DELEGATE( RetVal_FiveParams, DelegateName, RetValType, Param1Type, Param2Type, Param3Type, Param4Type, Param5Type )
-#define DECLARE_DYNAMIC_DELEGATE_RetVal_FiveParams( RetValType, DelegateName, Param1Type, Param1Name, Param2Type, Param2Name, Param3Type, Param3Name, Param4Type, Param4Name, Param5Type, Param5Name ) FUNC_DECLARE_DYNAMIC_DELEGATE_RETVAL( FWeakObjectPtr, RetVal_FiveParams, DelegateName, DelegateName##_DelegateWrapper, RetValType, FUNC_CONCAT( Param1Type InParam1, Param2Type InParam2, Param3Type InParam3, Param4Type InParam4, Param5Type InParam5 ), FUNC_CONCAT( *this, InParam1, InParam2, InParam3, InParam4, InParam5 ), RetValType, Param1Type, Param2Type, Param3Type, Param4Type, Param5Type )
+#define DECLARE_DYNAMIC_DELEGATE_RetVal_FiveParams( RetValType, DelegateName, Param1Type, Param1Name, Param2Type, Param2Name, Param3Type, Param3Name, Param4Type, Param4Name, Param5Type, Param5Name ) BODY_MACRO_COMBINE(CURRENT_FILE_ID,_,__LINE__,_DELEGATE) FUNC_DECLARE_DYNAMIC_DELEGATE_RETVAL( FWeakObjectPtr, RetVal_FiveParams, DelegateName, DelegateName##_DelegateWrapper, RetValType, FUNC_CONCAT( Param1Type InParam1, Param2Type InParam2, Param3Type InParam3, Param4Type InParam4, Param5Type InParam5 ), FUNC_CONCAT( *this, InParam1, InParam2, InParam3, InParam4, InParam5 ), RetValType, Param1Type, Param2Type, Param3Type, Param4Type, Param5Type )
 #undef FUNC_SUFFIX
 #undef FUNC_TEMPLATE_DECL
 #undef FUNC_TEMPLATE_DECL_TYPENAME
@@ -2641,8 +2641,8 @@
 #define DECLARE_DELEGATE_SixParams( DelegateName, Param1Type, Param2Type, Param3Type, Param4Type, Param5Type, Param6Type ) FUNC_DECLARE_DELEGATE( SixParams, DelegateName, void, Param1Type, Param2Type, Param3Type, Param4Type, Param5Type, Param6Type )
 #define DECLARE_MULTICAST_DELEGATE_SixParams( DelegateName, Param1Type, Param2Type, Param3Type, Param4Type, Param5Type, Param6Type ) FUNC_DECLARE_MULTICAST_DELEGATE( SixParams, DelegateName, void, Param1Type, Param2Type, Param3Type, Param4Type, Param5Type, Param6Type )
 #define DECLARE_EVENT_SixParams( OwningType, EventName, Param1Type, Param2Type, Param3Type, Param4Type, Param5Type, Param6Type ) FUNC_DECLARE_EVENT( OwningType, EventName, SixParams, void, Param1Type, Param2Type, Param3Type, Param4Type, Param5Type, Param6Type )
-#define DECLARE_DYNAMIC_DELEGATE_SixParams( DelegateName, Param1Type, Param1Name, Param2Type, Param2Name, Param3Type, Param3Name, Param4Type, Param4Name, Param5Type, Param5Name, Param6Type, Param6Name ) FUNC_DECLARE_DYNAMIC_DELEGATE( FWeakObjectPtr, SixParams, DelegateName, DelegateName##_DelegateWrapper, FUNC_CONCAT( Param1Type InParam1, Param2Type InParam2, Param3Type InParam3, Param4Type InParam4, Param5Type InParam5, Param6Type InParam6 ), FUNC_CONCAT( *this, InParam1, InParam2, InParam3, InParam4, InParam5, InParam6 ), void, Param1Type, Param2Type, Param3Type, Param4Type, Param5Type, Param6Type )
-#define DECLARE_DYNAMIC_MULTICAST_DELEGATE_SixParams( DelegateName, Param1Type, Param1Name, Param2Type, Param2Name, Param3Type, Param3Name, Param4Type, Param4Name, Param5Type, Param5Name, Param6Type, Param6Name ) FUNC_DECLARE_DYNAMIC_MULTICAST_DELEGATE( FWeakObjectPtr, SixParams, DelegateName, DelegateName##_DelegateWrapper, FUNC_CONCAT( Param1Type InParam1, Param2Type InParam2, Param3Type InParam3, Param4Type InParam4, Param5Type InParam5, Param6Type InParam6 ), FUNC_CONCAT( *this, InParam1, InParam2, InParam3, InParam4, InParam5, InParam6 ), void, Param1Type, Param2Type, Param3Type, Param4Type, Param5Type, Param6Type )
+#define DECLARE_DYNAMIC_DELEGATE_SixParams( DelegateName, Param1Type, Param1Name, Param2Type, Param2Name, Param3Type, Param3Name, Param4Type, Param4Name, Param5Type, Param5Name, Param6Type, Param6Name ) BODY_MACRO_COMBINE(CURRENT_FILE_ID,_,__LINE__,_DELEGATE) FUNC_DECLARE_DYNAMIC_DELEGATE( FWeakObjectPtr, SixParams, DelegateName, DelegateName##_DelegateWrapper, FUNC_CONCAT( Param1Type InParam1, Param2Type InParam2, Param3Type InParam3, Param4Type InParam4, Param5Type InParam5, Param6Type InParam6 ), FUNC_CONCAT( *this, InParam1, InParam2, InParam3, InParam4, InParam5, InParam6 ), void, Param1Type, Param2Type, Param3Type, Param4Type, Param5Type, Param6Type )
+#define DECLARE_DYNAMIC_MULTICAST_DELEGATE_SixParams( DelegateName, Param1Type, Param1Name, Param2Type, Param2Name, Param3Type, Param3Name, Param4Type, Param4Name, Param5Type, Param5Name, Param6Type, Param6Name ) BODY_MACRO_COMBINE(CURRENT_FILE_ID,_,__LINE__,_DELEGATE) FUNC_DECLARE_DYNAMIC_MULTICAST_DELEGATE( FWeakObjectPtr, SixParams, DelegateName, DelegateName##_DelegateWrapper, FUNC_CONCAT( Param1Type InParam1, Param2Type InParam2, Param3Type InParam3, Param4Type InParam4, Param5Type InParam5, Param6Type InParam6 ), FUNC_CONCAT( *this, InParam1, InParam2, InParam3, InParam4, InParam5, InParam6 ), void, Param1Type, Param2Type, Param3Type, Param4Type, Param5Type, Param6Type )
 #undef FUNC_SUFFIX
 #undef FUNC_TEMPLATE_DECL
 #undef FUNC_TEMPLATE_DECL_TYPENAME
@@ -2844,7 +2844,7 @@
 #include "DelegateSignatureImpl.inl"
 
 #define DECLARE_DELEGATE_RetVal_SixParams( RetValType, DelegateName, Param1Type, Param2Type, Param3Type, Param4Type, Param5Type, Param6Type ) FUNC_DECLARE_DELEGATE( RetVal_SixParams, DelegateName, RetValType, Param1Type, Param2Type, Param3Type, Param4Type, Param5Type, Param6Type )
-#define DECLARE_DYNAMIC_DELEGATE_RetVal_SixParams( RetValType, DelegateName, Param1Type, Param1Name, Param2Type, Param2Name, Param3Type, Param3Name, Param4Type, Param4Name, Param5Type, Param5Name, Param6Type, Param6Name ) FUNC_DECLARE_DYNAMIC_DELEGATE_RETVAL( FWeakObjectPtr, RetVal_SixParams, DelegateName, DelegateName##_DelegateWrapper, RetValType, FUNC_CONCAT( Param1Type InParam1, Param2Type InParam2, Param3Type InParam3, Param4Type InParam4, Param5Type InParam5, Param6Type InParam6 ), FUNC_CONCAT( *this, InParam1, InParam2, InParam3, InParam4, InParam5, InParam6 ), RetValType, Param1Type, Param2Type, Param3Type, Param4Type, Param5Type, Param6Type )
+#define DECLARE_DYNAMIC_DELEGATE_RetVal_SixParams( RetValType, DelegateName, Param1Type, Param1Name, Param2Type, Param2Name, Param3Type, Param3Name, Param4Type, Param4Name, Param5Type, Param5Name, Param6Type, Param6Name ) BODY_MACRO_COMBINE(CURRENT_FILE_ID,_,__LINE__,_DELEGATE) FUNC_DECLARE_DYNAMIC_DELEGATE_RETVAL( FWeakObjectPtr, RetVal_SixParams, DelegateName, DelegateName##_DelegateWrapper, RetValType, FUNC_CONCAT( Param1Type InParam1, Param2Type InParam2, Param3Type InParam3, Param4Type InParam4, Param5Type InParam5, Param6Type InParam6 ), FUNC_CONCAT( *this, InParam1, InParam2, InParam3, InParam4, InParam5, InParam6 ), RetValType, Param1Type, Param2Type, Param3Type, Param4Type, Param5Type, Param6Type )
 #undef FUNC_SUFFIX
 #undef FUNC_TEMPLATE_DECL
 #undef FUNC_TEMPLATE_DECL_TYPENAME
@@ -3048,8 +3048,8 @@
 #define DECLARE_DELEGATE_SevenParams( DelegateName, Param1Type, Param2Type, Param3Type, Param4Type, Param5Type, Param6Type, Param7Type ) FUNC_DECLARE_DELEGATE( SevenParams, DelegateName, void, Param1Type, Param2Type, Param3Type, Param4Type, Param5Type, Param6Type, Param7Type )
 #define DECLARE_MULTICAST_DELEGATE_SevenParams( DelegateName, Param1Type, Param2Type, Param3Type, Param4Type, Param5Type, Param6Type, Param7Type ) FUNC_DECLARE_MULTICAST_DELEGATE( SevenParams, DelegateName, void, Param1Type, Param2Type, Param3Type, Param4Type, Param5Type, Param6Type, Param7Type )
 #define DECLARE_EVENT_SevenParams( OwningType, EventName, Param1Type, Param2Type, Param3Type, Param4Type, Param5Type, Param6Type, Param7Type ) FUNC_DECLARE_EVENT( OwningType, EventName, SevenParams, void, Param1Type, Param2Type, Param3Type, Param4Type, Param5Type, Param6Type, Param7Type )
-#define DECLARE_DYNAMIC_DELEGATE_SevenParams( DelegateName, Param1Type, Param1Name, Param2Type, Param2Name, Param3Type, Param3Name, Param4Type, Param4Name, Param5Type, Param5Name, Param6Type, Param6Name, Param7Type, Param7Name ) FUNC_DECLARE_DYNAMIC_DELEGATE( FWeakObjectPtr, SevenParams, DelegateName, DelegateName##_DelegateWrapper, FUNC_CONCAT( Param1Type InParam1, Param2Type InParam2, Param3Type InParam3, Param4Type InParam4, Param5Type InParam5, Param6Type InParam6, Param7Type InParam7 ), FUNC_CONCAT( *this, InParam1, InParam2, InParam3, InParam4, InParam5, InParam6, InParam7 ), void, Param1Type, Param2Type, Param3Type, Param4Type, Param5Type, Param6Type, Param7Type )
-#define DECLARE_DYNAMIC_MULTICAST_DELEGATE_SevenParams( DelegateName, Param1Type, Param1Name, Param2Type, Param2Name, Param3Type, Param3Name, Param4Type, Param4Name, Param5Type, Param5Name, Param6Type, Param6Name, Param7Type, Param7Name ) FUNC_DECLARE_DYNAMIC_MULTICAST_DELEGATE( FWeakObjectPtr, SevenParams, DelegateName, DelegateName##_DelegateWrapper, FUNC_CONCAT( Param1Type InParam1, Param2Type InParam2, Param3Type InParam3, Param4Type InParam4, Param5Type InParam5, Param6Type InParam6, Param7Type InParam7 ), FUNC_CONCAT( *this, InParam1, InParam2, InParam3, InParam4, InParam5, InParam6, InParam7 ), void, Param1Type, Param2Type, Param3Type, Param4Type, Param5Type, Param6Type, Param7Type )
+#define DECLARE_DYNAMIC_DELEGATE_SevenParams( DelegateName, Param1Type, Param1Name, Param2Type, Param2Name, Param3Type, Param3Name, Param4Type, Param4Name, Param5Type, Param5Name, Param6Type, Param6Name, Param7Type, Param7Name ) BODY_MACRO_COMBINE(CURRENT_FILE_ID,_,__LINE__,_DELEGATE) FUNC_DECLARE_DYNAMIC_DELEGATE( FWeakObjectPtr, SevenParams, DelegateName, DelegateName##_DelegateWrapper, FUNC_CONCAT( Param1Type InParam1, Param2Type InParam2, Param3Type InParam3, Param4Type InParam4, Param5Type InParam5, Param6Type InParam6, Param7Type InParam7 ), FUNC_CONCAT( *this, InParam1, InParam2, InParam3, InParam4, InParam5, InParam6, InParam7 ), void, Param1Type, Param2Type, Param3Type, Param4Type, Param5Type, Param6Type, Param7Type )
+#define DECLARE_DYNAMIC_MULTICAST_DELEGATE_SevenParams( DelegateName, Param1Type, Param1Name, Param2Type, Param2Name, Param3Type, Param3Name, Param4Type, Param4Name, Param5Type, Param5Name, Param6Type, Param6Name, Param7Type, Param7Name ) BODY_MACRO_COMBINE(CURRENT_FILE_ID,_,__LINE__,_DELEGATE) FUNC_DECLARE_DYNAMIC_MULTICAST_DELEGATE( FWeakObjectPtr, SevenParams, DelegateName, DelegateName##_DelegateWrapper, FUNC_CONCAT( Param1Type InParam1, Param2Type InParam2, Param3Type InParam3, Param4Type InParam4, Param5Type InParam5, Param6Type InParam6, Param7Type InParam7 ), FUNC_CONCAT( *this, InParam1, InParam2, InParam3, InParam4, InParam5, InParam6, InParam7 ), void, Param1Type, Param2Type, Param3Type, Param4Type, Param5Type, Param6Type, Param7Type )
 #undef FUNC_SUFFIX
 #undef FUNC_TEMPLATE_DECL
 #undef FUNC_TEMPLATE_DECL_TYPENAME
@@ -3251,7 +3251,7 @@
 #include "DelegateSignatureImpl.inl"
 
 #define DECLARE_DELEGATE_RetVal_SevenParams( RetValType, DelegateName, Param1Type, Param2Type, Param3Type, Param4Type, Param5Type, Param6Type, Param7Type ) FUNC_DECLARE_DELEGATE( RetVal_SevenParams, DelegateName, RetValType, Param1Type, Param2Type, Param3Type, Param4Type, Param5Type, Param6Type, Param7Type )
-#define DECLARE_DYNAMIC_DELEGATE_RetVal_SevenParams( RetValType, DelegateName, Param1Type, Param1Name, Param2Type, Param2Name, Param3Type, Param3Name, Param4Type, Param4Name, Param5Type, Param5Name, Param6Type, Param6Name, Param7Type, Param7Name ) FUNC_DECLARE_DYNAMIC_DELEGATE_RETVAL( FWeakObjectPtr, RetVal_SevenParams, DelegateName, DelegateName##_DelegateWrapper, RetValType, FUNC_CONCAT( Param1Type InParam1, Param2Type InParam2, Param3Type InParam3, Param4Type InParam4, Param5Type InParam5, Param6Type InParam6, Param7Type InParam7 ), FUNC_CONCAT( *this, InParam1, InParam2, InParam3, InParam4, InParam5, InParam6, InParam7 ), RetValType, Param1Type, Param2Type, Param3Type, Param4Type, Param5Type, Param6Type, Param7Type )
+#define DECLARE_DYNAMIC_DELEGATE_RetVal_SevenParams( RetValType, DelegateName, Param1Type, Param1Name, Param2Type, Param2Name, Param3Type, Param3Name, Param4Type, Param4Name, Param5Type, Param5Name, Param6Type, Param6Name, Param7Type, Param7Name ) BODY_MACRO_COMBINE(CURRENT_FILE_ID,_,__LINE__,_DELEGATE) FUNC_DECLARE_DYNAMIC_DELEGATE_RETVAL( FWeakObjectPtr, RetVal_SevenParams, DelegateName, DelegateName##_DelegateWrapper, RetValType, FUNC_CONCAT( Param1Type InParam1, Param2Type InParam2, Param3Type InParam3, Param4Type InParam4, Param5Type InParam5, Param6Type InParam6, Param7Type InParam7 ), FUNC_CONCAT( *this, InParam1, InParam2, InParam3, InParam4, InParam5, InParam6, InParam7 ), RetValType, Param1Type, Param2Type, Param3Type, Param4Type, Param5Type, Param6Type, Param7Type )
 #undef FUNC_SUFFIX
 #undef FUNC_TEMPLATE_DECL
 #undef FUNC_TEMPLATE_DECL_TYPENAME
@@ -3455,8 +3455,8 @@
 #define DECLARE_DELEGATE_EightParams( DelegateName, Param1Type, Param2Type, Param3Type, Param4Type, Param5Type, Param6Type, Param7Type, Param8Type ) FUNC_DECLARE_DELEGATE( EightParams, DelegateName, void, Param1Type, Param2Type, Param3Type, Param4Type, Param5Type, Param6Type, Param7Type, Param8Type )
 #define DECLARE_MULTICAST_DELEGATE_EightParams( DelegateName, Param1Type, Param2Type, Param3Type, Param4Type, Param5Type, Param6Type, Param7Type, Param8Type ) FUNC_DECLARE_MULTICAST_DELEGATE( EightParams, DelegateName, void, Param1Type, Param2Type, Param3Type, Param4Type, Param5Type, Param6Type, Param7Type, Param8Type )
 #define DECLARE_EVENT_EightParams( OwningType, EventName, Param1Type, Param2Type, Param3Type, Param4Type, Param5Type, Param6Type, Param7Type, Param8Type ) FUNC_DECLARE_EVENT( OwningType, EventName, EightParams, void, Param1Type, Param2Type, Param3Type, Param4Type, Param5Type, Param6Type, Param7Type, Param8Type )
-#define DECLARE_DYNAMIC_DELEGATE_EightParams( DelegateName, Param1Type, Param1Name, Param2Type, Param2Name, Param3Type, Param3Name, Param4Type, Param4Name, Param5Type, Param5Name, Param6Type, Param6Name, Param7Type, Param7Name, Param8Type, Param8Name ) FUNC_DECLARE_DYNAMIC_DELEGATE( FWeakObjectPtr, EightParams, DelegateName, DelegateName##_DelegateWrapper, FUNC_CONCAT( Param1Type InParam1, Param2Type InParam2, Param3Type InParam3, Param4Type InParam4, Param5Type InParam5, Param6Type InParam6, Param7Type InParam7, Param8Type InParam8 ), FUNC_CONCAT( *this, InParam1, InParam2, InParam3, InParam4, InParam5, InParam6, InParam7, InParam8 ), void, Param1Type, Param2Type, Param3Type, Param4Type, Param5Type, Param6Type, Param7Type, Param8Type )
-#define DECLARE_DYNAMIC_MULTICAST_DELEGATE_EightParams( DelegateName, Param1Type, Param1Name, Param2Type, Param2Name, Param3Type, Param3Name, Param4Type, Param4Name, Param5Type, Param5Name, Param6Type, Param6Name, Param7Type, Param7Name, Param8Type, Param8Name ) FUNC_DECLARE_DYNAMIC_MULTICAST_DELEGATE( FWeakObjectPtr, EightParams, DelegateName, DelegateName##_DelegateWrapper, FUNC_CONCAT( Param1Type InParam1, Param2Type InParam2, Param3Type InParam3, Param4Type InParam4, Param5Type InParam5, Param6Type InParam6, Param7Type InParam7, Param8Type InParam8 ), FUNC_CONCAT( *this, InParam1, InParam2, InParam3, InParam4, InParam5, InParam6, InParam7, InParam8 ), void, Param1Type, Param2Type, Param3Type, Param4Type, Param5Type, Param6Type, Param7Type, Param8Type )
+#define DECLARE_DYNAMIC_DELEGATE_EightParams( DelegateName, Param1Type, Param1Name, Param2Type, Param2Name, Param3Type, Param3Name, Param4Type, Param4Name, Param5Type, Param5Name, Param6Type, Param6Name, Param7Type, Param7Name, Param8Type, Param8Name ) BODY_MACRO_COMBINE(CURRENT_FILE_ID,_,__LINE__,_DELEGATE) FUNC_DECLARE_DYNAMIC_DELEGATE( FWeakObjectPtr, EightParams, DelegateName, DelegateName##_DelegateWrapper, FUNC_CONCAT( Param1Type InParam1, Param2Type InParam2, Param3Type InParam3, Param4Type InParam4, Param5Type InParam5, Param6Type InParam6, Param7Type InParam7, Param8Type InParam8 ), FUNC_CONCAT( *this, InParam1, InParam2, InParam3, InParam4, InParam5, InParam6, InParam7, InParam8 ), void, Param1Type, Param2Type, Param3Type, Param4Type, Param5Type, Param6Type, Param7Type, Param8Type )
+#define DECLARE_DYNAMIC_MULTICAST_DELEGATE_EightParams( DelegateName, Param1Type, Param1Name, Param2Type, Param2Name, Param3Type, Param3Name, Param4Type, Param4Name, Param5Type, Param5Name, Param6Type, Param6Name, Param7Type, Param7Name, Param8Type, Param8Name ) BODY_MACRO_COMBINE(CURRENT_FILE_ID,_,__LINE__,_DELEGATE) FUNC_DECLARE_DYNAMIC_MULTICAST_DELEGATE( FWeakObjectPtr, EightParams, DelegateName, DelegateName##_DelegateWrapper, FUNC_CONCAT( Param1Type InParam1, Param2Type InParam2, Param3Type InParam3, Param4Type InParam4, Param5Type InParam5, Param6Type InParam6, Param7Type InParam7, Param8Type InParam8 ), FUNC_CONCAT( *this, InParam1, InParam2, InParam3, InParam4, InParam5, InParam6, InParam7, InParam8 ), void, Param1Type, Param2Type, Param3Type, Param4Type, Param5Type, Param6Type, Param7Type, Param8Type )
 #undef FUNC_SUFFIX
 #undef FUNC_TEMPLATE_DECL
 #undef FUNC_TEMPLATE_DECL_TYPENAME
@@ -3658,7 +3658,7 @@
 #include "DelegateSignatureImpl.inl"
 
 #define DECLARE_DELEGATE_RetVal_EightParams( RetValType, DelegateName, Param1Type, Param2Type, Param3Type, Param4Type, Param5Type, Param6Type, Param7Type, Param8Type ) FUNC_DECLARE_DELEGATE( RetVal_EightParams, DelegateName, RetValType, Param1Type, Param2Type, Param3Type, Param4Type, Param5Type, Param6Type, Param7Type, Param8Type )
-#define DECLARE_DYNAMIC_DELEGATE_RetVal_EightParams( RetValType, DelegateName, Param1Type, Param1Name, Param2Type, Param2Name, Param3Type, Param3Name, Param4Type, Param4Name, Param5Type, Param5Name, Param6Type, Param6Name, Param7Type, Param7Name, Param8Type, Param8Name ) FUNC_DECLARE_DYNAMIC_DELEGATE_RETVAL( FWeakObjectPtr, RetVal_EightParams, DelegateName, DelegateName##_DelegateWrapper, RetValType, FUNC_CONCAT( Param1Type InParam1, Param2Type InParam2, Param3Type InParam3, Param4Type InParam4, Param5Type InParam5, Param6Type InParam6, Param7Type InParam7, Param8Type InParam8 ), FUNC_CONCAT( *this, InParam1, InParam2, InParam3, InParam4, InParam5, InParam6, InParam7, InParam8 ), RetValType, Param1Type, Param2Type, Param3Type, Param4Type, Param5Type, Param6Type, Param7Type, Param8Type )
+#define DECLARE_DYNAMIC_DELEGATE_RetVal_EightParams( RetValType, DelegateName, Param1Type, Param1Name, Param2Type, Param2Name, Param3Type, Param3Name, Param4Type, Param4Name, Param5Type, Param5Name, Param6Type, Param6Name, Param7Type, Param7Name, Param8Type, Param8Name ) BODY_MACRO_COMBINE(CURRENT_FILE_ID,_,__LINE__,_DELEGATE) FUNC_DECLARE_DYNAMIC_DELEGATE_RETVAL( FWeakObjectPtr, RetVal_EightParams, DelegateName, DelegateName##_DelegateWrapper, RetValType, FUNC_CONCAT( Param1Type InParam1, Param2Type InParam2, Param3Type InParam3, Param4Type InParam4, Param5Type InParam5, Param6Type InParam6, Param7Type InParam7, Param8Type InParam8 ), FUNC_CONCAT( *this, InParam1, InParam2, InParam3, InParam4, InParam5, InParam6, InParam7, InParam8 ), RetValType, Param1Type, Param2Type, Param3Type, Param4Type, Param5Type, Param6Type, Param7Type, Param8Type )
 #undef FUNC_SUFFIX
 #undef FUNC_TEMPLATE_DECL
 #undef FUNC_TEMPLATE_DECL_TYPENAME
