Fix FFeedbackContextWindows/XboxOne::YesNoF() to match the base declaration (caught with /W4).

diff --git a/Engine/Source/Runtime/Core/Private/Windows/WindowsPlatformFeedbackContextPrivate.h b/Engine/Source/Runtime/Core/Private/Windows/WindowsPlatformFeedbackContextPrivate.h
--- a/Engine/Source/Runtime/Core/Private/Windows/WindowsPlatformFeedbackContextPrivate.h
+++ b/Engine/Source/Runtime/Core/Private/Windows/WindowsPlatformFeedbackContextPrivate.h
@@ -66,19 +66,19 @@ public:
 			GLog->Serialize( V, Verbosity, Category );
 		}
 	}
-	VARARG_BODY( bool, YesNof, const TCHAR*, VARARG_NONE )
+
+	virtual bool YesNof(const FText& Question) override
 	{
-		TCHAR TempStr[4096];
-		GET_VARARGS( TempStr, ARRAY_COUNT(TempStr), ARRAY_COUNT(TempStr)-1, Fmt, Fmt );
 		if( ( GIsClient || GIsEditor ) && ( ( GIsSilent != true ) && ( FApp::IsUnattended() != true ) ) )
 		{
-			return( ::MessageBox( NULL, TempStr, *NSLOCTEXT("Core", "Question", "Question").ToString(), MB_YESNO|MB_TASKMODAL ) == IDYES);
+			return( ::MessageBox( NULL, Question.ToString().GetCharArray().GetData(), *NSLOCTEXT("Core", "Question", "Question").ToString(), MB_YESNO|MB_TASKMODAL ) == IDYES);
 		}
 		else
 		{
 			return false;
 		}
 	}
+
 	FContextSupplier* GetContext() const
 	{
 		return Context;
