- CIS fix

diff --git a/Engine/Source/Runtime/Core/Private/Linux/LinuxPlatformStackWalk.cpp b/Engine/Source/Runtime/Core/Private/Linux/LinuxPlatformStackWalk.cpp
--- a/Engine/Source/Runtime/Core/Private/Linux/LinuxPlatformStackWalk.cpp
+++ b/Engine/Source/Runtime/Core/Private/Linux/LinuxPlatformStackWalk.cpp
@@ -6,6 +6,7 @@
 
 #include "CorePrivatePCH.h"
 #include "Misc/App.h"
+#include "LinuxPlatformCrashContext.h"
 #include <execinfo.h>
 #include <dlfcn.h>
 #include <cxxabi.h>
