Core: better name for local variables (TypeContainer)

diff --git a/Engine/Source/Runtime/Core/Public/Misc/TypeContainer.h b/Engine/Source/Runtime/Core/Public/Misc/TypeContainer.h
--- a/Engine/Source/Runtime/Core/Public/Misc/TypeContainer.h
+++ b/Engine/Source/Runtime/Core/Public/Misc/TypeContainer.h
@@ -230,11 +230,11 @@ public:
 	template<class R>
 	TSharedPtr<R> GetInstance()
 	{
-		const FName ClassName = R::GetTypeName();
+		const FName TypeName = R::GetTypeName();
 
 		FScopeLock Lock(&CriticalSection);
 		{
-			const TSharedPtr<IInstanceProvider>& Provider = Providers.FindRef(ClassName);
+			const TSharedPtr<IInstanceProvider>& Provider = Providers.FindRef(TypeName);
 
 			if (Provider.IsValid())
 			{
@@ -380,11 +380,11 @@ public:
 	template<class R>
 	void Unregister()
 	{
-		const FName ClassName = R::GetTypeName();
+		const FName TypeName = R::GetTypeName();
 
 		FScopeLock Lock(&CriticalSection);
 		{
-			Providers.Remove(ClassName);
+			Providers.Remove(TypeName);
 		}
 	}
 
