Optimise FBoxSphereBounds::TransformBy( const FTransform& M ) by transforming the axes and taking the sum of absolutes, rather than transforming all 8 corners.

On my machine 1000000 iterations went from 0.927s to 0.154s. ~6x improvement.

diff --git a/Engine/Source/Runtime/Core/Private/Math/BoxSphereBounds.cpp b/Engine/Source/Runtime/Core/Private/Math/BoxSphereBounds.cpp
--- a/Engine/Source/Runtime/Core/Private/Math/BoxSphereBounds.cpp
+++ b/Engine/Source/Runtime/Core/Private/Math/BoxSphereBounds.cpp
@@ -49,24 +49,11 @@ FBoxSphereBounds FBoxSphereBounds::TransformBy(const FTransform& M) const
 	FBoxSphereBounds Result;
 
 	Result.Origin = M.TransformPosition(Origin);
-	Result.BoxExtent = FVector::ZeroVector;
-
-	float Signs[2] = { -1.0f, 1.0f };
-
-	for (int32 X = 0; X < 2; X++)
-	{
-		for (int32 Y = 0; Y < 2; Y++)
-		{
-			for (int32 Z = 0; Z < 2; Z++)
-			{
-				FVector	Corner = M.TransformVector(FVector(Signs[X] * BoxExtent.X,Signs[Y] * BoxExtent.Y,Signs[Z] * BoxExtent.Z));
 
-				Result.BoxExtent.X = FMath::Max(Corner.X, Result.BoxExtent.X);
-				Result.BoxExtent.Y = FMath::Max(Corner.Y, Result.BoxExtent.Y);
-				Result.BoxExtent.Z = FMath::Max(Corner.Z, Result.BoxExtent.Z);
-			}
-		}
-	}
+	FVector4 XAxis = M.TransformVector(FVector(BoxExtent.X, 0.0, 0.0)).GetAbs();
+	FVector4 YAxis = M.TransformVector(FVector(0.0, BoxExtent.Y, 0.0)).GetAbs();
+	FVector4 ZAxis = M.TransformVector(FVector(0.0, 0.0, BoxExtent.Z)).GetAbs();
+ 	Result.BoxExtent = XAxis + YAxis + ZAxis;
 
 	Result.SphereRadius = SphereRadius * M.GetMaximumAxisScale();
 
