[INTEGRATE] Change 2426184 by [REDACTED_PERSON][REDACTED_HANDLE] on 2015/01/30 14:48:19

	Fix redraw crashes when resizing the editor on the Mac

diff --git a/Engine/Source/Runtime/Core/Private/Mac/MacApplication.cpp b/Engine/Source/Runtime/Core/Private/Mac/MacApplication.cpp
--- a/Engine/Source/Runtime/Core/Private/Mac/MacApplication.cpp
+++ b/Engine/Source/Runtime/Core/Private/Mac/MacApplication.cpp
@@ -460,6 +460,34 @@ void FMacApplication::CloseQueuedWindows()
 	}
 }
 
+void FMacApplication::InvalidateTextLayout(FCocoaWindow* Window)
+{
+	WindowsRequiringTextInvalidation.AddUnique( Window );
+}
+
+void FMacApplication::InvalidateTextLayouts()
+{
+	if( WindowsRequiringTextInvalidation.Num() > 0 )
+	{
+	   MainThreadCall(^{
+		   SCOPED_AUTORELEASE_POOL;
+	
+		   for( FCocoaWindow* CocoaWindow : WindowsRequiringTextInvalidation)
+		   {
+			   if(CocoaWindow && [CocoaWindow openGLView])
+			   {
+				   FCocoaTextView* TextView = (FCocoaTextView*)[CocoaWindow openGLView];
+				   [[TextView inputContext] invalidateCharacterCoordinates];
+			   }
+		   }
+
+		}, UE4IMEEventMode, true);
+		
+		WindowsRequiringTextInvalidation.Empty();
+	}
+
+}
+
 void FMacApplication::OnWindowDraggingFinished()
 {
 	if( DraggedWindow )
diff --git a/Engine/Source/Runtime/Core/Private/Mac/MacPlatformMisc.cpp b/Engine/Source/Runtime/Core/Private/Mac/MacPlatformMisc.cpp
--- a/Engine/Source/Runtime/Core/Private/Mac/MacPlatformMisc.cpp
+++ b/Engine/Source/Runtime/Core/Private/Mac/MacPlatformMisc.cpp
@@ -497,6 +497,7 @@ void FMacPlatformMisc::PumpMessages( bool bFromMainLoop )
 				bChachedMacMenuStateNeedsUpdate = false;
 			}
 
+			MacApplication->InvalidateTextLayouts();
 			MacApplication->CloseQueuedWindows();
 		}
 	}
diff --git a/Engine/Source/Runtime/Core/Private/Mac/MacTextInputMethodSystem.cpp b/Engine/Source/Runtime/Core/Private/Mac/MacTextInputMethodSystem.cpp
--- a/Engine/Source/Runtime/Core/Private/Mac/MacTextInputMethodSystem.cpp
+++ b/Engine/Source/Runtime/Core/Private/Mac/MacTextInputMethodSystem.cpp
@@ -53,14 +53,8 @@ namespace
 		if(ChangeType != ELayoutChangeType::Created && GetContextWindow().IsValid())
 		{
 			FCocoaWindow* CocoaWindow = (FCocoaWindow*)GetContextWindow()->GetOSWindowHandle();
-			MainThreadCall(^{
-				SCOPED_AUTORELEASE_POOL;
-				if(CocoaWindow && [CocoaWindow openGLView])
-				{
-					FCocoaTextView* TextView = (FCocoaTextView*)[CocoaWindow openGLView];
-					[[TextView inputContext] invalidateCharacterCoordinates];
-				}
-			}, UE4IMEEventMode, true);
+			
+			MacApplication->InvalidateTextLayout( CocoaWindow );
 		}
 	}
 	
diff --git a/Engine/Source/Runtime/Core/Public/Mac/MacApplication.h b/Engine/Source/Runtime/Core/Public/Mac/MacApplication.h
--- a/Engine/Source/Runtime/Core/Public/Mac/MacApplication.h
+++ b/Engine/Source/Runtime/Core/Public/Mac/MacApplication.h
@@ -111,8 +111,14 @@ public:
 
 	const TArray<TSharedRef<FMacWindow>>& GetAllWindows() const { return Windows; }
 
+
 	void CloseQueuedWindows();
 
+	/** Queues a window for text layout invalidation when safe */
+	void InvalidateTextLayout( FCocoaWindow* Window );
+
+	/** Invalidates all queued windows requiring text layout changes */
+	void InvalidateTextLayouts();
 private:
 	enum FMacApplicationEventTypes
 	{
@@ -180,6 +186,8 @@ private:
 
 	TArray<FCocoaWindow*> WindowsToClose;
 
+	TArray<FCocoaWindow*> WindowsRequiringTextInvalidation;
+
 	TSharedPtr<FMacTextInputMethodSystem> TextInputMethodSystem;
 
 	bool bIsWorkspaceSessionActive;
