Fix non-unity CIS

diff --git a/Engine/Source/Runtime/Core/Private/Mac/MacTextInputMethodSystem.cpp b/Engine/Source/Runtime/Core/Private/Mac/MacTextInputMethodSystem.cpp
--- a/Engine/Source/Runtime/Core/Private/Mac/MacTextInputMethodSystem.cpp
+++ b/Engine/Source/Runtime/Core/Private/Mac/MacTextInputMethodSystem.cpp
@@ -2,6 +2,7 @@
 
 #include "CorePrivatePCH.h"
 
+#include "MacApplication.h"
 #include "MacWindow.h"
 #include "MacTextInputMethodSystem.h"
 #include "CocoaTextView.h"
