limit the thread priority change to only the editor and not commandlets

diff --git a/Engine/Source/Editor/UnrealEd/Private/Editor.cpp b/Engine/Source/Editor/UnrealEd/Private/Editor.cpp
--- a/Engine/Source/Editor/UnrealEd/Private/Editor.cpp
+++ b/Engine/Source/Editor/UnrealEd/Private/Editor.cpp
@@ -6075,7 +6075,7 @@ bool UEditorEngine::ShouldThrottleCPUUsage() const
 		return false;
 	}
 
-	return bShouldThrottle;
+    return bShouldThrottle && !IsRunningCommandlet();
 }
 
 bool UEditorEngine::AreAllWindowsHidden() const
diff --git a/Engine/Source/Runtime/Core/Private/Mac/MacApplication.cpp b/Engine/Source/Runtime/Core/Private/Mac/MacApplication.cpp
--- a/Engine/Source/Runtime/Core/Private/Mac/MacApplication.cpp
+++ b/Engine/Source/Runtime/Core/Private/Mac/MacApplication.cpp
@@ -163,7 +163,7 @@ FMacApplication::FMacApplication()
 								}
 
 								// If editor thread doesn't have the focus, don't suck up too much CPU time.
-								if (GIsEditor)
+								if (GIsEditor && !IsRunningCommandlet())
 								{
 									// Boost our priority back to normal.
 									struct sched_param Sched;
@@ -209,7 +209,7 @@ FMacApplication::FMacApplication()
 								}
 
 								// If editor thread doesn't have the focus, don't suck up too much CPU time.
-								if (GIsEditor)
+								if (GIsEditor && !IsRunningCommandlet())
 								{
 									// Drop our priority to speed up whatever is in the foreground.
 									struct sched_param Sched;
