Mac non unity fix

diff --git a/Engine/Source/Runtime/Core/Private/Mac/MacApplication.cpp b/Engine/Source/Runtime/Core/Private/Mac/MacApplication.cpp
--- a/Engine/Source/Runtime/Core/Private/Mac/MacApplication.cpp
+++ b/Engine/Source/Runtime/Core/Private/Mac/MacApplication.cpp
@@ -13,6 +13,7 @@
 #include "IAnalyticsProvider.h"
 #include "CocoaThread.h"
 #include "ModuleManager.h"
+#include "CocoaTextView.h"
 
 static NSString* NSWindowDraggingFinished = @"NSWindowDraggingFinished";
 
