Optimise FBoxSphereBounds::TransformBy and FBox::TransformBy

diff --git a/Engine/Source/Runtime/Core/Private/Math/Box.cpp b/Engine/Source/Runtime/Core/Private/Math/Box.cpp
--- a/Engine/Source/Runtime/Core/Private/Math/Box.cpp
+++ b/Engine/Source/Runtime/Core/Private/Math/Box.cpp
@@ -45,79 +45,35 @@ FBox FBox::TransformBy(const FMatrix& M) const
 		return FBox(0);
 	}
 
-	VectorRegister Vertices[8];
-
-	VectorRegister m0 = VectorLoadAligned(M.M[0]);
-	VectorRegister m1 = VectorLoadAligned(M.M[1]);
-	VectorRegister m2 = VectorLoadAligned(M.M[2]);
-	VectorRegister m3 = VectorLoadAligned(M.M[3]);
-
-	Vertices[0] = VectorLoadFloat3(&Min);
-	Vertices[1] = VectorSetFloat3(Min.X, Min.Y, Max.Z);
-	Vertices[2] = VectorSetFloat3(Min.X, Max.Y, Min.Z);
-	Vertices[3] = VectorSetFloat3(Max.X, Min.Y, Min.Z);
-	Vertices[4] = VectorSetFloat3(Max.X, Max.Y, Min.Z);
-	Vertices[5] = VectorSetFloat3(Max.X, Min.Y, Max.Z);
-	Vertices[6] = VectorSetFloat3(Min.X, Max.Y, Max.Z);
-	Vertices[7] = VectorLoadFloat3(&Max);
-
-	VectorRegister r0 = VectorMultiply(VectorReplicate(Vertices[0],0), m0);
-	VectorRegister r1 = VectorMultiply(VectorReplicate(Vertices[1],0), m0);
-	VectorRegister r2 = VectorMultiply(VectorReplicate(Vertices[2],0), m0);
-	VectorRegister r3 = VectorMultiply(VectorReplicate(Vertices[3],0), m0);
-	VectorRegister r4 = VectorMultiply(VectorReplicate(Vertices[4],0), m0);
-	VectorRegister r5 = VectorMultiply(VectorReplicate(Vertices[5],0), m0);
-	VectorRegister r6 = VectorMultiply(VectorReplicate(Vertices[6],0), m0);
-	VectorRegister r7 = VectorMultiply(VectorReplicate(Vertices[7],0), m0);
-
-	r0 = VectorMultiplyAdd( VectorReplicate(Vertices[0],1), m1, r0);
-	r1 = VectorMultiplyAdd( VectorReplicate(Vertices[1],1), m1, r1);
-	r2 = VectorMultiplyAdd( VectorReplicate(Vertices[2],1), m1, r2);
-	r3 = VectorMultiplyAdd( VectorReplicate(Vertices[3],1), m1, r3);
-	r4 = VectorMultiplyAdd( VectorReplicate(Vertices[4],1), m1, r4);
-	r5 = VectorMultiplyAdd( VectorReplicate(Vertices[5],1), m1, r5);
-	r6 = VectorMultiplyAdd( VectorReplicate(Vertices[6],1), m1, r6);
-	r7 = VectorMultiplyAdd( VectorReplicate(Vertices[7],1), m1, r7);
-
-	r0 = VectorMultiplyAdd( VectorReplicate(Vertices[0],2), m2, r0);
-	r1 = VectorMultiplyAdd( VectorReplicate(Vertices[1],2), m2, r1);
-	r2 = VectorMultiplyAdd( VectorReplicate(Vertices[2],2), m2, r2);
-	r3 = VectorMultiplyAdd( VectorReplicate(Vertices[3],2), m2, r3);
-	r4 = VectorMultiplyAdd( VectorReplicate(Vertices[4],2), m2, r4);
-	r5 = VectorMultiplyAdd( VectorReplicate(Vertices[5],2), m2, r5);
-	r6 = VectorMultiplyAdd( VectorReplicate(Vertices[6],2), m2, r6);
-	r7 = VectorMultiplyAdd( VectorReplicate(Vertices[7],2), m2, r7);
-
-	r0 = VectorAdd(r0, m3);
-	r1 = VectorAdd(r1, m3);
-	r2 = VectorAdd(r2, m3);
-	r3 = VectorAdd(r3, m3);
-	r4 = VectorAdd(r4, m3);
-	r5 = VectorAdd(r5, m3);
-	r6 = VectorAdd(r6, m3);
-	r7 = VectorAdd(r7, m3);
-
 	FBox NewBox;
-	
-	VectorRegister min0 = VectorMin(r0, r1);
-	VectorRegister min1 = VectorMin(r2, r3);
-	VectorRegister min2 = VectorMin(r4, r5);
-	VectorRegister min3 = VectorMin(r6, r7);
-	VectorRegister max0 = VectorMax(r0, r1);
-	VectorRegister max1 = VectorMax(r2, r3);
-	VectorRegister max2 = VectorMax(r4, r5);
-	VectorRegister max3 = VectorMax(r6, r7);
-	
-	min0 = VectorMin(min0, min1);
-	min1 = VectorMin(min2, min3);
-	max0 = VectorMax(max0, max1);
-	max1 = VectorMax(max2, max3);
-	min0 = VectorMin(min0, min1);
-	max0 = VectorMax(max0, max1);
 
-	VectorStoreFloat3(min0, &NewBox.Min);
-	VectorStoreFloat3(max0, &NewBox.Max);
+	const VectorRegister VecMin = VectorLoadFloat3(&Min);
+	const VectorRegister VecMax = VectorLoadFloat3(&Max);
+
+	const VectorRegister m0 = VectorLoadAligned(M.M[0]);
+	const VectorRegister m1 = VectorLoadAligned(M.M[1]);
+	const VectorRegister m2 = VectorLoadAligned(M.M[2]);
+	const VectorRegister m3 = VectorLoadAligned(M.M[3]);
+
+	const VectorRegister Half = VectorSetFloat3(0.5f, 0.5f, 0.5f);
+	const VectorRegister Origin = VectorMultiply(VectorAdd(VecMax, VecMin), Half);
+	const VectorRegister Extent = VectorMultiply(VectorSubtract(VecMax, VecMin), Half);
 
+	VectorRegister NewOrigin = VectorMultiply(VectorReplicate(Origin, 0), m0);
+	NewOrigin = VectorMultiplyAdd(VectorReplicate(Origin, 1), m1, NewOrigin);
+	NewOrigin = VectorMultiplyAdd(VectorReplicate(Origin, 2), m2, NewOrigin);
+	NewOrigin = VectorAdd(NewOrigin, m3);
+
+	VectorRegister NewExtent = VectorAbs(VectorMultiply(VectorReplicate(Extent, 0), m0));
+	NewExtent = VectorAdd(NewExtent, VectorAbs(VectorMultiply(VectorReplicate(Extent, 1), m1)));
+	NewExtent = VectorAdd(NewExtent, VectorAbs(VectorMultiply(VectorReplicate(Extent, 2), m2)));
+
+	const VectorRegister NewVecMin = VectorSubtract(NewOrigin, NewExtent);
+	const VectorRegister NewVecMax = VectorAdd(NewOrigin, NewExtent);
+	
+	VectorStoreFloat3(NewVecMin, &NewBox.Min );
+	VectorStoreFloat3(NewVecMax, &NewBox.Max );
+	
 	NewBox.IsValid = 1;
 
 	return NewBox;
diff --git a/Engine/Source/Runtime/Core/Private/Math/BoxSphereBounds.cpp b/Engine/Source/Runtime/Core/Private/Math/BoxSphereBounds.cpp
--- a/Engine/Source/Runtime/Core/Private/Math/BoxSphereBounds.cpp
+++ b/Engine/Source/Runtime/Core/Private/Math/BoxSphereBounds.cpp
@@ -14,48 +14,38 @@ FBoxSphereBounds FBoxSphereBounds::TransformBy(const FMatrix& M) const
 {
 	FBoxSphereBounds Result;
 
-	Result.Origin = M.TransformPosition(Origin);
-	Result.BoxExtent = FVector::ZeroVector;
+	const VectorRegister VecOrigin = VectorLoadFloat3(&Origin);
+	const VectorRegister VecExtent = VectorLoadFloat3(&BoxExtent);
 
-	float Signs[2] = { -1.0f, 1.0f };
+	const VectorRegister m0 = VectorLoadAligned(M.M[0]);
+	const VectorRegister m1 = VectorLoadAligned(M.M[1]);
+	const VectorRegister m2 = VectorLoadAligned(M.M[2]);
+	const VectorRegister m3 = VectorLoadAligned(M.M[3]);
 
-	for (int32 X = 0; X < 2; X++)
-	{
-		for (int32 Y = 0; Y < 2; Y++)
-		{
-			for (int32 Z = 0; Z < 2; Z++)
-			{
-				FVector	Corner = M.TransformVector(FVector(Signs[X] * BoxExtent.X,Signs[Y] * BoxExtent.Y,Signs[Z] * BoxExtent.Z));
+	VectorRegister NewOrigin = VectorMultiply(VectorReplicate(VecOrigin, 0), m0);
+	NewOrigin = VectorMultiplyAdd(VectorReplicate(VecOrigin, 1), m1, NewOrigin);
+	NewOrigin = VectorMultiplyAdd(VectorReplicate(VecOrigin, 2), m2, NewOrigin);
+	NewOrigin = VectorAdd(NewOrigin, m3);
 
-				Result.BoxExtent.X = FMath::Max(Corner.X,Result.BoxExtent.X);
-				Result.BoxExtent.Y = FMath::Max(Corner.Y,Result.BoxExtent.Y);
-				Result.BoxExtent.Z = FMath::Max(Corner.Z,Result.BoxExtent.Z);
-			}
-		}
-	}
+	VectorRegister NewExtent = VectorAbs(VectorMultiply(VectorReplicate(VecExtent, 0), m0));
+	NewExtent = VectorAdd(NewExtent, VectorAbs(VectorMultiply(VectorReplicate(VecExtent, 1), m1)));
+	NewExtent = VectorAdd(NewExtent, VectorAbs(VectorMultiply(VectorReplicate(VecExtent, 2), m2)));
 
-	const FVector XAxis(M.M[0][0], M.M[0][1], M.M[0][2]);
-	const FVector YAxis(M.M[1][0], M.M[1][1], M.M[1][2]);
-	const FVector ZAxis(M.M[2][0], M.M[2][1], M.M[2][2]);
+	VectorStoreFloat3(NewExtent, &Result.BoxExtent);
+	VectorStoreFloat3(NewOrigin, &Result.Origin);
 
-	Result.SphereRadius = FMath::Sqrt(FMath::Max(XAxis | XAxis, FMath::Max(YAxis | YAxis, ZAxis | ZAxis))) * SphereRadius;
+	VectorRegister MaxRadius = VectorMultiply(m0, m0);
+	MaxRadius = VectorMultiplyAdd(m1, m1, MaxRadius);
+	MaxRadius = VectorMultiplyAdd(m2, m2, MaxRadius);
+	MaxRadius = VectorMax(VectorMax(MaxRadius, VectorReplicate(MaxRadius, 1)), VectorReplicate(MaxRadius, 2));
+	Result.SphereRadius = FMath::Sqrt(VectorGetComponent( MaxRadius, 0) ) * SphereRadius;
 
 	return Result;
 }
 
-
 FBoxSphereBounds FBoxSphereBounds::TransformBy(const FTransform& M) const
 {
-	FBoxSphereBounds Result;
-
-	Result.Origin = M.TransformPosition(Origin);
-
-	FVector4 XAxis = M.TransformVector(FVector(BoxExtent.X, 0.0, 0.0)).GetAbs();
-	FVector4 YAxis = M.TransformVector(FVector(0.0, BoxExtent.Y, 0.0)).GetAbs();
-	FVector4 ZAxis = M.TransformVector(FVector(0.0, 0.0, BoxExtent.Z)).GetAbs();
- 	Result.BoxExtent = XAxis + YAxis + ZAxis;
-
-	Result.SphereRadius = SphereRadius * M.GetMaximumAxisScale();
-
+	const FMatrix Mat = M.ToMatrixWithScale();
+	FBoxSphereBounds Result = TransformBy(Mat);
 	return Result;
 }
\ No newline at end of file
