[INTEGRATE] Change 2429887 by [REDACTED_PERSON][REDACTED_HANDLE] on 2015/02/03 09:50:31

	Allow FArchive to filter properties it is not interested in

diff --git a/Engine/Source/Runtime/Core/Public/Serialization/ArchiveBase.h b/Engine/Source/Runtime/Core/Public/Serialization/ArchiveBase.h
--- a/Engine/Source/Runtime/Core/Public/Serialization/ArchiveBase.h
+++ b/Engine/Source/Runtime/Core/Public/Serialization/ArchiveBase.h
@@ -870,6 +870,14 @@ public:
 		return false;
 	}
 
+	/**
+	 * Checks whether the archive wants to skip the property independent of the other flags
+	 */
+	virtual bool ShouldSkipProperty(const class UProperty* InProperty) const
+	{
+		return false;
+	}
+
 	/**
 	 * Sets the property that is currently being serialized
 	 * 
diff --git a/Engine/Source/Runtime/CoreUObject/Private/UObject/Property.cpp b/Engine/Source/Runtime/CoreUObject/Private/UObject/Property.cpp
--- a/Engine/Source/Runtime/CoreUObject/Private/UObject/Property.cpp
+++ b/Engine/Source/Runtime/CoreUObject/Private/UObject/Property.cpp
@@ -588,6 +588,9 @@ bool UProperty::ExportText_Direct
 
 bool UProperty::ShouldSerializeValue( FArchive& Ar ) const
 {
+	if (Ar.ShouldSkipProperty(this))
+		return false;
+
 	if (Ar.IsSaveGame() && !(PropertyFlags & CPF_SaveGame))
 		return false;
 
