Moved the initialization of the static empty FText instance back into FText::GetEmpty()

This allows us to control the initialization order of the empty FText instance, as the file-scope static was sometimes used before our empty FText has itself been constructed (in monolothic builds, when other statics tried to construct an empty FText).

diff --git a/Engine/Source/Runtime/Core/Private/Internationalization/Text.cpp b/Engine/Source/Runtime/Core/Private/Internationalization/Text.cpp
--- a/Engine/Source/Runtime/Core/Private/Internationalization/Text.cpp
+++ b/Engine/Source/Runtime/Core/Private/Internationalization/Text.cpp
@@ -226,8 +226,6 @@ FArchive& operator<<(FArchive& Ar, FNumberFormattingOptions& Value)
 bool FText::bEnableErrorCheckingResults = ENABLE_TEXT_ERROR_CHECKING_RESULTS;
 bool FText::bSuppressWarnings = false;
 
-const FText FText::StaticEmptyText = FText(FText::EInitToEmptyString::Value);
-
 //Some error text formats
 const FText FText::UnusedArgumentsError = LOCTEXT("Error_UnusedArguments", "ERR: The following arguments were unused ({0}).");
 const FText FText::CommentStartError = LOCTEXT("Error_CommentDoesntStartWithQMark", "ERR: The Comment for Arg Block {0} does not start with a '?'.");
@@ -240,7 +238,7 @@ const FText FText::UnEscapedCloseBraceOutsideOfArgumentBlock = LOCTEXT("Error_Un
 const FText FText::SerializationFailureError = LOCTEXT("Error_SerializationFailure", "ERR: Transient text cannot be serialized \"{0}\".");
 
 FText::FText()
-	: DisplayString( StaticEmptyText.DisplayString )
+	: DisplayString( GetEmpty().DisplayString )
 	, Flags(0)
 {
 }
diff --git a/Engine/Source/Runtime/Core/Public/Internationalization/Text.h b/Engine/Source/Runtime/Core/Public/Internationalization/Text.h
--- a/Engine/Source/Runtime/Core/Public/Internationalization/Text.h
+++ b/Engine/Source/Runtime/Core/Public/Internationalization/Text.h
@@ -124,6 +124,9 @@ public:
 
 	static const FText& GetEmpty()
 	{
+		// This is initialized inside this function as we need to be able to control the initialization order of the empty FText instance
+		// If this were a file-scope static, we can end up with other statics trying to construct an empty FText before our empty FText has itself been constructed
+		static const FText StaticEmptyText = FText(FText::EInitToEmptyString::Value);
 		return StaticEmptyText;
 	}
 
@@ -359,8 +362,6 @@ private:
 	static bool bSuppressWarnings;
 
 public:
-	static const FText StaticEmptyText;
-
 	//Some error text formats
 	static const FText UnusedArgumentsError;
 	static const FText CommentStartError;
