Expanded TArray::RemoveAllSwap with bAllowShrinking parameter #UE4

- merged over from demo branch

diff --git a/Engine/Source/Runtime/Core/Public/Containers/Array.h b/Engine/Source/Runtime/Core/Public/Containers/Array.h
--- a/Engine/Source/Runtime/Core/Public/Containers/Array.h
+++ b/Engine/Source/Runtime/Core/Public/Containers/Array.h
@@ -1968,13 +1968,13 @@ public:
 	 * @param Predicate Predicate class instance
 	 */
 	template <class PREDICATE_CLASS>
-	void RemoveAllSwap(const PREDICATE_CLASS& Predicate)
+	void RemoveAllSwap(const PREDICATE_CLASS& Predicate, bool bAllowShrinking = true)
 	{
 		for (int32 ItemIndex = 0; ItemIndex < Num();)
 		{
 			if (Predicate((*this)[ItemIndex]))
 			{
-				RemoveAtSwap(ItemIndex);
+				RemoveAtSwap(ItemIndex, 1, bAllowShrinking);
 			}
 			else
 			{
