#CORE - UE-9320 Mac crashes are not uploading to Crash Reporter (added an option to override the stack size for threads created for the thread pool)

diff --git a/Engine/Source/Programs/CrashReporter/CrashReportClient/Private/CrashReportClientApp.cpp b/Engine/Source/Programs/CrashReporter/CrashReportClient/Private/CrashReportClientApp.cpp
--- a/Engine/Source/Programs/CrashReporter/CrashReportClient/Private/CrashReportClientApp.cpp
+++ b/Engine/Source/Programs/CrashReporter/CrashReportClient/Private/CrashReportClientApp.cpp
@@ -94,6 +94,9 @@ FPlatformErrorReport LoadErrorReport()
 
 void RunCrashReportClient(const TCHAR* CommandLine)
 {
+	// Override the stack size for the thread pool.
+	FQueuedThreadPool::OverrideStackSize = 256 * 1024;
+
 	// Set up the main loop
 	GEngineLoop.PreInit(CommandLine);
 
diff --git a/Engine/Source/Runtime/Core/Private/HAL/ThreadingBase.cpp b/Engine/Source/Runtime/Core/Private/HAL/ThreadingBase.cpp
--- a/Engine/Source/Runtime/Core/Private/HAL/ThreadingBase.cpp
+++ b/Engine/Source/Runtime/Core/Private/HAL/ThreadingBase.cpp
@@ -426,6 +426,13 @@ public:
 		FScopeLock Lock(SynchQueue);
 		// Presize the array so there is no extra memory allocated
 		QueuedThreads.Empty(InNumQueuedThreads);
+
+		// Check for stack size override.
+		if( OverrideStackSize > StackSize )
+		{
+			StackSize = OverrideStackSize;
+		}
+
 		// Now create each thread and add it to the array
 		for (uint32 Count = 0; Count < InNumQueuedThreads && bWasSuccessful == true; Count++)
 		{
@@ -573,6 +580,7 @@ public:
 	}
 };
 
+uint32 FQueuedThreadPool::OverrideStackSize = 0;
 
 FQueuedThreadPool* FQueuedThreadPool::Allocate()
 {
diff --git a/Engine/Source/Runtime/Core/Public/HAL/ThreadingBase.h b/Engine/Source/Runtime/Core/Public/HAL/ThreadingBase.h
--- a/Engine/Source/Runtime/Core/Public/HAL/ThreadingBase.h
+++ b/Engine/Source/Runtime/Core/Public/HAL/ThreadingBase.h
@@ -623,6 +623,13 @@ public:
 	 * @return A new thread pool.
 	 */
 	static FQueuedThreadPool* Allocate();
+
+	/**
+	 *	Stack size for threads created for the thread pool. 
+	 *	Can be overridden by other projects.
+	 *	If 0 means to use the value passed in the Create method.
+	 */
+	static uint32 OverrideStackSize;
 };
 
 
