CIS fix for iOS
#ios

diff --git a/Engine/Source/Runtime/Core/Public/Math/UnrealMathNeon.h b/Engine/Source/Runtime/Core/Public/Math/UnrealMathNeon.h
--- a/Engine/Source/Runtime/Core/Public/Math/UnrealMathNeon.h
+++ b/Engine/Source/Runtime/Core/Public/Math/UnrealMathNeon.h
@@ -224,7 +224,7 @@ FORCEINLINE void VectorStore( VectorRegister Vec, void* Ptr )
  * @param Vec	Vector to store XYZ
  * @param Ptr	Unaligned memory pointer
  */
-FORCEINLINE void VectorStoreFloat3( VectorRegister& Vec, void* Ptr )
+FORCEINLINE void VectorStoreFloat3( const VectorRegister& Vec, void* Ptr )
 {
 	vst1q_lane_f32( ((float32_t *)Ptr) + 0, Vec, 0 );
 	vst1q_lane_f32( ((float32_t *)Ptr) + 1, Vec, 1 );
