We no longer limit the window size when reshaping a fullscreen window

UE-9036 - Changes to Resolution are not saved if Fullscreen is toggled on

Windows defaults to preventing you being able to size a window larger than the currently available desktop size. When running in D3D fullscreen mode, this means that it prevents you being able to resize a window larger than the current fullscreen resolution.

You can override this limitation by passing SWP_NOSENDCHANGING when calling SetWindowPos. This prevents the window being notified that it is about to change size (preventing the WM_GETMINMAXINFO message), and instead just notifies the window of its new size.

This CL also makes sure we call SetWindowPos after changing the window styles, as the MSDN stipulates that the DWM may cache this information until notified of a change by SWP_FRAMECHANGED.

diff --git a/Engine/Source/Runtime/Core/Private/Windows/WindowsWindow.cpp b/Engine/Source/Runtime/Core/Private/Windows/WindowsWindow.cpp
--- a/Engine/Source/Runtime/Core/Private/Windows/WindowsWindow.cpp
+++ b/Engine/Source/Runtime/Core/Private/Windows/WindowsWindow.cpp
@@ -195,6 +195,7 @@ void FWindowsWindow::Initialize( FWindowsApplication* const Application, const T
 		}
 
 		verify(SetWindowLong(HWnd, GWL_STYLE, WindowStyle));
+		::SetWindowPos(HWnd, nullptr, 0, 0, 0, 0, SWP_NOMOVE | SWP_NOSIZE | SWP_NOZORDER | SWP_FRAMECHANGED);
 
 		AdjustWindowRegion( ClientWidth, ClientHeight );
 	}
@@ -353,19 +354,17 @@ void FWindowsWindow::ReshapeWindow( int32 NewX, int32 NewY, int32 NewWidth, int3
 		NewHeight = FMath::Max( NewHeight, FMath::Min( OldHeight, MinRetainedHeight ) );
 	}
 
-
-	// NOTE: MoveWindow will trigger a WM_SIZE and our SWindow's cached size will be updated
-
 	const int32 OldWidth = WindowInfo.rcClient.right - WindowInfo.rcClient.left;
 	const int32 OldHeight = WindowInfo.rcClient.bottom - WindowInfo.rcClient.top;
 
+	// NOTE: SetWindowPos will trigger a WM_SIZE and our SWindow's cached size will be updated
 	if( NewX != WindowInfo.rcClient.left || NewY != WindowInfo.rcClient.top || NewWidth != OldWidth || NewHeight != OldHeight  )
 	{
-		const bool bRepaint = true;
-		::MoveWindow(HWnd, WindowX, WindowY, NewWidth, NewHeight, bRepaint);
+		// We use SWP_NOSENDCHANGING when in fullscreen mode to prevent Windows limiting our window size to the current resolution, as that 
+		// prevents us being able to change to a higher resolution while in fullscreen mode
+		::SetWindowPos( HWnd, nullptr, WindowX, WindowY, NewWidth, NewHeight, SWP_NOZORDER | SWP_NOACTIVATE | ((WindowMode == EWindowMode::Fullscreen) ? SWP_NOSENDCHANGING : 0) );
 	}
 
-
 	if( Definition->SizeWillChangeOften && bVirtualSizeChanged )
 	{
 		AdjustWindowRegion( VirtualWidth, VirtualHeight );
@@ -538,6 +537,7 @@ void FWindowsWindow::SetWindowMode( EWindowMode::Type NewWindowMode )
 			WindowFlags |= FullscreenFlags;
 
 			SetWindowLong(HWnd, GWL_STYLE, WindowFlags);
+			::SetWindowPos(HWnd, nullptr, 0, 0, 0, 0, SWP_NOMOVE | SWP_NOSIZE | SWP_NOZORDER | SWP_FRAMECHANGED);
 
 			if (!bTrueFullscreen)
 			{
@@ -585,6 +585,7 @@ void FWindowsWindow::SetWindowMode( EWindowMode::Type NewWindowMode )
 			WindowFlags &= ~FullscreenFlags;
 			WindowFlags |= RestoredFlags;
 			SetWindowLong(HWnd, GWL_STYLE, WindowFlags);
+			::SetWindowPos(HWnd, nullptr, 0, 0, 0, 0, SWP_NOMOVE | SWP_NOSIZE | SWP_NOZORDER | SWP_FRAMECHANGED);
 
 			::SetWindowPlacement(HWnd, &PreFullscreenWindowPlacement);
 		}
