Global symbol resolution when loading shared objects.

- Needed for Steam to work when its library is loaded by a plugin.

Merging CL 2435807 from my UT branch to main.

diff --git a/Engine/Source/Runtime/Core/Private/Linux/LinuxPlatformProcess.cpp b/Engine/Source/Runtime/Core/Private/Linux/LinuxPlatformProcess.cpp
--- a/Engine/Source/Runtime/Core/Private/Linux/LinuxPlatformProcess.cpp
+++ b/Engine/Source/Runtime/Core/Private/Linux/LinuxPlatformProcess.cpp
@@ -13,7 +13,7 @@ void* FLinuxPlatformProcess::GetDllHandle( const TCHAR* Filename )
 {
 	check( Filename );
 	FString AbsolutePath = FPaths::ConvertRelativePathToFull(Filename);
-	void *Handle = dlopen( TCHAR_TO_ANSI(*AbsolutePath), RTLD_LAZY | RTLD_LOCAL );
+	void *Handle = dlopen( TCHAR_TO_ANSI(*AbsolutePath), RTLD_LAZY | RTLD_GLOBAL );
 	if (!Handle)
 	{
 		UE_LOG(LogLinux, Warning, TEXT("dlopen failed: %s"), ANSI_TO_TCHAR(dlerror()) );
