Removed empty state from TInterval,  inverted intervals are empty, to make it easier to work with
Restricted it to numeric types only
Added utility functions

diff --git a/Engine/Source/Runtime/Core/Public/Math/Interval.h b/Engine/Source/Runtime/Core/Public/Math/Interval.h
--- a/Engine/Source/Runtime/Core/Public/Math/Interval.h
+++ b/Engine/Source/Runtime/Core/Public/Math/Interval.h
@@ -4,30 +4,28 @@
 
 
 /**
- * Template for interval of numbers.
+ * Template for numeric interval
  */
 template<typename ElementType> struct TInterval
 {
+	static_assert(TIsArithmeticType<ElementType>::Value, "Interval can be used only with numeric types");
+	
 	/** Holds the lower bound of the interval. */
 	ElementType Min;
 	
 	/** Holds the upper bound of the interval. */
 	ElementType Max;
-	
-	/** Holds a flag indicating whether the interval is empty. */
-	bool bIsEmpty;
-
+		
 public:
 
 	/**
 	 * Default constructor.
 	 *
-	 * The interval is initialized to [0, 0].
+	 * The interval is empty
 	 */
 	TInterval()
-		: Min(0)
-		, Max(0)
-		, bIsEmpty(true)
+		: Min(TNumericLimits<ElementType>::Max())
+		, Max(TNumericLimits<ElementType>::Lowest())
 	{ }
 
     /**
@@ -39,7 +37,6 @@ public:
 	TInterval( ElementType InMin, ElementType InMax )
 		: Min(InMin)
 		, Max(InMax)
-		, bIsEmpty(InMin >= InMax)
 	{ }
 
 public:
@@ -51,7 +48,7 @@ public:
 	 */
 	void operator+= ( ElementType X )
 	{
-		if (!bIsEmpty)
+		if (!IsEmpty())
 		{
 			Min += X;
 			Max += X;
@@ -65,7 +62,7 @@ public:
 	 */
 	void operator-= ( ElementType X )
 	{
-		if (!bIsEmpty)
+		if (!IsEmpty())
 		{
 			Min -= X;
 			Max -= X;
@@ -73,6 +70,37 @@ public:
 	}
 
 public:
+	
+	/**
+	 * Computes the size of this interval.
+	 *
+	 * @return Interval size.
+	 */
+	ElementType Size() const
+	{
+		return (Max - Min);
+	}
+
+	/**
+	 * Whether interval is empty.
+	 *
+	 * @return false when interval is empty, true otherwise
+	 */
+	ElementType IsEmpty() const
+	{
+		return (Min >= Max);
+	}
+	
+	/**
+	 * Checks whether this interval contains the specified element.
+	 *
+	 * @param Element The element to check.
+	 * @return true if the range interval the element, false otherwise.
+	 */
+	bool Contains( const ElementType& Element ) const
+	{
+		return !IsEmpty() && (Element >= Min && Element <= Max);
+	}
 
 	/**
 	 * Expands this interval to both sides by the specified amount.
@@ -81,7 +109,7 @@ public:
 	 */
 	void Expand( ElementType ExpandAmount )
 	{
-		if (!bIsEmpty)
+		if (!IsEmpty())
 		{
 			Min -= ExpandAmount;
 			Max += ExpandAmount;
@@ -95,11 +123,10 @@ public:
 	 */
 	void Include( ElementType X )
 	{
-		if (bIsEmpty)
+		if (IsEmpty())
 		{
 			Min = X;
 			Max = X;
-			bIsEmpty = false;
 		}
 		else
 		{
@@ -115,6 +142,22 @@ public:
 		}
 	}
 
+	/**
+	 * Interval interpolation
+	 *
+	 * @param Alpha interpolation amount
+	 * @return interpolation result
+	 */
+	ElementType Interpolate( float Alpha ) const
+	{
+		if (!IsEmpty())
+		{
+			return Min + ElementType(Alpha*Size());
+		}
+		
+		return ElementType();
+	}
+
 public:
 
 	/**
@@ -126,13 +169,36 @@ public:
 	 */
 	friend TInterval Intersect( const TInterval& A, const TInterval& B )
 	{
-		if (A.bIsEmpty || B.bIsEmpty)
+		if (A.IsEmpty() || B.IsEmpty())
 		{
 			return TInterval();
 		}
 
 		return TInterval(FMath::Max(A.Min, B.Min), FMath::Min(A.Max, B.Max));
 	}
+
+	/**
+	 * Serializes the interval.
+	 *
+	 * @param Ar The archive to serialize into.
+	 * @param Interval The interval to serialize.
+	 * @return Reference to the Archive after serialization.
+	 */
+	friend class FArchive& operator<<( class FArchive& Ar, TInterval& Interval )
+	{
+		return Ar << Interval.Min << Interval.Max;
+	}
+	
+	/**
+	 * Gets the hash for the specified interval.
+	 *
+	 * @param Interval The Interval to get the hash for.
+	 * @return Hash value.
+	 */
+	friend uint32 GetTypeHash(const TInterval& Interval)
+	{
+		return HashCombine(GetTypeHash(Interval.Min), GetTypeHash(Interval.Max));
+	}
 };
 
 /* Default intervals for built-in types
