Fix FIntVector::ToString()

diff --git a/Engine/Source/Runtime/Core/Public/Math/IntVector.h b/Engine/Source/Runtime/Core/Public/Math/IntVector.h
--- a/Engine/Source/Runtime/Core/Public/Math/IntVector.h
+++ b/Engine/Source/Runtime/Core/Public/Math/IntVector.h
@@ -382,7 +382,7 @@ FORCEINLINE int32 FIntVector::Size() const
 
 FORCEINLINE FString FIntVector::ToString() const
 {
-	return FString::Printf(TEXT("X=%3.3f Y=%3.3f Z=%3.3f"), X, Y, Z);
+	return FString::Printf(TEXT("X=%d Y=%d Z=%d"), X, Y, Z);
 }
 
 FORCEINLINE uint32 GetTypeHash(const FIntVector& Vector)
