CrashReport - Added re-entrancy guard to NewReportEnsure

diff --git a/Engine/Source/Runtime/Core/Private/Windows/WindowsPlatformCrashContext.cpp b/Engine/Source/Runtime/Core/Private/Windows/WindowsPlatformCrashContext.cpp
--- a/Engine/Source/Runtime/Core/Private/Windows/WindowsPlatformCrashContext.cpp
+++ b/Engine/Source/Runtime/Core/Private/Windows/WindowsPlatformCrashContext.cpp
@@ -376,11 +376,25 @@ int32 ReportCrashUsingCrashReportClient(EXCEPTION_POINTERS* ExceptionInfo, const
 
 } // end anonymous namespace
 
+static FCriticalSection EnsureLock;
+static bool bReentranceGuard = false;
+
 /** 
  * Report an ensure to the crash reporting system
  */
 void NewReportEnsure( const TCHAR* ErrorMessage )
 {
+	// Simple re-entrance guard.
+	EnsureLock.Lock();
+
+	if( bReentranceGuard )
+	{
+		EnsureLock.Unlock();
+		return;
+	}
+
+	bReentranceGuard = true;
+
 #if WINVER > 0x502	// Windows Error Reporting is not supported on Windows XP
 #if !PLATFORM_SEH_EXCEPTIONS_DISABLED
 	__try
@@ -395,6 +409,9 @@ void NewReportEnsure( const TCHAR* ErrorMessage )
 	}
 #endif
 #endif	// WINVER
+
+	bReentranceGuard = false;
+	EnsureLock.Unlock();
 }
 
 #include "HideWindowsPlatformTypes.h"
