Don't call SlateApplication::OnCursorSet in FMacApplication::ProcessNSEvent() if the current event's window was closed as a result of handling of mouse down or mouse up

diff --git a/Engine/Source/Runtime/Core/Private/Mac/MacApplication.cpp b/Engine/Source/Runtime/Core/Private/Mac/MacApplication.cpp
--- a/Engine/Source/Runtime/Core/Private/Mac/MacApplication.cpp
+++ b/Engine/Source/Runtime/Core/Private/Mac/MacApplication.cpp
@@ -739,7 +739,7 @@ void FMacApplication::ProcessNSEvent(NSEvent* const Event)
 					MessageHandler->OnMouseDown(CurrentEventWindow, Button);
 				}
 
-				if (!DraggedWindow && !GetCapture())
+				if (CurrentEventWindow->GetWindowHandle() && !DraggedWindow && !GetCapture())
 				{
 					MessageHandler->OnCursorSet();
 				}
@@ -775,7 +775,7 @@ void FMacApplication::ProcessNSEvent(NSEvent* const Event)
 
 			MessageHandler->OnMouseUp(Button);
 
-			if (CurrentEventWindow.IsValid() && !DraggedWindow && !GetCapture())
+			if (CurrentEventWindow.IsValid() && CurrentEventWindow->GetWindowHandle() && !DraggedWindow && !GetCapture())
 			{
 				MessageHandler->OnCursorSet();
 			}
