Removed unnecessary derived classes of ICUBreakIterator that merely call its constructor with a hardcoded argument. The only exposed factory methods that already exist instead instantiate it and pass the hardcoded argument.

diff --git a/Engine/Source/Runtime/Core/Private/Internationalization/ICUCharacterBoundaryIterator.cpp b/Engine/Source/Runtime/Core/Private/Internationalization/ICUCharacterBoundaryIterator.cpp
--- a/Engine/Source/Runtime/Core/Private/Internationalization/ICUCharacterBoundaryIterator.cpp
+++ b/Engine/Source/Runtime/Core/Private/Internationalization/ICUCharacterBoundaryIterator.cpp
@@ -6,20 +6,9 @@
 #if UE_ENABLE_ICU
 #include "ICUBreakIterator.h"
 
-class FICUCharacterBoundaryIterator : public FICUBreakIterator
-{
-public:
-	FICUCharacterBoundaryIterator();
-};
-
-FICUCharacterBoundaryIterator::FICUCharacterBoundaryIterator()
-	: FICUBreakIterator(FICUBreakIteratorManager::Get().CreateCharacterBoundaryIterator())
-{
-}
-
 TSharedRef<IBreakIterator> FBreakIterator::CreateCharacterBoundaryIterator()
 {
-	return MakeShareable(new FICUCharacterBoundaryIterator());
+	return MakeShareable(new FICUBreakIterator(FICUBreakIteratorManager::Get().CreateCharacterBoundaryIterator()));
 }
 
 #endif
diff --git a/Engine/Source/Runtime/Core/Private/Internationalization/ICULineBreakIterator.cpp b/Engine/Source/Runtime/Core/Private/Internationalization/ICULineBreakIterator.cpp
--- a/Engine/Source/Runtime/Core/Private/Internationalization/ICULineBreakIterator.cpp
+++ b/Engine/Source/Runtime/Core/Private/Internationalization/ICULineBreakIterator.cpp
@@ -6,20 +6,9 @@
 #if UE_ENABLE_ICU
 #include "ICUBreakIterator.h"
 
-class FICULineBreakIterator : public FICUBreakIterator
-{
-public:
-	FICULineBreakIterator();
-};
-
-FICULineBreakIterator::FICULineBreakIterator()
-	: FICUBreakIterator(FICUBreakIteratorManager::Get().CreateLineBreakIterator())
-{
-}
-
 TSharedRef<IBreakIterator> FBreakIterator::CreateLineBreakIterator()
 {
-	return MakeShareable(new FICULineBreakIterator());
+	return MakeShareable(new FICUBreakIterator(FICUBreakIteratorManager::Get().CreateLineBreakIterator()));
 }
 
 #endif
diff --git a/Engine/Source/Runtime/Core/Private/Internationalization/ICUWordBreakIterator.cpp b/Engine/Source/Runtime/Core/Private/Internationalization/ICUWordBreakIterator.cpp
--- a/Engine/Source/Runtime/Core/Private/Internationalization/ICUWordBreakIterator.cpp
+++ b/Engine/Source/Runtime/Core/Private/Internationalization/ICUWordBreakIterator.cpp
@@ -6,20 +6,9 @@
 #if UE_ENABLE_ICU
 #include "ICUBreakIterator.h"
 
-class FICUWordBreakIterator : public FICUBreakIterator
-{
-public:
-	FICUWordBreakIterator();
-};
-
-FICUWordBreakIterator::FICUWordBreakIterator()
-	: FICUBreakIterator(FICUBreakIteratorManager::Get().CreateWordBreakIterator())
-{
-}
-
 TSharedRef<IBreakIterator> FBreakIterator::CreateWordBreakIterator()
 {
-	return MakeShareable(new FICUWordBreakIterator());
+	return MakeShareable(new FICUBreakIterator(FICUBreakIteratorManager::Get().CreateWordBreakIterator()));
 }
 
 #endif
