Added stat tracking for time spent in word wrapping logic for canvas text. Modified ICU culture wrapper initialization so that ICU resources are acquired when first needed to minimize the cost for cultures that are only used for querying names. Added break iterators to the resources that ICU culture wrappers track so that they can remain resident after first use.

diff --git a/Engine/Source/Runtime/Core/Private/Internationalization/ICUBreakIterator.cpp b/Engine/Source/Runtime/Core/Private/Internationalization/ICUBreakIterator.cpp
--- a/Engine/Source/Runtime/Core/Private/Internationalization/ICUBreakIterator.cpp
+++ b/Engine/Source/Runtime/Core/Private/Internationalization/ICUBreakIterator.cpp
@@ -29,8 +29,7 @@ FICUBreakIteratorManager& FICUBreakIteratorManager::Get()
 
 TWeakPtr<icu::BreakIterator> FICUBreakIteratorManager::CreateCharacterBoundaryIterator()
 {
-	UErrorCode ICUStatus = U_ZERO_ERROR;
-	TSharedRef<icu::BreakIterator> Iterator = MakeShareable(icu::BreakIterator::createCharacterInstance(icu::Locale::getDefault(), ICUStatus));
+	TSharedRef<icu::BreakIterator> Iterator = MakeShareable( FInternationalization::Get().GetDefaultCulture()->Implementation->GetBreakIterator(EBreakIteratorType::Grapheme)->clone() );
 	{
 		FScopeLock ScopeLock(&AllocatedIteratorsCS);
 		AllocatedIterators.Add(Iterator);
@@ -40,8 +39,7 @@ TWeakPtr<icu::BreakIterator> FICUBreakIteratorManager::CreateCharacterBoundaryIt
 
 TWeakPtr<icu::BreakIterator> FICUBreakIteratorManager::CreateWordBreakIterator()
 {
-	UErrorCode ICUStatus = U_ZERO_ERROR;
-	TSharedRef<icu::BreakIterator> Iterator = MakeShareable(icu::BreakIterator::createWordInstance(icu::Locale::getDefault(), ICUStatus));
+	TSharedRef<icu::BreakIterator> Iterator = MakeShareable( FInternationalization::Get().GetDefaultCulture()->Implementation->GetBreakIterator(EBreakIteratorType::Word)->clone() );
 	{
 		FScopeLock ScopeLock(&AllocatedIteratorsCS);
 		AllocatedIterators.Add(Iterator);
@@ -51,8 +49,7 @@ TWeakPtr<icu::BreakIterator> FICUBreakIteratorManager::CreateWordBreakIterator()
 
 TWeakPtr<icu::BreakIterator> FICUBreakIteratorManager::CreateLineBreakIterator()
 {
-	UErrorCode ICUStatus = U_ZERO_ERROR;
-	TSharedRef<icu::BreakIterator> Iterator = MakeShareable(icu::BreakIterator::createLineInstance(icu::Locale::getDefault(), ICUStatus));
+	TSharedRef<icu::BreakIterator> Iterator = MakeShareable( FInternationalization::Get().GetDefaultCulture()->Implementation->GetBreakIterator(EBreakIteratorType::Line)->clone() );
 	{
 		FScopeLock ScopeLock(&AllocatedIteratorsCS);
 		AllocatedIterators.Add(Iterator);
diff --git a/Engine/Source/Runtime/Core/Private/Internationalization/ICUCulture.cpp b/Engine/Source/Runtime/Core/Private/Internationalization/ICUCulture.cpp
--- a/Engine/Source/Runtime/Core/Private/Internationalization/ICUCulture.cpp
+++ b/Engine/Source/Runtime/Core/Private/Internationalization/ICUCulture.cpp
@@ -5,9 +5,37 @@
 #if UE_ENABLE_ICU
 #include "ICUUtilities.h"
 #include "ICUCulture.h"
+#include "Function.h"
 
 namespace
 {
+	TSharedRef<const icu::BreakIterator> CreateBreakIterator( const icu::Locale& ICULocale, const EBreakIteratorType Type)
+	{
+		UErrorCode ICUStatus = U_ZERO_ERROR;
+		TFunction<icu::BreakIterator* (const icu::Locale&, UErrorCode&)> FactoryFunction;
+		switch (Type)
+		{
+		case EBreakIteratorType::Grapheme:
+			FactoryFunction = icu::BreakIterator::createCharacterInstance;
+			break;
+		case EBreakIteratorType::Word:
+			FactoryFunction = icu::BreakIterator::createWordInstance;
+			break;
+		case EBreakIteratorType::Line:
+			FactoryFunction = icu::BreakIterator::createLineInstance;
+			break;
+		case EBreakIteratorType::Sentence:
+			FactoryFunction = icu::BreakIterator::createSentenceInstance;
+			break;
+		case EBreakIteratorType::Title:
+			FactoryFunction = icu::BreakIterator::createTitleInstance;
+			break;
+		}
+		TSharedPtr<const icu::BreakIterator> Ptr = MakeShareable( FactoryFunction(ICULocale, ICUStatus) );
+		checkf(Ptr.IsValid(), TEXT("Creating a break iterator object failed using locale %s. Perhaps this locale has no data."), StringCast<TCHAR>(ICULocale.getName()).Get());
+		return Ptr.ToSharedRef();
+	}
+
 	TSharedRef<const icu::Collator, ESPMode::ThreadSafe> CreateCollator( const icu::Locale& ICULocale )
 	{
 		UErrorCode ICUStatus = U_ZERO_ERROR;
@@ -70,13 +98,6 @@ namespace
 
 FCulture::FICUCultureImplementation::FICUCultureImplementation(const FString& LocaleName)
 	: ICULocale( TCHAR_TO_ANSI( *LocaleName ) )
-	, ICUCollator( CreateCollator( ICULocale ) )
-	, ICUDecimalFormat( CreateDecimalFormat( ICULocale ) )
-	, ICUCurrencyFormat( CreateCurrencyFormat( ICULocale ) )
-	, ICUPercentFormat( CreatePercentFormat( ICULocale ) )
-	, ICUDateFormat ( CreateDateFormat( ICULocale ) )
-	, ICUTimeFormat ( CreateTimeFormat( ICULocale ) )
-	, ICUDateTimeFormat ( CreateDateTimeFormat( ICULocale ) )
 {
 
 }
@@ -210,11 +231,52 @@ FString FCulture::FICUCultureImplementation::GetVariant() const
 	return ICULocale.getVariant();
 }
 
-TSharedRef<const icu::Collator, ESPMode::ThreadSafe> FCulture::FICUCultureImplementation::GetCollator(const ETextComparisonLevel::Type ComparisonLevel) const
+TSharedRef<const icu::BreakIterator> FCulture::FICUCultureImplementation::GetBreakIterator(const EBreakIteratorType Type)
+{
+	TSharedPtr<const icu::BreakIterator> Result;
+
+	switch (Type)
+	{
+	case EBreakIteratorType::Grapheme:
+		{
+			Result = ICUGraphemeBreakIterator.IsValid() ? ICUGraphemeBreakIterator : ( ICUGraphemeBreakIterator = CreateBreakIterator(ICULocale, Type) );
+		}
+		break;
+	case EBreakIteratorType::Word:
+		{
+			Result = ICUWordBreakIterator.IsValid() ? ICUWordBreakIterator : ( ICUWordBreakIterator = CreateBreakIterator(ICULocale, Type) );
+		}
+		break;
+	case EBreakIteratorType::Line:
+		{
+			Result = ICULineBreakIterator.IsValid() ? ICULineBreakIterator : ( ICULineBreakIterator = CreateBreakIterator(ICULocale, Type) );
+		}
+		break;
+	case EBreakIteratorType::Sentence:
+		{
+			Result = ICUSentenceBreakIterator.IsValid() ? ICUSentenceBreakIterator : ( ICUSentenceBreakIterator = CreateBreakIterator(ICULocale, Type) );
+		}
+		break;
+	case EBreakIteratorType::Title:
+		{
+			Result = ICUTitleBreakIterator.IsValid() ? ICUTitleBreakIterator : ( ICUTitleBreakIterator = CreateBreakIterator(ICULocale, Type) );
+		}
+		break;
+	}
+
+	return Result.ToSharedRef();
+}
+
+TSharedRef<const icu::Collator, ESPMode::ThreadSafe> FCulture::FICUCultureImplementation::GetCollator(const ETextComparisonLevel::Type ComparisonLevel)
 {
+	if (!ICUCollator.IsValid())
+	{
+		ICUCollator = CreateCollator( ICULocale );
+	}
+
 	UErrorCode ICUStatus = U_ZERO_ERROR;
 	const bool bIsDefault = (ComparisonLevel == ETextComparisonLevel::Default);
-	const TSharedRef<const icu::Collator, ESPMode::ThreadSafe> DefaultCollator( ICUCollator );
+	const TSharedRef<const icu::Collator, ESPMode::ThreadSafe> DefaultCollator( ICUCollator.ToSharedRef() );
 	if(bIsDefault)
 	{
 		return DefaultCollator;
@@ -227,10 +289,15 @@ TSharedRef<const icu::Collator, ESPMode::ThreadSafe> FCulture::FICUCultureImplem
 	}
 }
 
-TSharedRef<const icu::DecimalFormat> FCulture::FICUCultureImplementation::GetDecimalFormatter(const FNumberFormattingOptions* const Options) const
+TSharedRef<const icu::DecimalFormat> FCulture::FICUCultureImplementation::GetDecimalFormatter(const FNumberFormattingOptions* const Options)
 {
+	if (!ICUDecimalFormat.IsValid())
+	{
+		ICUDecimalFormat = CreateDecimalFormat( ICULocale );
+	}
+
 	const bool bIsDefault = Options == NULL;
-	const TSharedRef<const icu::DecimalFormat> DefaultFormatter( ICUDecimalFormat );
+	const TSharedRef<const icu::DecimalFormat> DefaultFormatter( ICUDecimalFormat.ToSharedRef() );
 	if(bIsDefault)
 	{
 		return DefaultFormatter;
@@ -251,10 +318,15 @@ TSharedRef<const icu::DecimalFormat> FCulture::FICUCultureImplementation::GetDec
 	}
 }
 
-TSharedRef<const icu::DecimalFormat> FCulture::FICUCultureImplementation::GetCurrencyFormatter(const FString& CurrencyCode, const FNumberFormattingOptions* const Options) const
+TSharedRef<const icu::DecimalFormat> FCulture::FICUCultureImplementation::GetCurrencyFormatter(const FString& CurrencyCode, const FNumberFormattingOptions* const Options)
 {
+	if (!ICUCurrencyFormat.IsValid())
+	{
+		ICUCurrencyFormat = CreateCurrencyFormat( ICULocale );
+	}
+
 	const bool bIsDefault = Options == NULL && CurrencyCode.IsEmpty();
-	const TSharedRef<const icu::DecimalFormat> DefaultFormatter( ICUCurrencyFormat );
+	const TSharedRef<const icu::DecimalFormat> DefaultFormatter( ICUCurrencyFormat.ToSharedRef() );
 
 	if(bIsDefault)
 	{
@@ -285,10 +357,15 @@ TSharedRef<const icu::DecimalFormat> FCulture::FICUCultureImplementation::GetCur
 	}
 }
 
-TSharedRef<const icu::DecimalFormat> FCulture::FICUCultureImplementation::GetPercentFormatter(const FNumberFormattingOptions* const Options) const
+TSharedRef<const icu::DecimalFormat> FCulture::FICUCultureImplementation::GetPercentFormatter(const FNumberFormattingOptions* const Options)
 {
+	if (!ICUPercentFormat.IsValid())
+	{
+		ICUPercentFormat = CreatePercentFormat( ICULocale );
+	}
+
 	const bool bIsDefault = Options == NULL;
-	const TSharedRef<const icu::DecimalFormat> DefaultFormatter( ICUPercentFormat );
+	const TSharedRef<const icu::DecimalFormat> DefaultFormatter( ICUPercentFormat.ToSharedRef() );
 	if(bIsDefault)
 	{
 		return DefaultFormatter;
@@ -309,12 +386,17 @@ TSharedRef<const icu::DecimalFormat> FCulture::FICUCultureImplementation::GetPer
 	}
 }
 
-TSharedRef<const icu::DateFormat> FCulture::FICUCultureImplementation::GetDateFormatter(const EDateTimeStyle::Type DateStyle, const FString& TimeZone) const
+TSharedRef<const icu::DateFormat> FCulture::FICUCultureImplementation::GetDateFormatter(const EDateTimeStyle::Type DateStyle, const FString& TimeZone)
 {
+	if (!ICUDateFormat.IsValid())
+	{
+		ICUDateFormat = CreateDateFormat( ICULocale );
+	}
+
 	icu::UnicodeString InputTimeZoneID;
 	ICUUtilities::ConvertString(TimeZone, InputTimeZoneID, false);
 
-	const TSharedRef<const icu::DateFormat> DefaultFormatter( ICUDateFormat );
+	const TSharedRef<const icu::DateFormat> DefaultFormatter( ICUDateFormat.ToSharedRef() );
 
 	bool bIsDefaultTimeZone = TimeZone.IsEmpty();
 	if( !bIsDefaultTimeZone )
@@ -349,12 +431,17 @@ TSharedRef<const icu::DateFormat> FCulture::FICUCultureImplementation::GetDateFo
 	}
 }
 
-TSharedRef<const icu::DateFormat> FCulture::FICUCultureImplementation::GetTimeFormatter(const EDateTimeStyle::Type TimeStyle, const FString& TimeZone) const
+TSharedRef<const icu::DateFormat> FCulture::FICUCultureImplementation::GetTimeFormatter(const EDateTimeStyle::Type TimeStyle, const FString& TimeZone)
 {
+	if (!ICUTimeFormat.IsValid())
+	{
+		ICUTimeFormat = CreateTimeFormat( ICULocale );
+	}
+
 	icu::UnicodeString InputTimeZoneID;
 	ICUUtilities::ConvertString(TimeZone, InputTimeZoneID, false);
 
-	const TSharedRef<const icu::DateFormat> DefaultFormatter( ICUTimeFormat );
+	const TSharedRef<const icu::DateFormat> DefaultFormatter( ICUTimeFormat.ToSharedRef() );
 
 	bool bIsDefaultTimeZone = TimeZone.IsEmpty();
 	if( !bIsDefaultTimeZone )
@@ -389,12 +476,17 @@ TSharedRef<const icu::DateFormat> FCulture::FICUCultureImplementation::GetTimeFo
 	}
 }
 
-TSharedRef<const icu::DateFormat> FCulture::FICUCultureImplementation::GetDateTimeFormatter(const EDateTimeStyle::Type DateStyle, const EDateTimeStyle::Type TimeStyle, const FString& TimeZone) const
+TSharedRef<const icu::DateFormat> FCulture::FICUCultureImplementation::GetDateTimeFormatter(const EDateTimeStyle::Type DateStyle, const EDateTimeStyle::Type TimeStyle, const FString& TimeZone)
 {
+	if (!ICUDateTimeFormat.IsValid())
+	{
+		ICUDateTimeFormat = CreateDateTimeFormat( ICULocale );
+	}
+
 	icu::UnicodeString InputTimeZoneID;
 	ICUUtilities::ConvertString(TimeZone, InputTimeZoneID, false);
 
-	const TSharedRef<const icu::DateFormat> DefaultFormatter( ICUDateTimeFormat );
+	const TSharedRef<const icu::DateFormat> DefaultFormatter( ICUDateTimeFormat.ToSharedRef() );
 
 	bool bIsDefaultTimeZone = TimeZone.IsEmpty();
 	if( !bIsDefaultTimeZone )
@@ -429,4 +521,5 @@ TSharedRef<const icu::DateFormat> FCulture::FICUCultureImplementation::GetDateTi
 		return Formatter;
 	}
 }
+
 #endif
\ No newline at end of file
diff --git a/Engine/Source/Runtime/Core/Private/Internationalization/ICUCulture.h b/Engine/Source/Runtime/Core/Private/Internationalization/ICUCulture.h
--- a/Engine/Source/Runtime/Core/Private/Internationalization/ICUCulture.h
+++ b/Engine/Source/Runtime/Core/Private/Internationalization/ICUCulture.h
@@ -4,6 +4,7 @@
 #if UE_ENABLE_ICU
 #include "Internationalization/Text.h"
 #include <unicode/locid.h>
+#include <unicode/brkiter.h>
 #include <unicode/coll.h>
 #include <unicode/numfmt.h>
 #include <unicode/decimfmt.h>
@@ -99,10 +100,20 @@ inline icu::DecimalFormat::ERoundingMode UEToICU(const ERoundingMode RoundingMod
 	return Value;
 }
 
+enum class EBreakIteratorType
+{
+	Grapheme,
+	Word,
+	Line,
+	Sentence,
+	Title
+};
+
 class FCulture::FICUCultureImplementation
 {
 	friend FCulture;
 	friend FText;
+	friend class FICUBreakIteratorManager;
 
 	FICUCultureImplementation(const FString& LocaleName);
 
@@ -136,21 +147,27 @@ class FCulture::FICUCultureImplementation
 
 	FString GetVariant() const;
 
-	TSharedRef<const icu::Collator, ESPMode::ThreadSafe> GetCollator(const ETextComparisonLevel::Type ComparisonLevel) const;
-	TSharedRef<const icu::DecimalFormat> GetDecimalFormatter(const FNumberFormattingOptions* const Options = NULL) const;
-	TSharedRef<const icu::DecimalFormat> GetCurrencyFormatter(const FString& CurrencyCode = FString(), const FNumberFormattingOptions* const Options = NULL) const;
-	TSharedRef<const icu::DecimalFormat> GetPercentFormatter(const FNumberFormattingOptions* const Options = NULL) const;
-	TSharedRef<const icu::DateFormat> GetDateFormatter(const EDateTimeStyle::Type DateStyle, const FString& TimeZone) const;
-	TSharedRef<const icu::DateFormat> GetTimeFormatter(const EDateTimeStyle::Type TimeStyle, const FString& TimeZone) const;
-	TSharedRef<const icu::DateFormat> GetDateTimeFormatter(const EDateTimeStyle::Type DateStyle, const EDateTimeStyle::Type TimeStyle, const FString& TimeZone) const;
+	TSharedRef<const icu::BreakIterator> GetBreakIterator(const EBreakIteratorType Type);
+	TSharedRef<const icu::Collator, ESPMode::ThreadSafe> GetCollator(const ETextComparisonLevel::Type ComparisonLevel);
+	TSharedRef<const icu::DecimalFormat> GetDecimalFormatter(const FNumberFormattingOptions* const Options = NULL);
+	TSharedRef<const icu::DecimalFormat> GetCurrencyFormatter(const FString& CurrencyCode = FString(), const FNumberFormattingOptions* const Options = NULL);
+	TSharedRef<const icu::DecimalFormat> GetPercentFormatter(const FNumberFormattingOptions* const Options = NULL);
+	TSharedRef<const icu::DateFormat> GetDateFormatter(const EDateTimeStyle::Type DateStyle, const FString& TimeZone);
+	TSharedRef<const icu::DateFormat> GetTimeFormatter(const EDateTimeStyle::Type TimeStyle, const FString& TimeZone);
+	TSharedRef<const icu::DateFormat> GetDateTimeFormatter(const EDateTimeStyle::Type DateStyle, const EDateTimeStyle::Type TimeStyle, const FString& TimeZone);
 
 	icu::Locale ICULocale;
-	const TSharedRef<const icu::Collator, ESPMode::ThreadSafe> ICUCollator;
-	const TSharedRef<const icu::DecimalFormat> ICUDecimalFormat;
-	const TSharedRef<const icu::DecimalFormat> ICUCurrencyFormat;
-	const TSharedRef<const icu::DecimalFormat> ICUPercentFormat;
-	const TSharedRef<const icu::DateFormat> ICUDateFormat;
-	const TSharedRef<const icu::DateFormat> ICUTimeFormat;
-	const TSharedRef<const icu::DateFormat> ICUDateTimeFormat;
+	TSharedPtr<const icu::BreakIterator> ICUGraphemeBreakIterator;
+	TSharedPtr<const icu::BreakIterator> ICUWordBreakIterator;
+	TSharedPtr<const icu::BreakIterator> ICULineBreakIterator;
+	TSharedPtr<const icu::BreakIterator> ICUSentenceBreakIterator;
+	TSharedPtr<const icu::BreakIterator> ICUTitleBreakIterator;
+	TSharedPtr<const icu::Collator, ESPMode::ThreadSafe> ICUCollator;
+	TSharedPtr<const icu::DecimalFormat> ICUDecimalFormat;
+	TSharedPtr<const icu::DecimalFormat> ICUCurrencyFormat;
+	TSharedPtr<const icu::DecimalFormat> ICUPercentFormat;
+	TSharedPtr<const icu::DateFormat> ICUDateFormat;
+	TSharedPtr<const icu::DateFormat> ICUTimeFormat;
+	TSharedPtr<const icu::DateFormat> ICUDateTimeFormat;
 };
 #endif
\ No newline at end of file
diff --git a/Engine/Source/Runtime/Core/Private/Internationalization/ICUInternationalization.cpp b/Engine/Source/Runtime/Core/Private/Internationalization/ICUInternationalization.cpp
--- a/Engine/Source/Runtime/Core/Private/Internationalization/ICUInternationalization.cpp
+++ b/Engine/Source/Runtime/Core/Private/Internationalization/ICUInternationalization.cpp
@@ -1,4 +1,4 @@
-// Copyright 1998-2015 Epic Games, Inc. All Rights Reserved.
+﻿// Copyright 1998-2015 Epic Games, Inc. All Rights Reserved.
 #include "CorePrivatePCH.h"
 #include "ICUInternationalization.h"
 
@@ -284,7 +284,7 @@ FCulturePtr FICUInternationalization::FindOrMakeCulture(const FString& Name, con
 		// Confirm if data for the desired culture exists.
 		if (UResourceBundle* ICUResourceBundle = ures_open(nullptr, StringCast<char>(*CanonicalName).Get(), &ICUStatus))
 		{
-			// Make the culture only if it actually has some form of data and doesn't fallback to default "root" data, unless overriden ot allow it.
+			// Make the culture only if it actually has some form of data and doesn't fallback to default "root" data, unless overidden to allow it.
 			if (ICUStatus != U_USING_DEFAULT_WARNING || AllowDefaultFallback)
 			{
 				FCulturePtr NewCulture = FCulture::Create(CanonicalName);
diff --git a/Engine/Source/Runtime/Engine/Private/UserInterface/Canvas.cpp b/Engine/Source/Runtime/Engine/Private/UserInterface/Canvas.cpp
--- a/Engine/Source/Runtime/Engine/Private/UserInterface/Canvas.cpp
+++ b/Engine/Source/Runtime/Engine/Private/UserInterface/Canvas.cpp
@@ -23,6 +23,7 @@ DEFINE_STAT(STAT_Canvas_FlushTime);
 DEFINE_STAT(STAT_Canvas_DrawTextureTileTime);
 DEFINE_STAT(STAT_Canvas_DrawMaterialTileTime);
 DEFINE_STAT(STAT_Canvas_DrawStringTime);
+DEFINE_STAT(STAT_Canvas_WordWrappingTime);
 DEFINE_STAT(STAT_Canvas_GetBatchElementsTime);
 DEFINE_STAT(STAT_Canvas_AddTileRenderTime);
 DEFINE_STAT(STAT_Canvas_NumBatchesCreated);
@@ -35,21 +36,25 @@ FCanvasWordWrapper::FCanvasWordWrapper()
 
 void FCanvasWordWrapper::Execute(const TCHAR* const InString, const FTextSizingParameters& InParameters, TArray<FWrappedStringElement>& OutStrings, FWrappedLineData* const OutWrappedLineData)
 {
-	FWrappingState WrappingState(InString, FCString::Strlen(InString), InParameters, OutStrings, OutWrappedLineData);
+	SCOPE_CYCLE_COUNTER(STAT_Canvas_WordWrappingTime);
 
+	FWrappingState WrappingState(InString, FCString::Strlen(InString), InParameters, OutStrings, OutWrappedLineData);
 	if (WrappingState.WrappedLineData)
 	{
 		WrappingState.WrappedLineData->Empty();
 	}
 
-	GraphemeBreakIterator->SetString(WrappingState.String, WrappingState.StringLength);
-	LineBreakIterator->SetString(WrappingState.String, WrappingState.StringLength);
+	if (WrappingState.StringLength > 0)
+	{
+		GraphemeBreakIterator->SetString(WrappingState.String, WrappingState.StringLength);
+		LineBreakIterator->SetString(WrappingState.String, WrappingState.StringLength);
 
-	for(int32 i = 0; i < WrappingState.StringLength; ++i) // Sanity check: Doesn't seem valid to have more lines than code units.
-	{	
-		if( !ProcessLine(WrappingState) )
-		{
-			break;
+		for(int32 i = 0; i < WrappingState.StringLength; ++i) // Sanity check: Doesn't seem valid to have more lines than code units.
+		{	
+			if( !ProcessLine(WrappingState) )
+			{
+				break;
+			}
 		}
 	}
 }
diff --git a/Engine/Source/Runtime/Engine/Public/EngineStats.h b/Engine/Source/Runtime/Engine/Public/EngineStats.h
--- a/Engine/Source/Runtime/Engine/Public/EngineStats.h
+++ b/Engine/Source/Runtime/Engine/Public/EngineStats.h
@@ -163,6 +163,7 @@ DECLARE_CYCLE_STAT_EXTERN(TEXT("Flush Time"),STAT_Canvas_FlushTime,STATGROUP_Can
 DECLARE_CYCLE_STAT_EXTERN(TEXT("Draw Texture Tile Time"),STAT_Canvas_DrawTextureTileTime,STATGROUP_Canvas, );
 DECLARE_CYCLE_STAT_EXTERN(TEXT("Draw Material Tile Time"),STAT_Canvas_DrawMaterialTileTime,STATGROUP_Canvas, );
 DECLARE_CYCLE_STAT_EXTERN(TEXT("Draw String Time"),STAT_Canvas_DrawStringTime,STATGROUP_Canvas, );
+DECLARE_CYCLE_STAT_EXTERN(TEXT("Word Wrapping Time"),STAT_Canvas_WordWrappingTime,STATGROUP_Canvas, );
 DECLARE_CYCLE_STAT_EXTERN(TEXT("Get Batched Element Time"),STAT_Canvas_GetBatchElementsTime,STATGROUP_Canvas, );
 DECLARE_CYCLE_STAT_EXTERN(TEXT("Add Material Tile Time"),STAT_Canvas_AddTileRenderTime,STATGROUP_Canvas, );
 DECLARE_DWORD_COUNTER_STAT_EXTERN(TEXT("Num Batches Created"),STAT_Canvas_NumBatchesCreated,STATGROUP_Canvas, );	
