Workaround for window being maximized after fullscreen switch.

- Fixes UE-8735 (ShooterGame stuck at desktop resolution) for Compiz.

diff --git a/Engine/Source/Runtime/Core/Private/Linux/LinuxWindow.cpp b/Engine/Source/Runtime/Core/Private/Linux/LinuxWindow.cpp
--- a/Engine/Source/Runtime/Core/Private/Linux/LinuxWindow.cpp
+++ b/Engine/Source/Runtime/Core/Private/Linux/LinuxWindow.cpp
@@ -346,9 +346,15 @@ void FLinuxWindow::SetWindowMode( EWindowMode::Type NewWindowMode )
 
 			case EWindowMode::Windowed:
 			{
+				// when going back to windowed from desktop, make window smaller (but not too small),
+				// since some too smart window managers (Compiz) will maximize the window if it's set to desktop size.
+				// @FIXME: [RCL] 2015-02-10: this is a hack.
+				int SmallerWidth = FMath::Max(100, VirtualWidth - 100);
+				int SmallerHeight = FMath::Max(100, VirtualHeight - 100);
+				SDL_SetWindowSize(HWnd, SmallerWidth, SmallerHeight);
+
 				SDL_SetWindowFullscreen( HWnd, 0 );
 				SDL_SetWindowBordered( HWnd, SDL_TRUE );
-				SDL_SetWindowSize( HWnd, VirtualWidth, VirtualHeight );
 
 				SDL_SetWindowGrab( HWnd, SDL_FALSE );
 
