Fixed non-Unity build issue in ICUBreakIterator.cpp.

diff --git a/Engine/Source/Runtime/Core/Private/Internationalization/ICUBreakIterator.cpp b/Engine/Source/Runtime/Core/Private/Internationalization/ICUBreakIterator.cpp
--- a/Engine/Source/Runtime/Core/Private/Internationalization/ICUBreakIterator.cpp
+++ b/Engine/Source/Runtime/Core/Private/Internationalization/ICUBreakIterator.cpp
@@ -5,6 +5,7 @@
 
 #if UE_ENABLE_ICU
 #include "ICUTextCharacterIterator.h"
+#include "ICUCulture.h"
 
 FICUBreakIteratorManager* FICUBreakIteratorManager::Singleton = nullptr;
 
