TInterval replaced IsEmpty with IsValid

diff --git a/Engine/Source/Runtime/Core/Public/Math/Interval.h b/Engine/Source/Runtime/Core/Public/Math/Interval.h
--- a/Engine/Source/Runtime/Core/Public/Math/Interval.h
+++ b/Engine/Source/Runtime/Core/Public/Math/Interval.h
@@ -21,7 +21,7 @@ public:
 	/**
 	 * Default constructor.
 	 *
-	 * The interval is empty
+	 * The interval is invalid
 	 */
 	TInterval()
 		: Min(TNumericLimits<ElementType>::Max())
@@ -48,7 +48,7 @@ public:
 	 */
 	void operator+= ( ElementType X )
 	{
-		if (!IsEmpty())
+		if (IsValid())
 		{
 			Min += X;
 			Max += X;
@@ -62,7 +62,7 @@ public:
 	 */
 	void operator-= ( ElementType X )
 	{
-		if (!IsEmpty())
+		if (IsValid())
 		{
 			Min -= X;
 			Max -= X;
@@ -82,13 +82,13 @@ public:
 	}
 
 	/**
-	 * Whether interval is empty.
+	 * Whether interval is valid (Min <= Max).
 	 *
-	 * @return false when interval is empty, true otherwise
+	 * @return false when interval is invalid, true otherwise
 	 */
-	ElementType IsEmpty() const
+	ElementType IsValid() const
 	{
-		return (Min >= Max);
+		return (Min <= Max);
 	}
 	
 	/**
@@ -99,7 +99,7 @@ public:
 	 */
 	bool Contains( const ElementType& Element ) const
 	{
-		return !IsEmpty() && (Element >= Min && Element <= Max);
+		return IsValid() && (Element >= Min && Element <= Max);
 	}
 
 	/**
@@ -109,7 +109,7 @@ public:
 	 */
 	void Expand( ElementType ExpandAmount )
 	{
-		if (!IsEmpty())
+		if (IsValid())
 		{
 			Min -= ExpandAmount;
 			Max += ExpandAmount;
@@ -123,7 +123,7 @@ public:
 	 */
 	void Include( ElementType X )
 	{
-		if (IsEmpty())
+		if (!IsValid())
 		{
 			Min = X;
 			Max = X;
@@ -150,7 +150,7 @@ public:
 	 */
 	ElementType Interpolate( float Alpha ) const
 	{
-		if (!IsEmpty())
+		if (IsValid())
 		{
 			return Min + ElementType(Alpha*Size());
 		}
@@ -169,12 +169,12 @@ public:
 	 */
 	friend TInterval Intersect( const TInterval& A, const TInterval& B )
 	{
-		if (A.IsEmpty() || B.IsEmpty())
+		if (A.IsValid() && B.IsValid())
 		{
-			return TInterval();
+			return TInterval(FMath::Max(A.Min, B.Min), FMath::Min(A.Max, B.Max));
 		}
 
-		return TInterval(FMath::Max(A.Min, B.Min), FMath::Min(A.Max, B.Max));
+		return TInterval();
 	}
 
 	/**
