fix for appl crash when locking the screen on iOS 8
UE-6738
#ios

diff --git a/Engine/Source/Runtime/Core/Private/IOS/IOSAppDelegate.cpp b/Engine/Source/Runtime/Core/Private/IOS/IOSAppDelegate.cpp
--- a/Engine/Source/Runtime/Core/Private/IOS/IOSAppDelegate.cpp
+++ b/Engine/Source/Runtime/Core/Private/IOS/IOSAppDelegate.cpp
@@ -349,7 +349,7 @@ void InstallSignalHandlers()
 	if (bSuspend)
 	{
 		FAppEntry::Suspend();
-	}
+	}wi
 	else
 	{
 		FAppEntry::Resume();
@@ -540,8 +540,8 @@ void InstallSignalHandlers()
 	 */
 
 	FCoreDelegates::ApplicationWillDeactivateDelegate.Broadcast();
+	[self ToggleSuspend:true];
 	[self ToggleAudioSession:false];
-    [self ToggleSuspend:true];
 }
 
 - (void)applicationDidEnterBackground:(UIApplication *)application
