CIS fix
#ios

diff --git a/Engine/Source/Runtime/Core/Private/IOS/IOSAppDelegate.cpp b/Engine/Source/Runtime/Core/Private/IOS/IOSAppDelegate.cpp
--- a/Engine/Source/Runtime/Core/Private/IOS/IOSAppDelegate.cpp
+++ b/Engine/Source/Runtime/Core/Private/IOS/IOSAppDelegate.cpp
@@ -349,7 +349,7 @@ void InstallSignalHandlers()
 	if (bSuspend)
 	{
 		FAppEntry::Suspend();
-	}wi
+	}
 	else
 	{
 		FAppEntry::Resume();
