Increased the default thread stack size on Mac to 2 MB

diff --git a/Engine/Source/Runtime/Core/Public/Mac/MacPlatformRunnableThread.h b/Engine/Source/Runtime/Core/Public/Mac/MacPlatformRunnableThread.h
--- a/Engine/Source/Runtime/Core/Public/Mac/MacPlatformRunnableThread.h
+++ b/Engine/Source/Runtime/Core/Public/Mac/MacPlatformRunnableThread.h
@@ -30,7 +30,7 @@ private:
 	virtual int GetDefaultStackSize() override
 	{
 		// default is 512 KB, we need more
-		return 1024 * 1024;
+		return 2 * 1024 * 1024;
 	}
 
 	/**
