- Fixed a DEPRECATED usage clang compiler issue that doesn't break VC apparently

diff --git a/Engine/Source/Runtime/Core/Public/Math/Interval.h b/Engine/Source/Runtime/Core/Public/Math/Interval.h
--- a/Engine/Source/Runtime/Core/Public/Math/Interval.h
+++ b/Engine/Source/Runtime/Core/Public/Math/Interval.h
@@ -244,8 +244,7 @@ public:
 		enum { Value = true }; \
 	};
 
-DEPRECATED(4.8, "FInterval is deprecated, please use FFloatInterval instead.")
-DEFINE_INTERVAL_WRAPPER_STRUCT(FInterval, float)
-
 DEFINE_INTERVAL_WRAPPER_STRUCT(FFloatInterval, float)
 DEFINE_INTERVAL_WRAPPER_STRUCT(FInt32Interval, int32)
+
+struct DEPRECATED(4.8, "FInterval is deprecated, please use FFloatInterval instead.") FInterval : public FFloatInterval { } ;
