Stats - Stats for async task running on the thread pool

diff --git a/Engine/Source/Runtime/Core/Public/Async/AsyncWork.h b/Engine/Source/Runtime/Core/Public/Async/AsyncWork.h
--- a/Engine/Source/Runtime/Core/Public/Async/AsyncWork.h
+++ b/Engine/Source/Runtime/Core/Public/Async/AsyncWork.h
@@ -6,6 +6,8 @@
 
 #pragma once
 
+DECLARE_STATS_GROUP(TEXT("ThreadPool Async Tasks"), STATGROUP_ThreadPoolAsyncTasks, STATCAT_Advanced);
+
 /**
 	FAutoDeleteAsyncTask - template task for jobs that delete themselves when complete
 
@@ -27,10 +29,10 @@
 			... do the work here
 		}
 
-		static const TCHAR *Name()
-		{
-			return TEXT("ExampleAutoDeleteAsyncTask");
-		}
+ 		FORCEINLINE TStatId GetStatId() const
+ 		{
+			RETURN_QUICK_DECLARE_CYCLE_STAT(ExampleAutoDeleteAsyncTask, STATGROUP_ThreadPoolAsyncTasks);
+ 		}
 	};
 
 	start an example job
@@ -76,6 +78,8 @@ class FAutoDeleteAsyncTask
 	**/
 	void DoWork()
 	{		
+		FScopeCycleCounter Scope(Task.GetStatId(), true); 
+
 		Task.DoWork();		
 		delete this;
 	}
@@ -206,9 +210,9 @@ public:
 			... do the work here
 		}
 
-		static const TCHAR *Name()
+		FORCEINLINE TStatId GetStatId() const
 		{
-			return TEXT("ExampleAsyncTask");
+			RETURN_QUICK_DECLARE_CYCLE_STAT(ExampleAsyncTask, STATGROUP_ThreadPoolAsyncTasks);
 		}
 	};
 
@@ -293,7 +297,9 @@ class FAsyncTask
 	* Tells the user job to do the work, sometimes called synchronously, sometimes from the thread pool. Calls the event tracker.
 	**/
 	void DoWork()
-	{		
+	{	
+		FScopeCycleCounter Scope(Task.GetStatId(), true); 
+
 		Task.DoWork();		
 		check(WorkNotFinishedCounter.GetValue() == 1);
 		WorkNotFinishedCounter.Decrement();
@@ -619,12 +625,10 @@ public:
 		// Uncompress from memory to memory.
 		verify( FCompression::UncompressMemory( Flags, UncompressedBuffer, UncompressedSize, CompressedBuffer, CompressedSize, bIsSourceMemoryPadded ) );
 	}
-	/** Give the name for external event viewers
-	* @return	the name to display in external event viewers
-	*/
-	static const TCHAR *Name()
+
+	FORCEINLINE TStatId GetStatId() const
 	{
-		return TEXT("FAsyncUncompress");
+		RETURN_QUICK_DECLARE_CYCLE_STAT(FAsyncUncompress, STATGROUP_ThreadPoolAsyncTasks);
 	}
 };
 
