Core - Added native platform GetCommandLine (implemented only for Windows)

diff --git a/Engine/Source/Runtime/Core/Private/Windows/WindowsPlatformMisc.cpp b/Engine/Source/Runtime/Core/Private/Windows/WindowsPlatformMisc.cpp
--- a/Engine/Source/Runtime/Core/Private/Windows/WindowsPlatformMisc.cpp
+++ b/Engine/Source/Runtime/Core/Private/Windows/WindowsPlatformMisc.cpp
@@ -1874,6 +1874,11 @@ uint32 FWindowsPlatformMisc::GetLastError()
 	return (uint32)::GetLastError();
 }
 
+TCHAR* FWindowsPlatformMisc::GetCommandLine()
+{
+	return ::GetCommandLine();
+}
+
 bool FWindowsPlatformMisc::CoInitialize()
 {
 	HRESULT hr = ::CoInitialize(NULL);
diff --git a/Engine/Source/Runtime/Core/Public/GenericPlatform/GenericPlatformMisc.h b/Engine/Source/Runtime/Core/Public/GenericPlatform/GenericPlatformMisc.h
--- a/Engine/Source/Runtime/Core/Public/GenericPlatform/GenericPlatformMisc.h
+++ b/Engine/Source/Runtime/Core/Public/GenericPlatform/GenericPlatformMisc.h
@@ -438,15 +438,11 @@ public:
 	}
 
 	/**
-	 * Rebuild the commandline if needed
-	 *
-	 * @param NewCommandLine The commandline to fill out
-	 *
-	 * @return true if NewCommandLine should be pushed to FCommandLine
+	 * @return the executable's command line
 	 */
-	FORCEINLINE static bool ResetCommandLine(TCHAR NewCommandLine[16384])
+	FORCEINLINE static TCHAR* GetCommandLine()
 	{
-		return 0;
+		return nullptr;
 	}
 
 	/** 
diff --git a/Engine/Source/Runtime/Core/Public/Windows/WindowsPlatformMisc.h b/Engine/Source/Runtime/Core/Public/Windows/WindowsPlatformMisc.h
--- a/Engine/Source/Runtime/Core/Public/Windows/WindowsPlatformMisc.h
+++ b/Engine/Source/Runtime/Core/Public/Windows/WindowsPlatformMisc.h
@@ -82,9 +82,8 @@ struct CORE_API FWindowsPlatformMisc
 	static void LoadStartupModules();
 
 	static uint32 GetLastError();
-
+	static TCHAR* GetCommandLine();
 	static void RaiseException( uint32 ExceptionCode );
-
 	static bool SetStoredValue(const FString& InStoreId, const FString& InSectionName, const FString& InKeyName, const FString& InValue);
 	static bool GetStoredValue(const FString& InStoreId, const FString& InSectionName, const FString& InKeyName, FString& OutValue);
 
