Stats - Stats for async task running on the thread pool (CIS fix)

diff --git a/Engine/Source/Runtime/Core/Private/Stats/StatsFile.cpp b/Engine/Source/Runtime/Core/Private/Stats/StatsFile.cpp
--- a/Engine/Source/Runtime/Core/Private/Stats/StatsFile.cpp
+++ b/Engine/Source/Runtime/Core/Private/Stats/StatsFile.cpp
@@ -64,12 +64,9 @@ public:
 		Outer->FramesInfo.Add( FStatsFrameInfo( FrameFileOffset, ThreadCycles ) );
 	}
 
-	/**
-	* @return	the name to display in external event viewers
-	*/
-	static const TCHAR *Name()
+	FORCEINLINE TStatId GetStatId() const
 	{
-		return TEXT( "FAsyncStatsWriteWorker" );
+		RETURN_QUICK_DECLARE_CYCLE_STAT( FAsyncStatsWrite, STATGROUP_ThreadPoolAsyncTasks );
 	}
 };
 
