[INTEGRATE] Change 2435941 by [REDACTED_PERSON][REDACTED_HANDLE] on 2015/02/06 15:40:38

	UnrealHeaderTool: Improve default-value parsing for things like FText::FromString, and it now treats C++ types as case-sensitive

diff --git a/Engine/Source/Programs/UnrealHeaderTool/Private/HeaderParser.cpp b/Engine/Source/Programs/UnrealHeaderTool/Private/HeaderParser.cpp
--- a/Engine/Source/Programs/UnrealHeaderTool/Private/HeaderParser.cpp
+++ b/Engine/Source/Programs/UnrealHeaderTool/Private/HeaderParser.cpp
@@ -7194,10 +7194,14 @@ bool FHeaderParser::DefaultValueStringCppFormatToInnerFormat(const UProperty* Pr
 {
 	OutForm = FString();
 	if (!Property || CppForm.IsEmpty())
+	{
 		return false;
+	}
 
 	if (Property->IsA(UClassProperty::StaticClass()) || Property->IsA(UObjectPropertyBase::StaticClass()))
+	{
 		return FDefaultValueHelper::Is(CppForm, TEXT("NULL")) || FDefaultValueHelper::Is(CppForm, TEXT("nullptr")) || FDefaultValueHelper::Is(CppForm, TEXT("0"));
+	}
 
 	if( !Property->IsA(UStructProperty::StaticClass()) )
 	{
diff --git a/Engine/Source/Runtime/Core/Private/Misc/DefaultValueHelper.cpp b/Engine/Source/Runtime/Core/Private/Misc/DefaultValueHelper.cpp
--- a/Engine/Source/Runtime/Core/Private/Misc/DefaultValueHelper.cpp
+++ b/Engine/Source/Runtime/Core/Private/Misc/DefaultValueHelper.cpp
@@ -14,7 +14,7 @@ bool FDefaultValueHelper::Is(const FString& Source, const TCHAR* CompareStr)
 		return false;
 	}
 
-	if( Source.Find(CompareStr) != Pos )
+	if( Source.Find(CompareStr, ESearchCase::CaseSensitive) != Pos )
 	{
 		return false;
 	}
@@ -95,7 +95,7 @@ bool FDefaultValueHelper::Trim(const TCHAR* & Start, const TCHAR* End )
 
 FString FDefaultValueHelper::GetUnqualifiedEnumValue(const FString& Source)
 {
-	auto SeparatorPosition = Source.Find(FString("::"));
+	auto SeparatorPosition = Source.Find(FString("::"), ESearchCase::CaseSensitive);
 	if (SeparatorPosition == INDEX_NONE)
 	{
 		return Source;
@@ -118,7 +118,7 @@ bool FDefaultValueHelper::HasWhitespaces(const FString& Source)
 }
 
 
-bool FDefaultValueHelper::GetParameters(const FString& Source, const FString& TypeName, FString &OutForm)
+bool FDefaultValueHelper::GetParameters(const FString& Source, const FString& TypeName, FString& OutForm)
 {
 	int32 Pos = 0;
 
@@ -128,7 +128,7 @@ bool FDefaultValueHelper::GetParameters(const FString& Source, const FString& Ty
 	}
 
 	// find the beginning of actual val after "TypeName ( "
-	if( Source.Find(TypeName) != Pos )
+	if( Source.Find(TypeName, ESearchCase::CaseSensitive) != Pos )
 	{
 		return false;
 	}
@@ -373,7 +373,7 @@ bool FDefaultValueHelper::IsStringValidLinearColor(const FString& Source)
 }
 
 
-bool FDefaultValueHelper::StringFromCppString(const FString& Source, const FString& TypeName, FString &OutForm)
+bool FDefaultValueHelper::StringFromCppString(const FString& Source, const FString& TypeName, FString& OutForm)
 {
 	int32 Pos = 0, PendingParentheses = 0;
 
@@ -383,13 +383,39 @@ bool FDefaultValueHelper::StringFromCppString(const FString& Source, const FStri
 	}
 
 	// remove "TypeName ( "
-	if( Source.Find(TypeName) == Pos )
+	if (Source.Find(TypeName, ESearchCase::CaseSensitive) == Pos)
 	{
 		Pos += TypeName.Len();
 		if( !Trim(Pos, Source) )
 		{
 			return false;
 		}
+
+		if (Source.Find(FString("::"), ESearchCase::CaseSensitive) == Pos)
+		{
+			Pos += 2;
+
+			if (!Trim(Pos, Source))
+			{
+				return false;
+			}
+
+			const FString AllowedFunctionName(TEXT("FromString"));
+			if (Source.Find(AllowedFunctionName, ESearchCase::CaseSensitive, ESearchDir::FromStart, Pos) == Pos)
+			{
+				Pos += AllowedFunctionName.Len();
+			}
+			else
+			{
+				return false;
+			}
+
+			if (!Trim(Pos, Source))
+			{
+				return false;
+			}
+		}
+
 		if( TS(TEXT("(")) != Source[Pos++] )
 		{
 			return false;
@@ -403,7 +429,7 @@ bool FDefaultValueHelper::StringFromCppString(const FString& Source, const FStri
 
 	// remove "TEXT ( "
 	const TCHAR* TextStr = TEXT("TEXT");
-	if( Source.Find(TextStr) == Pos )
+	if (Source.Find(TextStr, ESearchCase::CaseSensitive) == Pos)
 	{
 		Pos += FCString::Strlen(TextStr);
 		if( !Trim(Pos, Source) )
