Added missing destructor

diff --git a/Engine/Source/Runtime/Core/Public/Stats/StatsFile.h b/Engine/Source/Runtime/Core/Public/Stats/StatsFile.h
--- a/Engine/Source/Runtime/Core/Public/Stats/StatsFile.h
+++ b/Engine/Source/Runtime/Core/Public/Stats/StatsFile.h
@@ -351,6 +351,12 @@ protected:
 	/** Default constructor. */
 	IStatsWriteFile();
 
+public:
+	/** Destructor. */
+	virtual ~IStatsWriteFile()
+	{}
+
+protected:
 	/** NewFrame delegate handle  */
 	FDelegateHandle DataDelegateHandle;
 
