Order initialization fix

diff --git a/Engine/Source/Runtime/Core/Private/Stats/StatsFile.cpp b/Engine/Source/Runtime/Core/Private/Stats/StatsFile.cpp
--- a/Engine/Source/Runtime/Core/Private/Stats/StatsFile.cpp
+++ b/Engine/Source/Runtime/Core/Private/Stats/StatsFile.cpp
@@ -75,12 +75,12 @@ FStatsThreadState::FStatsThreadState( FString const& Filename )
 	, MinFrameSeen( -1 )
 	, LastFullFrameMetaAndNonFrame( -1 )
 	, LastFullFrameProcessed( -1 )
-	, CurrentGameFrame( -1 )
-	, CurrentRenderFrame( -1 )
-	, MaxNumStatMessages(0)
-	, TotalNumStatMessages(0)
+	, MaxNumStatMessages( 0 )
+	, TotalNumStatMessages( 0 )
 	, bWasLoaded( true )
 	, bFindMemoryExtensiveStats( false )
+	, CurrentGameFrame( -1 )
+	, CurrentRenderFrame( -1 )	
 {
 	const int64 Size = IFileManager::Get().FileSize( *Filename );
 	if( Size < 4 )
