Stats - Fixed a crash when checking the commandline and name conflict with GetCommandLineW

diff --git a/Engine/Source/Runtime/Core/Private/GenericPlatform/GenericPlatformMisc.cpp b/Engine/Source/Runtime/Core/Private/GenericPlatform/GenericPlatformMisc.cpp
--- a/Engine/Source/Runtime/Core/Private/GenericPlatform/GenericPlatformMisc.cpp
+++ b/Engine/Source/Runtime/Core/Private/GenericPlatform/GenericPlatformMisc.cpp
@@ -871,3 +871,8 @@ FString FGenericPlatformMisc::GetOperatingSystemId()
 	// not implemented by default.
 	return FString();
 }
+
+TCHAR* FGenericPlatformMisc::GetNativeCommandLine()
+{
+	return TEXT("");
+}
diff --git a/Engine/Source/Runtime/Core/Private/HAL/UnrealMemory.cpp b/Engine/Source/Runtime/Core/Private/HAL/UnrealMemory.cpp
--- a/Engine/Source/Runtime/Core/Private/HAL/UnrealMemory.cpp
+++ b/Engine/Source/Runtime/Core/Private/HAL/UnrealMemory.cpp
@@ -43,9 +43,9 @@ void GCreateMalloc()
 
 	
 #if	STATS
-	const bool bCreateMallocProfiler = FParse::Param( FPlatformMisc::GetCommandLine(), TEXT( "MemoryProfiler" ) );
+	const bool bCreateMallocProfiler = FParse::Param( FPlatformMisc::GetNativeCommandLine(), TEXT( "MemoryProfiler" ) );
 
-	// Create the stats malloc profiler proxy, needs to be enabled separately.
+	// Create the stats malloc profiler proxy.
 	if( bCreateMallocProfiler )
 	{
 		GMalloc = FStatsMallocProfilerProxy::Get();
diff --git a/Engine/Source/Runtime/Core/Private/Windows/WindowsPlatformMisc.cpp b/Engine/Source/Runtime/Core/Private/Windows/WindowsPlatformMisc.cpp
--- a/Engine/Source/Runtime/Core/Private/Windows/WindowsPlatformMisc.cpp
+++ b/Engine/Source/Runtime/Core/Private/Windows/WindowsPlatformMisc.cpp
@@ -1874,7 +1874,7 @@ uint32 FWindowsPlatformMisc::GetLastError()
 	return (uint32)::GetLastError();
 }
 
-TCHAR* FWindowsPlatformMisc::GetCommandLine()
+TCHAR* FWindowsPlatformMisc::GetNativeCommandLine()
 {
 	return ::GetCommandLine();
 }
diff --git a/Engine/Source/Runtime/Core/Public/GenericPlatform/GenericPlatformMisc.h b/Engine/Source/Runtime/Core/Public/GenericPlatform/GenericPlatformMisc.h
--- a/Engine/Source/Runtime/Core/Public/GenericPlatform/GenericPlatformMisc.h
+++ b/Engine/Source/Runtime/Core/Public/GenericPlatform/GenericPlatformMisc.h
@@ -440,10 +440,7 @@ public:
 	/**
 	 * @return the executable's command line
 	 */
-	FORCEINLINE static TCHAR* GetCommandLine()
-	{
-		return nullptr;
-	}
+	static TCHAR* GetNativeCommandLine();
 
 	/** 
 	 *	Set the value for the given section and key in the platform specific key->value store
diff --git a/Engine/Source/Runtime/Core/Public/Windows/WindowsPlatformMisc.h b/Engine/Source/Runtime/Core/Public/Windows/WindowsPlatformMisc.h
--- a/Engine/Source/Runtime/Core/Public/Windows/WindowsPlatformMisc.h
+++ b/Engine/Source/Runtime/Core/Public/Windows/WindowsPlatformMisc.h
@@ -82,7 +82,7 @@ struct CORE_API FWindowsPlatformMisc
 	static void LoadStartupModules();
 
 	static uint32 GetLastError();
-	static TCHAR* GetCommandLine();
+	static TCHAR* GetNativeCommandLine();
 	static void RaiseException( uint32 ExceptionCode );
 	static bool SetStoredValue(const FString& InStoreId, const FString& InSectionName, const FString& InKeyName, const FString& InValue);
 	static bool GetStoredValue(const FString& InStoreId, const FString& InSectionName, const FString& InKeyName, FString& OutValue);
