CLang fixes, added const

diff --git a/Engine/Source/Runtime/Core/Private/GenericPlatform/GenericPlatformMisc.cpp b/Engine/Source/Runtime/Core/Private/GenericPlatform/GenericPlatformMisc.cpp
--- a/Engine/Source/Runtime/Core/Private/GenericPlatform/GenericPlatformMisc.cpp
+++ b/Engine/Source/Runtime/Core/Private/GenericPlatform/GenericPlatformMisc.cpp
@@ -872,7 +872,7 @@ FString FGenericPlatformMisc::GetOperatingSystemId()
 	return FString();
 }
 
-TCHAR* FGenericPlatformMisc::GetNativeCommandLine()
+const TCHAR* FGenericPlatformMisc::GetNativeCommandLine()
 {
 	return TEXT("");
 }
diff --git a/Engine/Source/Runtime/Core/Private/Windows/WindowsPlatformMisc.cpp b/Engine/Source/Runtime/Core/Private/Windows/WindowsPlatformMisc.cpp
--- a/Engine/Source/Runtime/Core/Private/Windows/WindowsPlatformMisc.cpp
+++ b/Engine/Source/Runtime/Core/Private/Windows/WindowsPlatformMisc.cpp
@@ -1874,7 +1874,7 @@ uint32 FWindowsPlatformMisc::GetLastError()
 	return (uint32)::GetLastError();
 }
 
-TCHAR* FWindowsPlatformMisc::GetNativeCommandLine()
+const TCHAR* FWindowsPlatformMisc::GetNativeCommandLine()
 {
 	return ::GetCommandLine();
 }
diff --git a/Engine/Source/Runtime/Core/Public/GenericPlatform/GenericPlatformMisc.h b/Engine/Source/Runtime/Core/Public/GenericPlatform/GenericPlatformMisc.h
--- a/Engine/Source/Runtime/Core/Public/GenericPlatform/GenericPlatformMisc.h
+++ b/Engine/Source/Runtime/Core/Public/GenericPlatform/GenericPlatformMisc.h
@@ -440,7 +440,7 @@ public:
 	/**
 	 * @return the executable's command line
 	 */
-	static TCHAR* GetNativeCommandLine();
+	static const TCHAR* GetNativeCommandLine();
 
 	/** 
 	 *	Set the value for the given section and key in the platform specific key->value store
diff --git a/Engine/Source/Runtime/Core/Public/Windows/WindowsPlatformMisc.h b/Engine/Source/Runtime/Core/Public/Windows/WindowsPlatformMisc.h
--- a/Engine/Source/Runtime/Core/Public/Windows/WindowsPlatformMisc.h
+++ b/Engine/Source/Runtime/Core/Public/Windows/WindowsPlatformMisc.h
@@ -82,7 +82,7 @@ struct CORE_API FWindowsPlatformMisc
 	static void LoadStartupModules();
 
 	static uint32 GetLastError();
-	static TCHAR* GetNativeCommandLine();
+	static const TCHAR* GetNativeCommandLine();
 	static void RaiseException( uint32 ExceptionCode );
 	static bool SetStoredValue(const FString& InStoreId, const FString& InSectionName, const FString& InKeyName, const FString& InValue);
 	static bool GetStoredValue(const FString& InStoreId, const FString& InSectionName, const FString& InKeyName, FString& OutValue);
