Adding EngineVersion to FArchive to be able to guard against mismatched versions with previous releases.

- Extracted basic engine version numbers from FEngineVersion into FEngineVersionBase to be able to use that with FArchiveBase
- FEngineVersion now derives from FEngineVersionBase
- Added code to propagate FEngineVersionBase from linkers to FArchiveBase

diff --git a/Engine/Source/Runtime/AssetRegistry/Private/PackageReader.cpp b/Engine/Source/Runtime/AssetRegistry/Private/PackageReader.cpp
--- a/Engine/Source/Runtime/AssetRegistry/Private/PackageReader.cpp
+++ b/Engine/Source/Runtime/AssetRegistry/Private/PackageReader.cpp
@@ -89,8 +89,10 @@ bool FPackageReader::OpenPackageFile(const FString& InPackageFilename)
 	//make sure the filereader gets the correct version number (it defaults to latest version)
 	SetUE4Ver(PackageFileSummary.GetFileVersionUE4());
 	SetLicenseeUE4Ver(PackageFileSummary.GetFileVersionLicenseeUE4());
+	SetEngineVer(PackageFileSummary.EngineVersion);
 	Loader->SetUE4Ver(PackageFileSummary.GetFileVersionUE4());
 	Loader->SetLicenseeUE4Ver(PackageFileSummary.GetFileVersionLicenseeUE4());
+	Loader->SetEngineVer(PackageFileSummary.EngineVersion);
 
 	const FCustomVersionContainer& PackageFileSummaryVersions = PackageFileSummary.GetCustomVersionContainer();
 	SetCustomVersions(PackageFileSummaryVersions);
diff --git a/Engine/Source/Runtime/Core/Private/Misc/EngineVersion.cpp b/Engine/Source/Runtime/Core/Private/Misc/EngineVersion.cpp
--- a/Engine/Source/Runtime/Core/Private/Misc/EngineVersion.cpp
+++ b/Engine/Source/Runtime/Core/Private/Misc/EngineVersion.cpp
@@ -5,6 +5,82 @@
 #include "Runtime/Launch/Resources/Version.h"
 #include "EngineBuildSettings.h"
 
+FEngineVersionBase::FEngineVersionBase()
+: Major(0)
+, Minor(0)
+, Patch(0)
+, Changelist(0)
+{
+}
+
+FEngineVersionBase::FEngineVersionBase(uint16 InMajor, uint16 InMinor, uint16 InPatch, uint32 InChangelist)
+: Major(InMajor)
+, Minor(InMinor)
+, Patch(InPatch)
+, Changelist(InChangelist)
+{
+}
+
+uint32 FEngineVersionBase::GetChangelist() const
+{
+	// Mask to ignore licensee bit
+	return Changelist & (uint32)0x7fffffffU;
+}
+
+bool FEngineVersionBase::IsLicenseeVersion() const
+{
+	// Check for licensee bit
+	return (Changelist & (uint32)0x80000000U) != 0;
+}
+
+bool FEngineVersionBase::IsEmpty() const
+{
+	return Major == 0 && Minor == 0 && Patch == 0;
+}
+
+bool FEngineVersionBase::IsPromotedBuild() const
+{
+	return GetChangelist() != 0;
+}
+
+EVersionComparison FEngineVersionBase::GetNewest(const FEngineVersionBase &First, const FEngineVersionBase &Second, EVersionComponent *OutComponent)
+{
+	EVersionComponent LocalComponent = EVersionComponent::Minor;
+	auto& Component = OutComponent ? *OutComponent : LocalComponent;
+
+	// Compare major versions
+	if (First.GetMajor() != Second.GetMajor())
+	{
+		Component = EVersionComponent::Major;
+		return (First.GetMajor() > Second.GetMajor()) ? EVersionComparison::First : EVersionComparison::Second;
+	}
+
+	// Compare minor versions
+	if (First.GetMinor() != Second.GetMinor())
+	{
+		Component = EVersionComponent::Minor;
+		return (First.GetMinor() > Second.GetMinor()) ? EVersionComparison::First : EVersionComparison::Second;
+	}
+
+	// Compare patch versions
+	if (First.GetPatch() != Second.GetPatch())
+	{
+		Component = EVersionComponent::Patch;
+		return (First.GetPatch() > Second.GetPatch()) ? EVersionComparison::First : EVersionComparison::Second;
+	}
+
+	// Compare changelists (only if they're both from the same vendor)
+	if (First.IsLicenseeVersion() == Second.IsLicenseeVersion() && First.GetChangelist() != Second.GetChangelist())
+	{
+		Component = EVersionComponent::Changelist;
+		return (First.GetChangelist() > Second.GetChangelist()) ? EVersionComparison::First : EVersionComparison::Second;
+	}
+
+	// Otherwise they're the same
+	return EVersionComparison::Neither;
+}
+
+
 FEngineVersion::FEngineVersion()
 {
 	Empty();
@@ -24,34 +100,12 @@
 	Branch = InBranch;
 }
 
-uint32 FEngineVersion::GetChangelist() const
-{
-	// Mask to ignore licensee bit
-	return Changelist & (uint32)0x7fffffffU;
-}
-
-bool FEngineVersion::IsLicenseeVersion() const
-{
-	// Check for licensee bit
-	return (Changelist & (uint32)0x80000000U) != 0;
-}
-
 void FEngineVersion::Empty()
 {
 	Set(0, 0, 0, 0, FString());
 }
 
-bool FEngineVersion::IsEmpty() const
+bool FEngineVersion::IsCompatibleWith(const FEngineVersionBase &Other) const
 {
-	return Major == 0 && Minor == 0 && Patch == 0;
-}
-
-bool FEngineVersion::IsPromotedBuild() const
-{
-	return GetChangelist() != 0;
-}
-
-bool FEngineVersion::IsCompatibleWith(const FEngineVersion &Other) const
-{
 	// If this or the other is not a promoted build, always assume compatibility. 
 	if(!IsPromotedBuild() || !Other.IsPromotedBuild())
@@ -139,43 +193,6 @@
 	return true;
 }
 
-EVersionComparison FEngineVersion::GetNewest(const FEngineVersion &First, const FEngineVersion &Second, EVersionComponent *OutComponent)
-{
-	EVersionComponent LocalComponent = EVersionComponent::Minor;
-	auto& Component = OutComponent ? *OutComponent : LocalComponent;
-
-	// Compare major versions
-	if(First.Major != Second.Major)
-	{
-		Component = EVersionComponent::Major;
-		return (First.Major > Second.Major)? EVersionComparison::First : EVersionComparison::Second;
-	}
-
-	// Compare minor versions
-	if(First.Minor != Second.Minor)
-	{
-		Component = EVersionComponent::Minor;
-		return (First.Minor > Second.Minor)? EVersionComparison::First : EVersionComparison::Second;
-	}
-
-	// Compare patch versions
-	if(First.Patch != Second.Patch)
-	{
-		Component = EVersionComponent::Patch;
-		return (First.Patch > Second.Patch)? EVersionComparison::First : EVersionComparison::Second;
-	}
-
-	// Compare changelists (only if they're both from the same vendor)
-	if(First.IsLicenseeVersion() == Second.IsLicenseeVersion() && First.GetChangelist() != Second.GetChangelist())
-	{
-		Component = EVersionComponent::Changelist;
-		return (First.GetChangelist() > Second.GetChangelist()) ? EVersionComparison::First : EVersionComparison::Second;
-	}
-
-	// Otherwise they're the same
-	return EVersionComparison::Neither;
-}
-
 void operator<<(FArchive &Ar, FEngineVersion &Version)
 {
 	Ar << Version.Major;
diff --git a/Engine/Source/Runtime/Core/Private/Serialization/Archive.cpp b/Engine/Source/Runtime/Core/Private/Serialization/Archive.cpp
--- a/Engine/Source/Runtime/Core/Private/Serialization/Archive.cpp
+++ b/Engine/Source/Runtime/Core/Private/Serialization/Archive.cpp
@@ -62,6 +62,7 @@ void FArchive::Reset()
 	ArNetVer							= GEngineNegotiationVersion;
 	ArUE4Ver							= GPackageFileUE4Version;
 	ArLicenseeUE4Ver					= GPackageFileLicenseeUE4Version;
+	ArEngineVer = GEngineVersion;
 	ArIsLoading							= false;
 	ArIsSaving							= false;
 	ArIsTransacting						= false;
@@ -100,6 +101,7 @@ void FArchive::CopyTrivialFArchiveStatusMembers(const FArchive& ArchiveToCopy)
 	ArNetVer                             = ArchiveToCopy.ArNetVer;
 	ArUE4Ver                             = ArchiveToCopy.ArUE4Ver;
 	ArLicenseeUE4Ver                     = ArchiveToCopy.ArLicenseeUE4Ver;
+	ArEngineVer                          = ArchiveToCopy.ArEngineVer;
 	ArIsLoading                          = ArchiveToCopy.ArIsLoading;
 	ArIsSaving                           = ArchiveToCopy.ArIsSaving;
 	ArIsTransacting                      = ArchiveToCopy.ArIsTransacting;
diff --git a/Engine/Source/Runtime/Core/Public/Misc/EngineVersion.h b/Engine/Source/Runtime/Core/Public/Misc/EngineVersion.h
--- a/Engine/Source/Runtime/Core/Public/Misc/EngineVersion.h
+++ b/Engine/Source/Runtime/Core/Public/Misc/EngineVersion.h
@@ -2,29 +2,11 @@
 
 #pragma once
 
-
-/** Enum for the components of a version string. */
-enum class EVersionComponent
-{
-	Major,					///< Major version increments introduce breaking API changes.
-	Minor,					///< Minor version increments add additional functionality without breaking existing APIs.
-	Patch,					///< Patch version increments fix existing functionality without changing the API.
-	Changelist,				///< The pre-release field adds additional versioning through a series of comparable dotted strings or numbers.
-	Branch,					///<
-};
-
-
-/** Components of a version string. */
-enum class EVersionComparison
-{
-	Neither,
-	First,
-	Second,
-};
+#include "EngineVersionBase.h"
 
 
 /** Utility functions. */
-class CORE_API FEngineVersion
+class CORE_API FEngineVersion : public FEngineVersionBase
 {
 public:
 
@@ -37,50 +19,23 @@ public:
 	/** Sets the version to the given values. */
 	void Set(uint16 InMajor, uint16 InMinor, uint16 InPatch, uint32 InChangelist, const FString &InBranch);
 
-	/** Returns the changelist number corresponding to this version. */
-	uint32 GetChangelist() const;	
-
-	/** Checks if the changelist number represents licensee changelist number. */
-	bool IsLicenseeVersion() const;
-
 	/** Clears the object. */
 	void Empty();
 
-	/** Returns whether the current version is empty. */
-	bool IsEmpty() const;
-
-	/** Returns whether the engine version is for a promoted build. */
-	bool IsPromotedBuild() const; 
-
 	/** Checks compatibility with another version object. */
-	bool IsCompatibleWith(const FEngineVersion &Other) const;
+	bool IsCompatibleWith(const FEngineVersionBase &Other) const;
 
 	/** Generates a version string */
 	FString ToString(EVersionComponent LastComponent = EVersionComponent::Branch) const;
 
 	/** Parses a version object from a string. Returns true on success. */
 	static bool Parse(const FString &Text, FEngineVersion &OutVersion);
-	
-	/** Returns the newest of two versions, and the component at which they differ */
-	static EVersionComparison GetNewest(const FEngineVersion &First, const FEngineVersion &Second, EVersionComponent *OutComponent);
 
 	/** Serialization function */
-	friend CORE_API void operator<<(FArchive &Ar, FEngineVersion &Version);
+	friend CORE_API void operator<<(class FArchive &Ar, FEngineVersion &Version);
 
 private:
 
-	/** Major version number. */
-	uint16 Major;
-
-	/** Minor version number. */
-	uint16 Minor;
-
-	/** Patch version number. */
-	uint16 Patch;
-
-	/** Changelist number. This is used to arbitrate when Major/Minor/Patch version numbers match. Use GetChangelist() instead of using this member directly. */
-	uint32 Changelist;
-
 	/** Branch name. */
 	FString Branch;
 };
diff --git a/Engine/Source/Runtime/Core/Public/Misc/EngineVersionBase.h b/Engine/Source/Runtime/Core/Public/Misc/EngineVersionBase.h
new file mode 100644
--- /dev/null
+++ b/Engine/Source/Runtime/Core/Public/Misc/EngineVersionBase.h
@@ -0,0 +1,82 @@
+// Copyright 1998-2015 Epic Games, Inc. All Rights Reserved.
+
+#pragma once
+
+/** Enum for the components of a version string. */
+enum class EVersionComponent
+{
+	Major,					///< Major version increments introduce breaking API changes.
+	Minor,					///< Minor version increments add additional functionality without breaking existing APIs.
+	Patch,					///< Patch version increments fix existing functionality without changing the API.
+	Changelist,				///< The pre-release field adds additional versioning through a series of comparable dotted strings or numbers.
+	Branch,					///<
+};
+
+
+/** Components of a version string. */
+enum class EVersionComparison
+{
+	Neither,
+	First,
+	Second,
+};
+
+
+/** Base class for the EngineVersion class. Holds basic version numbers. */
+class CORE_API FEngineVersionBase
+{
+public:
+
+	/** Empty constructor. Initializes the version to 0.0.0-0. */
+	FEngineVersionBase();
+
+	/** Constructs a version from the given components. */
+	FEngineVersionBase(uint16 InMajor, uint16 InMinor, uint16 InPatch = 0, uint32 InChangelist = 0);
+
+	/** Returns the changelist number corresponding to this version. */
+	uint32 GetChangelist() const;	
+
+	/** Returns the Major version number corresponding to this version. */
+	FORCEINLINE uint16 GetMajor() const
+	{
+		return Major;
+	}
+
+	/** Returns the Minor version number corresponding to this version. */
+	FORCEINLINE uint16 GetMinor() const
+	{
+		return Minor;
+	}
+
+	/** Returns the Patch version number corresponding to this version. */
+	FORCEINLINE uint16 GetPatch() const
+	{
+		return Patch;
+	}
+
+	/** Checks if the changelist number represents licensee changelist number. */
+	bool IsLicenseeVersion() const;
+
+	/** Returns whether the current version is empty. */
+	bool IsEmpty() const;
+
+	/** Returns whether the engine version is for a promoted build. */
+	bool IsPromotedBuild() const; 
+
+	/** Returns the newest of two versions, and the component at which they differ */
+	static EVersionComparison GetNewest(const FEngineVersionBase &First, const FEngineVersionBase &Second, EVersionComponent *OutComponent);
+
+protected:
+
+	/** Major version number. */
+	uint16 Major;
+
+	/** Minor version number. */
+	uint16 Minor;
+
+	/** Patch version number. */
+	uint16 Patch;
+
+	/** Changelist number. This is used to arbitrate when Major/Minor/Patch version numbers match. Use GetChangelist() instead of using this member directly. */
+	uint32 Changelist;
+};
diff --git a/Engine/Source/Runtime/Core/Public/Serialization/ArchiveBase.h b/Engine/Source/Runtime/Core/Public/Serialization/ArchiveBase.h
--- a/Engine/Source/Runtime/Core/Public/Serialization/ArchiveBase.h
+++ b/Engine/Source/Runtime/Core/Public/Serialization/ArchiveBase.h
@@ -5,7 +5,7 @@
 #include "Containers/EnumAsByte.h"
 #include "HAL/PlatformProperties.h"
 #include "Misc/Compression.h"
-
+#include "Misc/EngineVersionBase.h"
 
 class FAssetPtr;
 class FCustomVersionContainer;
@@ -547,6 +547,11 @@ public:
 		return ArLicenseeUE4Ver;
 	}
 
+	FORCEINLINE FEngineVersionBase EngineVer() const
+	{
+		return ArEngineVer;
+	}
+
 	/**
 	 * Registers the custom version to the archive.  This is used to inform the archive that custom version information is about to be stored.
 	 * There is no effect when the archive is being loaded from.
@@ -706,7 +711,7 @@ public:
 	}
 
 	/**
-	 * Sets the archive licensee version number. Used by the code that makes sure that ULinkerLoad's 
+	 * Sets the archive version number. Used by the code that makes sure that ULinkerLoad's 
 	 * internal archive versions match the file reader it creates.
 	 *
 	 * @param UE4Ver	new version number
@@ -727,6 +732,17 @@ public:
 		ArLicenseeUE4Ver = InVer;
 	}
 
+	/**
+	* Sets the archive engine version. Used by the code that makes sure that ULinkerLoad's
+	* internal archive versions match the file reader it creates.
+	*
+	* @param InVer	new version number
+	*/
+	void SetEngineVer(const FEngineVersionBase& InVer)
+	{
+		ArEngineVer = InVer;
+	}
+
 	/**
 	 * Gets the custom version numbers for this archive.
 	 *
@@ -919,6 +935,9 @@ protected:
 	/** Holds the archive version for licensees. */
 	int32 ArLicenseeUE4Ver;
 
+	/** Holds the engine version. */
+	FEngineVersionBase ArEngineVer;
+
 private:
 
 	/**
diff --git a/Engine/Source/Runtime/CoreUObject/Private/Misc/WorldCompositionUtility.cpp b/Engine/Source/Runtime/CoreUObject/Private/Misc/WorldCompositionUtility.cpp
--- a/Engine/Source/Runtime/CoreUObject/Private/Misc/WorldCompositionUtility.cpp
+++ b/Engine/Source/Runtime/CoreUObject/Private/Misc/WorldCompositionUtility.cpp
@@ -105,6 +105,7 @@ bool FWorldTileInfo::Read(const FString& InPackageFileName, FWorldTileInfo& OutI
 
 		//make sure the filereader gets the correct version number (it defaults to latest version)
 		FileReader->SetUE4Ver(FileSummary.GetFileVersionUE4());
+		FileReader->SetEngineVer(FileSummary.EngineVersion);
 
 		// Load the structure
 		*FileReader << OutInfo;
diff --git a/Engine/Source/Runtime/CoreUObject/Private/UObject/LinkerLoad.cpp b/Engine/Source/Runtime/CoreUObject/Private/UObject/LinkerLoad.cpp
--- a/Engine/Source/Runtime/CoreUObject/Private/UObject/LinkerLoad.cpp
+++ b/Engine/Source/Runtime/CoreUObject/Private/UObject/LinkerLoad.cpp
@@ -831,10 +831,11 @@ ULinkerLoad::ELinkerStatus ULinkerLoad::CreateLoader()
 		}
 
 		// Set status info.
-		ArUE4Ver		= GPackageFileUE4Version;
-		ArLicenseeUE4Ver	= GPackageFileLicenseeUE4Version;
-		ArIsLoading		= true;
-		ArIsPersistent	= true;
+		ArUE4Ver = GPackageFileUE4Version;
+		ArLicenseeUE4Ver = GPackageFileLicenseeUE4Version;
+		ArEngineVer = GEngineVersion;
+		ArIsLoading = true;
+		ArIsPersistent = true;
 
 		// Reset all custom versions
 		ResetCustomVersions();
@@ -972,9 +973,11 @@ ULinkerLoad::ELinkerStatus ULinkerLoad::SerializePackageFileSummary()
 		// Loader needs to be the same version.
 		Loader->SetUE4Ver(Summary.GetFileVersionUE4());
 		Loader->SetLicenseeUE4Ver(Summary.GetFileVersionLicenseeUE4());
+		Loader->SetEngineVer(Summary.EngineVersion);
 
 		ArUE4Ver = Summary.GetFileVersionUE4();
 		ArLicenseeUE4Ver = Summary.GetFileVersionLicenseeUE4();
+		ArEngineVer = Summary.EngineVersion;
 
 		const FCustomVersionContainer& SummaryVersions = Summary.GetCustomVersionContainer();
 		Loader->SetCustomVersions(SummaryVersions);
