Fix non-unity compile

diff --git a/Engine/Source/Runtime/Core/Private/Serialization/Archive.cpp b/Engine/Source/Runtime/Core/Private/Serialization/Archive.cpp
--- a/Engine/Source/Runtime/Core/Private/Serialization/Archive.cpp
+++ b/Engine/Source/Runtime/Core/Private/Serialization/Archive.cpp
@@ -6,6 +6,7 @@
 #include "CorePrivatePCH.h"
 #include "Serialization/Archive.h"
 #include "Serialization/CustomVersion.h"
+#include "EngineVersion.h"
 
 /*-----------------------------------------------------------------------------
 	FArchiveProxy implementation.
