Xbox One exception catching, stack walking, log flushing on exception, minidumps

diff --git a/Engine/Source/Runtime/Core/Public/HAL/ExceptionHandling.h b/Engine/Source/Runtime/Core/Public/HAL/ExceptionHandling.h
--- a/Engine/Source/Runtime/Core/Public/HAL/ExceptionHandling.h
+++ b/Engine/Source/Runtime/Core/Public/HAL/ExceptionHandling.h
@@ -18,6 +18,10 @@ extern CORE_API TCHAR MiniDumpFilenameW[1024];
 // @TODO yrx 2014-10-09 These methods are specific to windows, remove from here.
 extern CORE_API int32 ReportCrash( LPEXCEPTION_POINTERS ExceptionInfo );
 extern CORE_API void NewReportEnsure( const TCHAR* ErrorMessage );
+#elif PLATFORM_XBOXONE
+// @TODO yrx 2014-10-09 Should be move to another file
+extern CORE_API int32 ReportCrash( int ExceptionCode, LPEXCEPTION_POINTERS ExceptionInfo );
+extern CORE_API void NewReportEnsure( const TCHAR* ErrorMessage );
 #elif PLATFORM_MAC
 // @TODO yrx 2014-10-09 Should be move to another file
 extern CORE_API int32 ReportCrash( ucontext_t *Context, int32 Signal, struct __siginfo* Info );
