Enable vector intrinsics for x86-64 Linux as well.

- Scalar FTransform is broken (see UE-9933 for repro).
- Merging CL 2447718 from my UT branch.
- Fix weapon flicker in UT.

diff --git a/Engine/Source/Runtime/Core/Public/Linux/LinuxPlatform.h b/Engine/Source/Runtime/Core/Public/Linux/LinuxPlatform.h
--- a/Engine/Source/Runtime/Core/Public/Linux/LinuxPlatform.h
+++ b/Engine/Source/Runtime/Core/Public/Linux/LinuxPlatform.h
@@ -45,12 +45,12 @@ typedef FLinuxPlatformTypes FPlatformTypes;
 
 #define PLATFORM_USES_DYNAMIC_RHI					1
 
-// only enable vectorintrinsics on x86 for now
-#if defined(_M_IX86) || defined(__i386__)
+// only enable vectorintrinsics on x86(-64) for now
+#if defined(_M_IX86) || defined(__i386__) || defined(_M_X64) || defined(__x86_64__) || defined (__amd64__) 
 	#define PLATFORM_ENABLE_VECTORINTRINSICS		1
 #else
 	#define PLATFORM_ENABLE_VECTORINTRINSICS		0
-#endif // defined(_M_IX86) || defined(__i386__)
+#endif // defined(_M_IX86) || defined(__i386__) || defined(_M_X64) || defined(__x86_64__) || defined (__amd64__) 
 
 // Function type macros.
 #define VARARGS													/* Functions with variable arguments */
