Core - Fixed a stats warning, TAsync thread has an unique name now

diff --git a/Engine/Source/Runtime/Core/Public/Async/Async.h b/Engine/Source/Runtime/Core/Public/Async/Async.h
--- a/Engine/Source/Runtime/Core/Public/Async/Async.h
+++ b/Engine/Source/Runtime/Core/Public/Async/Async.h
@@ -266,7 +266,10 @@ TFuture<ResultType> Async(EAsyncExecution Execution, TFunction<ResultType()> Fun
 		{
 			TPromise<FRunnableThread*> ThreadPromise;
 			TAsyncRunnable<ResultType>* Runnable = new TAsyncRunnable<ResultType>(MoveTemp(Function), MoveTemp(Promise), ThreadPromise.GetFuture());
-			FRunnableThread* RunnableThread = FRunnableThread::Create(Runnable, TEXT("TAsync"));
+
+			static FThreadSafeCounter TAsyncThreadIndex;
+			const FString TAsyncThreadName = FString::Printf( TEXT( "TAsync %d" ), TAsyncThreadIndex.Add(1) );
+			FRunnableThread* RunnableThread = FRunnableThread::Create(Runnable, *TAsyncThreadName);
 
 			check(RunnableThread != nullptr);
 
