Fixed the problem in Mac editor with toolbar icons highlighting while dragging a window. Also, improved mouse cursor position reporting and window position updating during dragging.

diff --git a/Engine/Source/Runtime/Core/Private/Mac/MacApplication.cpp b/Engine/Source/Runtime/Core/Private/Mac/MacApplication.cpp
--- a/Engine/Source/Runtime/Core/Private/Mac/MacApplication.cpp
+++ b/Engine/Source/Runtime/Core/Private/Mac/MacApplication.cpp
@@ -323,7 +323,7 @@ bool FMacApplication::IsCursorDirectlyOverSlateWindow() const
 	SCOPED_AUTORELEASE_POOL;
 	const NSInteger WindowNumber = [NSWindow windowNumberAtPoint:[NSEvent mouseLocation] belowWindowWithWindowNumber:0];
 	NSWindow* const Window = [NSApp windowWithWindowNumber:WindowNumber];
-	return Window && [Window isKindOfClass:[FCocoaWindow class]];
+	return Window && [Window isKindOfClass:[FCocoaWindow class]] && Window != DraggedWindow;
 }
 
 void FMacApplication::ProcessEvent(FMacEvent const* const Event)
@@ -622,17 +622,6 @@ void FMacApplication::ProcessNSEvent(NSEvent* const Event)
 
 			FMacCursor* MacCursor = (FMacCursor*)Cursor.Get();
 
-			// Cocoa does not update NSWindow's frame until user stops dragging the window, so while window is being dragged, we calculate
-			// its position based on mouse move delta
-			if (CurrentEventWindow.IsValid() && DraggedWindow && DraggedWindow == NativeWindow)
-			{
-				const int32 X = FMath::TruncToInt(CurrentEventWindow->PositionX + [Event deltaX]);
-				const int32 Y = FMath::TruncToInt(CurrentEventWindow->PositionY + [Event deltaY]);
-				CurrentEventWindow->PositionX = X;
-				CurrentEventWindow->PositionY = Y;
-				MessageHandler->OnMovedWindow(CurrentEventWindow.ToSharedRef(), X, Y);
-			}
-
 			if (bUsingHighPrecisionMouseInput)
 			{
 				// Under OS X we disassociate the cursor and mouse position during hi-precision mouse input.
@@ -676,23 +665,31 @@ void FMacApplication::ProcessNSEvent(NSEvent* const Event)
 			}
 			else
 			{
-				NSPoint CursorPos = [Event locationInWindow];
-				if ([Event GetWindow])
-				{
-					CursorPos.x += [Event windowPosition].x;
-					CursorPos.y += [Event windowPosition].y;
-				}
-				CursorPos.y--; // The y coordinate in the point returned by locationInWindow starts from a base of 1
+				NSPoint CursorPos = [NSEvent mouseLocation];
+				CursorPos.y--; // The y coordinate of the point returned by mouseLocation starts from a base of 1
 				const FVector2D MousePosition = FVector2D(CursorPos.x, FPlatformMisc::ConvertSlateYPositionToCocoa(CursorPos.y));
 				FVector2D CurrentPosition = MousePosition * MacCursor->GetMouseScaling();
+				const FVector2D MouseDelta = CurrentPosition - MacCursor->GetPosition();
 				if (MacCursor->UpdateCursorClipping(CurrentPosition))
 				{
 					MacCursor->SetPosition(CurrentPosition.X, CurrentPosition.Y);
 				}
 				else
 				{
-					MacCursor->UpdateCurrentPosition(MousePosition);
+					MacCursor->UpdateCurrentPosition(CurrentPosition);
 				}
+
+				// Cocoa does not update NSWindow's frame until user stops dragging the window, so while window is being dragged, we calculate
+				// its position based on mouse move delta
+				if (CurrentEventWindow.IsValid() && DraggedWindow && DraggedWindow == NativeWindow)
+				{
+					const int32 X = FMath::TruncToInt(CurrentEventWindow->PositionX + MouseDelta.X);
+					const int32 Y = FMath::TruncToInt(CurrentEventWindow->PositionY + MouseDelta.Y);
+					MessageHandler->OnMovedWindow(CurrentEventWindow.ToSharedRef(), X, Y);
+					CurrentEventWindow->PositionX = X;
+					CurrentEventWindow->PositionY = Y;
+				}
+
 				MessageHandler->OnMouseMove();
 			}
 
@@ -980,13 +977,8 @@ FCocoaWindow* FMacApplication::FindEventWindow( NSEvent* Event )
 		}
 		else
 		{
-			NSPoint CursorPos = [Event locationInWindow];
-			if ([Event GetWindow])
-			{
-				CursorPos.x += [Event windowPosition].x;
-				CursorPos.y += [Event windowPosition].y;
-			}
-			CursorPos.y--; // The y coordinate in the point returned by locationInWindow starts from a base of 1
+			NSPoint CursorPos = [NSEvent mouseLocation];
+			CursorPos.y--; // The y coordinate of the point returned by mouseLocation starts from a base of 1
 			TSharedPtr<FMacWindow> WindowUnderCursor = LocateWindowUnderCursor(CursorPos);
 			if (WindowUnderCursor.IsValid())
 			{
diff --git a/Engine/Source/Runtime/Core/Private/Mac/MacCursor.cpp b/Engine/Source/Runtime/Core/Private/Mac/MacCursor.cpp
--- a/Engine/Source/Runtime/Core/Private/Mac/MacCursor.cpp
+++ b/Engine/Source/Runtime/Core/Private/Mac/MacCursor.cpp
@@ -109,11 +109,9 @@ FMacCursor::FMacCursor()
 	// Set the default cursor
 	SetType( EMouseCursor::Default );
 
-	CGEventRef Event = CGEventCreate(NULL);
-	CGPoint CursorPos = CGEventGetLocation(Event);
-	CFRelease(Event);
-
-	CurrentPosition = FVector2D(FMath::TruncToFloat(CursorPos.x), FMath::TruncToFloat(CursorPos.y));
+	NSPoint CursorPos = [NSEvent mouseLocation];
+	CursorPos.y--; // The y coordinate of the point returned by mouseLocation starts from a base of 1
+	CurrentPosition = FVector2D(FMath::TruncToFloat(CursorPos.x), FMath::TruncToFloat(FPlatformMisc::ConvertSlateYPositionToCocoa(CursorPos.y)));
 }
 
 FMacCursor::~FMacCursor()
