Exchange now uses Memswap instead of move semantics.
All Exchange overloads removed, since they were all doing the same thing as the generic one.

diff --git a/Engine/Source/Runtime/Core/Public/Containers/Array.h b/Engine/Source/Runtime/Core/Public/Containers/Array.h
--- a/Engine/Source/Runtime/Core/Public/Containers/Array.h
+++ b/Engine/Source/Runtime/Core/Public/Containers/Array.h
@@ -2609,34 +2609,6 @@ template <typename T,typename Allocator> void* operator new( size_t Size, TArray
 	return &Array[Index];
 }
 
-/**
- * A specialization of the exchange macro that avoids reallocating when
- * exchanging two arrays.
- *
- * @param FirstArrayToExchange First array to exchange.
- * @param SecondArrayToExchange Second array to exchange.
- */
-template <typename T>
-inline void Exchange(TArray<T>& FirstArrayToExchange, TArray<T>& SecondArrayToExchange)
-{
-	FMemory::Memswap(&FirstArrayToExchange, &SecondArrayToExchange, sizeof(TArray<T>));
-}
-
-/**
- * A specialization of the exchange macro that avoids reallocating when
- * exchanging two arrays.
- *
- * Non-default allocator version.
- *
- * @param FirstArrayToExchange First array to exchange.
- * @param SecondArrayToExchange Second array to exchange.
- */
-template<typename ElementType, typename Allocator>
-inline void Exchange(TArray<ElementType, Allocator>& FirstArrayToExchange, TArray<ElementType, Allocator>& SecondArrayToExchange)
-{
-	FMemory::Memswap(&FirstArrayToExchange, &SecondArrayToExchange, sizeof(TArray<ElementType, Allocator>));
-}
-
 /*-----------------------------------------------------------------------------
 	MRU array.
 -----------------------------------------------------------------------------*/
@@ -3326,34 +3298,6 @@ template <typename T,typename Allocator> void* operator new( size_t Size, TIndir
 	return &Array[Index];
 }
 
-/**
- * A specialization of the exchange macro that avoids reallocating when
- * exchanging two arrays.
- *
- * @param FirstArrayToExchange First array to exchange.
- * @param SecondArrayToExchange Second array to exchange.
- */
-template <typename T>
-inline void Exchange(TIndirectArray<T>& FirstArrayToExchange, TIndirectArray<T>& SecondArrayToExchange)
-{
-	FMemory::Memswap(&FirstArrayToExchange, &SecondArrayToExchange, sizeof(TIndirectArray<T>));
-}
-
-/**
- * A specialization of the exchange macro that avoids reallocating when
- * exchanging two arrays.
- *
- * Non-default allocator version.
- *
- * @param FirstArrayToExchange First array to exchange.
- * @param SecondArrayToExchange Second array to exchange.
- */
-template<typename ElementType,typename Allocator>
-inline void Exchange(TIndirectArray<ElementType, Allocator>& FirstArrayToExchange, TIndirectArray<ElementType, Allocator>& SecondArrayToExchange)
-{
-	FMemory::Memswap(&FirstArrayToExchange, &SecondArrayToExchange, sizeof(TIndirectArray<ElementType, Allocator>));
-}
-
 /*-----------------------------------------------------------------------------
 	Transactional array.
 -----------------------------------------------------------------------------*/
diff --git a/Engine/Source/Runtime/Core/Public/Containers/BitArray.h b/Engine/Source/Runtime/Core/Public/Containers/BitArray.h
--- a/Engine/Source/Runtime/Core/Public/Containers/BitArray.h
+++ b/Engine/Source/Runtime/Core/Public/Containers/BitArray.h
@@ -879,11 +879,3 @@ private:
 		CurrentBitIndex = BaseBitIndex + NumBitsPerDWORD - 1 - FMath::CountLeadingZeros(this->Mask);
 	}
 };
-
-
-/** A specialization of the exchange macro that avoids reallocating when exchanging two bit arrays. */
-template<typename Allocator>
-FORCEINLINE void Exchange(TBitArray<Allocator>& A,TBitArray<Allocator>& B)
-{
-	FMemory::Memswap(&A,&B,sizeof(TBitArray<Allocator>));
-}
diff --git a/Engine/Source/Runtime/Core/Public/Containers/ChunkedArray.h b/Engine/Source/Runtime/Core/Public/Containers/ChunkedArray.h
--- a/Engine/Source/Runtime/Core/Public/Containers/ChunkedArray.h
+++ b/Engine/Source/Runtime/Core/Public/Containers/ChunkedArray.h
@@ -241,17 +241,3 @@ template <typename T,uint32 TargetBytesPerChunk> void* operator new( size_t Size
 	const int32 Index = ChunkedArray.Add(1);
 	return &ChunkedArray(Index);
 }
-
-
-/**
- * A specialization of the exchange macro that avoids reallocating when
- * exchanging two arrays.
- *
- * @param FirstArrayToExchange First array to exchange.
- * @param SecondArrayToExchange Second array to exchange.
- */
-template <typename T, uint32 TargetBytesPerChunk>
-inline void Exchange(TChunkedArray<T,TargetBytesPerChunk>& FirstArrayToExchange, TChunkedArray<T,TargetBytesPerChunk>& SecondArrayToExchange)
-{
-	FMemory::Memswap(&FirstArrayToExchange, &SecondArrayToExchange, sizeof(TChunkedArray<T,TargetBytesPerChunk>));
-}
diff --git a/Engine/Source/Runtime/Core/Public/Containers/Set.h b/Engine/Source/Runtime/Core/Public/Containers/Set.h
--- a/Engine/Source/Runtime/Core/Public/Containers/Set.h
+++ b/Engine/Source/Runtime/Core/Public/Containers/Set.h
@@ -1148,14 +1148,3 @@ struct TContainerTraits<TSet<ElementType, KeyFuncs, Allocator> > : public TConta
 		TContainerTraits<typename TSet<ElementType, KeyFuncs, Allocator>::ElementArrayType>::MoveWillEmptyContainer &&
 		TAllocatorTraits<typename Allocator::HashAllocator>::SupportsMove };
 };
-
-/** A specialization of the exchange function that avoids reallocating when exchanging two sets. */
-template<
-	typename InElementType,
-	typename KeyFuncs /*= DefaultKeyFuncs<ElementType>*/,
-	typename Allocator /*= FDefaultSetAllocator*/
->
-FORCEINLINE void Exchange(TSet<InElementType,KeyFuncs,Allocator>& A,TSet<InElementType,KeyFuncs,Allocator>& B)
-{
-	FMemory::Memswap(&A,&B,sizeof(TSet<InElementType,KeyFuncs,Allocator>));
-}
diff --git a/Engine/Source/Runtime/Core/Public/Containers/SparseArray.h b/Engine/Source/Runtime/Core/Public/Containers/SparseArray.h
--- a/Engine/Source/Runtime/Core/Public/Containers/SparseArray.h
+++ b/Engine/Source/Runtime/Core/Public/Containers/SparseArray.h
@@ -818,14 +818,3 @@ inline void* operator new(size_t Size,const FSparseArrayAllocationInfo& Allocati
 	ASSUME(Allocation.Pointer);
 	return Allocation.Pointer;
 }
-
-/** A specialization of the exchange macro that avoids reallocating when exchanging two arrays. */
-template <typename ElementType,typename Allocator>
-inline void Exchange(
-	TSparseArray<ElementType,Allocator>& A,
-	TSparseArray<ElementType,Allocator>& B
-	)
-{
-	FMemory::Memswap( &A, &B, sizeof(A) );
-}
-
diff --git a/Engine/Source/Runtime/Core/Public/Containers/UnrealString.h b/Engine/Source/Runtime/Core/Public/Containers/UnrealString.h
--- a/Engine/Source/Runtime/Core/Public/Containers/UnrealString.h
+++ b/Engine/Source/Runtime/Core/Public/Containers/UnrealString.h
@@ -1404,11 +1404,6 @@ public:
 
 		return Result;
 	}
-
-	FORCEINLINE friend void Exchange(FString& A, FString& B)
-	{
-		Exchange(A.Data, B.Data);
-	}
 };
 
 template<>
diff --git a/Engine/Source/Runtime/Core/Public/Templates/UnrealTemplate.h b/Engine/Source/Runtime/Core/Public/Templates/UnrealTemplate.h
--- a/Engine/Source/Runtime/Core/Public/Templates/UnrealTemplate.h
+++ b/Engine/Source/Runtime/Core/Public/Templates/UnrealTemplate.h
@@ -7,6 +7,7 @@
 #include "TypeWrapper.h"
 #include "HAL/Platform.h"
 #include "Templates/UnrealTypeTraits.h"
+#include "UnrealMemory.h"
 
 
 /*-----------------------------------------------------------------------------
@@ -134,13 +135,6 @@ char (&ArrayCountHelper(const T (&)[N]))[N];
 #endif
 
 
-/*-----------------------------------------------------------------------------
-	Allocators.
------------------------------------------------------------------------------*/
-
-template <class T> class TAllocator
-{};
-
 /**
  * works just like std::min_element.
  */
@@ -373,13 +367,14 @@ FORCEINLINE T&& Forward(typename TRemoveReference<T>::Type&& Obj)
 /**
  * Swap two values, using moves if possible
  */
-template< class T > inline void Swap(T& A, T& B)
+template <typename T>
+inline void Swap(T& A, T& B)
 {
-	T Temp = MoveTemp(A);
-	A = MoveTemp(B);
-	B = MoveTemp(Temp);
+	FMemory::Memswap(&A, &B, sizeof(T));
 }
-template< class T > inline void Exchange(T& A, T& B)
+
+template <typename T>
+inline void Exchange(T& A, T& B)
 {
 	Swap(A, B);
 }
