[INTEGRATE] Change 2449107 by [REDACTED_PERSON][REDACTED_HANDLE] on 2015/02/17 15:28:42

	Fix deadlock crash setting the window title while also focusing the editor.

	https://jira.ol.epicgames.net/browse/UE-9976

diff --git a/Engine/Source/Runtime/Core/Private/Mac/MacWindow.cpp b/Engine/Source/Runtime/Core/Private/Mac/MacWindow.cpp
--- a/Engine/Source/Runtime/Core/Private/Mac/MacWindow.cpp
+++ b/Engine/Source/Runtime/Core/Private/Mac/MacWindow.cpp
@@ -537,16 +537,19 @@ bool FMacWindow::IsForegroundWindow() const
 
 void FMacWindow::SetText(const TCHAR* const Text)
 {
-	CFStringRef CFName = FPlatformString::TCHARToCFString(Text);
-	MainThreadCall(^{
-		SCOPED_AUTORELEASE_POOL;
-		[WindowHandle setTitle: (NSString*)CFName];
-		if (IsRegularWindow())
-		{
-			[NSApp changeWindowsItem:WindowHandle title:(NSString*)CFName filename:NO];
-		}
-		CFRelease(CFName);
-	}, UE4NilEventMode, true);
+	if( FString([WindowHandle title]) != FString(Text) )
+	{
+		CFStringRef CFName = FPlatformString::TCHARToCFString( Text );
+		MainThreadCall(^{
+			SCOPED_AUTORELEASE_POOL;
+			[WindowHandle setTitle: (NSString*)CFName];
+			if (IsRegularWindow())
+			{
+				[NSApp changeWindowsItem: WindowHandle title: (NSString*)CFName filename: NO];
+			}
+			CFRelease( CFName );
+		}, UE4NilEventMode, true);
+	}
 }
 
 bool FMacWindow::IsRegularWindow() const
