Fixed TSF IME hang when Wacom tablet is connected

UETOOL-215 - Editor hangs when a graphics tablet with handwriting input enabled is connected

TSF was requesting the run information with a buffer size smaller than the size of the requested string. We now make sure we honor the maximum run buffer size when returning the run information.

diff --git a/Engine/Source/Runtime/Core/Private/Windows/TextStoreACP.cpp b/Engine/Source/Runtime/Core/Private/Windows/TextStoreACP.cpp
--- a/Engine/Source/Runtime/Core/Private/Windows/TextStoreACP.cpp
+++ b/Engine/Source/Runtime/Core/Private/Windows/TextStoreACP.cpp
@@ -548,7 +548,7 @@ STDAPI FTextStoreACP::GetText(LONG acpStart, LONG acpEnd, __RPC__out_ecount_part
 	if(prgRunInfo && ulRunInfoReq > 0)
 	{
 		// We only edit text that's been stripped of any markup, so we have no need to provide multiple runs.
-		prgRunInfo[0].uCount = Length;
+		prgRunInfo[0].uCount = FMath::Min(static_cast<uint32>(ulRunInfoReq), Length);
 		prgRunInfo[0].type = TS_RT_PLAIN;
 		++(*pulRunInfoOut);
 	}
