Adding assert to warn early if someone attempts to queue async IO request with the Destination buffer set to NULL.

diff --git a/Engine/Source/Runtime/Core/Private/Serialization/AsyncIOSystemBase.cpp b/Engine/Source/Runtime/Core/Private/Serialization/AsyncIOSystemBase.cpp
--- a/Engine/Source/Runtime/Core/Private/Serialization/AsyncIOSystemBase.cpp
+++ b/Engine/Source/Runtime/Core/Private/Serialization/AsyncIOSystemBase.cpp
@@ -30,8 +30,10 @@ uint64 FAsyncIOSystemBase::QueueIORequest(
 	FThreadSafeCounter* Counter,
 	EAsyncIOPriority Priority )
 {
-	FScopeLock ScopeLock( CriticalSection );
-	check( Offset != INDEX_NONE );
+	check(Offset != INDEX_NONE);
+	check(Dest != nullptr || Size == 0);
+
+	FScopeLock ScopeLock(CriticalSection);
 
 	// Create an IO request containing passed in information.
 	FAsyncIORequest IORequest;
