GitHub pull request #674

diff --git a/Engine/Source/Runtime/Core/Public/Async/TaskGraphInterfaces.h b/Engine/Source/Runtime/Core/Public/Async/TaskGraphInterfaces.h
--- a/Engine/Source/Runtime/Core/Public/Async/TaskGraphInterfaces.h
+++ b/Engine/Source/Runtime/Core/Public/Async/TaskGraphInterfaces.h
@@ -222,7 +222,7 @@ protected:
 	}
 	/** 
 	 *	Sets the desired execution thread. This is not part of the constructor because this information may not be known quite yet duiring construction.
-	 *	@param InNumberOfPrerequistitesOutstanding; the number of prerequisites outstanding. We actually add one to this internally to prevent the task from firing while we are setting up the task
+	 *	@param InThreadToExecuteOn; the desired thread to execute on.
 	 **/
 	void SetThreadToExecuteOn(ENamedThreads::Type InThreadToExecuteOn)
 	{
