Moved construction of FTaskbarList to FWindowsApplication constructor

Previously this was being created for every unhandled windows message that was routed through WindowsApplication, which occaisonally caused significant stalls waiting for the construction of a new COM interface.

diff --git a/Engine/Source/Runtime/Core/Private/Windows/WindowsApplication.cpp b/Engine/Source/Runtime/Core/Private/Windows/WindowsApplication.cpp
--- a/Engine/Source/Runtime/Core/Private/Windows/WindowsApplication.cpp
+++ b/Engine/Source/Runtime/Core/Private/Windows/WindowsApplication.cpp
@@ -78,6 +78,8 @@ FWindowsApplication::FWindowsApplication( const HINSTANCE HInstance, const HICON
 		TextInputMethodSystem.Reset();
 	}
 
+	TaskbarList = FTaskbarList::Create();
+
 	// Get initial display metrics. (display information for existing desktop, before we start changing resolutions)
 	FDisplayMetrics::GetDisplayMetrics(InitialDisplayMetrics);
 
@@ -630,10 +632,6 @@ LRESULT CALLBACK FWindowsApplication::AppWndProc(HWND hwnd, uint32 msg, WPARAM w
 
 int32 FWindowsApplication::ProcessMessage( HWND hwnd, uint32 msg, WPARAM wParam, LPARAM lParam )
 {
-	//The message id for when the taskbar button becoems created.
-	//This is set to s proper id once a WM_CREATE message is received.
-	static DWORD wmTaskbarButtonCreate = (DWORD)-1; 
-
 	TSharedPtr< FWindowsWindow > CurrentNativeEventWindowPtr = FindWindowByHWND( Windows, hwnd );
 
 	if( Windows.Num() && CurrentNativeEventWindowPtr.IsValid() )
@@ -1098,11 +1096,7 @@ int32 FWindowsApplication::ProcessMessage( HWND hwnd, uint32 msg, WPARAM wParam,
 			break;
 		
 		case WM_CREATE:
-			{
-				wmTaskbarButtonCreate = RegisterWindowMessage(TEXT("TaskbarButtonCreated"));
-				return 0;
-			}
-			break;
+			return 0;
 
 		case WM_DEVICECHANGE:
 			{
@@ -1110,11 +1104,6 @@ int32 FWindowsApplication::ProcessMessage( HWND hwnd, uint32 msg, WPARAM wParam,
 				QueryConnectedMice();
 			}
 		}
-
-		if (wmTaskbarButtonCreate)
-		{
-			TaskbarList = FTaskbarList::Create();
-		}
 	}
 
 	return DefWindowProc(hwnd, msg, wParam, lParam);
