Some minor changes to the ini code.
- updated the comment surrounding config commandline options to avoid incorrect usage.

diff --git a/Engine/Source/Runtime/Core/Private/Misc/ConfigCacheIni.cpp b/Engine/Source/Runtime/Core/Private/Misc/ConfigCacheIni.cpp
--- a/Engine/Source/Runtime/Core/Private/Misc/ConfigCacheIni.cpp
+++ b/Engine/Source/Runtime/Core/Private/Misc/ConfigCacheIni.cpp
@@ -615,9 +615,9 @@ static void OverrideFromCommandline(FConfigFile* File, const FString& Filename)
 {
 	FString Settings;
 	// look for this filename on the commandline in the format:
-	//		-ini:IniName:Section1.Key1=Value1,Section2.Key2=Value2
+	//		-ini:IniName:[Section1]:Key1=Value1,[Section2]:Key2=Value2
 	// for example:
-	//		-ini:Engine:/Script/Engine.Engine:bSmoothFrameRate=False,TextureStreaming:PoolSize=100
+	//		-ini:Engine:[/Script/Engine.Engine]:bSmoothFrameRate=False,[TextureStreaming]:PoolSize=100
 	//			(will update the cache after the final combined engine.ini)
 	if (FParse::Value(FCommandLine::Get(), *FString::Printf(TEXT("%s%s"), *CommandlineOverrideSpecifiers::IniSwitchIdentifier, *FPaths::GetBaseFilename(Filename)), Settings, false))
 	{
