Linux: tab dragging fix.

- PR #830 contributed by [REDACTED_PERSON].
- Also removes unused TrackActivationChanges().

diff --git a/Engine/Source/Runtime/Core/Private/Linux/LinuxApplication.cpp b/Engine/Source/Runtime/Core/Private/Linux/LinuxApplication.cpp
--- a/Engine/Source/Runtime/Core/Private/Linux/LinuxApplication.cpp
+++ b/Engine/Source/Runtime/Core/Private/Linux/LinuxApplication.cpp
@@ -185,43 +185,6 @@ bool FLinuxApplication::GeneratesKeyCharMessage(const SDL_KeyboardEvent & KeyDow
 		(Sym != SDLK_DOWN && Sym != SDLK_LEFT && Sym != SDLK_RIGHT && Sym != SDLK_UP && Sym != SDLK_DELETE);
 }
 
-void FLinuxApplication::TrackActivationChanges(const TSharedPtr<FLinuxWindow> Window, EWindowActivation::Type Event)
-{
-	bool bNeedToNotify = (Window != CurrentlyActiveWindow) || Event == EWindowActivation::Deactivate;
-	
-	UE_LOG(LogLinuxWindow, Verbose, TEXT("TrackActivationChanges: bNeedToNotify=%s (Window: %p, CurrentlyActiveWindow: %p, Event: %d)"),
-		bNeedToNotify ? TEXT("true") : TEXT("false"),
-		Window.Get(),
-		CurrentlyActiveWindow.Get(),
-		static_cast<int32>(Event));
-	
-	if (bNeedToNotify)
-	{
-		// see if previous window needs to be deactivated
-		if (CurrentlyActiveWindow.IsValid())
-		{
-			UE_LOG(LogLinuxWindow, Verbose, TEXT("TrackActivationChanges: Deactivating previous window %p"), CurrentlyActiveWindow.Get());
-			MessageHandler->OnWindowActivationChanged(CurrentlyActiveWindow.ToSharedRef(), EWindowActivation::Deactivate);
-		}
-			
-		if (Event != EWindowActivation::Deactivate)
-		{
-			CurrentlyActiveWindow = Window;
-			UE_LOG(LogLinuxWindow, Verbose, TEXT("TrackActivationChanges: New active window is %p, notifying it"), CurrentlyActiveWindow.Get());
-			
-			if (CurrentlyActiveWindow.IsValid())
-			{
-				MessageHandler->OnWindowActivationChanged(CurrentlyActiveWindow.ToSharedRef(), Event);
-			}
-		}
-		else
-		{
-			CurrentlyActiveWindow = nullptr;
-			UE_LOG(LogLinuxWindow, Verbose, TEXT("TrackActivationChanges: don't have any active window"));
-		}
-	}
-}
-
 void FLinuxApplication::ProcessDeferredMessage( SDL_Event Event )
 {
 	// This function can be reentered when entering a modal tick loop.
@@ -675,35 +638,39 @@ void FLinuxApplication::ProcessDeferredMessage( SDL_Event Event )
 				case SDL_WINDOWEVENT_DESTROY:
 					{
 						RemoveEventWindow(CurrentEventWindow->GetHWnd());
-
-						// Remove the window from the focus tracker.
 						RemoveRevertFocusWindow(CurrentEventWindow->GetHWnd());
-						
-						MessageHandler->OnWindowClose(CurrentEventWindow.ToSharedRef());
 
 						// After destroying the window we have to decide what to do with the
 						// focus. We should check if the current window has a parent so we can give
 						// the focus to the parent. If not we get a window from the Windows list
 						// and set the focus to the most top one.
-						if (CurrentEventWindow->IsRegularWindow() && CurrentEventWindow->GetParent().IsValid())
-						{
-							CurrentEventWindow->GetParent()->BringToFront();
-							CurrentEventWindow->GetParent()->SetWindowFocus();
-						} 
-						else if (CurrentEventWindow->IsRegularWindow() && RevertFocusStack.Num() > 0)
+ 						if (CurrentEventWindow->IsRegularWindow() && CurrentEventWindow->GetParent().IsValid())
+ 						{
+ 							CurrentEventWindow->GetParent()->BringToFront();
+ 							CurrentEventWindow->GetParent()->SetWindowFocus();
+ 						} 
+ 						else if (CurrentEventWindow->IsRegularWindow() && RevertFocusStack.Num() > 0)
 						{
-							TSharedPtr< FLinuxWindow > TopmostWindow = RevertFocusStack.Top();
+ 							TSharedPtr< FLinuxWindow > TopmostWindow = RevertFocusStack.Top();
 							if (TopmostWindow.IsValid())
 							{
 								TopmostWindow->BringToFront();
 								TopmostWindow->SetWindowFocus();
 							}
 						}
+
+						SDL_DestroyWindow(CurrentEventWindow->GetHWnd());
+
+						MessageHandler->OnWindowClose(CurrentEventWindow.ToSharedRef());
 					}
 					break;
 
 				case SDL_WINDOWEVENT_SHOWN:
 					{
+						if (CurrentEventWindow->IsRegularWindow())
+						{
+							CurrentEventWindow->SetWindowFocus();
+						}
 						int Width, Height;
 
 						SDL_GetWindowSize(NativeWindow, &Width, &Height);
diff --git a/Engine/Source/Runtime/Core/Private/Linux/LinuxWindow.cpp b/Engine/Source/Runtime/Core/Private/Linux/LinuxWindow.cpp
--- a/Engine/Source/Runtime/Core/Private/Linux/LinuxWindow.cpp
+++ b/Engine/Source/Runtime/Core/Private/Linux/LinuxWindow.cpp
@@ -228,7 +228,6 @@ void FLinuxWindow::BringToFront( bool bForce )
 /** Native windows should implement this function by asking the OS to destroy OS-specific resource associated with the window (e.g. Win32 window handle) */
 void FLinuxWindow::Destroy()
 {
-	SDL_DestroyWindow( HWnd );
 	OwningApplication->SendDestroyEvent(HWnd);
 }
 
diff --git a/Engine/Source/Runtime/Core/Public/Linux/LinuxApplication.h b/Engine/Source/Runtime/Core/Public/Linux/LinuxApplication.h
--- a/Engine/Source/Runtime/Core/Public/Linux/LinuxApplication.h
+++ b/Engine/Source/Runtime/Core/Public/Linux/LinuxApplication.h
@@ -96,14 +96,6 @@ private:
 	 */
 	static bool GeneratesKeyCharMessage(const SDL_KeyboardEvent & KeyDownEvent);
 
-	/**
-	 * Tracks window activation changes and sends notifications as required
-	 * 
-	 * @param Window window that caused the activation event
-	 * @param Event type of the activation change
-	 */
-	void TrackActivationChanges(const TSharedPtr<FLinuxWindow> Window, EWindowActivation::Type Event);
-
 private:
 
 	struct SDLControllerState
@@ -114,8 +106,6 @@ private:
 		bool analogOverThreshold[10];
 	};
 
-// 	static const FIntPoint MinimizedWindowPosition;
-
 	bool bUsingHighPrecisionMouseInput;
 
 	TArray< SDL_Event > PendingEvents;
