CrashReport - Shortened a crash-specific directory to avoid issues with long names

diff --git a/Engine/Source/Runtime/Core/Private/Linux/LinuxPlatformCrashContext.cpp b/Engine/Source/Runtime/Core/Private/Linux/LinuxPlatformCrashContext.cpp
--- a/Engine/Source/Runtime/Core/Private/Linux/LinuxPlatformCrashContext.cpp
+++ b/Engine/Source/Runtime/Core/Private/Linux/LinuxPlatformCrashContext.cpp
@@ -679,8 +679,7 @@ void DLLEXPORT GenerateCrashInfoAndLaunchReporter(const FLinuxCrashContext & Con
 #if !IS_PROGRAM
 
 	// create a crash-specific directory
-	FString CrashInfoFolder = FString::Printf(TEXT("crashinfo-%s-pid-%d-%s-%s"), FApp::GetGameName(), getpid(), 
-		*FDateTime::Now().ToString(), *FGuid::NewGuid().ToString());
+	FString CrashInfoFolder = FString::Printf(TEXT("crashinfo-%s-pid-%d-%s"), FApp::GetGameName(), getpid(), *FGuid::NewGuid().ToString());
 	FString CrashInfoAbsolute = FPaths::ConvertRelativePathToFull(CrashInfoFolder);
 	if (IFileManager::Get().MakeDirectory(*CrashInfoFolder))
 	{
