Config File issue fixes
- Writing a default config file with no hierarchy caused a crash.
- Moved all of the relevant array checking code into a check for if we are processing an array.

diff --git a/Engine/Source/Runtime/Core/Private/Misc/ConfigCacheIni.cpp b/Engine/Source/Runtime/Core/Private/Misc/ConfigCacheIni.cpp
--- a/Engine/Source/Runtime/Core/Private/Misc/ConfigCacheIni.cpp
+++ b/Engine/Source/Runtime/Core/Private/Misc/ConfigCacheIni.cpp
@@ -1278,26 +1278,33 @@ void FConfigFile::ProcessSourceAndCheckAgainstBackup()
 
 void FConfigFile::ProcessPropertyAndWriteForDefaults( const TArray< FString >& InCompletePropertyToProcess, FString& OutText, const FString& SectionName, const FString& PropertyName )
 {
-	FConfigCacheIni Hierarchy(EConfigCacheType::Temporary);
-	const FString& LastFileInHierarchy = SourceIniHierarchy.Last().Filename;
-
-	// Build a config file out of this default configs hierarchy.
-	FConfigFile& DefaultConfigFile = Hierarchy.Add(LastFileInHierarchy, FConfigFile());
-	for (TArray<FIniFilename>::TIterator HierarchyFileIt(SourceIniHierarchy); HierarchyFileIt; ++HierarchyFileIt)
-	{
-		DefaultConfigFile.Combine(HierarchyFileIt->Filename);
-	}
-	
-	// Handle array elements from the configs hierarchy.
-	if (PropertyName.StartsWith(TEXT("+")) || InCompletePropertyToProcess.Num() > 1)
+	// Only process against a hierarchy if this config file has one.
+	if (SourceIniHierarchy.Num() > 0)
 	{
-		TArray<FString> ArrayProperties;
-		Hierarchy.GetArray(*SectionName, *PropertyName.Replace(TEXT("+"), TEXT("")), ArrayProperties, *LastFileInHierarchy);
-
-		for (const FString& NextElement : ArrayProperties)
+		// Handle array elements from the configs hierarchy.
+		if (PropertyName.StartsWith(TEXT("+")) || InCompletePropertyToProcess.Num() > 1)
 		{
-			FString PropertyNameWithRemoveOp = PropertyName.Replace(TEXT("+"), TEXT("-"));
-			OutText += FString::Printf(TEXT("%s=%s") LINE_TERMINATOR, *PropertyNameWithRemoveOp, *NextElement);
+			// Build a config file out of this default configs hierarchy.
+			FConfigCacheIni Hierarchy(EConfigCacheType::Temporary);
+			const FString& LastFileInHierarchy = SourceIniHierarchy.Last().Filename;
+
+			FConfigFile& DefaultConfigFile = Hierarchy.Add(LastFileInHierarchy, FConfigFile());
+			for (TArray<FIniFilename>::TIterator HierarchyFileIt(SourceIniHierarchy); HierarchyFileIt; ++HierarchyFileIt)
+			{
+				DefaultConfigFile.Combine(HierarchyFileIt->Filename);
+			}
+
+			// Remove any array elements from the default configs hierearchy, we will add these in below
+			// Note.	This compensates for an issue where strings in the hierarchy have a slightly different format
+			//			to how the config system wishes to serialize them.
+			TArray<FString> ArrayProperties;
+			Hierarchy.GetArray(*SectionName, *PropertyName.Replace(TEXT("+"), TEXT("")), ArrayProperties, *LastFileInHierarchy);
+
+			for (const FString& NextElement : ArrayProperties)
+			{
+				FString PropertyNameWithRemoveOp = PropertyName.Replace(TEXT("+"), TEXT("-"));
+				OutText += FString::Printf(TEXT("%s=%s") LINE_TERMINATOR, *PropertyNameWithRemoveOp, *NextElement);
+			}
 		}
 	}
 
