GitHub pull request #640

diff --git a/Engine/Source/Runtime/Core/Private/Windows/WindowsApplication.cpp b/Engine/Source/Runtime/Core/Private/Windows/WindowsApplication.cpp
--- a/Engine/Source/Runtime/Core/Private/Windows/WindowsApplication.cpp
+++ b/Engine/Source/Runtime/Core/Private/Windows/WindowsApplication.cpp
@@ -760,7 +760,6 @@ int32 FWindowsApplication::ProcessMessage( HWND hwnd, uint32 msg, WPARAM wParam,
 		case WM_NCMOUSEMOVE:
 		case WM_MOUSEMOVE:
 		case WM_MOUSEWHEEL:
-		case WM_SETCURSOR:
 			{
 				DeferMessage( CurrentNativeEventWindowPtr, hwnd, msg, wParam, lParam );
 				// Handled
@@ -768,6 +767,11 @@ int32 FWindowsApplication::ProcessMessage( HWND hwnd, uint32 msg, WPARAM wParam,
 			}
 			break;
 
+		case WM_SETCURSOR:
+			{
+				DeferMessage( CurrentNativeEventWindowPtr, hwnd, msg, wParam, lParam );
+			}
+			break;
 		// Mouse Movement
 		case WM_INPUT:
 			{
diff --git a/Engine/Source/Runtime/Core/Private/Windows/WindowsWindow.cpp b/Engine/Source/Runtime/Core/Private/Windows/WindowsWindow.cpp
--- a/Engine/Source/Runtime/Core/Private/Windows/WindowsWindow.cpp
+++ b/Engine/Source/Runtime/Core/Private/Windows/WindowsWindow.cpp
@@ -110,7 +110,7 @@ void FWindowsWindow::Initialize( FWindowsApplication* const Application, const T
 	{
 		// OS Window border setup
 		WindowExStyle = WS_EX_APPWINDOW;
-		WindowStyle = WS_POPUP | WS_OVERLAPPED | WS_SYSMENU | WS_MINIMIZEBOX | WS_MAXIMIZEBOX | WS_BORDER | WS_CAPTION;
+		WindowStyle = WS_POPUP | WS_OVERLAPPED | WS_SYSMENU | WS_MINIMIZEBOX | WS_MAXIMIZEBOX | WS_SIZEBOX | WS_BORDER | WS_CAPTION;
 
 		// Note SizeX and SizeY should be the size of the client area.  We need to get the actual window size by adjusting the client size to account for standard windows border around the window
 		RECT WindowRect = { 0, 0, ClientWidth, ClientWidth };
@@ -519,7 +519,7 @@ void FWindowsWindow::SetWindowMode( EWindowMode::Type NewWindowMode )
 		// Setup Win32 Flags to be used for Fullscreen mode
 		LONG WindowFlags = GetWindowLong(HWnd, GWL_STYLE);
 		const LONG FullscreenFlags = WS_POPUP;
-		const LONG RestoredFlags = WS_MAXIMIZEBOX | WS_MINIMIZEBOX | WS_CAPTION | WS_SYSMENU | WS_OVERLAPPED | WS_BORDER;
+		const LONG RestoredFlags = WS_MAXIMIZEBOX | WS_MINIMIZEBOX | WS_CAPTION | WS_SYSMENU | WS_OVERLAPPED | WS_BORDER | WS_SIZEBOX;
 
 		// If we're not in fullscreen, make it so
 		if( NewWindowMode == EWindowMode::WindowedFullscreen || NewWindowMode == EWindowMode::Fullscreen)
diff --git a/Engine/Source/Runtime/SlateRHIRenderer/Private/SlateRHIRenderer.cpp b/Engine/Source/Runtime/SlateRHIRenderer/Private/SlateRHIRenderer.cpp
--- a/Engine/Source/Runtime/SlateRHIRenderer/Private/SlateRHIRenderer.cpp
+++ b/Engine/Source/Runtime/SlateRHIRenderer/Private/SlateRHIRenderer.cpp
@@ -1212,8 +1212,12 @@ void FSlateRHIRenderer::RequestResize( const TSharedPtr<SWindow>& Window, uint32
 
 	if( ViewInfo )
 	{
-		ViewInfo->DesiredWidth = NewWidth;
-		ViewInfo->DesiredHeight = NewHeight;
+		if (ViewInfo->DesiredWidth != NewWidth ||
+			ViewInfo->DesiredHeight != NewHeight) {
+			ViewInfo->DesiredWidth = NewWidth;
+			ViewInfo->DesiredHeight = NewHeight;
+			FSystemResolution::RequestResolutionChange(NewWidth, NewHeight, EWindowMode::Windowed);
+		}
 	}
 }
 
