Linux: support priority modifier in CreateProc.

- Makes compiling shaders in the background (by SCW) more tolerable (as was intended).
- Also some minor fixes to string handling.

diff --git a/Engine/Source/Runtime/Core/Private/Linux/LinuxPlatformProcess.cpp b/Engine/Source/Runtime/Core/Private/Linux/LinuxPlatformProcess.cpp
--- a/Engine/Source/Runtime/Core/Private/Linux/LinuxPlatformProcess.cpp
+++ b/Engine/Source/Runtime/Core/Private/Linux/LinuxPlatformProcess.cpp
@@ -409,7 +409,7 @@ void FLinuxPlatformProcess::LaunchURL(const TCHAR* URL, const TCHAR* Parms, FStr
 {
 	// @todo This ignores params and error; mostly a stub
 	pid_t pid = fork();
-	UE_LOG(LogHAL, Verbose, TEXT("FLinuxPlatformProcess::LaunchURL: '%s'"), TCHAR_TO_ANSI(URL));
+	UE_LOG(LogHAL, Verbose, TEXT("FLinuxPlatformProcess::LaunchURL: '%s'"), URL);
 	if (pid == 0)
 	{
 		exit(execl("/usr/bin/xdg-open", "xdg-open", TCHAR_TO_ANSI(URL), (char *)0));
@@ -421,7 +421,7 @@ FProcHandle FLinuxPlatformProcess::CreateProc(const TCHAR* URL, const TCHAR* Par
 	// @TODO bLaunchHidden bLaunchReallyHidden are not handled
 	// We need an absolute path to executable
 	FString ProcessPath = URL;
-	if (*URL != '/')
+	if (*URL != TEXT('/'))
 	{
 		ProcessPath = FPaths::ConvertRelativePathToFull(ProcessPath);
 	}
@@ -557,13 +557,81 @@ FProcHandle FLinuxPlatformProcess::CreateProc(const TCHAR* URL, const TCHAR* Par
 		posix_spawn_file_actions_adddup2(&FileActions, PipeWriteHandle->GetHandle(), STDOUT_FILENO);
 	}
 
-	int ErrNo = posix_spawn(&ChildPid, TCHAR_TO_ANSI(*ProcessPath), &FileActions, NULL, Argv, environ);
+	int ErrNo = posix_spawn(&ChildPid, TCHAR_TO_ANSI(*ProcessPath), &FileActions, nullptr, Argv, environ);
 	posix_spawn_file_actions_destroy(&FileActions);
 	if (ErrNo != 0)
 	{
 		UE_LOG(LogHAL, Fatal, TEXT("FLinuxPlatformProcess::CreateProc: posix_spawn() failed (%d, %s)"), ErrNo, ANSI_TO_TCHAR(strerror(ErrNo)));
 		return FProcHandle();	// produce knowingly invalid handle if for some reason Fatal log (above) returns
 	}
+
+	// renice the child (subject to race condition).
+	// Why this instead of posix_spawn_setschedparam()? 
+	// Because posix_spawnattr priority is unusable under Linux due to min/max priority range being [0;0] for the default scheduler
+	if (PriorityModifier != 0)
+	{
+		errno = 0;
+		int TheirCurrentPrio = getpriority(PRIO_PROCESS, ChildPid);
+
+		if (errno != 0)
+		{
+			int ErrNo = errno;
+			UE_LOG(LogHAL, Warning, TEXT("FLinuxPlatformProcess::CreateProc: could not get child's priority, errno=%d (%s)"),
+				ErrNo,
+				ANSI_TO_TCHAR(strerror(ErrNo))
+			);
+			
+			// proceed anyway...
+			TheirCurrentPrio = 0;
+		}
+
+		rlimit PrioLimits;
+		int MaxPrio = 0;
+		if (getrlimit(RLIMIT_NICE, &PrioLimits) == -1)
+		{
+			int ErrNo = errno;
+			UE_LOG(LogHAL, Warning, TEXT("FLinuxPlatformProcess::CreateProc: could not get priority limits (RLIMIT_NICE), errno=%d (%s)"),
+				ErrNo,
+				ANSI_TO_TCHAR(strerror(ErrNo))
+			);
+
+			// proceed anyway...
+		}
+		else
+		{
+			MaxPrio = PrioLimits.rlim_cur;
+		}
+
+		int NewPrio = TheirCurrentPrio;
+		if (PriorityModifier > 0)
+		{
+			// decrease the nice value - will perhaps fail, it's up to the user to run with proper permissions
+			NewPrio -= 10;
+		}
+		else
+		{
+			NewPrio += 10;
+		}
+
+		// cap to [RLIMIT_NICE, 19]
+		NewPrio = FMath::Min(19, NewPrio);
+		NewPrio = FMath::Max(MaxPrio, NewPrio);	// MaxPrio is actually the _lowest_ numerically priority
+
+		if (setpriority(PRIO_PROCESS, ChildPid, NewPrio) == -1)
+		{
+			int ErrNo = errno;
+			UE_LOG(LogHAL, Warning, TEXT("FLinuxPlatformProcess::CreateProc: could not change child's priority (nice value) from %d to %d, errno=%d (%s)"),
+				TheirCurrentPrio, NewPrio,
+				ErrNo,
+				ANSI_TO_TCHAR(strerror(ErrNo))
+			);
+		}
+		else
+		{
+			UE_LOG(LogHAL, Log, TEXT("Changed child's priority (nice value) to %d (change from %d)"), NewPrio, TheirCurrentPrio);
+		}
+	}
+
 	else
 	{
 		UE_LOG(LogHAL, Log, TEXT("FLinuxPlatformProcess::CreateProc: spawned child %d"), ChildPid);
