Linux: fix unability to rename after latest UI changes.

- PR #836 contributed by [REDACTED_PERSON].

diff --git a/Engine/Source/Runtime/Core/Private/Linux/LinuxApplication.cpp b/Engine/Source/Runtime/Core/Private/Linux/LinuxApplication.cpp
--- a/Engine/Source/Runtime/Core/Private/Linux/LinuxApplication.cpp
+++ b/Engine/Source/Runtime/Core/Private/Linux/LinuxApplication.cpp
@@ -658,6 +658,11 @@ void FLinuxApplication::ProcessDeferredMessage( SDL_Event Event )
 								TopmostWindow->SetWindowFocus();
 							}
 						}
+						else if (CurrentEventWindow->IsPopupMenuWindow() && CurrentEventWindow->GetParent().IsValid() && !bIsDragWindowButtonPressed && bActivateApp)
+						{
+							CurrentEventWindow->GetParent()->BringToFront();
+							CurrentEventWindow->GetParent()->SetWindowFocus();
+						}
 
 						SDL_DestroyWindow(CurrentEventWindow->GetHWnd());
 
