GitHub 846: https://github.com/EpicGames/UnrealEngine/pull/846

std::nothrow_t& parameters made const.

diff --git a/Engine/Source/Runtime/Core/Public/Modules/Boilerplate/ModuleBoilerplate.h b/Engine/Source/Runtime/Core/Public/Modules/Boilerplate/ModuleBoilerplate.h
--- a/Engine/Source/Runtime/Core/Public/Modules/Boilerplate/ModuleBoilerplate.h
+++ b/Engine/Source/Runtime/Core/Public/Modules/Boilerplate/ModuleBoilerplate.h
@@ -19,14 +19,14 @@
 #endif
 
 #define REPLACEMENT_OPERATOR_NEW_AND_DELETE \
-	OPERATOR_NEW_MSVC_PRAGMA void* operator new  ( size_t Size                  ) OPERATOR_NEW_THROW_SPEC      { return FMemory::Malloc( Size ); } \
-	OPERATOR_NEW_MSVC_PRAGMA void* operator new[]( size_t Size                  ) OPERATOR_NEW_THROW_SPEC      { return FMemory::Malloc( Size ); } \
-	OPERATOR_NEW_MSVC_PRAGMA void* operator new  ( size_t Size, std::nothrow_t& ) OPERATOR_NEW_NOTHROW_SPEC    { return FMemory::Malloc( Size ); } \
-	OPERATOR_NEW_MSVC_PRAGMA void* operator new[]( size_t Size, std::nothrow_t& ) OPERATOR_NEW_NOTHROW_SPEC    { return FMemory::Malloc( Size ); } \
-	void operator delete  ( void* Ptr )                                           OPERATOR_DELETE_THROW_SPEC   { FMemory::Free( Ptr ); } \
-	void operator delete[]( void* Ptr )                                           OPERATOR_DELETE_THROW_SPEC   { FMemory::Free( Ptr ); } \
-	void operator delete  ( void* Ptr, std::nothrow_t& )                          OPERATOR_DELETE_NOTHROW_SPEC { FMemory::Free( Ptr ); } \
-	void operator delete[]( void* Ptr, std::nothrow_t& )                          OPERATOR_DELETE_NOTHROW_SPEC { FMemory::Free( Ptr ); }
+	OPERATOR_NEW_MSVC_PRAGMA void* operator new  ( size_t Size                        ) OPERATOR_NEW_THROW_SPEC      { return FMemory::Malloc( Size ); } \
+	OPERATOR_NEW_MSVC_PRAGMA void* operator new[]( size_t Size                        ) OPERATOR_NEW_THROW_SPEC      { return FMemory::Malloc( Size ); } \
+	OPERATOR_NEW_MSVC_PRAGMA void* operator new  ( size_t Size, const std::nothrow_t& ) OPERATOR_NEW_NOTHROW_SPEC    { return FMemory::Malloc( Size ); } \
+	OPERATOR_NEW_MSVC_PRAGMA void* operator new[]( size_t Size, const std::nothrow_t& ) OPERATOR_NEW_NOTHROW_SPEC    { return FMemory::Malloc( Size ); } \
+	void operator delete  ( void* Ptr )                                                 OPERATOR_DELETE_THROW_SPEC   { FMemory::Free( Ptr ); } \
+	void operator delete[]( void* Ptr )                                                 OPERATOR_DELETE_THROW_SPEC   { FMemory::Free( Ptr ); } \
+	void operator delete  ( void* Ptr, const std::nothrow_t& )                          OPERATOR_DELETE_NOTHROW_SPEC { FMemory::Free( Ptr ); } \
+	void operator delete[]( void* Ptr, const std::nothrow_t& )                          OPERATOR_DELETE_NOTHROW_SPEC { FMemory::Free( Ptr ); }
 
 // in DLL builds, these are done per-module, otherwise we just need one in the application
 // visual studio cannot find cross dll data for visualizers, so these provide access
