Don't make context menus key windows on Mac

diff --git a/Engine/Source/Runtime/Core/Private/Mac/CocoaWindow.cpp b/Engine/Source/Runtime/Core/Private/Mac/CocoaWindow.cpp
--- a/Engine/Source/Runtime/Core/Private/Mac/CocoaWindow.cpp
+++ b/Engine/Source/Runtime/Core/Private/Mac/CocoaWindow.cpp
@@ -126,7 +126,7 @@ NSString* NSPerformDragOperation = @"NSPerformDragOperation";
 - (BOOL)canBecomeKeyWindow
 {
 	SCOPED_AUTORELEASE_POOL;
-	return bAcceptsInput && ![self ignoresMouseEvents];
+	return bAcceptsInput && ![self ignoresMouseEvents] && ([self styleMask] != NSBorderlessWindowMask);
 }
 
 - (BOOL)validateMenuItem:(NSMenuItem*)MenuItem
