Rollback //depot/UE4/Engine/Source/Runtime/Core/Private/Mac/CocoaWindow.cpp to revision 25

diff --git a/Engine/Source/Runtime/Core/Private/Mac/CocoaWindow.cpp b/Engine/Source/Runtime/Core/Private/Mac/CocoaWindow.cpp
--- a/Engine/Source/Runtime/Core/Private/Mac/CocoaWindow.cpp
+++ b/Engine/Source/Runtime/Core/Private/Mac/CocoaWindow.cpp
@@ -126,7 +126,7 @@ NSString* NSPerformDragOperation = @"NSPerformDragOperation";
 - (BOOL)canBecomeKeyWindow
 {
 	SCOPED_AUTORELEASE_POOL;
-	return bAcceptsInput && ![self ignoresMouseEvents] && ([self styleMask] != NSBorderlessWindowMask);
+	return bAcceptsInput && ![self ignoresMouseEvents];
 }
 
 - (BOOL)validateMenuItem:(NSMenuItem*)MenuItem
