Speculative fix for UE-10541

It's possible that the taskbar list is being destroyed too late (after com has been un-initialized)

diff --git a/Engine/Source/Runtime/Core/Private/Windows/WindowsApplication.cpp b/Engine/Source/Runtime/Core/Private/Windows/WindowsApplication.cpp
--- a/Engine/Source/Runtime/Core/Private/Windows/WindowsApplication.cpp
+++ b/Engine/Source/Runtime/Core/Private/Windows/WindowsApplication.cpp
@@ -187,6 +187,8 @@ void FWindowsApplication::DestroyApplication()
 	// Restore accessibility shortcuts and remove the saved state from the .ini
 	AllowAccessibilityShortcutKeys(true);
 	GConfig->EmptySection(TEXT("WindowsApplication.Accessibility"), GEngineIni);
+
+	TaskbarList = nullptr;
 }
 
 void FWindowsApplication::ShutDownAfterError()
@@ -194,6 +196,8 @@ void FWindowsApplication::ShutDownAfterError()
 	// Restore accessibility shortcuts and remove the saved state from the .ini
 	AllowAccessibilityShortcutKeys(true);
 	GConfig->EmptySection(TEXT("WindowsApplication.Accessibility"), GEngineIni);
+
+	TaskbarList = nullptr;
 }
 
 bool FWindowsApplication::RegisterClass( const HINSTANCE HInstance, const HICON HIcon )
