Swapped the way we treat main and key windows in Slate on Mac to be consistent with native window buttons behavior

diff --git a/Engine/Source/Runtime/Core/Private/Mac/CocoaWindow.cpp b/Engine/Source/Runtime/Core/Private/Mac/CocoaWindow.cpp
--- a/Engine/Source/Runtime/Core/Private/Mac/CocoaWindow.cpp
+++ b/Engine/Source/Runtime/Core/Private/Mac/CocoaWindow.cpp
@@ -120,13 +120,13 @@ NSString* NSPerformDragOperation = @"NSPerformDragOperation";
 - (BOOL)canBecomeMainWindow
 {
 	SCOPED_AUTORELEASE_POOL;
-	return bAcceptsInput && ([self styleMask] != NSBorderlessWindowMask);
+	return bAcceptsInput && ![self ignoresMouseEvents];
 }
 
 - (BOOL)canBecomeKeyWindow
 {
 	SCOPED_AUTORELEASE_POOL;
-	return bAcceptsInput && ![self ignoresMouseEvents];
+	return bAcceptsInput && ([self styleMask] != NSBorderlessWindowMask);
 }
 
 - (BOOL)validateMenuItem:(NSMenuItem*)MenuItem
@@ -246,7 +246,7 @@ NSString* NSPerformDragOperation = @"NSPerformDragOperation";
 	}
 }
 
-- (void)windowDidBecomeKey:(NSNotification*)Notification
+- (void)windowDidBecomeMain:(NSNotification*)Notification
 {
 	SCOPED_AUTORELEASE_POOL;
 	if ([NSApp isHidden] == NO)
@@ -257,7 +257,7 @@ NSString* NSPerformDragOperation = @"NSPerformDragOperation";
 	FMacEvent::SendToGameRunLoop(Notification, EMacEventSendMethod::Async, @[ NSDefaultRunLoopMode, UE4ShowEventMode, UE4CloseEventMode, UE4FullscreenEventMode ]);
 }
 
-- (void)windowDidResignKey:(NSNotification*)Notification
+- (void)windowDidResignMain:(NSNotification*)Notification
 {
 	SCOPED_AUTORELEASE_POOL;
 	[self setMovable: YES];
diff --git a/Engine/Source/Runtime/Core/Private/Mac/MacApplication.cpp b/Engine/Source/Runtime/Core/Private/Mac/MacApplication.cpp
--- a/Engine/Source/Runtime/Core/Private/Mac/MacApplication.cpp
+++ b/Engine/Source/Runtime/Core/Private/Mac/MacApplication.cpp
@@ -368,13 +368,13 @@ void FMacApplication::ProcessEvent(FMacEvent const* const Event)
 			{
 				MacApplication->OnWindowDidResize(CocoaWindow);
 			}
-			else if (NotificationName == NSWindowDidBecomeKeyNotification)
+			else if (NotificationName == NSWindowDidBecomeMainNotification)
 			{
-				MacApplication->OnWindowDidBecomeKey(CocoaWindow);
+				MacApplication->OnWindowDidBecomeMain(CocoaWindow);
 			}
-			else if (NotificationName == NSWindowDidResignKeyNotification)
+			else if (NotificationName == NSWindowDidResignMainNotification)
 			{
-				MacApplication->OnWindowDidResignKey(CocoaWindow);
+				MacApplication->OnWindowDidResignMain(CocoaWindow);
 			}
 			else if (NotificationName == NSWindowWillMoveNotification)
 			{
@@ -1198,7 +1198,7 @@ void FMacApplication::OnDragDrop( FCocoaWindow* Window )
 	}
 }
 
-void FMacApplication::OnWindowDidBecomeKey( FCocoaWindow* Window )
+void FMacApplication::OnWindowDidBecomeMain( FCocoaWindow* Window )
 {
 	TSharedPtr< FMacWindow > EventWindow = FindWindowByNSWindow( Windows, &WindowsMutex, Window );
 	if (EventWindow.IsValid())
@@ -1207,7 +1207,7 @@ void FMacApplication::OnWindowDidBecomeKey( FCocoaWindow* Window )
 	}
 }
 
-void FMacApplication::OnWindowDidResignKey( FCocoaWindow* Window )
+void FMacApplication::OnWindowDidResignMain( FCocoaWindow* Window )
 {
 	TSharedPtr< FMacWindow > EventWindow = FindWindowByNSWindow( Windows, &WindowsMutex, Window );
 	if (EventWindow.IsValid())
diff --git a/Engine/Source/Runtime/Core/Private/Mac/MacWindow.cpp b/Engine/Source/Runtime/Core/Private/Mac/MacWindow.cpp
--- a/Engine/Source/Runtime/Core/Private/Mac/MacWindow.cpp
+++ b/Engine/Source/Runtime/Core/Private/Mac/MacWindow.cpp
@@ -532,7 +532,7 @@ int32 FMacWindow::GetWindowBorderSize() const
 bool FMacWindow::IsForegroundWindow() const
 {
 	SCOPED_AUTORELEASE_POOL;
-	return [WindowHandle isMainWindow];
+	return [WindowHandle isKeyWindow];
 }
 
 void FMacWindow::SetText(const TCHAR* const Text)
diff --git a/Engine/Source/Runtime/Core/Public/Mac/MacApplication.h b/Engine/Source/Runtime/Core/Public/Mac/MacApplication.h
--- a/Engine/Source/Runtime/Core/Public/Mac/MacApplication.h
+++ b/Engine/Source/Runtime/Core/Public/Mac/MacApplication.h
@@ -95,8 +95,8 @@ public:
 	void OnDragOut(FCocoaWindow* Window);
 	void OnDragDrop(FCocoaWindow* Window);
 
-	void OnWindowDidBecomeKey(FCocoaWindow* Window);
-	void OnWindowDidResignKey(FCocoaWindow* Window);
+	void OnWindowDidBecomeMain(FCocoaWindow* Window);
+	void OnWindowDidResignMain(FCocoaWindow* Window);
 	void OnWindowWillMove(FCocoaWindow* Window);
 	void OnWindowDidMove(FCocoaWindow* Window);
 	void OnWindowWillStartLiveResize(FCocoaWindow* Window);
