Further reduced number of toupper()/tolower() calls during file operations.

(Merging CL 2456477 from my branch to main, with changes).

diff --git a/Engine/Source/Runtime/Core/Private/Misc/Paths.cpp b/Engine/Source/Runtime/Core/Private/Misc/Paths.cpp
--- a/Engine/Source/Runtime/Core/Private/Misc/Paths.cpp
+++ b/Engine/Source/Runtime/Core/Private/Misc/Paths.cpp
@@ -465,27 +465,27 @@ bool FPaths::IsDrive(const FString& InPath)
 {
 	FString ConvertedPathString = InPath;
 
-	ConvertedPathString = ConvertedPathString.Replace(TEXT("/"), TEXT("\\"));
+	ConvertedPathString = ConvertedPathString.Replace(TEXT("/"), TEXT("\\"), ESearchCase::CaseSensitive);
 	const TCHAR* ConvertedPath= *ConvertedPathString;
 
 	// Does Path refer to a drive letter or BNC path?
-	if( FCString::Stricmp(ConvertedPath,TEXT(""))==0 )
+	if (ConvertedPath[0] == TCHAR(0))
 	{
 		return true;
 	}
-	else if( FChar::ToUpper(ConvertedPath[0])!=FChar::ToLower(ConvertedPath[0]) && ConvertedPath[1]==':' && ConvertedPath[2]==0 )
+	else if (FChar::ToUpper(ConvertedPath[0])!=FChar::ToLower(ConvertedPath[0]) && ConvertedPath[1]==TEXT(':') && ConvertedPath[2]==0)
 	{
 		return true;
 	}
-	else if( FCString::Stricmp(ConvertedPath,TEXT("\\"))==0 )
+	else if (FCString::Strcmp(ConvertedPath,TEXT("\\"))==0)
 	{
 		return true;
 	}
-	else if( FCString::Stricmp(ConvertedPath,TEXT("\\\\"))==0 )
+	else if (FCString::Strcmp(ConvertedPath,TEXT("\\\\"))==0)
 	{
 		return true;
 	}
-	else if( ConvertedPath[0]=='\\' && ConvertedPath[1]=='\\' && !FCString::Strchr(ConvertedPath+2,'\\') )
+	else if (ConvertedPath[0]==TEXT('\\') && ConvertedPath[1]==TEXT('\\') && !FCString::Strchr(ConvertedPath+2,TEXT('\\')))
 	{
 		return true;
 	}
@@ -566,10 +566,10 @@ bool FPaths::IsRelative(const FString& InPath)
 {
 	// The previous implementation of this function seemed to handle normalized and unnormalized paths, so this one does too for legacy reasons.
 
-	const bool IsRooted =	InPath.StartsWith(TEXT("\\\\"))	||												// "\\" for UNC or "network" paths.
-							InPath.StartsWith(TEXT("//"))	||												// Equivalent to "\\", considering normalization replaces "\\" with "//".
-							InPath.StartsWith(TEXT("\\"))	||												// Root of the current directory on Windows
-							InPath.StartsWith(TEXT("/"))	||												// Root of the current directory on Windows, root on UNIX-likes.
+	const bool IsRooted =	InPath.StartsWith(TEXT("\\\\"), ESearchCase::CaseSensitive)	||												// "\\" for UNC or "network" paths.
+							InPath.StartsWith(TEXT("//"), ESearchCase::CaseSensitive)	||												// Equivalent to "\\", considering normalization replaces "\\" with "//".
+							InPath.StartsWith(TEXT("\\"), ESearchCase::CaseSensitive)	||												// Root of the current directory on Windows
+							InPath.StartsWith(TEXT("/"), ESearchCase::CaseSensitive)	||												// Root of the current directory on Windows, root on UNIX-likes.
 							InPath.StartsWith(TEXT("root:/")) ||											// Feature packs use this
 							(InPath.Len() >= 2 && FChar::IsAlpha(InPath[0]) && InPath[1] == TEXT(':'));	// Starts with "<DriveLetter>:"
 
@@ -707,8 +707,8 @@ void FPaths::MakeStandardFilename(FString& InPath)
 
 void FPaths::MakePlatformFilename( FString& InPath )
 {
-	InPath.ReplaceInline(TEXT("\\"), FPlatformMisc::GetDefaultPathSeparator(), ESearchCase::CaseSensitive);
-	InPath.ReplaceInline(TEXT("/"), FPlatformMisc::GetDefaultPathSeparator(), ESearchCase::CaseSensitive);
+	InPath.ReplaceInline(TEXT( "\\" ), FPlatformMisc::GetDefaultPathSeparator(), ESearchCase::CaseSensitive);
+	InPath.ReplaceInline(TEXT( "/" ), FPlatformMisc::GetDefaultPathSeparator(), ESearchCase::CaseSensitive);
 }
 
 bool FPaths::MakePathRelativeTo( FString& InPath, const TCHAR* InRelativeTo )
@@ -848,8 +848,8 @@ bool FPaths::ValidatePath( const FString& InPath, FText* OutReason )
 	}
 
 	// Walk each part of the path looking for name errors
-	for(int32 StartPos = 0, EndPos = Standardized.Find(TEXT("/")); ; 
-		StartPos = EndPos + 1, EndPos = Standardized.Find(TEXT("/"), ESearchCase::IgnoreCase, ESearchDir::FromStart, StartPos)
+	for(int32 StartPos = 0, EndPos = Standardized.Find(TEXT("/"), ESearchCase::CaseSensitive); ; 
+		StartPos = EndPos + 1, EndPos = Standardized.Find(TEXT("/"), ESearchCase::CaseSensitive, ESearchDir::FromStart, StartPos)
 		)
 	{
 		const bool bIsLastPart = EndPos == INDEX_NONE;
