fix, probably legitimate, for UE-6890.

    On browser pointerlock the mouse "jumps" from outside the canvas to "locked and hidden and inside the canvas". This causes the WindowEnter event to be skipped. Add code to generate a fake WindowEnter event which fixes the slate window state.

diff --git a/Engine/Source/Runtime/Core/Private/HTML5/HTML5Application.cpp b/Engine/Source/Runtime/Core/Private/HTML5/HTML5Application.cpp
--- a/Engine/Source/Runtime/Core/Private/HTML5/HTML5Application.cpp
+++ b/Engine/Source/Runtime/Core/Private/HTML5/HTML5Application.cpp
@@ -18,6 +18,25 @@ EM_BOOL mouse_callback(int eventType, const EmscriptenMouseEvent *e, void *userD
 	return 1;
 }
 
+EM_BOOL pointerlockchange_callback(int eventType, const EmscriptenPointerlockChangeEvent *Event, void *userData)
+{
+	UE_LOG(LogHTML5Application, Verbose, TEXT("PointerLockChangedEvent: Active:%d"), Event->isActive );
+
+	static uint Prev = 0;
+	// Generate a fake WindowsEnter event when the pointerlock goes from inactive to active. 
+	if (Event->isActive && Prev == 0)
+	{
+		SDL_Event event;
+		SDL_zero(event);
+		event.type = SDL_WINDOWEVENT;
+		event.window.event = SDL_WINDOWEVENT_ENTER;
+		SDL_PushEvent(&event);
+	}
+	Prev = Event->isActive; 
+
+	return 1; 
+}
+
 #endif 
 
 static const uint32 MaxWarmUpTicks = 10; 
@@ -41,6 +60,7 @@ FHTML5Application::FHTML5Application()
 
  	// work around emscripten bug where deffered browser requests are not called if there are no callbacks.
 	emscripten_set_mousedown_callback("#canvas",0,1,mouse_callback);
+	emscripten_set_pointerlockchange_callback(0,0,true,pointerlockchange_callback);
 #endif 
 
 }
