Removed FMacWindow::Enable(). Solves a problem with OS X ignoring windows while we display a modal window.

diff --git a/Engine/Source/Runtime/Core/Private/Mac/CocoaTextView.cpp b/Engine/Source/Runtime/Core/Private/Mac/CocoaTextView.cpp
--- a/Engine/Source/Runtime/Core/Private/Mac/CocoaTextView.cpp
+++ b/Engine/Source/Runtime/Core/Private/Mac/CocoaTextView.cpp
@@ -43,14 +43,6 @@
 	}
 }
 
-/** Forward mouse events up to the window rather than through the responder chain - thus avoiding
- *	the hidden titlebar controls. Normal windows just use the responder chain as usual.
- */
-- (BOOL)acceptsFirstMouse:(NSEvent *)Event
-{
-	return YES;
-}
-
 - (void)mouseDown:(NSEvent *)theEvent
 {
 	SCOPED_AUTORELEASE_POOL;
diff --git a/Engine/Source/Runtime/Core/Private/Mac/MacWindow.cpp b/Engine/Source/Runtime/Core/Private/Mac/MacWindow.cpp
--- a/Engine/Source/Runtime/Core/Private/Mac/MacWindow.cpp
+++ b/Engine/Source/Runtime/Core/Private/Mac/MacWindow.cpp
@@ -475,14 +475,6 @@ void FMacWindow::SetOpacity( const float InOpacity )
 	}, UE4NilEventMode, true);
 }
 
-void FMacWindow::Enable( bool bEnable )
-{
-	MainThreadCall(^{
-		SCOPED_AUTORELEASE_POOL;
-		[WindowHandle setIgnoresMouseEvents: !bEnable];
-	}, UE4NilEventMode, true);
-}
-
 bool FMacWindow::IsPointInWindow( int32 X, int32 Y ) const
 {
 	SCOPED_AUTORELEASE_POOL;
diff --git a/Engine/Source/Runtime/Core/Public/Mac/MacWindow.h b/Engine/Source/Runtime/Core/Public/Mac/MacWindow.h
--- a/Engine/Source/Runtime/Core/Public/Mac/MacWindow.h
+++ b/Engine/Source/Runtime/Core/Public/Mac/MacWindow.h
@@ -63,8 +63,6 @@ public:
 
 	virtual void SetOpacity( const float InOpacity ) override;
 
-	virtual void Enable( bool bEnable ) override;
-
 	virtual bool IsPointInWindow( int32 X, int32 Y ) const override;
 
 	virtual int32 GetWindowBorderSize() const override;
