Don't report mouse move to Slate on Mac if the cursor is not over a Slate window.

diff --git a/Engine/Source/Runtime/Core/Private/Mac/MacApplication.cpp b/Engine/Source/Runtime/Core/Private/Mac/MacApplication.cpp
--- a/Engine/Source/Runtime/Core/Private/Mac/MacApplication.cpp
+++ b/Engine/Source/Runtime/Core/Private/Mac/MacApplication.cpp
@@ -663,7 +663,7 @@ void FMacApplication::ProcessNSEvent(NSEvent* const Event)
 				MacCursor->WarpCursor(HighPrecisionMousePos.X, HighPrecisionMousePos.Y);
 				MessageHandler->OnRawMouseMove(Delta.X, Delta.Y);
 			}
-			else
+			else if (CurrentEventWindow.IsValid())
 			{
 				NSPoint CursorPos = [NSEvent mouseLocation];
 				CursorPos.y--; // The y coordinate of the point returned by mouseLocation starts from a base of 1
@@ -681,7 +681,7 @@ void FMacApplication::ProcessNSEvent(NSEvent* const Event)
 
 				// Cocoa does not update NSWindow's frame until user stops dragging the window, so while window is being dragged, we calculate
 				// its position based on mouse move delta
-				if (CurrentEventWindow.IsValid() && DraggedWindow && DraggedWindow == NativeWindow)
+				if (DraggedWindow && DraggedWindow == NativeWindow)
 				{
 					const int32 X = FMath::TruncToInt(CurrentEventWindow->PositionX + MouseDelta.X);
 					const int32 Y = FMath::TruncToInt(CurrentEventWindow->PositionY + MouseDelta.Y);
