FDelegateHandle::IsValid and FDelegateHandle::Reset functions added.
Some existing code fixed up to use these.

diff --git a/Engine/Source/Editor/UnrealEd/Private/AutoReimport/FileCache.cpp b/Engine/Source/Editor/UnrealEd/Private/AutoReimport/FileCache.cpp
--- a/Engine/Source/Editor/UnrealEd/Private/AutoReimport/FileCache.cpp
+++ b/Engine/Source/Editor/UnrealEd/Private/AutoReimport/FileCache.cpp
@@ -145,7 +145,7 @@ void FFileCache::Destroy()
 
 void FFileCache::UnbindWatcher()
 {
-	if (WatcherDelegate == FDelegateHandle())
+	if (!WatcherDelegate.IsValid())
 	{
 		return;
 	}
@@ -158,7 +158,7 @@ void FFileCache::UnbindWatcher()
 		}
 	}
 
-	WatcherDelegate = FDelegateHandle();
+	WatcherDelegate.Reset();
 }
 
 TOptional<FDirectoryState> FFileCache::ReadCache() const
diff --git a/Engine/Source/Runtime/Core/Private/GenericPlatform/GenericPlatformTime.cpp b/Engine/Source/Runtime/Core/Private/GenericPlatform/GenericPlatformTime.cpp
--- a/Engine/Source/Runtime/Core/Private/GenericPlatform/GenericPlatformTime.cpp
+++ b/Engine/Source/Runtime/Core/Private/GenericPlatform/GenericPlatformTime.cpp
@@ -173,8 +173,8 @@ struct FCPUTimeDump
 
 		FTicker::GetCoreTicker().RemoveTicker( GetCPUTimeDelegateHandle );
 		FTicker::GetCoreTicker().RemoveTicker( CPUTimeDumpDelegateHandle );
-		GetCPUTimeDelegateHandle   = FDelegateHandle();
-		CPUTimeDumpDelegateHandle  = FDelegateHandle();
+		GetCPUTimeDelegateHandle .Reset();
+		CPUTimeDumpDelegateHandle.Reset();
 
 		if( Delay == 0 )
 		{
diff --git a/Engine/Source/Runtime/Core/Public/Delegates/IDelegateInstance.h b/Engine/Source/Runtime/Core/Public/Delegates/IDelegateInstance.h
--- a/Engine/Source/Runtime/Core/Public/Delegates/IDelegateInstance.h
+++ b/Engine/Source/Runtime/Core/Public/Delegates/IDelegateInstance.h
@@ -55,6 +55,16 @@ public:
 	{
 	}
 
+	bool IsValid() const
+	{
+		return ID != 0;
+	}
+
+	void Reset()
+	{
+		ID = 0;
+	}
+
 private:
 	friend bool operator==(const FDelegateHandle& Lhs, const FDelegateHandle& Rhs)
 	{
diff --git a/Engine/Source/Runtime/Online/OnlineSubsystemUtils/Private/OnlineSessionClient.cpp b/Engine/Source/Runtime/Online/OnlineSubsystemUtils/Private/OnlineSessionClient.cpp
--- a/Engine/Source/Runtime/Online/OnlineSubsystemUtils/Private/OnlineSessionClient.cpp
+++ b/Engine/Source/Runtime/Online/OnlineSubsystemUtils/Private/OnlineSessionClient.cpp
@@ -227,7 +227,7 @@ void UOnlineSessionClient::DestroyExistingSession_Impl(FDelegateHandle& OutResul
 	}
 	else
 	{
-		OutResult = FDelegateHandle();
+		OutResult.Reset();
 		Delegate.ExecuteIfBound(SessionName, true);
 	}
 }
