Core - Added GetBranch to the FEngineVersion class

diff --git a/Engine/Source/Runtime/Core/Public/Misc/EngineVersion.h b/Engine/Source/Runtime/Core/Public/Misc/EngineVersion.h
--- a/Engine/Source/Runtime/Core/Public/Misc/EngineVersion.h
+++ b/Engine/Source/Runtime/Core/Public/Misc/EngineVersion.h
@@ -34,6 +34,12 @@ public:
 	/** Serialization function */
 	friend CORE_API void operator<<(class FArchive &Ar, FEngineVersion &Version);
 
+	/** Returns the branch name corresponding to this version. */
+	const FString GetBranch() const
+	{
+		return Branch.Replace( TEXT( "+" ), TEXT( "/" ) );
+	}
+
 private:
 
 	/** Branch name. */
