CrashReport - More properties added to unified crash reporting

diff --git a/Engine/Source/Runtime/Core/Private/GenericPlatform/GenericPlatformContext.cpp b/Engine/Source/Runtime/Core/Private/GenericPlatform/GenericPlatformContext.cpp
--- a/Engine/Source/Runtime/Core/Private/GenericPlatform/GenericPlatformContext.cpp
+++ b/Engine/Source/Runtime/Core/Private/GenericPlatform/GenericPlatformContext.cpp
@@ -11,8 +11,11 @@ bool FGenericCrashContext::bIsInitialized = false;
 
 namespace NCachedCrashContextProperties
 {
-	static FString CrashDescriptionCommonBuffer;
+	static FString ExecutableName;
+	static FString PlatformName;
+	static FString PlatformNameIni;
 	static FString BaseDir;
+	static FString RootDir;
 	static FString EpicAccountId;
 	static FString MachineIdStr;
 	static FString OsVersion;
@@ -23,12 +26,19 @@ namespace NCachedCrashContextProperties
 	static FString CPUBrand;
 	static FString PrimaryGPUBrand;
 	static FString UserName;
+	static FString DefaultLocale;
+
+	static FString CrashGUID;
 }
 using namespace NCachedCrashContextProperties;
 
 void FGenericCrashContext::Initialize()
 {
+	ExecutableName = FPlatformProcess::ExecutableName();
+	PlatformName = FPlatformProperties::PlatformName();
+	PlatformNameIni = FPlatformProperties::IniPlatformName();
 	BaseDir = FPlatformProcess::BaseDir();
+	RootDir = FPlatformMisc::RootDir();
 	EpicAccountId = FPlatformMisc::GetEpicAccountId();
 	MachineIdStr = FPlatformMisc::GetMachineId().ToString( EGuidFormats::Digits );
 	FPlatformMisc::GetOSVersions(OsVersion,OsSubVersion);
@@ -39,33 +49,44 @@ void FGenericCrashContext::Initialize()
 	CPUBrand = FPlatformMisc::GetCPUBrand();
 	PrimaryGPUBrand = FPlatformMisc::GetPrimaryGPUBrand();
 	UserName = FPlatformProcess::UserName();
+	DefaultLocale = FPlatformMisc::GetDefaultLocale();
+
+	const FGuid Guid = FGuid::NewGuid();
+	CrashGUID = FString::Printf( TEXT( "UE4CC-%s-%s" ), *PlatformNameIni, *Guid.ToString( EGuidFormats::Digits ) );
 
 	bIsInitialized = true;
 }
 
-
-FGenericCrashContext::FGenericCrashContext() :
-CommonBuffer( CrashDescriptionCommonBuffer )
+FGenericCrashContext::FGenericCrashContext()
 {
 	CommonBuffer.Reserve( 32768 );
 }
 
 void FGenericCrashContext::SerializeContentToBuffer()
 {
+	// Must conform against:
+	// https://www.securecoding.cert.org/confluence/display/seccode/SIG30-C.+Call+only+asynchronous-safe+functions+within+signal+handlers
 	AddHeader();
 
+	BeginSection( TEXT( "RuntimeProperties" ) );
+	AddCrashProperty( TEXT( "Version" ), (int32)ECrashDescVersions::VER_2_AddedNewProperties );
 	// ProcessId will be used by the crash report client to get more information about the crashed process.
 	// It should be enough to open the process, get the memory stats, a list of loaded modules etc.
 	// It means that the crashed process needs to be alive as long as the crash report client is working on the process.
 	AddCrashProperty( TEXT( "ProcessId" ), FPlatformProcess::GetCurrentProcessId() );
 
 	// Add common crash properties.
-	AddCrashProperty( TEXT( "GameName" ), *FString::Printf( TEXT( "UE4-%s" ), FApp::GetGameName() ) );
-	AddCrashProperty( TEXT( "PlatformName" ), *FString( FPlatformProperties::PlatformName() ) );
+	AddCrashProperty( TEXT( "GameName" ), FApp::GetGameName() );
+	AddCrashProperty( TEXT( "ExecutableName" ), *ExecutableName );
+	AddCrashProperty( TEXT( "BuildConfiguration" ), EBuildConfigurations::ToString( FApp::GetBuildConfiguration() ) );
+
+	AddCrashProperty( TEXT( "PlatformName" ), *PlatformName );
+	AddCrashProperty( TEXT( "PlatformNameIni" ), *PlatformNameIni );
 	AddCrashProperty( TEXT( "EngineMode" ), FPlatformMisc::GetEngineMode() );
 	AddCrashProperty( TEXT( "EngineVersion" ), ENGINE_VERSION_STRING );
 	AddCrashProperty( TEXT( "CommandLine" ), FCommandLine::IsInitialized() ? FCommandLine::Get() : TEXT("") );
 	AddCrashProperty( TEXT( "LanguageLCID" ), FInternationalization::Get().GetCurrentCulture()->GetLCID() );
+	AddCrashProperty( TEXT( "DefaultLocale" ), *DefaultLocale );
 
 	AddCrashProperty( TEXT( "IsUE4Release" ), (int32)FRocketSupport::IsRocket() );
 
@@ -80,13 +101,14 @@ void FGenericCrashContext::SerializeContentToBuffer()
 
 
 	AddCrashProperty( TEXT( "BaseDir" ), *BaseDir );
+	AddCrashProperty( TEXT( "RootDir" ), *RootDir );
 	AddCrashProperty( TEXT( "MachineId" ), *MachineIdStr );
 	AddCrashProperty( TEXT( "EpicAccountId" ), *EpicAccountId );
 
 	AddCrashProperty( TEXT( "CallStack" ), TEXT( "" ) );
 	AddCrashProperty( TEXT( "SourceContext" ), TEXT( "" ) );
 	AddCrashProperty( TEXT( "UserDescription" ), TEXT( "" ) );
-	AddCrashProperty( TEXT( "ErrorMessage" ), (const TCHAR*)GErrorMessage );
+	AddCrashProperty( TEXT( "ErrorMessage" ), (const TCHAR*)GErrorMessage ); // GErrorMessage may be broken.
 
 	// @TODO yrx 2014-10-08 Move to the crash report client.
 	/*if( CanUseUnsafeAPI() )
@@ -96,14 +118,16 @@ void FGenericCrashContext::SerializeContentToBuffer()
 	*/
 
 	// Add misc stats.
-	AddCrashProperty( TEXT( "Misc.NumberOfCores" ), NumberOfCores );
-	AddCrashProperty( TEXT( "Misc.NumberOfCoresIncludingHyperthreads" ), NumberOfCoresIncludingHyperthreads );
-	AddCrashProperty( TEXT( "Misc.Is64bitOperatingSystem" ), (int32)FPlatformMisc::Is64bitOperatingSystem() );
+	AddCrashProperty( TEXT( "MiscNumberOfCores" ), NumberOfCores );
+	AddCrashProperty( TEXT( "MiscNumberOfCoresIncludingHyperthreads" ), NumberOfCoresIncludingHyperthreads );
+	AddCrashProperty( TEXT( "MiscIs64bitOperatingSystem" ), (int32)FPlatformMisc::Is64bitOperatingSystem() );
+
+	AddCrashProperty( TEXT( "MiscCPUVendor" ), *CPUVendor );
+	AddCrashProperty( TEXT( "MiscCPUBrand" ), *CPUBrand );
+	AddCrashProperty( TEXT( "MiscPrimaryGPUBrand" ), *PrimaryGPUBrand );
+	AddCrashProperty( TEXT( "MiscOSVersionMajor" ), *OsVersion );
+	AddCrashProperty( TEXT( "MiscOSVersionMinor" ), *OsSubVersion );
 
-	AddCrashProperty( TEXT( "Misc.CPUVendor" ), *CPUVendor );
-	AddCrashProperty( TEXT( "Misc.CPUBrand" ), *CPUBrand );
-	AddCrashProperty( TEXT( "Misc.PrimaryGPUBrand" ), *PrimaryGPUBrand );
-	AddCrashProperty( TEXT( "Misc.OSVersion" ), *FString::Printf( TEXT( "%s, %s" ), *OsVersion, *OsSubVersion ) );
 
 	// @TODO yrx 2014-10-08 Move to the crash report client.
 	/*if( CanUseUnsafeAPI() )
@@ -111,8 +135,8 @@ void FGenericCrashContext::SerializeContentToBuffer()
 		uint64 AppDiskTotalNumberOfBytes = 0;
 		uint64 AppDiskNumberOfFreeBytes = 0;
 		FPlatformMisc::GetDiskTotalAndFreeSpace( FPlatformProcess::BaseDir(), AppDiskTotalNumberOfBytes, AppDiskNumberOfFreeBytes );
-		AddCrashProperty( TEXT( "Misc.AppDiskTotalNumberOfBytes" ), AppDiskTotalNumberOfBytes );
-		AddCrashProperty( TEXT( "Misc.AppDiskNumberOfFreeBytes" ), AppDiskNumberOfFreeBytes );
+		AddCrashProperty( TEXT( "MiscAppDiskTotalNumberOfBytes" ), AppDiskTotalNumberOfBytes );
+		AddCrashProperty( TEXT( "MiscAppDiskNumberOfFreeBytes" ), AppDiskNumberOfFreeBytes );
 	}*/
 
 	// FPlatformMemory::GetConstants is called in the GCreateMalloc, so we can assume it is always valid.
@@ -120,43 +144,62 @@ void FGenericCrashContext::SerializeContentToBuffer()
 		// Add memory stats.
 		const FPlatformMemoryConstants& MemConstants = FPlatformMemory::GetConstants();
 
-		AddCrashProperty( TEXT( "MemoryStats.TotalPhysical" ), (uint64)MemConstants.TotalPhysical );
-		AddCrashProperty( TEXT( "MemoryStats.TotalVirtual" ), (uint64)MemConstants.TotalVirtual );
-		AddCrashProperty( TEXT( "MemoryStats.PageSize" ), (uint64)MemConstants.PageSize );
-		AddCrashProperty( TEXT( "MemoryStats.TotalPhysicalGB" ), MemConstants.TotalPhysicalGB );
+		AddCrashProperty( TEXT( "MemoryStatsTotalPhysical" ), (uint64)MemConstants.TotalPhysical );
+		AddCrashProperty( TEXT( "MemoryStatsTotalVirtual" ), (uint64)MemConstants.TotalVirtual );
+		AddCrashProperty( TEXT( "MemoryStatsPageSize" ), (uint64)MemConstants.PageSize );
+		AddCrashProperty( TEXT( "MemoryStatsTotalPhysicalGB" ), MemConstants.TotalPhysicalGB );
 	}
 
 	// @TODO yrx 2014-10-08 Move to the crash report client.
 	/*if( CanUseUnsafeAPI() )
 	{
 		const FPlatformMemoryStats MemStats = FPlatformMemory::GetStats();
-		AddCrashProperty( TEXT( "MemoryStats.AvailablePhysical" ), (uint64)MemStats.AvailablePhysical );
-		AddCrashProperty( TEXT( "MemoryStats.AvailableVirtual" ), (uint64)MemStats.AvailableVirtual );
-		AddCrashProperty( TEXT( "MemoryStats.UsedPhysical" ), (uint64)MemStats.UsedPhysical );
-		AddCrashProperty( TEXT( "MemoryStats.PeakUsedPhysical" ), (uint64)MemStats.PeakUsedPhysical );
-		AddCrashProperty( TEXT( "MemoryStats.UsedVirtual" ), (uint64)MemStats.UsedVirtual );
-		AddCrashProperty( TEXT( "MemoryStats.PeakUsedVirtual" ), (uint64)MemStats.PeakUsedVirtual );
-		AddCrashProperty( TEXT( "MemoryStats.bIsOOM" ), (int32)FPlatformMemory::bIsOOM );
+		AddCrashProperty( TEXT( "MemoryStatsAvailablePhysical" ), (uint64)MemStats.AvailablePhysical );
+		AddCrashProperty( TEXT( "MemoryStatsAvailableVirtual" ), (uint64)MemStats.AvailableVirtual );
+		AddCrashProperty( TEXT( "MemoryStatsUsedPhysical" ), (uint64)MemStats.UsedPhysical );
+		AddCrashProperty( TEXT( "MemoryStatsPeakUsedPhysical" ), (uint64)MemStats.PeakUsedPhysical );
+		AddCrashProperty( TEXT( "MemoryStatsUsedVirtual" ), (uint64)MemStats.UsedVirtual );
+		AddCrashProperty( TEXT( "MemoryStatsPeakUsedVirtual" ), (uint64)MemStats.PeakUsedVirtual );
+		AddCrashProperty( TEXT( "MemoryStatsbIsOOM" ), (int32)FPlatformMemory::bIsOOM );
 	}*/
 
-	// @TODO yrx 2014-09-10 
 	//Architecture
 	//CrashedModuleName
 	//LoadedModules
-	AddFooter();
-
+	EndSection( TEXT( "RuntimeProperties" ) );
+	
 	// Add platform specific properties.
-	CommonBuffer += TEXT( "<PlatformSpecificProperties>" ) LINE_TERMINATOR;
+	BeginSection( TEXT("PlatformProperties") );
 	AddPlatformSpecificProperties();
-	CommonBuffer += TEXT( "</PlatformSpecificProperties>" ) LINE_TERMINATOR;
+	EndSection( TEXT( "PlatformProperties" ) );
+
+	AddFooter();
+}
+
+const FString& FGenericCrashContext::GetUniqueCrashName()
+{
+	return CrashGUID;
 }
 
 void FGenericCrashContext::SerializeAsXML( const TCHAR* Filename )
 {
 	SerializeContentToBuffer();
-	// @TODO yrx 2014-10-08 We need mechanism to store the file that is aware of crashes and OOMs.
-	// Probably will be replaced with some kind of Inter-Process Communication
-	FFileHelper::SaveStringToFile( CommonBuffer, Filename, FFileHelper::EEncodingOptions::ForceUTF8 );
+	// Use OS build-in functionality instead.
+	FFileHelper::SaveStringToFile( CommonBuffer, Filename, FFileHelper::EEncodingOptions::AutoDetect );
+}
+
+void FGenericCrashContext::AddCrashProperty( const TCHAR* PropertyName, const TCHAR* PropertyValue )
+{
+	CommonBuffer += TEXT( "<" );
+	CommonBuffer += PropertyName;
+	CommonBuffer += TEXT( ">" );
+
+	CommonBuffer += PropertyValue;
+
+	CommonBuffer += TEXT( "</" );
+	CommonBuffer += PropertyName;
+	CommonBuffer += TEXT( ">" );
+	CommonBuffer += LINE_TERMINATOR;
 }
 
 void FGenericCrashContext::AddPlatformSpecificProperties()
@@ -168,26 +211,26 @@ void FGenericCrashContext::AddPlatformSpecificProperties()
 void FGenericCrashContext::AddHeader()
 {
 	CommonBuffer += TEXT( "<?xml version=\"1.0\" encoding=\"UTF-8\"?>" ) LINE_TERMINATOR;
-	CommonBuffer += TEXT( "<RuntimeCrashDescription>" ) LINE_TERMINATOR;
-
-	AddCrashProperty( TEXT( "Version" ), (int32)ECrashDescVersions::VER_2_AddedNewProperties );
+	BeginSection( TEXT("FGenericCrashContext") );
 }
 
 void FGenericCrashContext::AddFooter()
 {
-	CommonBuffer += TEXT( "</RuntimeCrashDescription>" ) LINE_TERMINATOR;
+	EndSection( TEXT( "FGenericCrashContext" ) );
 }
-
-void FGenericCrashContext::AddCrashProperty( const TCHAR* PropertyName, const TCHAR* PropertyValue )
+\
+void FGenericCrashContext::BeginSection( const TCHAR* SectionName )
 {
 	CommonBuffer += TEXT( "<" );
-	CommonBuffer += PropertyName;
+	CommonBuffer += SectionName;
 	CommonBuffer += TEXT( ">" );
+	CommonBuffer += LINE_TERMINATOR;
+}
 
-	CommonBuffer += PropertyValue;
-
+void FGenericCrashContext::EndSection( const TCHAR* SectionName )
+{
 	CommonBuffer += TEXT( "</" );
-	CommonBuffer += PropertyName;
+	CommonBuffer += SectionName;
 	CommonBuffer += TEXT( ">" );
 	CommonBuffer += LINE_TERMINATOR;
 }
diff --git a/Engine/Source/Runtime/Core/Public/GenericPlatform/GenericPlatformContext.h b/Engine/Source/Runtime/Core/Public/GenericPlatform/GenericPlatformContext.h
--- a/Engine/Source/Runtime/Core/Public/GenericPlatform/GenericPlatformContext.h
+++ b/Engine/Source/Runtime/Core/Public/GenericPlatform/GenericPlatformContext.h
@@ -49,6 +49,11 @@ public:
 		return CommonBuffer;
 	}
 
+	/**
+	 * @return a globally unique crash name.
+	 */
+	const FString& GetUniqueCrashName();
+
 	/** Serializes crash's informations to the specified filename. */
 	void SerializeAsXML( const TCHAR* Filename );
 
@@ -72,8 +77,11 @@ private:
 	/** Writes footer to the buffer. */
 	void AddFooter();
 
+	void BeginSection( const TCHAR* SectionName );
+	void EndSection( const TCHAR* SectionName );
+
 	/** The buffer used to store the crash's properties. */
-	FString& CommonBuffer;
+	FString CommonBuffer;
 
 	// FNoncopyable
 	FGenericCrashContext( const FGenericCrashContext& );
