Don't wait for main thread in FMacPlatformMisc::ActivateApplication()

diff --git a/Engine/Source/Runtime/Core/Private/Mac/MacPlatformMisc.cpp b/Engine/Source/Runtime/Core/Private/Mac/MacPlatformMisc.cpp
--- a/Engine/Source/Runtime/Core/Private/Mac/MacPlatformMisc.cpp
+++ b/Engine/Source/Runtime/Core/Private/Mac/MacPlatformMisc.cpp
@@ -345,7 +345,7 @@ void FMacPlatformMisc::ActivateApplication()
 {
 	MainThreadCall(^{
 		[NSApp activateIgnoringOtherApps:YES];
-	});
+	}, NSDefaultRunLoopMode, false);
 }
 
 bool FMacPlatformMisc::ControlScreensaver(EScreenSaverAction Action)
