Prevent OS X from putting UE4 apps that don't use UI (for example the editor runnign a commandlet) into App Nap mode.

diff --git a/Engine/Source/Runtime/Core/Private/Mac/MacPlatformMisc.cpp b/Engine/Source/Runtime/Core/Private/Mac/MacPlatformMisc.cpp
--- a/Engine/Source/Runtime/Core/Private/Mac/MacPlatformMisc.cpp
+++ b/Engine/Source/Runtime/Core/Private/Mac/MacPlatformMisc.cpp
@@ -211,6 +211,7 @@ static MacApplicationInfo GMacAppInfo;
 
 UpdateCachedMacMenuStateProc FMacPlatformMisc::UpdateCachedMacMenuState = nullptr;
 bool FMacPlatformMisc::bChachedMacMenuStateNeedsUpdate = true;
+id<NSObject> FMacPlatformMisc::CommandletActivity = nil;
 
 void FMacPlatformMisc::PlatformPreInit()
 {
@@ -313,6 +314,21 @@ void FMacPlatformMisc::PlatformPostInit(bool ShowSplashScreen)
 
 		UpdateWindowMenu();
 	}
+
+	if (!MacApplication)
+	{
+		// No MacApplication means that app is a dedicated server, commandline tool or the editor running a commandlet. In these cases we don't want OS X to put our app into App Nap mode.
+		CommandletActivity = [[NSProcessInfo processInfo] beginActivityWithOptions:NSActivityUserInitiated reason:IsRunningCommandlet() ? @"Running commandlet" : @"Running dedicated server"];
+	}
+}
+
+void FMacPlatformMisc::PlatformTearDown()
+{
+	if (CommandletActivity)
+	{
+		[[NSProcessInfo processInfo] endActivity:CommandletActivity];
+		CommandletActivity = nil;
+	}
 }
 
 void FMacPlatformMisc::UpdateWindowMenu()
diff --git a/Engine/Source/Runtime/Core/Public/Mac/MacPlatformMisc.h b/Engine/Source/Runtime/Core/Public/Mac/MacPlatformMisc.h
--- a/Engine/Source/Runtime/Core/Public/Mac/MacPlatformMisc.h
+++ b/Engine/Source/Runtime/Core/Public/Mac/MacPlatformMisc.h
@@ -18,6 +18,7 @@ struct CORE_API FMacPlatformMisc : public FGenericPlatformMisc
 	static void PlatformPreInit();
 	static void PlatformInit();
 	static void PlatformPostInit(bool ShowSplashScreen = false);
+	static void PlatformTearDown();
 	static class GenericApplication* CreateApplication();
 	static void GetEnvironmentVariable(const TCHAR* VariableName, TCHAR* Result, int32 ResultLength);
 	static void SetEnvironmentVar(const TCHAR* VariableName, const TCHAR* Value);
@@ -173,6 +174,8 @@ struct CORE_API FMacPlatformMisc : public FGenericPlatformMisc
 	static FString GetOperatingSystemId();
 
 	static bool bChachedMacMenuStateNeedsUpdate;
+
+	static id<NSObject> CommandletActivity;
 };
 
 #ifdef __OBJC__
