Safari fullscreen fix. UE-10530
Implement LaunchURL for HTML5. UE-10852
Possible fix for bogus mouse coordinates in Safari

diff --git a/Engine/Build/HTML5/Game.html.template b/Engine/Build/HTML5/Game.html.template
--- a/Engine/Build/HTML5/Game.html.template
+++ b/Engine/Build/HTML5/Game.html.template
@@ -83,7 +83,7 @@ video {
     </div>
 
     <div class="emscripten">
-      <input type="button" value="Fullscreen" onclick="Module.requestFullScreen(false)">
+      <input type="button" value="Fullscreen" id="fullscreen_request">
       <input type="button" value="Pause" onclick="Module['pauseMainLoop']();">
       <input type="button" value="Resume" onclick="Module['resumeMainLoop']();">
       <!--input type="button" value="Quit" onclick="__exit(0)" -->
@@ -125,7 +125,7 @@ var UE4 = {
     } catch(e) {
         return function() {}
     }
-  }
+  },
 };
 
 function resizeListener(width, height) {
diff --git a/Engine/Source/Runtime/Core/Private/HTML5/HTML5InputInterface.cpp b/Engine/Source/Runtime/Core/Private/HTML5/HTML5InputInterface.cpp
--- a/Engine/Source/Runtime/Core/Private/HTML5/HTML5InputInterface.cpp
+++ b/Engine/Source/Runtime/Core/Private/HTML5/HTML5InputInterface.cpp
@@ -60,10 +60,26 @@ TSharedRef< FHTML5InputInterface > FHTML5InputInterface::Create(  const TSharedR
 	return MakeShareable( new FHTML5InputInterface( InMessageHandler, InCursor ) );
 }
 
+EM_BOOL mouse_move_callback(int evType, const EmscriptenMouseEvent* evt, void* userData) {
+	/* rescale (in case canvas is being scaled)*/
+	double client_w, client_h, xscale, yscale;
+	int canvas_w, canvas_h, canvas_fs;
+	emscripten_get_canvas_size(&canvas_w, &canvas_h, &canvas_fs);
+	emscripten_get_element_css_size(NULL, &client_w, &client_h);
+	xscale = canvas_w/client_w;
+	yscale = canvas_h/client_h;
+	int calc_x = (int)(evt->canvasX*xscale + .5);
+	int calc_y = (int)(evt->canvasY*yscale + .5);
+	UE_LOG(LogHTML5Input, Verbose, TEXT("MouseMoveCB Pos(%d or %d, %d or %d) XRel:%d YRel:%d"), evt->canvasX, calc_x, evt->canvasY, calc_y, evt->movementX, evt->movementY);
+	(*((const TSharedPtr< ICursor >*)userData))->SetPosition(calc_x, calc_y);
+	return 0;
+}
+
 FHTML5InputInterface::FHTML5InputInterface( const TSharedRef< FGenericApplicationMessageHandler >& InMessageHandler, const TSharedPtr< ICursor >& InCursor )
 	: MessageHandler( InMessageHandler )
 	, Cursor( InCursor )
 {
+	emscripten_set_mousemove_callback("canvas", (void*)&Cursor, 1, mouse_move_callback);
 }
 
 void FHTML5InputInterface::SetMessageHandler( const TSharedRef< FGenericApplicationMessageHandler >& InMessageHandler )
@@ -236,7 +252,7 @@ void FHTML5InputInterface::Tick(float DeltaTime, const SDL_Event& Event,TSharedR
 			break; 
 		case SDL_MOUSEMOTION:
 			{
-				Cursor->SetPosition(Event.motion.x, Event.motion.y);
+				//Cursor->SetPosition(Event.motion.x, Event.motion.y);
 				MessageHandler->OnRawMouseMove(Event.motion.xrel, Event.motion.yrel);
 				MessageHandler->OnMouseMove(); 
 				UE_LOG(LogHTML5Input, Verbose, TEXT("MouseMotion Pos(%d, %d) XRel:%d YRel:%d"), Event.motion.x, Event.motion.y, Event.motion.xrel, Event.motion.yrel);
diff --git a/Engine/Source/Runtime/Core/Private/HTML5/HTML5PlatformProcess.cpp b/Engine/Source/Runtime/Core/Private/HTML5/HTML5PlatformProcess.cpp
--- a/Engine/Source/Runtime/Core/Private/HTML5/HTML5PlatformProcess.cpp
+++ b/Engine/Source/Runtime/Core/Private/HTML5/HTML5PlatformProcess.cpp
@@ -88,3 +88,9 @@ bool FHTML5PlatformProcess::SupportsMultithreading()
 {
    return false; 
 }
+
+void FHTML5PlatformProcess::LaunchURL(const TCHAR* URL, const TCHAR* Parms, FString* Error) 
+{
+	auto TmpURL = StringCast<ANSICHAR>(URL);
+	EM_ASM_ARGS({var InUrl = Pointer_stringify($0); console.log("Opening "+InUrl); window.open(InUrl);}, (ANSICHAR*)TmpURL.Get());
+}
diff --git a/Engine/Source/Runtime/Core/Public/HTML5/HTML5PlatformProcess.h b/Engine/Source/Runtime/Core/Public/HTML5/HTML5PlatformProcess.h
--- a/Engine/Source/Runtime/Core/Public/HTML5/HTML5PlatformProcess.h
+++ b/Engine/Source/Runtime/Core/Public/HTML5/HTML5PlatformProcess.h
@@ -36,6 +36,7 @@ struct CORE_API FHTML5PlatformProcess : public FGenericPlatformProcess
 	static class FEvent* CreateSynchEvent(bool bIsManualReset = 0);
 	static class FRunnableThread* CreateRunnableThread();
 	static bool SupportsMultithreading();
+	static void LaunchURL( const TCHAR* URL, const TCHAR* Parms, FString* Error );
 };
 
 typedef FHTML5PlatformProcess FPlatformProcess;
diff --git a/Engine/Source/Runtime/OpenGLDrv/Private/HTML5/HTML5OpenGL.cpp b/Engine/Source/Runtime/OpenGLDrv/Private/HTML5/HTML5OpenGL.cpp
--- a/Engine/Source/Runtime/OpenGLDrv/Private/HTML5/HTML5OpenGL.cpp
+++ b/Engine/Source/Runtime/OpenGLDrv/Private/HTML5/HTML5OpenGL.cpp
@@ -4,6 +4,7 @@
 #include <SDL.h>
 #if !PLATFORM_HTML5_WIN32
 #include <emscripten.h>
+#include <html5.h>
 #endif
 
 DEFINE_LOG_CATEGORY_STATIC(LogHTML5OpenGL, Log, All);
@@ -151,10 +152,15 @@ struct FPlatformOpenGLContext
 	}
 };
 
+extern "C"
+EM_BOOL request_fullscreen_callback(int eventType, const EmscriptenMouseEvent* evt, void* user);
+
 struct FPlatformOpenGLDevice
 {
 	FPlatformOpenGLDevice()
 	{
+		emscripten_set_click_callback("fullscreen_request", nullptr, true, request_fullscreen_callback);
+
 		SharedContext = new FPlatformOpenGLContext; 
 
 		SDL_GL_SetAttribute(SDL_GL_CONTEXT_EGL, 1);
@@ -336,16 +342,30 @@ void PlatformRestoreDesktopDisplayMode()
 {
 }
 
+#if PLATFORM_HTML5_BROWSER
 extern "C"
 {
-#if PLATFORM_HTML5_BROWSER
 	// callback from javascript. 
 	void resize_game(int w, int h)
 	{
       	// to-do: remove this separate code path. 
 		UE_LOG(LogHTML5OpenGL, Verbose, TEXT("resize_game(%d, %d) callback"), w, h);
 	}
-#endif 
+
+	EM_BOOL request_fullscreen_callback(int eventType, const EmscriptenMouseEvent* evt, void* user)
+	{
+		
+		EmscriptenFullscreenStrategy FSStrat;
+		FMemory::Memzero(FSStrat);
+		FSStrat.scaleMode = EMSCRIPTEN_FULLSCREEN_SCALE_STRETCH;//EMSCRIPTEN_FULLSCREEN_SCALE_ASPECT;// : EMSCRIPTEN_FULLSCREEN_SCALE_STRETCH;
+		FSStrat.canvasResolutionScaleMode = EMSCRIPTEN_FULLSCREEN_CANVAS_SCALE_HIDEF;
+		FSStrat.filteringMode = EMSCRIPTEN_FULLSCREEN_FILTERING_DEFAULT;
+		emscripten_request_fullscreen_strategy("canvas", true, &FSStrat);
+		
+		//emscripten_request_fullscreen("canvas", true);
+		return 0;
+	}
 }
+#endif
 
 
