CrashReport - Compilation fix

diff --git a/Engine/Source/Runtime/Core/Private/GenericPlatform/GenericPlatformContext.cpp b/Engine/Source/Runtime/Core/Private/GenericPlatform/GenericPlatformContext.cpp
--- a/Engine/Source/Runtime/Core/Private/GenericPlatform/GenericPlatformContext.cpp
+++ b/Engine/Source/Runtime/Core/Private/GenericPlatform/GenericPlatformContext.cpp
@@ -218,7 +218,7 @@ void FGenericCrashContext::AddFooter()
 {
 	EndSection( TEXT( "FGenericCrashContext" ) );
 }
-\
+
 void FGenericCrashContext::BeginSection( const TCHAR* SectionName )
 {
 	CommonBuffer += TEXT( "<" );
diff --git a/Engine/Source/Runtime/Core/Public/GenericPlatform/GenericPlatformContext.h b/Engine/Source/Runtime/Core/Public/GenericPlatform/GenericPlatformContext.h
--- a/Engine/Source/Runtime/Core/Public/GenericPlatform/GenericPlatformContext.h
+++ b/Engine/Source/Runtime/Core/Public/GenericPlatform/GenericPlatformContext.h
@@ -4,6 +4,39 @@
 
 class FString;
 
+/** 
+ * Symbol information associated with a program counter. 
+ * FString version.
+ * To be used by external tools.
+ */
+struct FProgramCounterSymbolInfoEx
+{
+	/** Module name. */
+	FString	ModuleName;
+
+	/** Function name. */
+	FString	FunctionName;
+
+	/** Filename. */
+	FString	Filename;
+
+	/** Line number in file. */
+	uint32	LineNumber;
+
+	/** Symbol displacement of address.	*/
+	uint64	SymbolDisplacement;
+
+	/** Program counter offset into module. */
+	uint64	OffsetInModule;
+
+	/** Program counter. */
+	uint64	ProgramCounter;
+
+	/** Default constructor. */
+	FProgramCounterSymbolInfoEx( FString InModuleName, FString InFunctionName, FString InFilename, uint32 InLineNumber, uint64 InSymbolDisplacement, uint64 InOffsetInModule, uint64 InProgramCounter );
+};
+
+
 /** Enumerates crash description versions. */
 enum class ECrashDescVersions
 {
