CrashReport - Build version has three parameters now

diff --git a/Engine/Source/Runtime/Core/Private/Linux/LinuxPlatformCrashContext.cpp b/Engine/Source/Runtime/Core/Private/Linux/LinuxPlatformCrashContext.cpp
--- a/Engine/Source/Runtime/Core/Private/Linux/LinuxPlatformCrashContext.cpp
+++ b/Engine/Source/Runtime/Core/Private/Linux/LinuxPlatformCrashContext.cpp
@@ -597,7 +597,7 @@ void GenerateWindowsErrorReport(const FString & WERPath)
 		WriteLine(ReportFile, TEXT("\t<ProblemSignatures>"));
 		WriteLine(ReportFile, TEXT("\t\t<EventType>APPCRASH</EventType>"));
 		WriteLine(ReportFile, *FString::Printf(TEXT("\t\t<Parameter0>UE4-%s</Parameter0>"), FApp::GetGameName()));
-		WriteLine(ReportFile, *FString::Printf(TEXT("\t\t<Parameter1>%d.%d.%d.0</Parameter1>"), ENGINE_MAJOR_VERSION, ENGINE_MINOR_VERSION, ENGINE_PATCH_VERSION));
+		WriteLine(ReportFile, *FString::Printf(TEXT("\t\t<Parameter1>%d.%d.%d</Parameter1>"), ENGINE_MAJOR_VERSION, ENGINE_MINOR_VERSION, ENGINE_PATCH_VERSION));
 		WriteLine(ReportFile, TEXT("\t\t<Parameter2>0</Parameter2>"));													// FIXME: supply valid?
 		WriteLine(ReportFile, TEXT("\t\t<Parameter3>Unknown Fault Module</Parameter3>"));										// FIXME: supply valid?
 		WriteLine(ReportFile, TEXT("\t\t<Parameter4>0.0.0.0</Parameter4>"));													// FIXME: supply valid?
diff --git a/Engine/Source/Runtime/Core/Private/Mac/MacPlatformMisc.cpp b/Engine/Source/Runtime/Core/Private/Mac/MacPlatformMisc.cpp
--- a/Engine/Source/Runtime/Core/Private/Mac/MacPlatformMisc.cpp
+++ b/Engine/Source/Runtime/Core/Private/Mac/MacPlatformMisc.cpp
@@ -1326,7 +1326,7 @@ void FMacCrashContext::GenerateWindowsErrorReport(char const* WERPath) const
 		WriteUTF16String(ReportFile, ItoTCHAR(ENGINE_MINOR_VERSION, 10));
 		WriteUTF16String(ReportFile, TEXT("."));
 		WriteUTF16String(ReportFile, ItoTCHAR(ENGINE_PATCH_VERSION, 10));
-		WriteLine(ReportFile, TEXT(".0</Parameter1>"));
+		WriteLine(ReportFile, TEXT("</Parameter1>"));
 
 		// App time stamp
 		WriteLine(ReportFile, TEXT("\t\t<Parameter2>528f2d37</Parameter2>"));													// FIXME: supply valid?
diff --git a/Engine/Source/Runtime/Core/Private/Windows/WindowsPlatformCrashContext.cpp b/Engine/Source/Runtime/Core/Private/Windows/WindowsPlatformCrashContext.cpp
--- a/Engine/Source/Runtime/Core/Private/Windows/WindowsPlatformCrashContext.cpp
+++ b/Engine/Source/Runtime/Core/Private/Windows/WindowsPlatformCrashContext.cpp
@@ -91,7 +91,7 @@ void GetCrashDescription(WER_REPORT_INFORMATION& ReportInformation)
  */
 void GetModuleVersion( TCHAR* ModuleName, TCHAR* StringBuffer, DWORD MaxSize )
 {
-	StringCchCopy( StringBuffer, MaxSize, TEXT( "0.0.0.0" ) );
+	StringCchCopy( StringBuffer, MaxSize, TEXT( "0.0.0" ) );
 	
 	DWORD Handle = 0;
 	DWORD InfoSize = GetFileVersionInfoSize( ModuleName, &Handle );
@@ -107,8 +107,8 @@ void GetModuleVersion( TCHAR* ModuleName, TCHAR* StringBuffer, DWORD MaxSize )
 			UINT InfoLength = 0;
 			if( VerQueryValue( VersionInfo.GetData(), TEXT( "\\" ), ( void** )&FixedFileInfo, &InfoLength ) )
 			{
-				StringCchPrintf( StringBuffer, MaxSize, TEXT( "%u.%u.%u.%u" ), 
-					HIWORD( FixedFileInfo->dwProductVersionMS ), LOWORD( FixedFileInfo->dwProductVersionMS ), HIWORD( FixedFileInfo->dwProductVersionLS ), LOWORD( FixedFileInfo->dwProductVersionLS ) );
+				StringCchPrintf( StringBuffer, MaxSize, TEXT( "%u.%u.%u" ), 
+					HIWORD( FixedFileInfo->dwProductVersionMS ), LOWORD( FixedFileInfo->dwProductVersionMS ), HIWORD( FixedFileInfo->dwProductVersionLS ) );
 			}
 		}
 	}
