CrashReport - Introduces a new runtime crash context (partially)

diff --git a/Engine/Source/Runtime/Core/Private/Windows/WindowsPlatformCrashContext.cpp b/Engine/Source/Runtime/Core/Private/Windows/WindowsPlatformCrashContext.cpp
--- a/Engine/Source/Runtime/Core/Private/Windows/WindowsPlatformCrashContext.cpp
+++ b/Engine/Source/Runtime/Core/Private/Windows/WindowsPlatformCrashContext.cpp
@@ -61,7 +61,9 @@ bool WriteMinidump(const TCHAR* Path, LPEXCEPTION_POINTERS ExceptionInfo)
 	CrashContextStreamInformation.UserStreamCount = 1;
 	CrashContextStreamInformation.UserStreamArray = &CrashContextStream;
 
-	const BOOL Result = MiniDumpWriteDump(GetCurrentProcess(), GetCurrentProcessId(), FileHandle, MiniDumpNormal, &DumpExceptionInfo, &CrashContextStreamInformation, NULL);
+	MINIDUMP_TYPE MinidumpType = MiniDumpNormal;//(MINIDUMP_TYPE)(MiniDumpWithPrivateReadWriteMemory|MiniDumpWithDataSegs|MiniDumpWithHandleData|MiniDumpWithFullMemoryInfo|MiniDumpWithThreadInfo|MiniDumpWithUnloadedModules);
+
+	const BOOL Result = MiniDumpWriteDump(GetCurrentProcess(), GetCurrentProcessId(), FileHandle, MinidumpType, &DumpExceptionInfo, &CrashContextStreamInformation, NULL);
 	CloseHandle(FileHandle);
 
 	return Result == TRUE;
@@ -192,33 +194,6 @@ void SetReportParameters( HREPORT ReportHandle, EXCEPTION_POINTERS* ExceptionInf
 	Result = WerReportSetParameter( ReportHandle, WER_P9, TEXT( "BranchBaseDir" ), StringBuffer );
 }
 
-/**
- * Add a minidump to the Windows Error Report
- *
- * Note this has to be a minidump (and not a microdump) for the dbgeng functions to work correctly
- */
-void AddMiniDump(HREPORT ReportHandle, EXCEPTION_POINTERS* ExceptionInfo)
-{
-	// Add GetMinidumpFile
-	const FString MinidumpFileName = FString::Printf(TEXT("%sDump%d.dmp"), *FPaths::GameLogDir(), FDateTime::UtcNow().GetTicks());
-	
-	if (WriteMinidump(*MinidumpFileName, ExceptionInfo))
-	{
-		WerReportAddFile(ReportHandle, *MinidumpFileName, WerFileTypeMinidump, WER_FILE_ANONYMOUS_DATA);
-	}
-}
-
-/** 
- * Add miscellaneous files to the report. Currently the log and the video file
- */
-void AddMiscFiles( HREPORT ReportHandle )
-{
-	FString LogFileName = FPaths::GameLogDir() / FApp::GetGameName() + TEXT( ".log" );
-	WerReportAddFile( ReportHandle, *LogFileName, WerFileTypeOther, WER_FILE_ANONYMOUS_DATA ); 
-
-	FString CrashVideoPath = FPaths::GameLogDir() / TEXT( "CrashVideo.avi" );
-	WerReportAddFile( ReportHandle, *CrashVideoPath, WerFileTypeOther, WER_FILE_ANONYMOUS_DATA );
-}
 
 /**
  * Force WER queuing on or off
@@ -312,11 +287,27 @@ int32 ReportCrashUsingCrashReportClient(EXCEPTION_POINTERS* ExceptionInfo, const
 			// Set the standard set of a crash parameters
 			SetReportParameters( ReportHandle, ExceptionInfo, ErrorMessage );
 
-			// Add a manually generated minidump
-			AddMiniDump( ReportHandle, ExceptionInfo );
+			{
+				// No super safe due to dynamic memory allocations, but at least enables new functionality.
+				// Introduces a new runtime crash context. Will replace all Windows related crash reporting.
+				FPlatformCrashContext CrashContext;
+
+				const FString CrashContextXMLPath = FPaths::Combine( *FPaths::GameLogDir(), *CrashContext.GetUniqueCrashName(), FPlatformCrashContext::CrashContextRuntimeXMLNameW );
+				CrashContext.SerializeAsXML( *CrashContextXMLPath );
+				WerReportAddFile( ReportHandle, *CrashContextXMLPath, WerFileTypeOther, WER_FILE_ANONYMOUS_DATA );
+
+				const FString MinidumpFileName = FPaths::Combine( *FPaths::GameLogDir(), *CrashContext.GetUniqueCrashName(), TEXT("UE4Minidump.dmp") );
+				if( WriteMinidump( *MinidumpFileName, ExceptionInfo ) )
+				{
+					WerReportAddFile( ReportHandle, *MinidumpFileName, WerFileTypeMinidump, WER_FILE_ANONYMOUS_DATA );
+				}
 
-			// Add the log and video
-			AddMiscFiles( ReportHandle );
+				const FString LogFileName = FPaths::GameLogDir() / FApp::GetGameName() + TEXT( ".log" );
+				WerReportAddFile( ReportHandle, *LogFileName, WerFileTypeOther, WER_FILE_ANONYMOUS_DATA );
+
+				const FString CrashVideoPath = FPaths::GameLogDir() / TEXT( "CrashVideo.avi" );
+				WerReportAddFile( ReportHandle, *CrashVideoPath, WerFileTypeOther, WER_FILE_ANONYMOUS_DATA );		
+			}
 
 			// Submit
 			WER_SUBMIT_RESULT SubmitResult;
