Fix for assert getting exe name on startup in HTML5.
Implements FPlatformProcess::ExecutableName()

diff --git a/Engine/Source/Runtime/Core/Private/HTML5/HTML5PlatformProcess.cpp b/Engine/Source/Runtime/Core/Private/HTML5/HTML5PlatformProcess.cpp
--- a/Engine/Source/Runtime/Core/Private/HTML5/HTML5PlatformProcess.cpp
+++ b/Engine/Source/Runtime/Core/Private/HTML5/HTML5PlatformProcess.cpp
@@ -94,3 +94,8 @@ void FHTML5PlatformProcess::LaunchURL(const TCHAR* URL, const TCHAR* Parms, FStr
 	auto TmpURL = StringCast<ANSICHAR>(URL);
 	EM_ASM_ARGS({var InUrl = Pointer_stringify($0); console.log("Opening "+InUrl); window.open(InUrl);}, (ANSICHAR*)TmpURL.Get());
 }
+
+const TCHAR* FHTML5PlatformProcess::ExecutableName(bool bRemoveExtension)
+{
+	return FApp::GetGameName();
+}
diff --git a/Engine/Source/Runtime/Core/Public/HTML5/HTML5PlatformProcess.h b/Engine/Source/Runtime/Core/Public/HTML5/HTML5PlatformProcess.h
--- a/Engine/Source/Runtime/Core/Public/HTML5/HTML5PlatformProcess.h
+++ b/Engine/Source/Runtime/Core/Public/HTML5/HTML5PlatformProcess.h
@@ -37,6 +37,7 @@ struct CORE_API FHTML5PlatformProcess : public FGenericPlatformProcess
 	static class FRunnableThread* CreateRunnableThread();
 	static bool SupportsMultithreading();
 	static void LaunchURL( const TCHAR* URL, const TCHAR* Parms, FString* Error );
+	static const TCHAR* ExecutableName(bool bRemoveExtension = true);
 };
 
 typedef FHTML5PlatformProcess FPlatformProcess;
