Exit the editor running a cook on the fly server when user closes the log window

diff --git a/Engine/Source/Runtime/Core/Private/Mac/MacPlatformOutputDevices.cpp b/Engine/Source/Runtime/Core/Private/Mac/MacPlatformOutputDevices.cpp
--- a/Engine/Source/Runtime/Core/Private/Mac/MacPlatformOutputDevices.cpp
+++ b/Engine/Source/Runtime/Core/Private/Mac/MacPlatformOutputDevices.cpp
@@ -145,6 +145,16 @@ void FOutputDeviceMacError::HandleError()
 	FCoreDelegates::OnShutdownAfterError.Broadcast();
 }
 
+@implementation FMacConsoleWindow
+- (void)windowWillClose:(NSNotification*)Notification
+{
+	if (!MacApplication && [[NSApp orderedWindows] count] == 1)
+	{
+		exit(0);
+	}
+}
+@end
+
 ////////////////////////////////////////
 // FOutputDeviceConsoleMac
 ////////////////////////////////////////
@@ -201,10 +211,11 @@ void FOutputDeviceConsoleMac::CreateConsole()
 	}
 
 	MainThreadCall(^{
-		ConsoleHandle = [[NSWindow alloc] initWithContentRect: NSMakeRect(ConsolePosX, ConsolePosY, ConsoleWidth, ConsoleHeight)
+		ConsoleHandle = [[FMacConsoleWindow alloc] initWithContentRect: NSMakeRect(ConsolePosX, ConsolePosY, ConsoleWidth, ConsoleHeight)
 										styleMask: (NSTitledWindowMask | NSClosableWindowMask | NSMiniaturizableWindowMask | NSResizableWindowMask)
 										  backing: NSBackingStoreBuffered
 											defer: NO];
+		[ConsoleHandle setDelegate:ConsoleHandle];
 
 		ScrollView = [[NSScrollView alloc] initWithFrame:[[ConsoleHandle contentView] frame]];
 		NSSize ContentSize = [ScrollView contentSize];
@@ -233,7 +244,7 @@ void FOutputDeviceConsoleMac::CreateConsole()
 			[ConsoleHandle center];
 		}
 
-		[ConsoleHandle setOpaque: YES];
+		[ConsoleHandle setOpaque:YES];
 		[ConsoleHandle makeKeyAndOrderFront:nil];
 		
 		if(!MacApplication)
diff --git a/Engine/Source/Runtime/Core/Private/Mac/MacPlatformOutputDevicesPrivate.h b/Engine/Source/Runtime/Core/Private/Mac/MacPlatformOutputDevicesPrivate.h
--- a/Engine/Source/Runtime/Core/Private/Mac/MacPlatformOutputDevicesPrivate.h
+++ b/Engine/Source/Runtime/Core/Private/Mac/MacPlatformOutputDevicesPrivate.h
@@ -6,6 +6,9 @@
 
 #pragma once
 
+@interface FMacConsoleWindow : NSWindow<NSWindowDelegate>
+@end
+
 class FOutputDeviceMacError : public FOutputDeviceError
 {
 public:
@@ -37,23 +40,23 @@ class CORE_API FOutputDeviceConsoleMac : public FOutputDeviceConsole
 {
 private:
 	/** Handle to the console log window */
-	NSWindow *ConsoleHandle;
-	
+	FMacConsoleWindow* ConsoleHandle;
+
 	/** Handle to window's content view */
-	NSTextView *TextView;
-	
+	NSTextView* TextView;
+
 	/** Handle to window's scroll view */
-	NSScrollView *ScrollView;
+	NSScrollView* ScrollView;
 
 	/** Handle text color of TextView */
-	NSDictionary *TextViewTextColor;
-	
+	NSDictionary* TextViewTextColor;
+
 	/** Critical section for Serialize() */
 	FCriticalSection CriticalSection;
-	
+
 	/** Number of outstanding dispatched Cocoa tasks */
 	uint64 OutstandingTasks;
-		
+
 	/**
 	 * Saves the console window's position and size to the game .ini
 	 */
@@ -63,7 +66,7 @@ private:
 	 * Creates console window
 	 */
 	void CreateConsole();
-	
+
 	/**
 	 * Destroys console window
 	 */
@@ -73,7 +76,7 @@ private:
 	 *
 	 */
 	void SetDefaultTextColor();
-	
+
 public:
 
 	/** 
@@ -106,4 +109,3 @@ public:
 	 */
 	void Serialize( const TCHAR* Data, ELogVerbosity::Type Verbosity, const class FName& Category );
 };
-
