Core - Fixed a dead-lock when deleting a thread from the same thread

diff --git a/Engine/Source/Runtime/Core/Public/Async/Async.h b/Engine/Source/Runtime/Core/Public/Async/Async.h
--- a/Engine/Source/Runtime/Core/Public/Async/Async.h
+++ b/Engine/Source/Runtime/Core/Public/Async/Async.h
@@ -140,17 +140,7 @@ public:
 
 	// FRunnable interface
 
-	virtual uint32 Run() override
-	{
-		Promise.SetValue(Function());
-
-		FRunnableThread* Thread = ThreadFuture.Get();
-
-		delete Thread;
-		delete this;
-
-		return 0;
-	}
+	virtual uint32 Run() override;
 
 private:
 
@@ -289,3 +279,20 @@ TFuture<ResultType> Async(EAsyncExecution Execution, TFunction<ResultType()> Fun
 
 	return MoveTemp(Future);
 }
+
+template<typename ResultType> uint32 TAsyncRunnable<ResultType>::Run()
+{
+	Promise.SetValue( Function() );
+
+	FRunnableThread* Thread = ThreadFuture.Get();
+
+	// Enqueue deletion of the thread to a different thread.
+	Async<bool>( EAsyncExecution::TaskGraph, [=]()
+	{
+		delete Thread;
+		delete this;
+		return true;
+	} );
+
+	return 0;
+}
\ No newline at end of file
