Async: Added support for TAsync<void>

diff --git a/Engine/Source/Runtime/Core/Private/Tests/Async/AsyncTest.cpp b/Engine/Source/Runtime/Core/Private/Tests/Async/AsyncTest.cpp
--- a/Engine/Source/Runtime/Core/Private/Tests/Async/AsyncTest.cpp
+++ b/Engine/Source/Runtime/Core/Private/Tests/Async/AsyncTest.cpp
@@ -41,5 +41,12 @@ bool FAsyncTest::RunTest(const FString& Parameters)
 		TestEqual(TEXT("Thread pool task must return expected value"), Result, 123);
 	}
 
+	TFunction<void()> VoidTask = []() {};
+	
+	// TAsync<void>
+	{
+		auto Future = Async(EAsyncExecution::TaskGraph, VoidTask);
+	}
+
 	return true;
 }
diff --git a/Engine/Source/Runtime/Core/Public/Async/Async.h b/Engine/Source/Runtime/Core/Public/Async/Async.h
--- a/Engine/Source/Runtime/Core/Public/Async/Async.h
+++ b/Engine/Source/Runtime/Core/Public/Async/Async.h
@@ -22,6 +22,27 @@ enum class EAsyncExecution
 };
 
 
+/**
+ * Template for setting a promise's value from a function.
+ */
+template<typename ResultType>
+void SetPromise(TPromise<ResultType>& Promise, TFunction<ResultType()> Function)
+{
+	Promise.SetValue(Function());
+}
+
+
+/**
+ * Template for setting a promise's value from a function (specialization for void results).
+ */
+template<>
+void SetPromise(TPromise<void>& Promise, TFunction<void()> Function)
+{
+	Function();
+	Promise.SetValue();
+}
+
+
 /**
  * Base class for asynchronous functions that are executed in the Task Graph system.
  */
@@ -91,7 +112,7 @@ public:
 	 */
 	void DoTask(ENamedThreads::Type CurrentThread, const FGraphEventRef& MyCompletionGraphEvent)
 	{
-		Promise.SetValue(Function());
+		SetPromise(Promise, Function);
 	}
 
 	/**
@@ -181,7 +202,7 @@ public:
 
 	virtual void DoThreadedWork() override
 	{
-		Promise.SetValue(Function());
+		SetPromise(Promise, Function);
 		delete this;
 	}
 
@@ -280,9 +301,10 @@ TFuture<ResultType> Async(EAsyncExecution Execution, TFunction<ResultType()> Fun
 	return MoveTemp(Future);
 }
 
+
 template<typename ResultType> uint32 TAsyncRunnable<ResultType>::Run()
 {
-	Promise.SetValue( Function() );
+	SetPromise(Promise, Function);
 
 	FRunnableThread* Thread = ThreadFuture.Get();
 
@@ -295,4 +317,4 @@ template<typename ResultType> uint32 TAsyncRunnable<ResultType>::Run()
 	} );
 
 	return 0;
-}
\ No newline at end of file
+}
