A workaround for a rare Mac crash on creating a new window while bringing the editor to front

diff --git a/Engine/Source/Runtime/Core/Private/Mac/MacWindow.cpp b/Engine/Source/Runtime/Core/Private/Mac/MacWindow.cpp
--- a/Engine/Source/Runtime/Core/Private/Mac/MacWindow.cpp
+++ b/Engine/Source/Runtime/Core/Private/Mac/MacWindow.cpp
@@ -429,7 +429,7 @@ EWindowMode::Type FMacWindow::GetWindowMode() const
 
 bool FMacWindow::IsMaximized() const
 {
-	return WindowHandle->bZoomed;
+	return WindowHandle && WindowHandle->bZoomed;
 }
 
 bool FMacWindow::IsMinimized() const
