Fixed duplicated symbol error in Async.h when linking on Mac

diff --git a/Engine/Source/Runtime/Core/Public/Async/Async.h b/Engine/Source/Runtime/Core/Public/Async/Async.h
--- a/Engine/Source/Runtime/Core/Public/Async/Async.h
+++ b/Engine/Source/Runtime/Core/Public/Async/Async.h
@@ -26,7 +26,7 @@ enum class EAsyncExecution
  * Template for setting a promise's value from a function.
  */
 template<typename ResultType>
-void SetPromise(TPromise<ResultType>& Promise, TFunction<ResultType()> Function)
+inline void SetPromise(TPromise<ResultType>& Promise, TFunction<ResultType()> Function)
 {
 	Promise.SetValue(Function());
 }
@@ -36,7 +36,7 @@ void SetPromise(TPromise<ResultType>& Promise, TFunction<ResultType()> Function)
  * Template for setting a promise's value from a function (specialization for void results).
  */
 template<>
-void SetPromise(TPromise<void>& Promise, TFunction<void()> Function)
+inline void SetPromise(TPromise<void>& Promise, TFunction<void()> Function)
 {
 	Function();
 	Promise.SetValue();
