Core - Another async fix

diff --git a/Engine/Source/Runtime/Core/Public/Async/Async.h b/Engine/Source/Runtime/Core/Public/Async/Async.h
--- a/Engine/Source/Runtime/Core/Public/Async/Async.h
+++ b/Engine/Source/Runtime/Core/Public/Async/Async.h
@@ -220,6 +220,15 @@ private:
 	TPromise<ResultType> Promise;
 };
 
+/** Helper struct used to generate unique ids for the stats. */
+struct FAsyncIndex
+{
+	CORE_API static int32 GetNext()
+	{
+		static FThreadSafeCounter TAsyncThreadIndex;
+		return TAsyncThreadIndex.Add( 1 );
+	}
+};
 
 /**
  * Executes a given function asynchronously.
@@ -277,9 +286,8 @@ TFuture<ResultType> Async(EAsyncExecution Execution, TFunction<ResultType()> Fun
 		{
 			TPromise<FRunnableThread*> ThreadPromise;
 			TAsyncRunnable<ResultType>* Runnable = new TAsyncRunnable<ResultType>(MoveTemp(Function), MoveTemp(Promise), ThreadPromise.GetFuture());
-
-			static FThreadSafeCounter TAsyncThreadIndex;
-			const FString TAsyncThreadName = FString::Printf( TEXT( "TAsync %d" ), TAsyncThreadIndex.Add(1) );
+			
+			const FString TAsyncThreadName = FString::Printf( TEXT( "TAsync %d" ), FAsyncIndex::GetNext() );
 			FRunnableThread* RunnableThread = FRunnableThread::Create(Runnable, *TAsyncThreadName);
 
 			check(RunnableThread != nullptr);
@@ -309,11 +317,10 @@ template<typename ResultType> uint32 TAsyncRunnable<ResultType>::Run()
 	FRunnableThread* Thread = ThreadFuture.Get();
 
 	// Enqueue deletion of the thread to a different thread.
-	Async<bool>( EAsyncExecution::TaskGraph, [=]()
+	Async<void>( EAsyncExecution::TaskGraph, [=]()
 	{
 		delete Thread;
 		delete this;
-		return true;
 	} );
 
 	return 0;
