On OS X use PLCrashReporter instead of our own signal handlers as it is considerably more reliable.
- FMacPlatformProcess::IsSandboxedApplication will return true if and only if the application is running within an OS X application sandbox, this can be used to ensure that UE4 only accesses sandbox-safe APIs.
- The crashed application isn't held open waiting for the crash reporter as if you try it will then crash again in Apple's XPC code after exiting waitpid (with or without PLCrashReporter) & I've not yet been able to determine why.
- We can use a new 10.10 NSProcess call to access the OS version number - no need to access the CoreServices plist unless we want the build number (unavailable within the sandbox).
- Fixed some symbolication bugs, PLCrashReport crash reports will symbolicate reliably.
- We can't copy the crash text to the clipboard when handling an actual crash in OS X as that the code requires Objecive-C which is incompatible with POSIX-signal or Mach-O exception handling routines & will hang or crash the application again. This often results in an application icon that cannot be removed from the Dock & requires a force-restart of the machine.

diff --git a/Engine/Source/Developer/CrashDebugHelper/CrashDebugHelper.Build.cs b/Engine/Source/Developer/CrashDebugHelper/CrashDebugHelper.Build.cs
--- a/Engine/Source/Developer/CrashDebugHelper/CrashDebugHelper.Build.cs
+++ b/Engine/Source/Developer/CrashDebugHelper/CrashDebugHelper.Build.cs
@@ -7,6 +7,7 @@ public class CrashDebugHelper : ModuleRules
 	public CrashDebugHelper(TargetInfo Target)
 	{
 		PrivateIncludePaths.Add("Developer/CrashDebugHelper/Private");
+		PrivateIncludePaths.Add("ThirdParty/PLCrashReporter/plcrashreporter-master-5ae3b0a/Source");
 
 		PrivateDependencyModuleNames.AddRange(
 			new string[] {
diff --git a/Engine/Source/Developer/CrashDebugHelper/Private/Mac/CrashDebugHelperMac.cpp b/Engine/Source/Developer/CrashDebugHelper/Private/Mac/CrashDebugHelperMac.cpp
--- a/Engine/Source/Developer/CrashDebugHelper/Private/Mac/CrashDebugHelperMac.cpp
+++ b/Engine/Source/Developer/CrashDebugHelper/Private/Mac/CrashDebugHelperMac.cpp
@@ -4,6 +4,7 @@
 #include "Database.h"
 #include "EngineVersion.h"
 #include "ApplePlatformSymbolication.h"
+#include "CrashReporter.h"
 #include <cxxabi.h>
 
 #include "ISourceControlModule.h"
@@ -127,7 +128,24 @@ static int32 ParseGraphics(TCHAR const* CrashLog, FString& OutGPUDetails)
 static int32 ParseError(TCHAR const* CrashLog, FString& OutErrorDetails)
 {
 	bool bFound = false;
-	TCHAR const* Line = FCStringWide::Strstr(CrashLog, TEXT("Application Specific Information:"));
+	TCHAR const* Line = FCStringWide::Strstr(CrashLog, TEXT("Exception Codes:"));
+	if (Line)
+	{
+		Line += FCStringWide::Strlen(TEXT("Exception Codes:"));
+		check(Line);
+		TCHAR const* End = FCStringWide::Strchr(Line, TEXT('\r'));
+		if(!End)
+		{
+			End = FCStringWide::Strchr(Line, TEXT('\n'));
+		}
+		check(End);
+		
+		int32 Length = FMath::Min(PATH_MAX, (int32)((uintptr_t)(End - Line)));
+		OutErrorDetails.Append(Line, Length);
+		
+		bFound = true;
+	}
+	Line = FCStringWide::Strstr(CrashLog, TEXT("Application Specific Information:"));
 	if (Line)
 	{
 		Line = FCStringWide::Strchr(Line, TEXT('\n'));
@@ -141,6 +159,7 @@ static int32 ParseError(TCHAR const* CrashLog, FString& OutErrorDetails)
 		check(End);
 		
 		int32 Length = FMath::Min(PATH_MAX, (int32)((uintptr_t)(End - Line)));
+		OutErrorDetails += TEXT(" ");
 		OutErrorDetails.Append(Line, Length);
 		
 		bFound = true;
@@ -155,7 +174,11 @@ static int32 ParseExceptionCode(TCHAR const* CrashLog, uint32& OutExceptionCode)
 	if(Line)
 	{
 		TCHAR Buffer[257] = {0};
-		Found = swscanf(Line, TEXT("%*s %*s %*d (%256ls)"), Buffer);
+		Found = swscanf(Line, TEXT("%*s %*s %*s (%256ls)"), Buffer);
+		if(!Found)
+		{
+			Found = swscanf(Line, TEXT("%*s %*s %256ls"), Buffer);
+		}
 		if(Found)
 		{
 			TCHAR* End = FCStringWide::Strchr(Buffer, TEXT(')'));
@@ -195,10 +218,14 @@ static int32 ParseExceptionCode(TCHAR const* CrashLog, uint32& OutExceptionCode)
 			{
 				OutExceptionCode = SIGABRT;
 			}
-			else
+			else if(FString(Buffer).IsNumeric())
 			{
 				Found = swscanf(Buffer, TEXT("%u"), &OutExceptionCode);
 			}
+			else
+			{
+				OutExceptionCode = SIGUSR1;
+			}
 		}
 	}
 	return Found;
@@ -258,7 +285,17 @@ static int32 ParseThreadStackLine(TCHAR const* StackLine, FString& OutModuleName
 	TCHAR FunctionName[1025];
 	TCHAR FileName[257];
 	
-	int32 Found = swscanf(StackLine, TEXT("%*d %256ls 0x%lx %1024ls + %*d (%256ls:%d)"), ModuleName, &OutProgramCounter, FunctionName, FileName, &OutLineNumber);
+	int32 Found = swscanf(StackLine, TEXT("%*d %256ls 0x%lx"), ModuleName, &OutProgramCounter);
+	if(Found == 2)
+	{
+		uint64 FunctionAddress = 0;
+		uint32 FunctionOffset = 0;
+		if(swscanf(StackLine, TEXT("%*d %*ls %*lx 0x%lx + %d"), &FunctionAddress, &FunctionOffset) == 0)
+		{
+			Found += swscanf(StackLine, TEXT("%*d %*ls %*lx %1024ls + %*d (%256ls:%d)"), FunctionName, FileName, &OutLineNumber);
+		}
+	}
+	
 	switch(Found)
 	{
 		case 5:
@@ -391,6 +428,14 @@ static bool ParseModuleLine(TCHAR const* ModuleLine, FCrashModuleInfo& OutModule
 				++UUIDStart;
 				int32 Length = FMath::Min(64, (int32)((uintptr_t)(UUIDEnd - UUIDStart)));
 				OutModule.Report.Append(UUIDStart, Length);
+				if(!OutModule.Report.Contains(TEXT("-")))
+				{
+					OutModule.Report.InsertAt(8, TEXT('-'));
+					OutModule.Report.InsertAt(13, TEXT('-'));
+					OutModule.Report.InsertAt(18, TEXT('-'));
+					OutModule.Report.InsertAt(23, TEXT('-'));
+				}
+				OutModule.Report = OutModule.Report.ToUpper();
 				Found++;
 			}
 			
@@ -437,6 +482,8 @@ FCrashDebugHelperMac::~FCrashDebugHelperMac()
 
 bool FCrashDebugHelperMac::ParseCrashDump(const FString& InCrashDumpName, FCrashDebugInfo& OutCrashDebugInfo)
 {
+	SCOPED_AUTORELEASE_POOL;
+
 	if (bInitialized == false)
 	{
 		UE_LOG(LogCrashDebugHelper, Warning, TEXT("ParseCrashDump: CrashDebugHelper not initialized"));
@@ -444,12 +491,26 @@ bool FCrashDebugHelperMac::ParseCrashDump(const FString& InCrashDumpName, FCrash
 	}
 	
 	FString CrashDump;
-	if ( FFileHelper::LoadFileToString( CrashDump, *InCrashDumpName ) )
+	
+	NSString* CrashDumpPath = InCrashDumpName.GetNSString();
+	NSError* Error = nil;
+    NSData* Data = [NSData dataWithContentsOfFile: CrashDumpPath options: NSMappedRead error: &Error];
+	if(Data && !Error)
+	{
+		PLCrashReport* CrashLog = [[PLCrashReport alloc] initWithData: Data error: &Error];
+		if(CrashLog && !Error)
+		{
+			NSString* Report = [PLCrashReportTextFormatter stringValueForCrashReport: CrashLog withTextFormat: PLCrashReportTextFormatiOS];
+			CrashDump = FString(Report);
+		}
+	}
+	
+	if ( !CrashDump.IsEmpty() || FFileHelper::LoadFileToString( CrashDump, *InCrashDumpName ) )
 	{
 		// Only supports Apple crash report version 11
 		int32 ReportVersion = 0;
 		int32 Result = ParseReportVersion(*CrashDump, ReportVersion);
-		if(Result == 1 && ReportVersion == 11)
+		if(Result == 1 && (ReportVersion == 11 || ReportVersion == 104))
 		{
 			int32 Major = 0;
 			int32 Minor = 0;
@@ -506,11 +567,25 @@ bool FCrashDebugHelperMac::CreateMinidumpDiagnosticReport( const FString& InCras
 	bool bOK = false;
 	
 	FString CrashDump;
-	if ( FFileHelper::LoadFileToString( CrashDump, *InCrashDumpName ) )
+	
+	NSString* CrashDumpPath = InCrashDumpName.GetNSString();
+	NSError* Error = nil;
+    NSData* Data = [NSData dataWithContentsOfFile: CrashDumpPath options: NSMappedRead error: &Error];
+	if(Data && !Error)
+	{
+		PLCrashReport* CrashLog = [[PLCrashReport alloc] initWithData: Data error: &Error];
+		if(CrashLog && !Error)
+		{
+			NSString* Report = [PLCrashReportTextFormatter stringValueForCrashReport: CrashLog withTextFormat: PLCrashReportTextFormatiOS];
+			CrashDump = FString(Report);
+		}
+	}
+	
+	if ( !CrashDump.IsEmpty() || FFileHelper::LoadFileToString( CrashDump, *InCrashDumpName ) )
 	{
 		int32 ReportVersion = 0;
 		int32 Result = ParseReportVersion(*CrashDump, ReportVersion);
-		if(Result == 1 && ReportVersion == 11)
+		if(Result == 1 && (ReportVersion == 11 || ReportVersion == 104))
 		{
 			FString Error;
 			FString ModulePath;
diff --git a/Engine/Source/Runtime/Core/Core.Build.cs b/Engine/Source/Runtime/Core/Core.Build.cs
--- a/Engine/Source/Runtime/Core/Core.Build.cs
+++ b/Engine/Source/Runtime/Core/Core.Build.cs
@@ -79,9 +79,10 @@ public class Core : ModuleRules
 			AddThirdPartyPrivateStaticDependencies(Target, 
 				"IntelTBB",
 				"zlib",
-				"OpenGL"
+				"OpenGL",
+				"PLCrashReporter"
 				);
-			PublicFrameworks.AddRange(new string[] { "Cocoa", "Carbon", "IOKit" });
+			PublicFrameworks.AddRange(new string[] { "Cocoa", "Carbon", "IOKit", "Security" });
 			
 			if (UEBuildConfiguration.bBuildEditor == true)
 			{
diff --git a/Engine/Source/Runtime/Core/Private/Apple/ApplePlatformSymbolication.cpp b/Engine/Source/Runtime/Core/Private/Apple/ApplePlatformSymbolication.cpp
--- a/Engine/Source/Runtime/Core/Private/Apple/ApplePlatformSymbolication.cpp
+++ b/Engine/Source/Runtime/Core/Private/Apple/ApplePlatformSymbolication.cpp
@@ -42,6 +42,7 @@ extern "C"
 	CSSymbolicatorRef CSSymbolicatorCreateWithPid(pid_t pid);
 	CSSymbolicatorRef CSSymbolicatorCreateWithPathAndArchitecture(const char* path, cpu_type_t type);
 	
+	CSSymbolRef CSSymbolicatorGetSymbolWithAddressAtTime(CSSymbolicatorRef Symbolicator, vm_address_t Address, uint64_t Time);
 	CSSourceInfoRef CSSymbolicatorGetSourceInfoWithAddressAtTime(CSSymbolicatorRef Symbolicator, vm_address_t Address, uint64_t Time);
 	CSSymbolRef CSSymbolicatorGetSymbolWithMangledNameFromSymbolOwnerWithNameAtTime(CSSymbolicatorRef Symbolicator, const char* Symbol, const char* Name, uint64_t Time);
 	CSSymbolOwnerRef CSSymbolicatorGetSymbolOwnerWithUUIDAtTime(CSSymbolicatorRef Symbolicator, CFUUIDRef UUID, uint64_t Time);
@@ -51,6 +52,7 @@ extern "C"
 	CSSymbolOwnerRef CSSourceInfoGetSymbolOwner(CSSourceInfoRef Info);
 	
 	const char* CSSymbolOwnerGetName(CSSymbolOwnerRef symbol);
+	vm_address_t CSSymbolOwnerGetBaseAddress(CSSymbolOwnerRef Owner);
 	
 	int CSSourceInfoGetLineNumber(CSSourceInfoRef Info);
 	const char* CSSourceInfoGetPath(CSSourceInfoRef Info);
@@ -194,39 +196,38 @@ bool FApplePlatformSymbolication::SymbolInfoForStrippedSymbol(FApplePlatformSymb
 			CSSymbolOwnerRef SymbolOwner = CSSymbolicatorGetSymbolOwnerWithUUIDAtTime(Symbolicator, UUID, kCSNow);
 			if(!CSIsNull(SymbolOwner))
 			{
-				CSSourceInfoRef Symbol = CSSymbolicatorGetSourceInfoWithAddressAtTime(Symbolicator, (vm_address_t)ProgramCounter, kCSNow);
+				vm_address_t BaseAddress = CSSymbolOwnerGetBaseAddress(SymbolOwner);
 				
+				CSSymbolRef Symbol = CSSymbolicatorGetSymbolWithAddressAtTime(Symbolicator, ((vm_address_t)ProgramCounter) + BaseAddress, kCSNow);
 				if(!CSIsNull(Symbol))
 				{
-					Info.LineNumber = CSSourceInfoGetLineNumber(Symbol);
+					ANSICHAR const* DylibName = CSSymbolOwnerGetName(SymbolOwner);
+					FCStringAnsi::Strcpy(Info.ModuleName, DylibName);
 					
-					CSRange CodeRange = CSSourceInfoGetRange(Symbol);
-					Info.SymbolDisplacement = (ProgramCounter - CodeRange.Location);
-					
-					ANSICHAR const* FileName = CSSourceInfoGetPath(Symbol);
-					if(FileName)
+					ANSICHAR const* FunctionName = CSSymbolGetName(Symbol);
+					if(FunctionName)
 					{
-						FCStringAnsi::Sprintf(Info.Filename, FileName);
+						FCStringAnsi::Sprintf(Info.FunctionName, FunctionName);
 					}
 					
-					CSSymbolRef SymbolData = CSSourceInfoGetSymbol(Symbol);
-					if(!CSIsNull(SymbolData))
+					CSRange Range = CSSymbolGetRange(Symbol);
+					Info.SymbolDisplacement = (ProgramCounter - Range.Location);
+					Info.OffsetInModule = Range.Location;
+					Info.ProgramCounter = ProgramCounter;
+					
+					CSSourceInfoRef SymbolInfo = CSSymbolicatorGetSourceInfoWithAddressAtTime(Symbolicator, ((vm_address_t)ProgramCounter) + BaseAddress, kCSNow);
+					if(!CSIsNull(SymbolInfo))
 					{
-						ANSICHAR const* FunctionName = CSSymbolGetName(SymbolData);
-						if(FunctionName)
+						Info.LineNumber = CSSourceInfoGetLineNumber(SymbolInfo);
+						
+						ANSICHAR const* FileName = CSSourceInfoGetPath(SymbolInfo);
+						if(FileName)
 						{
-							FCStringAnsi::Sprintf(Info.FunctionName, FunctionName);
+							FCStringAnsi::Sprintf(Info.Filename, FileName);
 						}
 					}
 					
-					CSSymbolOwnerRef Owner = CSSourceInfoGetSymbolOwner(Symbol);
-					if(!CSIsNull(Owner))
-					{
-						ANSICHAR const* DylibName = CSSymbolOwnerGetName(Owner);
-						FCStringAnsi::Strcpy(Info.ModuleName, DylibName);
-						
-						bOK = Info.LineNumber != 0;
-					}
+					bOK = true;
 				}
 			}
 			
diff --git a/Engine/Source/Runtime/Core/Private/Mac/MacPlatformMisc.cpp b/Engine/Source/Runtime/Core/Private/Mac/MacPlatformMisc.cpp
--- a/Engine/Source/Runtime/Core/Private/Mac/MacPlatformMisc.cpp
+++ b/Engine/Source/Runtime/Core/Private/Mac/MacPlatformMisc.cpp
@@ -18,6 +18,7 @@
 #include "MacMallocZone.h"
 #include "ApplePlatformSymbolication.h"
 #include "MacPlatformCrashContext.h"
+#include "PLCrashReporter.h"
 
 #include <dlfcn.h>
 #include <IOKit/IOKitLib.h>
@@ -37,25 +38,43 @@
  * Information that cannot be obtained during a signal-handler is initialised here.
  * This ensures that we only call safe functions within the crash reporting handler.
  */
-struct MacApplicationInfo
+struct FMacApplicationInfo
 {
 	void Init()
 	{
 		SCOPED_AUTORELEASE_POOL;
 		
-		NSDictionary* SystemVersion = [NSDictionary dictionaryWithContentsOfFile: @"/System/Library/CoreServices/SystemVersion.plist"];
-		OSVersion = FString((NSString*)[SystemVersion objectForKey: @"ProductVersion"]);
-		FCStringAnsi::Strcpy(OSVersionUTF8, PATH_MAX+1, [[SystemVersion objectForKey: @"ProductVersion"] UTF8String]);
-		OSBuild = FString((NSString*)[SystemVersion objectForKey: @"ProductBuildVersion"]);
+		AppName = FApp::GetGameName();
+		FCStringAnsi::Strcpy(AppNameUTF8, PATH_MAX+1, TCHAR_TO_UTF8(*AppName));
+		
+		ExecutableName = FPlatformProcess::ExecutableName();
+		
+		AppPath = FPlatformProcess::GenerateApplicationPath(AppName, FApp::GetBuildConfiguration());
+
+		AppBundleID = FString([[NSBundle mainBundle] bundleIdentifier]);
+
+		bIsSandboxed = FPlatformProcess::IsSandboxedApplication();
+		
+		NumCores = FPlatformMisc::NumberOfCores();
+		
+		LCID = FString::Printf(TEXT("%d"), FInternationalization::Get().GetCurrentCulture()->GetLCID());
 		
-		TArray<FString> Components;
-		OSVersion.ParseIntoArray(&Components, TEXT("."), true);
-		uint8 ComponentValues[2] = {0};
-		for(uint32 i = 0; i < Components.Num() && i < 2; i++)
+		PrimaryGPU = FPlatformMisc::GetPrimaryGPUBrand();
+		
+		RunUUID = RunGUID();
+		
+		OSXVersion = [[NSProcessInfo processInfo] operatingSystemVersion];
+		OSVersion = FString::Printf(TEXT("%ld.%ld.%ld"), OSXVersion.majorVersion, OSXVersion.minorVersion, OSXVersion.patchVersion);
+		FCStringAnsi::Strcpy(OSVersionUTF8, PATH_MAX+1, TCHAR_TO_UTF8(*OSVersion));
+		
+		// OS X build number is only accessible on non-sandboxed applications as it resides outside the accessible sandbox
+		if(!bIsSandboxed)
 		{
-			TTypeFromString<uint8>::FromString(ComponentValues[i], *Components[i]);
+			NSDictionary* SystemVersion = [NSDictionary dictionaryWithContentsOfFile: @"/System/Library/CoreServices/SystemVersion.plist"];
+			OSBuild = FString((NSString*)[SystemVersion objectForKey: @"ProductBuildVersion"]);
 		}
-		RunningOnMavericks = ComponentValues[0] == 10 && ComponentValues[1] == 9;
+		
+		RunningOnMavericks = OSXVersion.majorVersion == 10 && OSXVersion.minorVersion == 9;
 		
 		char TempSysCtlBuffer[PATH_MAX] = {};
 		size_t TempSysCtlBufferSize = PATH_MAX;
@@ -92,12 +111,6 @@ struct MacApplicationInfo
 		TempSysCtlBufferSize = PATH_MAX+1;
 		sysctlbyname("machdep.cpu.brand_string", MachineCPUString, &TempSysCtlBufferSize, NULL, 0);
 		
-		FMacPlatformMisc::NumberOfCores();
-		FMacPlatformMisc::NumberOfCoresIncludingHyperthreads();
-		
-		AppName = FApp::GetGameName();
-		FCStringAnsi::Strcpy(AppNameUTF8, PATH_MAX+1, TCHAR_TO_UTF8(*AppName));
-		
 		gethostname(MachineName, ARRAY_COUNT(MachineName));
 		
 		FString CrashVideoPath = FPaths::GameLogDir() + TEXT("CrashVideo.avi");
@@ -137,14 +150,6 @@ struct MacApplicationInfo
 		FCStringAnsi::Strcpy(CrashReportClient, PATH_MAX+1, TCHAR_TO_UTF8(*ReportClient));
 		IFileManager::Get().MakeDirectory(*ReportPath, true);
 		
-		AppPath = FPlatformProcess::GenerateApplicationPath(FApp::GetName(), FApp::GetBuildConfiguration());
-		
-		AppBundleID = FString([[NSBundle mainBundle] bundleIdentifier]);
-		
-		LCID = FString::Printf(TEXT("%d"), FInternationalization::Get().GetCurrentCulture()->GetLCID());
-		
-		PrimaryGPU = FPlatformMisc::GetPrimaryGPUBrand();
-		
 		// Notification handler to check we are running from a battery - this only applies to MacBook's.
 		notify_handler_t PowerSourceNotifyHandler = ^(int32 Token){
 			RunningOnBattery = false;
@@ -173,11 +178,68 @@ struct MacApplicationInfo
 		uint32 Status = notify_register_dispatch(kIOPSNotifyPowerSource, &PowerSourceNotification, dispatch_get_main_queue(), PowerSourceNotifyHandler);
 		check(Status == NOTIFY_STATUS_OK);
 		
-		NumCores = FPlatformMisc::NumberOfCores();
-		
-		FPlatformMisc::CreateGuid(RunUUID);
+		NSString* PLCrashReportFile = [TemporaryCrashReportFolder().GetNSString() stringByAppendingPathComponent:TemporaryCrashReportName().GetNSString()];
+		[PLCrashReportFile getCString:PLCrashReportPath maxLength:PATH_MAX encoding:NSUTF8StringEncoding];
+	}
+	
+	~FMacApplicationInfo()
+	{
+		if(GMalloc != CrashMalloc)
+		{
+			delete CrashMalloc;
+		}
+		if(CrashReporter)
+		{
+			CrashReporter = nil;
+			[CrashReporter release];
+		}
+		if(PowerSourceNotification)
+		{
+			notify_cancel(PowerSourceNotification);
+			PowerSourceNotification = 0;
+		}
+	}
+	
+	static FGuid RunGUID()
+	{
+		static FGuid Guid;
+		if(!Guid.IsValid())
+		{
+			FPlatformMisc::CreateGuid(Guid);
+		}
+		return Guid;
 	}
 	
+	static FString TemporaryCrashReportFolder()
+	{
+		static FString PLCrashReportFolder;
+		if(PLCrashReportFolder.IsEmpty())
+		{
+			SCOPED_AUTORELEASE_POOL;
+			
+			NSArray* Paths = NSSearchPathForDirectoriesInDomains(NSCachesDirectory, NSUserDomainMask, YES);
+			NSString* CacheDir = [Paths objectAtIndex: 0];
+			
+			NSString* BundleID = [[NSBundle mainBundle] bundleIdentifier];
+			if(!BundleID)
+			{
+				BundleID = [[NSProcessInfo processInfo] processName];
+			}
+			check(BundleID);
+			
+			NSString* PLCrashReportFolderPath = [CacheDir stringByAppendingPathComponent: BundleID];
+			PLCrashReportFolder = FString(PLCrashReportFolderPath);
+		}
+		return PLCrashReportFolder;
+	}
+	
+	static FString TemporaryCrashReportName()
+	{
+		static FString PLCrashReportFileName(RunGUID().ToString() + TEXT(".plcrash"));
+		return PLCrashReportFileName;
+	}
+	
+	bool bIsSandboxed;
 	bool RunningOnBattery;
 	bool RunningOnMavericks;
 	int32 PowerSourceNotification;
@@ -185,6 +247,7 @@ struct MacApplicationInfo
 	char AppNameUTF8[PATH_MAX+1];
 	char AppLogPath[PATH_MAX+1];
 	char CrashReportPath[PATH_MAX+1];
+	char PLCrashReportPath[PATH_MAX+1];
 	char CrashReportClient[PATH_MAX+1];
 	char CrashReportVideo[PATH_MAX+1];
 	char OSVersionUTF8[PATH_MAX+1];
@@ -205,9 +268,15 @@ struct MacApplicationInfo
 	FString CommandLine;
 	FString BranchBaseDir;
 	FString PrimaryGPU;
+	FString ExecutableName;
+	NSOperatingSystemVersion OSXVersion;
 	FGuid RunUUID;
+	static PLCrashReporter* CrashReporter;
+	static FMacMallocCrashHandler* CrashMalloc;
 };
-static MacApplicationInfo GMacAppInfo;
+static FMacApplicationInfo GMacAppInfo;
+PLCrashReporter* FMacApplicationInfo::CrashReporter = nullptr;
+FMacMallocCrashHandler* FMacApplicationInfo::CrashMalloc = nullptr;
 
 UpdateCachedMacMenuStateProc FMacPlatformMisc::UpdateCachedMacMenuState = nullptr;
 bool FMacPlatformMisc::bChachedMacMenuStateNeedsUpdate = true;
@@ -626,15 +695,20 @@ const TCHAR* FMacPlatformMisc::GetSystemErrorMessage(TCHAR* OutBuffer, int32 Buf
 
 void FMacPlatformMisc::ClipboardCopy(const TCHAR* Str)
 {
-	SCOPED_AUTORELEASE_POOL;
+	// Don't attempt to copy the text to the clipboard if we've crashed or we'll crash again & become unkillable.
+	// The MallocZone used for crash reporting will be enabled before this call if we've crashed so that will do for testing.
+	if ( GMalloc != FMacApplicationInfo::CrashMalloc )
+	{
+		SCOPED_AUTORELEASE_POOL;
 
-	CFStringRef CocoaString = FPlatformString::TCHARToCFString(Str);
-	NSPasteboard *Pasteboard = [NSPasteboard generalPasteboard];
-	[Pasteboard clearContents];
-	NSPasteboardItem *Item = [[[NSPasteboardItem alloc] init] autorelease];
-	[Item setString: (NSString *)CocoaString forType: NSPasteboardTypeString];
-	[Pasteboard writeObjects:[NSArray arrayWithObject:Item]];
-	CFRelease(CocoaString);
+		CFStringRef CocoaString = FPlatformString::TCHARToCFString(Str);
+		NSPasteboard *Pasteboard = [NSPasteboard generalPasteboard];
+		[Pasteboard clearContents];
+		NSPasteboardItem *Item = [[[NSPasteboardItem alloc] init] autorelease];
+		[Item setString: (NSString *)CocoaString forType: NSPasteboardTypeString];
+		[Pasteboard writeObjects:[NSArray arrayWithObject:Item]];
+		CFRelease(CocoaString);
+	}
 }
 
 void FMacPlatformMisc::ClipboardPaste(class FString& Result)
@@ -1094,16 +1168,8 @@ bool FMacPlatformMisc::IsRunningOnMavericks()
 
 int32 FMacPlatformMisc::MacOSXVersionCompare(uint8 Major, uint8 Minor, uint8 Revision)
 {
-	TArray<FString> Components;
-	GMacAppInfo.OSVersion.ParseIntoArray(&Components, TEXT("."), true);
-	
 	uint8 TargetValues[3] = {Major, Minor, Revision};
-	uint8 ComponentValues[3] = {0, 0, 0};
-	
-	for(uint32 i = 0; i < Components.Num(); i++)
-	{
-		TTypeFromString<uint8>::FromString(ComponentValues[i], *Components[i]);
-	}
+	NSInteger ComponentValues[3] = {GMacAppInfo.OSXVersion.majorVersion, GMacAppInfo.OSXVersion.minorVersion, GMacAppInfo.OSXVersion.patchVersion};
 	
 	for(uint32 i = 0; i < 3; i++)
 	{
@@ -1140,7 +1206,6 @@ FString FMacPlatformMisc::GetOperatingSystemId()
 
 /** Global pointer to crash handler */
 void (* GCrashHandlerPointer)(const FGenericCrashContext& Context) = NULL;
-FMacMallocCrashHandler* GCrashMalloc = nullptr;
 
 /**
  * Good enough default crash reporter.
@@ -1167,16 +1232,16 @@ static void DefaultCrashHandler(FMacCrashContext const& Context)
 }
 
 /** True system-specific crash handler that gets called first */
-static void PlatformCrashHandler(int32 Signal, siginfo_t* Info, void* Context)
+static void PlatformCrashHandler(siginfo_t* Info, ucontext_t* Uap, void* Context)
 {
 	FApplePlatformSymbolication::SetSymbolicationAllowed( false );
 	
 	FMacCrashContext CrashContext;
-	CrashContext.InitFromSignal(Signal, Info, Context);
+	CrashContext.InitFromSignal((int32)Info->si_signo, Info, Context);
 	
 	// Switch to crash handler malloc to avoid malloc reentrancy
-	check(GCrashMalloc);
-	GCrashMalloc->Enable(&CrashContext, FPlatformTLS::GetCurrentThreadId());
+	check(FMacApplicationInfo::CrashMalloc);
+	FMacApplicationInfo::CrashMalloc->Enable(&CrashContext, FPlatformTLS::GetCurrentThreadId());
 	
 	if (GCrashHandlerPointer)
 	{
@@ -1232,24 +1297,35 @@ void FMacPlatformMisc::SetGracefulTerminationHandler()
 
 void FMacPlatformMisc::SetCrashHandler(void (* CrashHandler)(const FGenericCrashContext& Context))
 {
-	GCrashHandlerPointer = CrashHandler;
+	SCOPED_AUTORELEASE_POOL;
 	
-	// Configure the crash handler malloc zone to reserve some VM space for itself
-	GCrashMalloc = new FMacMallocCrashHandler( 128 * 1024 * 1024 );
+	GCrashHandlerPointer = CrashHandler;
 	
-	struct sigaction Action;
-	FMemory::Memzero(&Action, sizeof(struct sigaction));
-	Action.sa_sigaction = PlatformCrashHandler;
-	sigemptyset(&Action.sa_mask);
-	Action.sa_flags = SA_SIGINFO | SA_RESTART | SA_ONSTACK;
-	sigaction(SIGQUIT, &Action, NULL);	// SIGQUIT is a user-initiated "crash".
-	sigaction(SIGILL, &Action, NULL);
-	sigaction(SIGEMT, &Action, NULL);
-	sigaction(SIGFPE, &Action, NULL);
-	sigaction(SIGBUS, &Action, NULL);
-	sigaction(SIGSEGV, &Action, NULL);
-	sigaction(SIGSYS, &Action, NULL);
-	sigaction(SIGABRT, &Action, NULL);
+	if(!FMacApplicationInfo::CrashReporter && !FMacApplicationInfo::CrashMalloc)
+	{
+		// Configure the crash handler malloc zone to reserve some VM space for itself
+		FMacApplicationInfo::CrashMalloc = new FMacMallocCrashHandler( 128 * 1024 * 1024 );
+		
+		PLCrashReporterConfig* Config = [[[PLCrashReporterConfig alloc] initWithSignalHandlerType: PLCrashReporterSignalHandlerTypeBSD
+																		symbolicationStrategy: PLCrashReporterSymbolicationStrategyNone
+																		crashReportFolder:FMacApplicationInfo::TemporaryCrashReportFolder().GetNSString()
+																		crashReportName:FMacApplicationInfo::TemporaryCrashReportName().GetNSString()] autorelease];
+		FMacApplicationInfo::CrashReporter = [[PLCrashReporter alloc] initWithConfiguration: Config];
+		
+		PLCrashReporterCallbacks CrashReportCallback = {
+			.version = 0,
+			.context = nullptr,
+			.handleSignal = PlatformCrashHandler
+		};
+		
+		[FMacApplicationInfo::CrashReporter setCrashCallbacks: &CrashReportCallback];
+
+		NSError* Error = nil;
+		if (![FMacApplicationInfo::CrashReporter enableCrashReporterAndReturnError: &Error])
+		{
+			UE_LOG(LogMac, Log,  TEXT("Failed to enable PLCrashReporter: %s"), *FString([Error localizedDescription]) );
+		}
+	}
 }
 
 void FMacCrashContext::GenerateWindowsErrorReport(char const* WERPath) const
@@ -1331,20 +1407,15 @@ void FMacCrashContext::GenerateWindowsErrorReport(char const* WERPath) const
 		// App time stamp
 		WriteLine(ReportFile, TEXT("\t\t<Parameter2>528f2d37</Parameter2>"));													// FIXME: supply valid?
 		
-		uint64 BackTrace[7] = {0, 0, 0, 0, 0, 0, 0};
-		FPlatformStackWalk::CaptureStackBackTrace( BackTrace, 7, Context );
-		
-		if(BackTrace[6] != 0)
+		Dl_info DLInfo;
+		if(Info && Info->si_addr != 0 && dladdr(Info->si_addr, &DLInfo) != 0)
 		{
-			Dl_info Info;
-			dladdr((void*)BackTrace[6], &Info);
-			
 			// Crash Module name
 			WriteUTF16String(ReportFile, TEXT("\t\t<Parameter3>"));
-			if (FCStringAnsi::Strlen(Info.dli_fname))
+			if (DLInfo.dli_fname && FCStringAnsi::Strlen(DLInfo.dli_fname))
 			{
 				FMemory::Memzero(Line, PATH_MAX * sizeof(TCHAR));
-				FUTF8ToTCHAR_Convert::Convert(Line, PATH_MAX, Info.dli_fname, FCStringAnsi::Strlen(Info.dli_fname));
+				FUTF8ToTCHAR_Convert::Convert(Line, PATH_MAX, DLInfo.dli_fname, FCStringAnsi::Strlen(DLInfo.dli_fname));
 				WriteUTF16String(ReportFile, Line);
 			}
 			else
@@ -1356,7 +1427,7 @@ void FMacCrashContext::GenerateWindowsErrorReport(char const* WERPath) const
 			// Check header
 			uint32 Version = 0;
 			uint32 TimeStamp = 0;
-			struct mach_header_64* Header = (struct mach_header_64*)Info.dli_fbase;
+			struct mach_header_64* Header = (struct mach_header_64*)DLInfo.dli_fbase;
 			struct load_command *CurrentCommand = (struct load_command *)( (char *)Header + sizeof(struct mach_header_64) );
 			if( Header->magic == MH_MAGIC_64 )
 			{
@@ -1390,7 +1461,7 @@ void FMacCrashContext::GenerateWindowsErrorReport(char const* WERPath) const
 			
 			// IL Offset -> Function pointer
 			WriteUTF16String(ReportFile, TEXT("\t\t<Parameter7>"));
-			WriteUTF16String(ReportFile, ItoTCHAR(BackTrace[6], 16));
+			WriteUTF16String(ReportFile, ItoTCHAR((uint64)Info->si_addr, 16));
 			WriteLine(ReportFile, TEXT("</Parameter7>"));
 		}
 		
@@ -1446,267 +1517,24 @@ void FMacCrashContext::GenerateWindowsErrorReport(char const* WERPath) const
 	}
 }
 
-void FMacCrashContext::GenerateMinidump(char const* Path) const
+void FMacCrashContext::CopyMinidump(char const* OutputPath, char const* InputPath) const
 {
-	int ReportFile = open(Path, O_CREAT|O_WRONLY, 0766);
-	if (ReportFile != -1)
+	int ReportFile = open(OutputPath, O_CREAT|O_WRONLY, 0766);
+	int DumpFile = open(InputPath, O_RDONLY, 0766);
+	if (ReportFile != -1 && DumpFile != -1)
 	{
-		TCHAR Line[PATH_MAX] = {};
+		char Data[PATH_MAX];
 		
-		// write BOM
-		static uint16 ByteOrderMarker = 0xFEFF;
-		write(ReportFile, &ByteOrderMarker, sizeof(ByteOrderMarker));
-		
-		WriteUTF16String(ReportFile, TEXT("Process:\tUE4-"));
-		WriteUTF16String(ReportFile, *GMacAppInfo.AppName);
-		WriteUTF16String(ReportFile, TEXT(" ["));
-		WriteUTF16String(ReportFile, ItoTCHAR(FPlatformProcess::GetCurrentProcessId(), 10));
-		WriteLine(ReportFile, TEXT("]"));
-		
-		WriteUTF16String(ReportFile, TEXT("Path:\t"));
-		WriteLine(ReportFile, *GMacAppInfo.AppPath);
-		
-		WriteUTF16String(ReportFile, TEXT("Identifier:\t"));
-		WriteLine(ReportFile, *GMacAppInfo.AppBundleID);
-		
-		WriteUTF16String(ReportFile, TEXT("Version:\t"));
-		WriteUTF16String(ReportFile, ENGINE_VERSION_STRINGIFY(ENGINE_MAJOR_VERSION) ENGINE_VERSION_TEXT(".") ENGINE_VERSION_STRINGIFY(ENGINE_MINOR_VERSION) ENGINE_VERSION_TEXT(".") ENGINE_VERSION_STRINGIFY(ENGINE_PATCH_VERSION));
-		WriteUTF16String(ReportFile, TEXT(" ("));
-		WriteUTF16String(ReportFile, ENGINE_VERSION_STRING);
-		WriteLine(ReportFile, TEXT(")"));
-		
-		WriteLine(ReportFile, TEXT("Code Type:\tX86-64 (Native)"));
-		
-		WriteUTF16String(ReportFile, TEXT("Parent Process:\t"));
-		WriteUTF16String(ReportFile, *GMacAppInfo.ParentProcess);
-		WriteUTF16String(ReportFile, TEXT(" ["));
-		WriteUTF16String(ReportFile, ItoTCHAR(getppid(), 10));
-		WriteLine(ReportFile, TEXT("]"));
-		
-		WriteUTF16String(ReportFile, TEXT("Responsible:\t"));
-		WriteUTF16String(ReportFile, *GMacAppInfo.ParentProcess);
-		WriteUTF16String(ReportFile, TEXT(" ["));
-		WriteUTF16String(ReportFile, ItoTCHAR(getppid(), 10));
-		WriteLine(ReportFile, TEXT("]"));
-		
-		WriteUTF16String(ReportFile, TEXT("User ID:\t"));
-		WriteLine(ReportFile, ItoTCHAR(getuid(), 10));
-		
-		WriteLine(ReportFile, TEXT(""));
-		
-		WriteUTF16String(ReportFile, TEXT("Date/Time:\t"));
-		
-		time_t Now;
-		time(&Now);
-		struct tm LocalTime;
-		localtime_r(&Now, &LocalTime);
-		wcsftime(Line, PATH_MAX, TEXT("%Y-%m-%d %H:%M:%S %z"), &LocalTime);
-		WriteLine(ReportFile, Line);
-		
-		WriteUTF16String(ReportFile, TEXT("OS Version:\tMac OS X "));
-		WriteUTF16String(ReportFile, *GMacAppInfo.OSVersion);
-		WriteUTF16String(ReportFile, TEXT(" ("));
-		WriteUTF16String(ReportFile, *GMacAppInfo.OSBuild);
-		WriteLine(ReportFile, TEXT(")"));
-		
-		WriteLine(ReportFile, TEXT("Report Version:\t11"));
-		
-		WriteUTF16String(ReportFile, TEXT("Anonymous UUID:\t"));
-		WriteUTF16String(ReportFile, *GMacAppInfo.MachineUUID);
-										  
-		WriteLine(ReportFile, TEXT(""));
-		WriteLine(ReportFile, TEXT(""));
-		
-		WriteLine(ReportFile, TEXT("Crashed Thread:\t0"));
-		WriteLine(ReportFile, TEXT(""));
-		
-		WriteUTF16String(ReportFile, TEXT("Exception Type:\t"));
-		WriteUTF16String(ReportFile, ItoTCHAR(Info->si_errno, 10));
-		WriteUTF16String(ReportFile, TEXT(" ("));
-		WriteUTF16String(ReportFile, ItoTCHAR(Signal, 10));
-		WriteLine(ReportFile, TEXT(")"));
-		
-		WriteUTF16String(ReportFile, TEXT("Exception Codes:\t"));
-		WriteUTF16String(ReportFile, TEXT(" 0x"));
-		WriteUTF16String(ReportFile, ItoTCHAR(Info->si_code, 16));
-		WriteUTF16String(ReportFile, TEXT(", 0x"));
-		WriteLine(ReportFile, ItoTCHAR(Info->si_band, 16));
-		
-		WriteLine(ReportFile, TEXT("Application Specific Information:"));
-		FMemory::Memzero(Line, PATH_MAX * sizeof(TCHAR));
-		FUTF8ToTCHAR_Convert::Convert(Line, PATH_MAX, SignalDescription, FCStringAnsi::Strlen(SignalDescription));
-		WriteLine(ReportFile, Line);
-		
-		WriteLine(ReportFile, TEXT(""));
-		WriteLine(ReportFile, TEXT("Thread 0:"));
-		
-		static const int MAX_DEPTH = 100;
-		static uint64 StackTrace[MAX_DEPTH];
-		FMemory::Memzero( StackTrace );
-		FPlatformStackWalk::CaptureStackBackTrace( StackTrace, MAX_DEPTH, Context );
-		
-		// Skip the first two entries as they are inside the stack walking code.
-		Dl_info Info;
-		static const int32 IgnoreCount = 6;
-		int32 CurrentDepth = IgnoreCount;
-		// Allow the first entry to be NULL as the crash could have been caused by a call to a NULL function pointer,
-		// which would mean the top of the callstack is NULL.
-		while( StackTrace[CurrentDepth] || ( CurrentDepth == IgnoreCount ) )
-		{
-			WriteUTF16String(ReportFile, ItoTCHAR(CurrentDepth - IgnoreCount, 10));
-			WriteUTF16String(ReportFile, TEXT("\t"));
-			
-			if(dladdr((void*)StackTrace[CurrentDepth], &Info) == 0)
-			{
-				Info.dli_fbase = nullptr;
-				Info.dli_fname = nullptr;
-				Info.dli_saddr = nullptr;
-				Info.dli_sname = nullptr;
-			}
-			if(Info.dli_fname && FCStringAnsi::Strrchr(Info.dli_fname, '/'))
-			{
-				char const* Name = FCStringAnsi::Strrchr(Info.dli_fname, '/');
-				Name++;
-				FMemory::Memzero(Line, PATH_MAX * sizeof(TCHAR));
-				FUTF8ToTCHAR_Convert::Convert(Line, PATH_MAX, Name, FCStringAnsi::Strlen(Name));
-				WriteUTF16String(ReportFile, Line);
-			}
-			else
-			{
-				WriteUTF16String(ReportFile, TEXT("[Unknown]"));
-			}
-			WriteUTF16String(ReportFile, TEXT("\t0x"));
-			WriteUTF16String(ReportFile, ItoTCHAR(StackTrace[CurrentDepth], 16));
-			WriteUTF16String(ReportFile, TEXT(" "));
-			if(Info.dli_sname && FCStringAnsi::Strlen(Info.dli_sname))
-			{
-				FMemory::Memzero(Line, PATH_MAX * sizeof(TCHAR));
-				FUTF8ToTCHAR_Convert::Convert(Line, PATH_MAX, Info.dli_sname, FCStringAnsi::Strlen(Info.dli_sname));
-				WriteUTF16String(ReportFile, Line);
-			}
-			else
-			{
-				WriteUTF16String(ReportFile, TEXT("[Unknown]"));
-			}
-			WriteUTF16String(ReportFile, TEXT(" + "));
-			WriteLine(ReportFile, ItoTCHAR(StackTrace[CurrentDepth] - (uint64)Info.dli_saddr, 10));
-			
-			CurrentDepth++;
-		}
-		
-		WriteLine(ReportFile, TEXT(""));
-		
-		WriteLine(ReportFile, TEXT("Binary Images:"));
-		uint32 ModuleCount = _dyld_image_count();
-		for(uint32 Index = 0; Index < ModuleCount; Index++)
+		int Bytes = 0;
+		while((Bytes = read(DumpFile, Data, PATH_MAX)) > 0)
 		{
-			char const* ModulePath = _dyld_get_image_name(Index);
-			uint32 CompatVersion = 0;
-			uint32 CurrentVersion = 0;
-			uint8 UUID[16] = {0,0,0,0,0,0,0,0,0,0,0,0,0,0,0,0};
-			struct mach_header_64* Header = (struct mach_header_64*)_dyld_get_image_header(Index);
-			struct load_command *CurrentCommand = (struct load_command *)( (char *)Header + sizeof(struct mach_header_64) );
-			if( Header->magic == MH_MAGIC_64 )
-			{
-				uint8 Found = 0;
-				for( int32 i = 0; i < Header->ncmds && Found < 2; i++ )
-				{
-					if( CurrentCommand->cmd == LC_LOAD_DYLIB )
-					{
-						struct dylib_command *DylibCommand = (struct dylib_command *) CurrentCommand;
-						CompatVersion = DylibCommand->dylib.compatibility_version;
-						CurrentVersion = DylibCommand->dylib.current_version;
-						Found++;
-					}
-					else if ( CurrentCommand->cmd == LC_UUID )
-					{
-						struct uuid_command* UUIDCommand = (struct uuid_command*)CurrentCommand;
-						FMemory::Memcpy(UUID, UUIDCommand->uuid, 16);
-						Found++;
-					}
-					
-					CurrentCommand = (struct load_command *)( (char *)CurrentCommand + CurrentCommand->cmdsize );
-				}
-			}
-			
-			WriteUTF16String(ReportFile, TEXT("0x"));
-			WriteUTF16String(ReportFile, ItoTCHAR((uint64)Header, 16));
-			WriteUTF16String(ReportFile, TEXT(" -\t0x"));
-			WriteUTF16String(ReportFile, ItoTCHAR((uint64)Header + sizeof(struct mach_header_64) + Header->sizeofcmds, 16));
-			WriteUTF16String(ReportFile, TEXT(" "));
-			
-			char const* Name = FCStringAnsi::Strrchr(ModulePath, '/');
-            if ( Name )
-            {
-                Name++;
-            }
-            else
-            {
-                Name = ModulePath;
-            }
-			FMemory::Memzero(Line, PATH_MAX * sizeof(TCHAR));
-			FUTF8ToTCHAR_Convert::Convert(Line, PATH_MAX, Name, FCStringAnsi::Strlen(Name));
-			WriteUTF16String(ReportFile, Line);
-			WriteUTF16String(ReportFile, TEXT(" ("));
-			
-			WriteUTF16String(ReportFile, ItoTCHAR(((CompatVersion >> 16) & 0xffff), 10));
-			WriteUTF16String(ReportFile, TEXT("."));
-			WriteUTF16String(ReportFile, ItoTCHAR(((CompatVersion >> 8) & 0xff), 10));
-			WriteUTF16String(ReportFile, TEXT("."));
-			WriteUTF16String(ReportFile, ItoTCHAR((CompatVersion & 0xff), 10));
-			WriteUTF16String(ReportFile, TEXT(" - "));
-			WriteUTF16String(ReportFile, ItoTCHAR(((CurrentVersion >> 16) & 0xffff), 10));
-			WriteUTF16String(ReportFile, TEXT("."));
-			WriteUTF16String(ReportFile, ItoTCHAR(((CurrentVersion >> 8) & 0xff), 10));
-			WriteUTF16String(ReportFile, TEXT("."));
-			WriteUTF16String(ReportFile, ItoTCHAR((CurrentVersion & 0xff), 10));
-			WriteUTF16String(ReportFile, TEXT(") <"));
-			WriteUTF16String(ReportFile, ItoTCHAR(UUID[0], 16, 2));
-			WriteUTF16String(ReportFile, ItoTCHAR(UUID[1], 16, 2));
-			WriteUTF16String(ReportFile, ItoTCHAR(UUID[2], 16, 2));
-			WriteUTF16String(ReportFile, ItoTCHAR(UUID[3], 16, 2));
-			WriteUTF16String(ReportFile, TEXT("-"));
-			WriteUTF16String(ReportFile, ItoTCHAR(UUID[4], 16, 2));
-			WriteUTF16String(ReportFile, ItoTCHAR(UUID[5], 16, 2));
-			WriteUTF16String(ReportFile, TEXT("-"));
-			WriteUTF16String(ReportFile, ItoTCHAR(UUID[6], 16, 2));
-			WriteUTF16String(ReportFile, ItoTCHAR(UUID[7], 16, 2));
-			WriteUTF16String(ReportFile, TEXT("-"));
-			WriteUTF16String(ReportFile, ItoTCHAR(UUID[8], 16, 2));
-			WriteUTF16String(ReportFile, ItoTCHAR(UUID[9], 16, 2));
-			WriteUTF16String(ReportFile, TEXT("-"));
-			WriteUTF16String(ReportFile, ItoTCHAR(UUID[10], 16, 2));
-			WriteUTF16String(ReportFile, ItoTCHAR(UUID[11], 16, 2));
-			WriteUTF16String(ReportFile, ItoTCHAR(UUID[12], 16, 2));
-			WriteUTF16String(ReportFile, ItoTCHAR(UUID[13], 16, 2));
-			WriteUTF16String(ReportFile, ItoTCHAR(UUID[14], 16, 2));
-			WriteUTF16String(ReportFile, ItoTCHAR(UUID[15], 16, 2));
-			WriteUTF16String(ReportFile, TEXT("> "));
-			
-			FMemory::Memzero(Line, PATH_MAX * sizeof(TCHAR));
-			FUTF8ToTCHAR_Convert::Convert(Line, PATH_MAX, ModulePath, FCStringAnsi::Strlen(ModulePath));
-			WriteLine(ReportFile, Line);
+			write(ReportFile, Data, Bytes);
 		}
 		
-		WriteLine(ReportFile, TEXT(""));
-		
-		WriteLine(ReportFile, TEXT("System Profile:"));
-		WriteUTF16String(ReportFile, TEXT("Model: "));
-		WriteUTF16String(ReportFile, *GMacAppInfo.MachineModel);
-		WriteUTF16String(ReportFile, TEXT(", "));
-		WriteUTF16String(ReportFile, ItoTCHAR(GMacAppInfo.NumCores, 10));
-		WriteUTF16String(ReportFile, TEXT(" processors, "));
-		WriteUTF16String(ReportFile, ItoTCHAR(GMacAppInfo.NumCores, 10));
-		WriteUTF16String(ReportFile, TEXT("-Core "));
-
-		FMemory::Memzero(Line, PATH_MAX * sizeof(TCHAR));
-		FUTF8ToTCHAR_Convert::Convert(Line, PATH_MAX, GMacAppInfo.MachineCPUString, FCStringAnsi::Strlen(GMacAppInfo.MachineCPUString));
-		WriteLine(ReportFile, Line);
-		
-		WriteUTF16String(ReportFile, TEXT("Graphics: "));
-		WriteLine(ReportFile, *GMacAppInfo.PrimaryGPU);
-		
+		close(DumpFile);
 		close(ReportFile);
+		
+		unlink(InputPath);
 	}
 }
 
@@ -1726,8 +1554,7 @@ void FMacCrashContext::GenerateCrashInfoAndLaunchReporter() const
 	FCStringAnsi::Strcat(CrashInfoFolder, PATH_MAX, ItoANSI(GMacAppInfo.RunUUID.D, 16));
 
 	// Prevent CrashReportClient from spawning another CrashReportClient.
-	const TCHAR* ExecutableName = FPlatformProcess::ExecutableName();
-	const bool bCanRunCrashReportClient = FCString::Stristr( ExecutableName, TEXT( "CrashReportClient" ) ) == nullptr;
+	const bool bCanRunCrashReportClient = FCString::Stristr( *(GMacAppInfo.ExecutableName), TEXT( "CrashReportClient" ) ) == nullptr;
 
 	if(!mkdir(CrashInfoFolder, 0766) && bCanRunCrashReportClient)
 	{
@@ -1756,7 +1583,7 @@ void FMacCrashContext::GenerateCrashInfoAndLaunchReporter() const
 		// generate "minidump" (Apple crash log format)
 		FCStringAnsi::Strncpy(FilePath, CrashInfoFolder, PATH_MAX);
 		FCStringAnsi::Strcat(FilePath, PATH_MAX, "/minidump.dmp");
-		GenerateMinidump(FilePath);
+		CopyMinidump(FilePath, GMacAppInfo.PLCrashReportPath);
 		
 		// generate "info.txt" custom data for our server
 		FCStringAnsi::Strncpy(FilePath, CrashInfoFolder, PATH_MAX);
@@ -1836,12 +1663,8 @@ void FMacCrashContext::GenerateCrashInfoAndLaunchReporter() const
 			FCStringAnsi::Strcat(FilePath, PATH_MAX, "/");
 			execl(GMacAppInfo.CrashReportClient, "CrashReportClient", FilePath, NULL);
 		}
-		else
-		{
-			// Parent
-			int StatLoc = 0;
-			waitpid(ForkPID, &StatLoc, 0);
-		}
+		// We no longer wait here because on return the OS will scribble & crash again due to the behaviour of the XPC function
+		// OS X uses internally to launch/wait on the CrashReportClient. It is simpler, easier & safer to just die here like a good little Mac.app.
 	}
 	
 	_Exit(0);
diff --git a/Engine/Source/Runtime/Core/Private/Mac/MacPlatformProcess.cpp b/Engine/Source/Runtime/Core/Private/Mac/MacPlatformProcess.cpp
--- a/Engine/Source/Runtime/Core/Private/Mac/MacPlatformProcess.cpp
+++ b/Engine/Source/Runtime/Core/Private/Mac/MacPlatformProcess.cpp
@@ -616,6 +616,33 @@ bool FMacPlatformProcess::IsThisApplicationForeground()
 	return [NSApp isActive] && MacApplication && MacApplication->IsWorkspaceSessionActive();
 }
 
+bool FMacPlatformProcess::IsSandboxedApplication()
+{
+	SCOPED_AUTORELEASE_POOL;
+	
+	bool bIsSandboxedApplication = false;
+
+	SecStaticCodeRef SecCodeObj = nullptr;
+	NSURL* BundleURL = [[NSBundle mainBundle] bundleURL];
+    OSStatus Err = SecStaticCodeCreateWithPath((CFURLRef)BundleURL, kSecCSDefaultFlags, &SecCodeObj);
+	if (SecCodeObj)
+	{
+		check(Err == errSecSuccess);
+		
+		SecRequirementRef SandboxRequirement = nullptr;
+		Err = SecRequirementCreateWithString(CFSTR("entitlement[\"com.apple.security.app-sandbox\"] exists"), kSecCSDefaultFlags, &SandboxRequirement);
+		check(Err == errSecSuccess && SandboxRequirement);
+		
+		Err = SecStaticCodeCheckValidityWithErrors(SecCodeObj, kSecCSDefaultFlags, SandboxRequirement, nullptr);
+		
+		bIsSandboxedApplication = (Err == errSecSuccess);
+		
+		CFRelease(SecCodeObj);
+	}
+	
+	return bIsSandboxedApplication;
+}
+
 void FMacPlatformProcess::CleanFileCache()
 {
 	bool bShouldCleanShaderWorkingDirectory = true;
diff --git a/Engine/Source/Runtime/Core/Public/Mac/MacPlatformCrashContext.h b/Engine/Source/Runtime/Core/Public/Mac/MacPlatformCrashContext.h
--- a/Engine/Source/Runtime/Core/Public/Mac/MacPlatformCrashContext.h
+++ b/Engine/Source/Runtime/Core/Public/Mac/MacPlatformCrashContext.h
@@ -9,8 +9,8 @@ struct CORE_API FMacCrashContext : public FApplePlatformCrashContext
 	/** Mimics Windows WER format */
 	void GenerateWindowsErrorReport(char const* WERPath) const;
 
-	/** Creates (fake so far) minidump */
-	void GenerateMinidump(char const* Path) const;
+	/** Copies the PLCrashReporter minidump */
+	void CopyMinidump(char const* OutputPath, char const* InputPath) const;
 
 	/** Generates information for crash reporter */
 	void GenerateCrashInfoAndLaunchReporter() const;
diff --git a/Engine/Source/Runtime/Core/Public/Mac/MacPlatformProcess.h b/Engine/Source/Runtime/Core/Public/Mac/MacPlatformProcess.h
--- a/Engine/Source/Runtime/Core/Public/Mac/MacPlatformProcess.h
+++ b/Engine/Source/Runtime/Core/Public/Mac/MacPlatformProcess.h
@@ -60,6 +60,7 @@ struct CORE_API FMacPlatformProcess : public FGenericPlatformProcess
 	static bool IsApplicationRunning( uint32 ProcessId );
 	static FString GetApplicationName( uint32 ProcessId );
 	static bool IsThisApplicationForeground();
+	static bool IsSandboxedApplication();
 	static void LaunchFileInDefaultExternalApplication( const TCHAR* FileName, const TCHAR* Parms = NULL, ELaunchVerb::Type Verb = ELaunchVerb::Open );
 	static void ExploreFolder( const TCHAR* FilePath );
 	static FRunnableThread* CreateRunnableThread();
