Rewritten debug symbol handling for OS X to allow creation of dSYM bundles & symbol stripping of executables.
- The MacToolChain will emit dSYMs and strip executables when the UBT configuration enables bGeneratedSYMFile, just like iOS.
- Symbol stripping requires generating dSYMs to prevent creation of non-debuggable builds whose crash reports would be unresolvable.
- To avoid a dependency on the Private framework CoreSymbolication all symbols from that framework are loaded dynamically & can only be used within programs, not the game or the editor, as CoreSymbolication is incompatible with non-ANSI malloc implementations.
- Added an initial platform-agnostic API for querying debug symbol info, including a generic database format that can be queried on otherwise incompatible platforms.
- Added UnrealAtoS that emulates Apple's atos to resolve symbols using the generic database or the platform API (CoreSymbolication on OS X) which on OS X is used by the editor to gather symbol info for CodeView.
- Added DSymExporter which will export Apple debug symbol data from Mach-O binaries, including the payload within a .dSYM bundle, to the generic format so that the crash report server may one day read the data without needing a Mac to symbolicate reports.
- Initial SymbolDebugger & MinidumpDiagnostics support on OS X.

diff --git a/Engine/Source/Developer/CrashDebugHelper/Private/Mac/CrashDebugHelperMac.cpp b/Engine/Source/Developer/CrashDebugHelper/Private/Mac/CrashDebugHelperMac.cpp
--- a/Engine/Source/Developer/CrashDebugHelper/Private/Mac/CrashDebugHelperMac.cpp
+++ b/Engine/Source/Developer/CrashDebugHelper/Private/Mac/CrashDebugHelperMac.cpp
@@ -222,6 +222,7 @@ static int32 ParseExceptionCode(TCHAR const* CrashLog, uint32& OutExceptionCode)
 			}
 			else
 			{
+				ensure(false);
 				OutExceptionCode = SIGUSR1;
 			}
 		}
@@ -335,7 +336,7 @@ static int32 ParseThreadStackLine(TCHAR const* StackLine, FString& OutModuleName
 	return Found;
 }
 
-static int32 SymboliseStackInfo(FApplePlatformSymbolCache* SymbolCache, TArray<FCrashModuleInfo> const& ModuleInfo, FString ModuleName, uint64 const ProgramCounter, FString& OutFunctionName, FString& OutFileName, int32& OutLineNumber)
+static int32 SymboliseStackInfo(FPlatformSymbolDatabaseSet& SymbolCache, TArray<FCrashModuleInfo> const& ModuleInfo, FString ModuleName, uint64 const ProgramCounter, FString& OutFunctionName, FString& OutFileName, int32& OutLineNumber)
 {
 	FProgramCounterSymbolInfo Info;
 	int32 ValuesSymbolised = 0;
@@ -349,8 +350,18 @@ static int32 SymboliseStackInfo(FApplePlatformSymbolCache* SymbolCache, TArray<F
 			break;
 		}
 	}
-
-	if((Module.Name.Len() > 0) && FApplePlatformSymbolication::SymbolInfoForStrippedSymbol(SymbolCache, (ProgramCounter - Module.BaseOfImage), TCHAR_TO_UTF8(*Module.Name), TCHAR_TO_UTF8(*Module.Report), Info))
+	
+	FApplePlatformSymbolDatabase* Db = SymbolCache.Find(Module.Report);
+	if(!Db)
+	{
+		FApplePlatformSymbolDatabase Database;
+		if(FPlatformSymbolication::LoadSymbolDatabaseForBinary(TEXT(""), Module.Name, Module.Report, Database))
+		{
+			SymbolCache.Add(Database);
+			Db = SymbolCache.Find(Module.Report);
+		}
+	}
+	if((Module.Name.Len() > 0) && Db && FPlatformSymbolication::SymbolInfoForStrippedSymbol(*Db, ProgramCounter, Module.BaseOfImage, Module.Report, Info))
 	{
 		if(FCStringAnsi::Strlen(Info.FunctionName) > 0)
 		{
@@ -643,7 +654,8 @@ bool FCrashDebugHelperMac::CreateMinidumpDiagnosticReport( const FString& InCras
 				}
 			}
 			
-			FApplePlatformSymbolCache* SymbolCache = FApplePlatformSymbolication::CreateSymbolCache();
+			
+			FPlatformSymbolDatabaseSet SymbolCache;
 			
 			TCHAR const* ThreadStackLine = FindCrashedThreadStack(*CrashDump);
 			while(ThreadStackLine)
@@ -695,8 +707,6 @@ bool FCrashDebugHelperMac::CreateMinidumpDiagnosticReport( const FString& InCras
 				}
 			}
 			
-			FApplePlatformSymbolication::DestroySymbolCache(SymbolCache);
-			
 			CrashInfo.Threads.Push(ThreadInfo);
 			
 			bOK = true;
diff --git a/Engine/Source/Editor/UnrealEd/Private/SourceCodeNavigation.cpp b/Engine/Source/Editor/UnrealEd/Private/SourceCodeNavigation.cpp
--- a/Engine/Source/Editor/UnrealEd/Private/SourceCodeNavigation.cpp
+++ b/Engine/Source/Editor/UnrealEd/Private/SourceCodeNavigation.cpp
@@ -523,9 +523,10 @@ void FSourceCodeNavigationImpl::NavigateToFunctionSource( const FString& Functio
 			IndexModule64 = (struct mach_header_64 const*)IndexModule32;
 			LoadCommands = (struct load_command const*)(IndexModule64 + 1);
 			struct load_command const* Command = LoadCommands;
-			struct symtab_command const* SymbolTable = NULL;
-			struct dysymtab_command const* DsymTable = NULL;
-			for(uint32 CommandIndex = 0; CommandIndex < IndexModule32->ncmds; CommandIndex++)
+			struct symtab_command const* SymbolTable = nullptr;
+			struct dysymtab_command const* DsymTable = nullptr;
+			struct uuid_command* UUIDCommand = nullptr;
+			for(uint32 CommandIndex = 0; CommandIndex < IndexModule64->ncmds; CommandIndex++)
 			{
 				if (Command && Command->cmd == LC_SYMTAB)
 				{
@@ -535,10 +536,14 @@ void FSourceCodeNavigationImpl::NavigateToFunctionSource( const FString& Functio
 				{
 					DsymTable = (struct dysymtab_command const*)Command;
 				}
+				else if (Command && Command->cmd == LC_UUID)
+				{
+					UUIDCommand = (struct uuid_command*)Command;
+				}
 				Command = (struct load_command const*)(((char const*)Command) + Command->cmdsize);
 			}
 			
-			check(SymbolTable && DsymTable);
+			check(SymbolTable && DsymTable && UUIDCommand);
 			
 			IPlatformFile& PlatformFile = IPlatformFile::GetPlatformPhysical();
 			IFileHandle* File = PlatformFile.OpenRead(*FullModulePath);
@@ -624,11 +629,39 @@ void FSourceCodeNavigationImpl::NavigateToFunctionSource( const FString& Functio
 							
 							if(FunctionSymbolName == SymbolName)
 							{
-								FString ModuleName = FPaths::GetCleanFilename(FullModulePath);
-								FProgramCounterSymbolInfo Info;
-								if(FApplePlatformSymbolication::SymbolInfoForFunctionFromModule((char const*)(StringTable+SymbolEntry.n_un.n_strx), TCHAR_TO_UTF8(*ModuleName), Info))
+								CFUUIDBytes UUIDBytes;
+								FMemory::Memcpy(&UUIDBytes, UUIDCommand->uuid, sizeof(CFUUIDBytes));
+								CFUUIDRef UUIDRef = CFUUIDCreateFromUUIDBytes(kCFAllocatorDefault, UUIDBytes);
+								CFStringRef UUIDString = CFUUIDCreateString(kCFAllocatorDefault, UUIDRef);
+								FString UUID((NSString*)UUIDString);
+								CFRelease(UUIDString);
+								CFRelease(UUIDRef);
+							
+								uint64 Address = SymbolEntry.n_value;
+								uint64 BaseAddress = (uint64)IndexModule64;
+								FString AtoSCommand = FString::Printf(TEXT("\"%s\" -s %s -l 0x%lx 0x%lx"), *FullModulePath, *UUID, BaseAddress, Address);
+								int32 ReturnCode = 0;
+								FString Results;
+								
+								const FString AtoSPath = FString::Printf(TEXT("%sBinaries/Mac/UnrealAtoS"), *FPaths::EngineDir() );
+								FPlatformProcess::ExecProcess( *AtoSPath, *AtoSCommand, &ReturnCode, &Results, NULL );
+								if(ReturnCode == 0)
 								{
-									SourceCodeAccessor.OpenFileAtLine( Info.Filename, Info.LineNumber, 0 );
+									int32 FirstIndex = -1;
+									int32 LastIndex = -1;
+									if(Results.FindChar(TCHAR('('), FirstIndex) && Results.FindLastChar(TCHAR('('), LastIndex) && FirstIndex != LastIndex)
+									{
+										int32 CloseIndex = -1;
+										int32 ColonIndex = -1;
+										if(Results.FindLastChar(TCHAR(':'), ColonIndex) && Results.FindLastChar(TCHAR(')'), CloseIndex))
+										{
+											int32 FileNamePos = LastIndex+1;
+											int32 FileNameLen = ColonIndex-FileNamePos;
+											FString FileName = Results.Mid(FileNamePos, FileNameLen);
+											FString LineNumber = Results.Mid(ColonIndex + 1, CloseIndex-(ColonIndex + 1));
+											SourceCodeAccessor.OpenFileAtLine( FileName, FCString::Atoi(*LineNumber), 0 );
+										}
+									}
 								}
 								break;
 							}
diff --git a/Engine/Source/Programs/CrashReporter/CrashReportClient/Private/Mac/CrashReportClientMainMac.cpp b/Engine/Source/Programs/CrashReporter/CrashReportClient/Private/Mac/CrashReportClientMainMac.cpp
--- a/Engine/Source/Programs/CrashReporter/CrashReportClient/Private/Mac/CrashReportClientMainMac.cpp
+++ b/Engine/Source/Programs/CrashReporter/CrashReportClient/Private/Mac/CrashReportClientMainMac.cpp
@@ -6,28 +6,13 @@
 #include "CocoaThread.h"
 
 /**
- * Because crash reporters can crash, too
+ * Because crash reporters can crash, too - only used for Sandboxed applications
  */
 void CrashReporterCrashHandler(const FGenericCrashContext& GenericContext)
 {
-	const FMacCrashContext& Context = static_cast< const FMacCrashContext& >( GenericContext );
-
-	Context.ReportCrash();
-	if (GLog)
-	{
-		GLog->SetCurrentThreadAsMasterThread();
-		GLog->Flush();
-	}
-	if (GWarn)
-	{
-		GWarn->Flush();
-	}
-	if (GError)
-	{
-		GError->Flush();
-		GError->HandleError();
-	}
-	FPlatformMisc::RequestExit(true);
+	// We never emit a crash report in Sandboxed builds from CRC as if we do, then the crashed application's
+	// crash report is overwritten by the CRC's when trampolining into the Apple Crash Reporter.
+	_Exit(0);
 }
 
 static FString GSavedCommandLine;
@@ -54,7 +39,11 @@ static FString GSavedCommandLine;
 - (void) runGameThread:(id)Arg
 {
 	FPlatformMisc::SetGracefulTerminationHandler();
-	FPlatformMisc::SetCrashHandler(&CrashReporterCrashHandler);
+	// For sandboxed applications CRC can never report a crash, or we break trampolining into Apple's crash reporter.
+	if(FPlatformProcess::IsSandboxedApplication())
+	{
+		FPlatformMisc::SetCrashHandler(CrashReporterCrashHandler);
+	}
 	
 	RunCrashReportClient(*GSavedCommandLine);
 	
diff --git a/Engine/Source/Programs/CrashReporter/MinidumpDiagnostics/MinidumpDiagnostics.Target.cs b/Engine/Source/Programs/CrashReporter/MinidumpDiagnostics/MinidumpDiagnostics.Target.cs
--- a/Engine/Source/Programs/CrashReporter/MinidumpDiagnostics/MinidumpDiagnostics.Target.cs
+++ b/Engine/Source/Programs/CrashReporter/MinidumpDiagnostics/MinidumpDiagnostics.Target.cs
@@ -14,6 +14,7 @@ public class MinidumpDiagnosticsTarget : TargetRules
 	public override bool GetSupportedPlatforms( ref List<UnrealTargetPlatform> OutPlatforms )
 	{
 		OutPlatforms.Add( UnrealTargetPlatform.Win64 );
+		OutPlatforms.Add( UnrealTargetPlatform.Mac );
 		return true;
 	}
 
diff --git a/Engine/Source/Programs/CrashReporter/MinidumpDiagnostics/Private/Mac/MinidumpDiagnosticsMainMac.cpp b/Engine/Source/Programs/CrashReporter/MinidumpDiagnostics/Private/Mac/MinidumpDiagnosticsMainMac.cpp
new file mode 100644
--- /dev/null
+++ b/Engine/Source/Programs/CrashReporter/MinidumpDiagnostics/Private/Mac/MinidumpDiagnosticsMainMac.cpp
@@ -0,0 +1,25 @@
+// Copyright 1998-2015 Epic Games, Inc. All Rights Reserved.
+
+#include "MinidumpDiagnosticsApp.h"
+
+#include "Core.h"
+#include "RequiredProgramMainCPPInclude.h"
+#include "ExceptionHandling.h"
+
+IMPLEMENT_APPLICATION( MinidumpDiagnostics, "MinidumpDiagnostics" )
+
+ // Main entry point to the application
+INT32_MAIN_INT32_ARGC_TCHAR_ARGV()
+{
+	FPlatformMisc::SetGracefulTerminationHandler();
+	FPlatformMisc::SetCrashHandler(nullptr);
+	
+	GEngineLoop.PreInit( ArgC, ArgV );
+
+	const int32 Result = RunMinidumpDiagnostics( ArgC, ArgV );
+	
+	FEngineLoop::AppPreExit();
+	FEngineLoop::AppExit();
+	
+	return Result;
+}
diff --git a/Engine/Source/Programs/Mac/DsymExporter/DsymExporter.Build.cs b/Engine/Source/Programs/Mac/DsymExporter/DsymExporter.Build.cs
new file mode 100644
--- /dev/null
+++ b/Engine/Source/Programs/Mac/DsymExporter/DsymExporter.Build.cs
@@ -0,0 +1,20 @@
+// Copyright 1998-2015 Epic Games, Inc. All Rights Reserved.
+
+using UnrealBuildTool;
+
+public class DsymExporter : ModuleRules
+{
+	public DsymExporter( TargetInfo Target )
+	{
+		PrivateIncludePathModuleNames.Add( "Launch" );
+		PrivateIncludePaths.Add( "Runtime/Launch/Private" );
+	
+		PrivateDependencyModuleNames.AddRange(
+			new string[] 
+			{ 
+				"Core",
+				"Projects"
+			}
+			);
+	}
+}
diff --git a/Engine/Source/Programs/Mac/DsymExporter/DsymExporter.Target.cs b/Engine/Source/Programs/Mac/DsymExporter/DsymExporter.Target.cs
new file mode 100644
--- /dev/null
+++ b/Engine/Source/Programs/Mac/DsymExporter/DsymExporter.Target.cs
@@ -0,0 +1,64 @@
+// Copyright 1998-2015 Epic Games, Inc. All Rights Reserved.
+
+using UnrealBuildTool;
+using System.Collections.Generic;
+
+public class DsymExporterTarget : TargetRules
+{
+	public DsymExporterTarget( TargetInfo Target )
+	{
+		Type = TargetType.Program;
+	}
+
+	// TargetRules interface.
+	public override bool GetSupportedPlatforms( ref List<UnrealTargetPlatform> OutPlatforms )
+	{
+		OutPlatforms.Add( UnrealTargetPlatform.Mac );
+		return true;
+	}
+
+	public override void SetupBinaries(
+		TargetInfo Target,
+		ref List<UEBuildBinaryConfiguration> OutBuildBinaryConfigurations,
+		ref List<string> OutExtraModuleNames
+		)
+	{
+		OutBuildBinaryConfigurations.Add(
+			new UEBuildBinaryConfiguration(	InType: UEBuildBinaryType.Executable,
+				InModuleNames: new List<string>() { "DsymExporter" })
+			);
+	}
+
+	public override bool ShouldCompileMonolithic(UnrealTargetPlatform InPlatform, UnrealTargetConfiguration InConfiguration)
+	{
+		return true;
+	}
+
+	public override void SetupGlobalEnvironment(
+		TargetInfo Target,
+		ref LinkEnvironmentConfiguration OutLinkEnvironmentConfiguration,
+		ref CPPEnvironmentConfiguration OutCPPEnvironmentConfiguration
+		)
+	{
+		UEBuildConfiguration.bCompileLeanAndMeanUE = true;
+
+		// Don't need editor
+		UEBuildConfiguration.bBuildEditor = false;
+
+		// DsymExporter doesn't ever compile with the engine linked in
+		UEBuildConfiguration.bCompileAgainstEngine = false;
+
+		// DsymExporter has no exports, so no need to verify that a .lib and .exp file was emitted by the linker.
+		OutLinkEnvironmentConfiguration.bHasExports = false;
+
+		// Do NOT produce additional console app exe
+		OutLinkEnvironmentConfiguration.bIsBuildingConsoleApplication = true;
+	}
+    public override bool GUBP_AlwaysBuildWithTools(UnrealTargetPlatform InHostPlatform, bool bBuildingRocket, out bool bInternalToolOnly, out bool SeparateNode, out bool CrossCompile)
+    {
+        bInternalToolOnly = false;
+        SeparateNode = false;
+		CrossCompile = false;
+        return true;
+    }
+}
diff --git a/Engine/Source/Programs/Mac/DsymExporter/Private/DsymExporterApp.cpp b/Engine/Source/Programs/Mac/DsymExporter/Private/DsymExporterApp.cpp
new file mode 100644
--- /dev/null
+++ b/Engine/Source/Programs/Mac/DsymExporter/Private/DsymExporterApp.cpp
@@ -0,0 +1,54 @@
+// Copyright 1998-2015 Epic Games, Inc. All Rights Reserved.
+
+#include "DsymExporterApp.h"
+
+#include "GenericPlatformSymbolication.h"
+#if PLATFORM_MAC
+#include "ApplePlatformSymbolication.h"
+#endif
+#include "Serialization/Archive.h"
+
+int32 RunDsymExporter(int32 ArgC, TCHAR* Argv[])
+{
+	// Make sure we have at least a single parameter
+	if( ArgC < 2 )
+	{
+		UE_LOG( LogInit, Error, TEXT( "DsymExporter - not enough parameters." ) );
+		UE_LOG( LogInit, Error, TEXT( " ... usage: DsymExporter <Mach-O Binary Path> [Output Folder]" ) );
+		UE_LOG( LogInit, Error, TEXT( "<Mach-O Binary Path>: This is an absolute path to a Mach-O binary containing symbols, which may be the payload binary within an application, framework or dSYM bundle, an executable or dylib." ) );
+		UE_LOG( LogInit, Error, TEXT( "[Output Folder]: The folder to write the new symbol database to, the database will take the filename of the input plus the .udebugsymbols extension." ) );
+		return 1;
+	}
+	
+#if PLATFORM_MAC
+	FApplePlatformSymbolication::EnableCoreSymbolication(true);
+#endif
+	
+	FPlatformSymbolDatabase Symbols;
+	FString Signature;
+	bool bOK = FPlatformSymbolication::LoadSymbolDatabaseForBinary(TEXT(""), Argv[1], Signature, Symbols);
+	if(bOK)
+	{
+		FString OutputFolder = FPaths::GetPath(Argv[1]);
+		if(ArgC == 3)
+		{
+			OutputFolder = Argv[2];
+		}
+		
+		if(FPlatformSymbolication::SaveSymbolDatabaseForBinary(OutputFolder, FPaths::GetBaseFilename(Argv[1]), Symbols))
+		{
+			return 0;
+		}
+		else
+		{
+			return 1;
+		}
+	}
+	else
+	{
+		UE_LOG( LogInit, Error, TEXT( "DsymExporter - unable to parse debug symbols for Mach-O file." ) );
+		return 1;
+	}
+}
+
+
diff --git a/Engine/Source/Programs/Mac/DsymExporter/Private/Mac/DsymExporterMainMac.cpp b/Engine/Source/Programs/Mac/DsymExporter/Private/Mac/DsymExporterMainMac.cpp
new file mode 100644
--- /dev/null
+++ b/Engine/Source/Programs/Mac/DsymExporter/Private/Mac/DsymExporterMainMac.cpp
@@ -0,0 +1,25 @@
+// Copyright 1998-2015 Epic Games, Inc. All Rights Reserved.
+
+#include "DsymExporterApp.h"
+
+#include "Core.h"
+#include "RequiredProgramMainCPPInclude.h"
+#include "ExceptionHandling.h"
+
+IMPLEMENT_APPLICATION( DsymExporter, "DsymExporter" )
+
+ // Main entry point to the application
+INT32_MAIN_INT32_ARGC_TCHAR_ARGV()
+{
+	FPlatformMisc::SetGracefulTerminationHandler();
+	FPlatformMisc::SetCrashHandler(nullptr);
+	
+	GEngineLoop.PreInit( ArgC, ArgV );
+
+	const int32 Result = RunDsymExporter( ArgC, ArgV );
+	
+	FEngineLoop::AppPreExit();
+	FEngineLoop::AppExit();
+	
+	return Result;
+}
diff --git a/Engine/Source/Programs/Mac/DsymExporter/Public/DsymExporterApp.h b/Engine/Source/Programs/Mac/DsymExporter/Public/DsymExporterApp.h
new file mode 100644
--- /dev/null
+++ b/Engine/Source/Programs/Mac/DsymExporter/Public/DsymExporterApp.h
@@ -0,0 +1,10 @@
+// Copyright 1998-2015 Epic Games, Inc. All Rights Reserved.
+
+#pragma once
+
+#include "Core.h"
+
+// Export symbol data from Apple's .dSYM format to something that can be parsed on non-Apple platforms
+int32 RunDsymExporter(int32 ArgC, TCHAR* Argv[]);
+
+
diff --git a/Engine/Source/Programs/Mac/UnrealAtoS/Private/Mac/UnrealAtoSMainMac.cpp b/Engine/Source/Programs/Mac/UnrealAtoS/Private/Mac/UnrealAtoSMainMac.cpp
new file mode 100644
--- /dev/null
+++ b/Engine/Source/Programs/Mac/UnrealAtoS/Private/Mac/UnrealAtoSMainMac.cpp
@@ -0,0 +1,25 @@
+// Copyright 1998-2015 Epic Games, Inc. All Rights Reserved.
+
+#include "UnrealAtoSApp.h"
+
+#include "Core.h"
+#include "RequiredProgramMainCPPInclude.h"
+#include "ExceptionHandling.h"
+
+IMPLEMENT_APPLICATION( UnrealAtoS, "UnrealAtoS" )
+
+ // Main entry point to the application
+INT32_MAIN_INT32_ARGC_TCHAR_ARGV()
+{
+	FPlatformMisc::SetGracefulTerminationHandler();
+	FPlatformMisc::SetCrashHandler(nullptr);
+	
+	GEngineLoop.PreInit( ArgC, ArgV );
+
+	const int32 Result = RunUnrealAtoS( ArgC, ArgV );
+	
+	FEngineLoop::AppPreExit();
+	FEngineLoop::AppExit();
+	
+	return Result;
+}
diff --git a/Engine/Source/Programs/Mac/UnrealAtoS/Private/UnrealAtoSApp.cpp b/Engine/Source/Programs/Mac/UnrealAtoS/Private/UnrealAtoSApp.cpp
new file mode 100644
--- /dev/null
+++ b/Engine/Source/Programs/Mac/UnrealAtoS/Private/UnrealAtoSApp.cpp
@@ -0,0 +1,130 @@
+// Copyright 1998-2015 Epic Games, Inc. All Rights Reserved.
+
+#include "UnrealAtoSApp.h"
+
+#include "GenericPlatformSymbolication.h"
+#if PLATFORM_MAC
+#include "ApplePlatformSymbolication.h"
+#endif
+
+static void ShowUsage()
+{
+	UE_LOG( LogInit, Error, TEXT( "UnrealAtoS - not enough parameters." ) );
+	UE_LOG( LogInit, Error, TEXT( " ... usage: UnrealAtoS <binary-image> [-s <binary-signature>] [-d <usymbol-path>] [-l <module-load-address>] <program-counter>" ) );
+}
+
+int32 RunUnrealAtoS(int32 ArgC, TCHAR* Argv[])
+{
+	// Make sure we have at least a single parameter
+	if( ArgC < 4 )
+	{
+		ShowUsage();
+		return 1;
+	}
+	
+#if PLATFORM_MAC
+	FApplePlatformSymbolication::EnableCoreSymbolication(true);
+#endif
+
+	TCHAR* ModulePath = Argv[1];
+	TCHAR* ModuleSig = nullptr;
+	TCHAR* SymbolPath = nullptr;
+	uint64 ModuleBase = 0;
+	uint64 ProgramCounter = FCStringWide::Strtoui64(Argv[ArgC - 1], nullptr, 16);
+	
+	for(int32 ArgI = 2; ArgI < ArgC - 1;)
+	{
+		if(FCStringWide::Strcmp(Argv[ArgI], TEXT("-d")) == 0)
+		{
+			if(ArgC - 1 > ++ArgI)
+			{
+				SymbolPath = Argv[ArgI++];
+			}
+			else
+			{
+				ShowUsage();
+				return 1;
+			}
+		}
+		else if(FCStringWide::Strcmp(Argv[ArgI], TEXT("-l")) == 0)
+		{
+			if(ArgC - 1 > ++ArgI)
+			{
+				ModuleBase = FCStringWide::Strtoui64(Argv[ArgI++], nullptr, 16);
+			}
+			else
+			{
+				ShowUsage();
+				return 1;
+			}
+		}
+		else if(FCStringWide::Strcmp(Argv[ArgI], TEXT("-s")) == 0)
+		{
+			if(ArgC - 1 > ++ArgI)
+			{
+				ModuleSig = Argv[ArgI++];
+			}
+			else
+			{
+				ShowUsage();
+				return 1;
+			}
+		}
+	}
+	
+	if(ModulePath && ProgramCounter)
+	{
+		uint64 Address = ProgramCounter > ModuleBase ? (ProgramCounter - ModuleBase) : ProgramCounter;
+		FProgramCounterSymbolInfo Info;
+		
+		bool bOK = false;
+		
+		FString ModuleSignature = ModuleSig;
+		FPlatformSymbolDatabase SymbolDb;
+		if(FPlatformSymbolication::LoadSymbolDatabaseForBinary(FString(SymbolPath), ModulePath, ModuleSignature, SymbolDb))
+		{
+			bOK = FPlatformSymbolication::SymbolInfoForStrippedSymbol(SymbolDb, ProgramCounter, ModuleBase, ModuleSignature, Info);
+		}
+
+		if(bOK)
+		{
+			FString FunctionName = Info.FunctionName;
+			if(FunctionName.IsEmpty())
+			{
+				FunctionName = TEXT("[Unknown]");
+			}
+			
+			FString ModuleName = FPaths::GetCleanFilename(Info.ModuleName);
+			if(ModuleName.IsEmpty())
+			{
+				ModuleName = TEXT("[Unknown]");
+			}
+			
+			FString Output;
+			FString FileName = FPaths::GetCleanFilename(Info.Filename);
+			if(!FileName.IsEmpty() && Info.LineNumber > 0)
+			{
+				Output = FString::Printf(TEXT("%s (in %s) (%s:%d)\n"), *FunctionName, *ModuleName, *FileName, Info.LineNumber);
+			}
+			else
+			{
+				Output = FString::Printf(TEXT("%s (in %s) + %d\n"), *FunctionName, *ModuleName, Info.SymbolDisplacement);
+			}
+		
+			FPlatformMisc::LocalPrint(*Output);
+			return 0;
+		}
+		else
+		{
+			UE_LOG( LogInit, Error, TEXT( "UnrealAtoS - no such symbol." ) );
+			return 1;
+		}
+	}
+	else
+	{
+		UE_LOG( LogInit, Error, TEXT( "UnrealAtoS - invalid arguments." ) );
+		return 1;
+	}
+}
+
+
diff --git a/Engine/Source/Programs/Mac/UnrealAtoS/Public/UnrealAtoSApp.h b/Engine/Source/Programs/Mac/UnrealAtoS/Public/UnrealAtoSApp.h
new file mode 100644
--- /dev/null
+++ b/Engine/Source/Programs/Mac/UnrealAtoS/Public/UnrealAtoSApp.h
@@ -0,0 +1,10 @@
+// Copyright 1998-2015 Epic Games, Inc. All Rights Reserved.
+
+#pragma once
+
+#include "Core.h"
+
+// Unreal equivalent to Apple's atos that resolves debug symbol info. from a program counter
+int32 RunUnrealAtoS(int32 ArgC, TCHAR* Argv[]);
+
+
diff --git a/Engine/Source/Programs/Mac/UnrealAtoS/UnrealAtoS.Build.cs b/Engine/Source/Programs/Mac/UnrealAtoS/UnrealAtoS.Build.cs
new file mode 100644
--- /dev/null
+++ b/Engine/Source/Programs/Mac/UnrealAtoS/UnrealAtoS.Build.cs
@@ -0,0 +1,20 @@
+// Copyright 1998-2015 Epic Games, Inc. All Rights Reserved.
+
+using UnrealBuildTool;
+
+public class UnrealAtoS : ModuleRules
+{
+	public UnrealAtoS( TargetInfo Target )
+	{
+		PrivateIncludePathModuleNames.Add( "Launch" );
+		PrivateIncludePaths.Add( "Runtime/Launch/Private" );
+	
+		PrivateDependencyModuleNames.AddRange(
+			new string[] 
+			{ 
+				"Core",
+				"Projects"
+			}
+			);
+	}
+}
diff --git a/Engine/Source/Programs/Mac/UnrealAtoS/UnrealAtoS.Target.cs b/Engine/Source/Programs/Mac/UnrealAtoS/UnrealAtoS.Target.cs
new file mode 100644
--- /dev/null
+++ b/Engine/Source/Programs/Mac/UnrealAtoS/UnrealAtoS.Target.cs
@@ -0,0 +1,64 @@
+// Copyright 1998-2015 Epic Games, Inc. All Rights Reserved.
+
+using UnrealBuildTool;
+using System.Collections.Generic;
+
+public class UnrealAtoSTarget : TargetRules
+{
+	public UnrealAtoSTarget( TargetInfo Target )
+	{
+		Type = TargetType.Program;
+	}
+
+	// TargetRules interface.
+	public override bool GetSupportedPlatforms( ref List<UnrealTargetPlatform> OutPlatforms )
+	{
+		OutPlatforms.Add( UnrealTargetPlatform.Mac );
+		return true;
+	}
+
+	public override void SetupBinaries(
+		TargetInfo Target,
+		ref List<UEBuildBinaryConfiguration> OutBuildBinaryConfigurations,
+		ref List<string> OutExtraModuleNames
+		)
+	{
+		OutBuildBinaryConfigurations.Add(
+			new UEBuildBinaryConfiguration(	InType: UEBuildBinaryType.Executable,
+				InModuleNames: new List<string>() { "UnrealAtoS" })
+			);
+	}
+
+	public override bool ShouldCompileMonolithic(UnrealTargetPlatform InPlatform, UnrealTargetConfiguration InConfiguration)
+	{
+		return true;
+	}
+
+	public override void SetupGlobalEnvironment(
+		TargetInfo Target,
+		ref LinkEnvironmentConfiguration OutLinkEnvironmentConfiguration,
+		ref CPPEnvironmentConfiguration OutCPPEnvironmentConfiguration
+		)
+	{
+		UEBuildConfiguration.bCompileLeanAndMeanUE = true;
+
+		// Don't need editor
+		UEBuildConfiguration.bBuildEditor = false;
+
+		// DsymExporter doesn't ever compile with the engine linked in
+		UEBuildConfiguration.bCompileAgainstEngine = false;
+
+		// DsymExporter has no exports, so no need to verify that a .lib and .exp file was emitted by the linker.
+		OutLinkEnvironmentConfiguration.bHasExports = false;
+
+		// Do NOT produce additional console app exe
+		OutLinkEnvironmentConfiguration.bIsBuildingConsoleApplication = true;
+	}
+    public override bool GUBP_AlwaysBuildWithTools(UnrealTargetPlatform InHostPlatform, bool bBuildingRocket, out bool bInternalToolOnly, out bool SeparateNode, out bool CrossCompile)
+    {
+        bInternalToolOnly = false;
+        SeparateNode = false;
+		CrossCompile = false;
+        return true;
+    }
+}
diff --git a/Engine/Source/Programs/SymbolDebugger/SymbolDebugger.Target.cs b/Engine/Source/Programs/SymbolDebugger/SymbolDebugger.Target.cs
--- a/Engine/Source/Programs/SymbolDebugger/SymbolDebugger.Target.cs
+++ b/Engine/Source/Programs/SymbolDebugger/SymbolDebugger.Target.cs
@@ -16,6 +16,7 @@ public class SymbolDebuggerTarget : TargetRules
     public override bool GetSupportedPlatforms(ref List<UnrealTargetPlatform> OutPlatforms)
     {
         OutPlatforms.Add(UnrealTargetPlatform.Win64);
+        OutPlatforms.Add(UnrealTargetPlatform.Mac);
         return true;
     }
 
diff --git a/Engine/Source/Programs/UnrealBuildTool/Mac/MacToolChain.cs b/Engine/Source/Programs/UnrealBuildTool/Mac/MacToolChain.cs
--- a/Engine/Source/Programs/UnrealBuildTool/Mac/MacToolChain.cs
+++ b/Engine/Source/Programs/UnrealBuildTool/Mac/MacToolChain.cs
@@ -1077,6 +1077,13 @@ namespace UnrealBuildTool
 				}
 			}
 
+			// For Mac, generate the dSYM file if the config file is set to do so
+			if (BuildConfiguration.bGeneratedSYMFile == true && (!bIsBuildingLibrary || LinkEnvironment.Config.bIsBuildingDLL))
+			{
+				Log.TraceInformation("Generating dSYM file for {0} - this will add some time to your build...", Path.GetFileName(OutputFile.AbsolutePath));
+				RemoteOutputFile = GenerateDebugInfo(OutputFile);
+			}
+
 			return RemoteOutputFile;
 		}
 
@@ -1123,9 +1130,60 @@ namespace UnrealBuildTool
 			return RemoteOutputFile;
 		}
 
+		private static Dictionary<Action, string> DebugOutputMap = new Dictionary<Action, string>();
+		static public void RPCDebugInfoActionHandler(Action Action, out int ExitCode, out string Output)
+		{
+			RPCUtilHelper.RPCActionHandler (Action, out ExitCode, out Output);
+			if (DebugOutputMap.ContainsKey (Action)) 
+			{
+				if (ExitCode == 0) 
+				{
+					RPCUtilHelper.CopyDirectory (Action.ProducedItems[0].AbsolutePath, DebugOutputMap[Action], RPCUtilHelper.ECopyOptions.None);
+				}
+				DebugOutputMap.Remove (Action);
+			}
+		}
+
 		/**
-		 * Creates app bundle for a given executable
+		 * Generates debug info for a given executable
 		 * 
+		 * @param MachOBinary FileItem describing the executable or dylib to generate debug info for
+		 */
+		public FileItem GenerateDebugInfo(FileItem MachOBinary)
+		{
+			// Make a file item for the source and destination files
+			string FullDestPath = MachOBinary.AbsolutePath + ".dSYM";
+
+			FileItem OutputFile = FileItem.GetItemByPath(FullDestPath);
+			FileItem DestFile = LocalToRemoteFileItem(OutputFile, false);
+			FileItem InputFile = LocalToRemoteFileItem(MachOBinary, false);
+
+			// Make the compile action
+			Action GenDebugAction = new Action(ActionType.GenerateDebugInfo);
+			if (BuildHostPlatform.Current.Platform != UnrealTargetPlatform.Mac)
+			{
+				DebugOutputMap.Add (GenDebugAction, OutputFile.AbsolutePath);
+				GenDebugAction.ActionHandler = new Action.BlockingActionHandler(MacToolChain.RPCDebugInfoActionHandler);
+			}
+			GenDebugAction.WorkingDirectory = Path.GetFullPath(".");
+			GenDebugAction.CommandPath = "sh";
+
+			// note that the source and dest are switched from a copy command
+			GenDebugAction.CommandArguments = string.Format("-c '\"{0}\"usr/bin/xcrun dsymutil \"{1}\" -o \"{2}\"; \"{0}\"usr/bin/xcrun strip -S -X -x \"{1}\"'",
+				XcodeDeveloperDir,
+				InputFile.AbsolutePath,
+				DestFile.AbsolutePath);
+			GenDebugAction.PrerequisiteItems.Add(InputFile);
+			GenDebugAction.ProducedItems.Add(DestFile);
+			GenDebugAction.StatusDescription = GenDebugAction.CommandArguments;
+			GenDebugAction.bCanExecuteRemotely = false;
+
+			return DestFile;
+		}
+
+		/**
+		 * Creates app bundle for a given executable
+		 *
 		 * @param Executable FileItem describing the executable to generate app bundle for
 		 */
 		FileItem FinalizeAppBundle(LinkEnvironment LinkEnvironment, FileItem Executable, FileItem FixDylibOutputFile)
diff --git a/Engine/Source/Programs/UnrealBuildTool/Mac/UEBuildMac.cs b/Engine/Source/Programs/UnrealBuildTool/Mac/UEBuildMac.cs
--- a/Engine/Source/Programs/UnrealBuildTool/Mac/UEBuildMac.cs
+++ b/Engine/Source/Programs/UnrealBuildTool/Mac/UEBuildMac.cs
@@ -136,7 +136,10 @@ namespace UnrealBuildTool
 
 			InBuildTarget.GlobalCompileEnvironment.Config.Definitions.Add("WITH_TTS=0");
 			InBuildTarget.GlobalCompileEnvironment.Config.Definitions.Add("WITH_SPEECH_RECOGNITION=0");
-			InBuildTarget.GlobalCompileEnvironment.Config.Definitions.Add("WITH_DATABASE_SUPPORT=0");
+			if(!InBuildTarget.GlobalCompileEnvironment.Config.Definitions.Contains("WITH_DATABASE_SUPPORT=0") && !InBuildTarget.GlobalCompileEnvironment.Config.Definitions.Contains("WITH_DATABASE_SUPPORT=1"))
+			{
+				InBuildTarget.GlobalCompileEnvironment.Config.Definitions.Add("WITH_DATABASE_SUPPORT=0");
+			}
 		}
 
 		/**
diff --git a/Engine/Source/Runtime/Core/Core.Build.cs b/Engine/Source/Runtime/Core/Core.Build.cs
--- a/Engine/Source/Runtime/Core/Core.Build.cs
+++ b/Engine/Source/Runtime/Core/Core.Build.cs
@@ -88,7 +88,6 @@ public class Core : ModuleRules
 			{
 				PublicAdditionalLibraries.Add("/System/Library/PrivateFrameworks/MultitouchSupport.framework/Versions/Current/MultitouchSupport");
 			}
-			PublicAdditionalLibraries.Add("/System/Library/PrivateFrameworks/CoreSymbolication.framework/Versions/Current/CoreSymbolication");
 		}
 		else if (Target.Platform == UnrealTargetPlatform.IOS)
 		{
diff --git a/Engine/Source/Runtime/Core/Private/Apple/ApplePlatformCrashContext.cpp b/Engine/Source/Runtime/Core/Private/Apple/ApplePlatformCrashContext.cpp
--- a/Engine/Source/Runtime/Core/Private/Apple/ApplePlatformCrashContext.cpp
+++ b/Engine/Source/Runtime/Core/Private/Apple/ApplePlatformCrashContext.cpp
@@ -6,6 +6,7 @@
 
 #include "CorePrivatePCH.h"
 #include "Apple/ApplePlatformCrashContext.h"
+#include "Apple/ApplePlatformSymbolication.h"
 
 FApplePlatformCrashContext::FApplePlatformCrashContext()
 :	Signal(0)
diff --git a/Engine/Source/Runtime/Core/Private/Apple/ApplePlatformSymbolication.cpp b/Engine/Source/Runtime/Core/Private/Apple/ApplePlatformSymbolication.cpp
--- a/Engine/Source/Runtime/Core/Private/Apple/ApplePlatformSymbolication.cpp
+++ b/Engine/Source/Runtime/Core/Private/Apple/ApplePlatformSymbolication.cpp
@@ -10,238 +10,446 @@
 #include <CoreFoundation/CoreFoundation.h>
 #include <mach/mach.h>
 
-#if PLATFORM_MAC && IS_PROGRAM
 extern "C"
 {
-	struct CSTypeRef
+	struct CSRange
 	{
-		void* Buffer0;
-		void* Buffer1;
+		uint64 Location;
+		uint64 Length;
 	};
 	
+	typedef FApplePlatformSymbolCache CSTypeRef;
 	typedef CSTypeRef CSSymbolicatorRef;
 	typedef CSTypeRef CSSourceInfoRef;
 	typedef CSTypeRef CSSymbolRef;
 	typedef CSTypeRef CSSymbolOwnerRef;
 	
-	struct CSRange
-	{
-		uint64 Location;
-		uint64 Length;
-	};
+	typedef int (^CSSymbolIterator)(CSSymbolRef Symbol);
+	typedef int (^CSSourceInfoIterator)(CSSourceInfoRef SourceInfo);
 	
 	#define kCSNow								0x80000000u
 	
-	Boolean CSEqual(CSTypeRef CS1, CSTypeRef CS2);
-	CFIndex CSGetRetainCount(CSTypeRef CS);
-	Boolean CSIsNull(CSTypeRef CS);
-	CSTypeRef CSRetain(CSTypeRef CS);
-	void CSRelease(CSTypeRef CS);
-	void CSShow(CSTypeRef CS);
+	typedef Boolean (*CSEqualPtr)(CSTypeRef Cs1, CSTypeRef Cs2);
+	typedef Boolean (*CSIsNullPtr)(CSTypeRef CS);
+	typedef void (*CSReleasePtr)(CSTypeRef CS);
+	typedef void (*CSRetainPtr)(CSTypeRef CS);
 	
-	CSSymbolicatorRef CSSymbolicatorCreateWithPid(pid_t pid);
-	CSSymbolicatorRef CSSymbolicatorCreateWithPathAndArchitecture(const char* path, cpu_type_t type);
+	typedef CSSymbolicatorRef (*CSSymbolicatorCreateWithPidPtr)(pid_t pid);
+	typedef CSSymbolicatorRef (*CSSymbolicatorCreateWithPathAndArchitecturePtr)(const char* path, cpu_type_t type);
 	
-	CSSymbolRef CSSymbolicatorGetSymbolWithAddressAtTime(CSSymbolicatorRef Symbolicator, vm_address_t Address, uint64_t Time);
-	CSSourceInfoRef CSSymbolicatorGetSourceInfoWithAddressAtTime(CSSymbolicatorRef Symbolicator, vm_address_t Address, uint64_t Time);
-	CSSymbolRef CSSymbolicatorGetSymbolWithMangledNameFromSymbolOwnerWithNameAtTime(CSSymbolicatorRef Symbolicator, const char* Symbol, const char* Name, uint64_t Time);
-	CSSymbolOwnerRef CSSymbolicatorGetSymbolOwnerWithUUIDAtTime(CSSymbolicatorRef Symbolicator, CFUUIDRef UUID, uint64_t Time);
+	typedef CSSymbolRef (*CSSymbolicatorGetSymbolWithAddressAtTimePtr)(CSSymbolicatorRef Symbolicator, vm_address_t Address, uint64_t Time);
+	typedef CSSourceInfoRef (*CSSymbolicatorGetSourceInfoWithAddressAtTimePtr)(CSSymbolicatorRef Symbolicator, vm_address_t Address, uint64_t Time);
+	typedef CSSymbolOwnerRef (*CSSymbolicatorGetSymbolOwnerWithUUIDAtTimePtr)(CSSymbolicatorRef Symbolicator, CFUUIDRef UUID, uint64_t Time);
+	typedef CSSymbolOwnerRef (*CSSymbolicatorGetSymbolOwnerPtr)(CSSymbolicatorRef cs);
+	typedef int (*CSSymbolicatorForeachSymbolAtTimePtr)(CSSymbolicatorRef Symbolicator, uint64_t Time, CSSymbolIterator It);
 	
-	const char* CSSymbolGetName(CSSymbolRef Symbol);
-	CSRange CSSymbolGetRange(CSSymbolRef Symbol);
-	CSSymbolOwnerRef CSSourceInfoGetSymbolOwner(CSSourceInfoRef Info);
+	typedef const char* (*CSSymbolGetNamePtr)(CSSymbolRef Symbol);
+	typedef CSRange (*CSSymbolGetRangePtr)(CSSymbolRef Symbol);
+	typedef CSSymbolOwnerRef (*CSSourceInfoGetSymbolOwnerPtr)(CSSourceInfoRef Info);
+	typedef CSSymbolOwnerRef (*CSSymbolGetSymbolOwnerPtr)(CSSymbolRef Sym);
+	typedef int (*CSSymbolForeachSourceInfoPtr)(CSSymbolRef Sym, CSSourceInfoIterator It);
 	
-	const char* CSSymbolOwnerGetName(CSSymbolOwnerRef symbol);
-	vm_address_t CSSymbolOwnerGetBaseAddress(CSSymbolOwnerRef Owner);
+	typedef const char* (*CSSymbolOwnerGetNamePtr)(CSSymbolOwnerRef Owner);
+	typedef CFUUIDRef (*CSSymbolOwnerGetUUIDPtr)(CSSymbolOwnerRef Owner);
+	typedef vm_address_t (*CSSymbolOwnerGetBaseAddressPtr)(CSSymbolOwnerRef Owner);
 	
-	int CSSourceInfoGetLineNumber(CSSourceInfoRef Info);
-	const char* CSSourceInfoGetPath(CSSourceInfoRef Info);
-	CSRange CSSourceInfoGetRange(CSSourceInfoRef Info);
-	CSSymbolRef CSSourceInfoGetSymbol(CSSourceInfoRef Info);
+	typedef int (*CSSourceInfoGetLineNumberPtr)(CSSourceInfoRef Info);
+	typedef const char* (*CSSourceInfoGetPathPtr)(CSSourceInfoRef Info);
+	typedef CSRange (*CSSourceInfoGetRangePtr)(CSSourceInfoRef Info);
+	typedef CSSymbolRef (*CSSourceInfoGetSymbolPtr)(CSSourceInfoRef Info);
 }
-#endif
 
-struct FApplePlatformSymbolCache
+static bool GAllowApplePlatformSymbolication = false;
+static void* GCoreSymbolicationHandle = nullptr;
+static CSEqualPtr CSEqual = nullptr;
+static CSIsNullPtr CSIsNull = nullptr;
+static CSReleasePtr CSRelease = nullptr;
+static CSRetainPtr CSRetain = nullptr;
+static CSSymbolicatorCreateWithPidPtr CSSymbolicatorCreateWithPid = nullptr;
+static CSSymbolicatorCreateWithPathAndArchitecturePtr CSSymbolicatorCreateWithPathAndArchitecture = nullptr;
+static CSSymbolicatorGetSymbolWithAddressAtTimePtr CSSymbolicatorGetSymbolWithAddressAtTime = nullptr;
+static CSSymbolicatorGetSourceInfoWithAddressAtTimePtr CSSymbolicatorGetSourceInfoWithAddressAtTime = nullptr;
+static CSSymbolicatorGetSymbolOwnerWithUUIDAtTimePtr CSSymbolicatorGetSymbolOwnerWithUUIDAtTime = nullptr;
+static CSSymbolicatorGetSymbolOwnerPtr CSSymbolicatorGetSymbolOwner = nullptr;
+static CSSymbolicatorForeachSymbolAtTimePtr CSSymbolicatorForeachSymbolAtTime = nullptr;
+static CSSymbolGetNamePtr CSSymbolGetName = nullptr;
+static CSSymbolGetRangePtr CSSymbolGetRange = nullptr;
+static CSSourceInfoGetSymbolOwnerPtr CSSourceInfoGetSymbolOwner = nullptr;
+static CSSymbolGetSymbolOwnerPtr CSSymbolGetSymbolOwner = nullptr;
+static CSSymbolForeachSourceInfoPtr CSSymbolForeachSourceInfo = nullptr;
+static CSSymbolOwnerGetNamePtr CSSymbolOwnerGetName = nullptr;
+static CSSymbolOwnerGetUUIDPtr CSSymbolOwnerGetUUID = nullptr;
+static CSSymbolOwnerGetBaseAddressPtr CSSymbolOwnerGetBaseAddress = nullptr;
+static CSSourceInfoGetLineNumberPtr CSSourceInfoGetLineNumber = nullptr;
+static CSSourceInfoGetPathPtr CSSourceInfoGetPath = nullptr;
+static CSSourceInfoGetRangePtr CSSourceInfoGetRange = nullptr;
+static CSSourceInfoGetSymbolPtr CSSourceInfoGetSymbol = nullptr;
+
+FApplePlatformSymbolDatabase::FApplePlatformSymbolDatabase()
 {
-#if PLATFORM_MAC && IS_PROGRAM
-	TMap< FString, CSSymbolicatorRef > Symbolicators;
-#endif
-};
+	AppleDB.Buffer0 = nullptr;
+	AppleDB.Buffer1 = nullptr;
+}
 
-static bool GAllowApplePlatformSymbolication = (PLATFORM_MAC != 0) && (IS_PROGRAM != 0);
+FApplePlatformSymbolDatabase::FApplePlatformSymbolDatabase(FApplePlatformSymbolDatabase const& Other)
+{
+	operator=(Other);
+}
 
-void FApplePlatformSymbolication::SetSymbolicationAllowed(bool const bAllow)
+FApplePlatformSymbolDatabase::~FApplePlatformSymbolDatabase()
 {
-	GAllowApplePlatformSymbolication = bAllow;
+	if(GCoreSymbolicationHandle && !CSIsNull(AppleDB))
+	{
+		CSRelease(AppleDB);
+	}
+	AppleDB.Buffer0 = nullptr;
+	AppleDB.Buffer1 = nullptr;
 }
 
-FApplePlatformSymbolCache* FApplePlatformSymbolication::CreateSymbolCache(void)
+FApplePlatformSymbolDatabase& FApplePlatformSymbolDatabase::operator=(FApplePlatformSymbolDatabase const& Other)
 {
-#if PLATFORM_MAC && IS_PROGRAM
-	return GAllowApplePlatformSymbolication ? new FApplePlatformSymbolCache : nullptr;
-#else 
-	return nullptr;
-#endif
+	if(this != &Other)
+	{
+		GenericDB = Other.GenericDB;
+		AppleDB = Other.AppleDB;
+		CSRetain(AppleDB);
+	}
+	return *this;
 }
 
-void FApplePlatformSymbolication::DestroySymbolCache(FApplePlatformSymbolCache* Cache)
+void FApplePlatformSymbolication::EnableCoreSymbolication(bool const bEnable)
 {
-#if PLATFORM_MAC && IS_PROGRAM
-	if(GAllowApplePlatformSymbolication && Cache)
+	GAllowApplePlatformSymbolication = bEnable;
+	if(bEnable && !GCoreSymbolicationHandle)
 	{
-		for(auto It : Cache->Symbolicators)
+		GCoreSymbolicationHandle = FPlatformProcess::GetDllHandle(TEXT("/System/Library/PrivateFrameworks/CoreSymbolication.framework/Versions/Current/CoreSymbolication"));
+		if(GCoreSymbolicationHandle)
 		{
-			CSRelease( It.Value );
+			CSEqual = (CSEqualPtr)FPlatformProcess::GetDllExport(GCoreSymbolicationHandle, TEXT("CSEqual"));
+			GAllowApplePlatformSymbolication &= (CSEqual != nullptr);
+			
+			CSIsNull = (CSIsNullPtr)FPlatformProcess::GetDllExport(GCoreSymbolicationHandle, TEXT("CSIsNull"));
+			GAllowApplePlatformSymbolication &= (CSIsNull != nullptr);
+			
+			CSRelease = (CSReleasePtr)FPlatformProcess::GetDllExport(GCoreSymbolicationHandle, TEXT("CSRelease"));
+			GAllowApplePlatformSymbolication &= (CSRelease != nullptr);
+			
+			CSRetain = (CSRetainPtr)FPlatformProcess::GetDllExport(GCoreSymbolicationHandle, TEXT("CSRetain"));
+			GAllowApplePlatformSymbolication &= (CSRetain != nullptr);
+			
+			CSSymbolicatorCreateWithPid = (CSSymbolicatorCreateWithPidPtr)FPlatformProcess::GetDllExport(GCoreSymbolicationHandle, TEXT("CSSymbolicatorCreateWithPid"));
+			GAllowApplePlatformSymbolication &= (CSSymbolicatorCreateWithPid != nullptr);
+			
+			CSSymbolicatorCreateWithPathAndArchitecture = (CSSymbolicatorCreateWithPathAndArchitecturePtr)FPlatformProcess::GetDllExport(GCoreSymbolicationHandle, TEXT("CSSymbolicatorCreateWithPathAndArchitecture"));
+			GAllowApplePlatformSymbolication &= (CSSymbolicatorCreateWithPathAndArchitecture != nullptr);
+			
+			CSSymbolicatorGetSymbolWithAddressAtTime = (CSSymbolicatorGetSymbolWithAddressAtTimePtr)FPlatformProcess::GetDllExport(GCoreSymbolicationHandle, TEXT("CSSymbolicatorGetSymbolWithAddressAtTime"));
+			GAllowApplePlatformSymbolication &= (CSSymbolicatorGetSymbolWithAddressAtTime != nullptr);
+			
+			CSSymbolicatorGetSourceInfoWithAddressAtTime = (CSSymbolicatorGetSourceInfoWithAddressAtTimePtr)FPlatformProcess::GetDllExport(GCoreSymbolicationHandle, TEXT("CSSymbolicatorGetSourceInfoWithAddressAtTime"));
+			GAllowApplePlatformSymbolication &= (CSSymbolicatorGetSourceInfoWithAddressAtTime != nullptr);
+			
+			CSSymbolicatorGetSymbolOwnerWithUUIDAtTime = (CSSymbolicatorGetSymbolOwnerWithUUIDAtTimePtr)FPlatformProcess::GetDllExport(GCoreSymbolicationHandle, TEXT("CSSymbolicatorGetSymbolOwnerWithUUIDAtTime"));
+			GAllowApplePlatformSymbolication &= (CSSymbolicatorGetSymbolOwnerWithUUIDAtTime != nullptr);
+			
+			CSSymbolicatorGetSymbolOwner = (CSSymbolicatorGetSymbolOwnerPtr)FPlatformProcess::GetDllExport(GCoreSymbolicationHandle, TEXT("CSSymbolicatorGetSymbolOwner"));
+			GAllowApplePlatformSymbolication &= (CSSymbolicatorGetSymbolOwner != nullptr);
+			
+			CSSymbolicatorForeachSymbolAtTime = (CSSymbolicatorForeachSymbolAtTimePtr)FPlatformProcess::GetDllExport(GCoreSymbolicationHandle, TEXT("CSSymbolicatorForeachSymbolAtTime"));
+			GAllowApplePlatformSymbolication &= (CSSymbolicatorForeachSymbolAtTime != nullptr);
+			
+			CSSymbolGetName = (CSSymbolGetNamePtr)FPlatformProcess::GetDllExport(GCoreSymbolicationHandle, TEXT("CSSymbolGetName"));
+			GAllowApplePlatformSymbolication &= (CSSymbolGetName != nullptr);
+			
+			CSSymbolGetRange = (CSSymbolGetRangePtr)FPlatformProcess::GetDllExport(GCoreSymbolicationHandle, TEXT("CSSymbolGetRange"));
+			GAllowApplePlatformSymbolication &= (CSSymbolGetRange != nullptr);
+			
+			CSSourceInfoGetSymbolOwner = (CSSourceInfoGetSymbolOwnerPtr)FPlatformProcess::GetDllExport(GCoreSymbolicationHandle, TEXT("CSSourceInfoGetSymbolOwner"));
+			GAllowApplePlatformSymbolication &= (CSSourceInfoGetSymbolOwner != nullptr);
+			
+			CSSymbolGetSymbolOwner = (CSSymbolGetSymbolOwnerPtr)FPlatformProcess::GetDllExport(GCoreSymbolicationHandle, TEXT("CSSymbolGetSymbolOwner"));
+			GAllowApplePlatformSymbolication &= (CSSymbolGetSymbolOwner != nullptr);
+			
+			CSSymbolOwnerGetBaseAddress = (CSSymbolOwnerGetBaseAddressPtr)FPlatformProcess::GetDllExport(GCoreSymbolicationHandle, TEXT("CSSymbolOwnerGetBaseAddress"));
+			GAllowApplePlatformSymbolication &= (CSSymbolOwnerGetBaseAddress != nullptr);
+			
+			CSSymbolForeachSourceInfo = (CSSymbolForeachSourceInfoPtr)FPlatformProcess::GetDllExport(GCoreSymbolicationHandle, TEXT("CSSymbolForeachSourceInfo"));
+			GAllowApplePlatformSymbolication &= (CSSymbolForeachSourceInfo != nullptr);
+			
+			CSSymbolOwnerGetName = (CSSymbolOwnerGetNamePtr)FPlatformProcess::GetDllExport(GCoreSymbolicationHandle, TEXT("CSSymbolOwnerGetName"));
+			GAllowApplePlatformSymbolication &= (CSSymbolOwnerGetName != nullptr);
+			
+			CSSymbolOwnerGetUUID = (CSSymbolOwnerGetUUIDPtr)FPlatformProcess::GetDllExport(GCoreSymbolicationHandle, TEXT("CSSymbolOwnerGetUUID"));
+			GAllowApplePlatformSymbolication &= (CSSymbolOwnerGetUUID != nullptr);
+			
+			CSSourceInfoGetLineNumber = (CSSourceInfoGetLineNumberPtr)FPlatformProcess::GetDllExport(GCoreSymbolicationHandle, TEXT("CSSourceInfoGetLineNumber"));
+			GAllowApplePlatformSymbolication &= (CSSourceInfoGetLineNumber != nullptr);
+			
+			CSSourceInfoGetPath = (CSSourceInfoGetPathPtr)FPlatformProcess::GetDllExport(GCoreSymbolicationHandle, TEXT("CSSourceInfoGetPath"));
+			GAllowApplePlatformSymbolication &= (CSSourceInfoGetPath != nullptr);
+			
+			CSSourceInfoGetRange = (CSSourceInfoGetRangePtr)FPlatformProcess::GetDllExport(GCoreSymbolicationHandle, TEXT("CSSourceInfoGetRange"));
+			GAllowApplePlatformSymbolication &= (CSSourceInfoGetRange != nullptr);
+			
+			CSSourceInfoGetSymbol = (CSSourceInfoGetSymbolPtr)FPlatformProcess::GetDllExport(GCoreSymbolicationHandle, TEXT("CSSourceInfoGetSymbol"));
+			GAllowApplePlatformSymbolication &= (CSSourceInfoGetSymbol != nullptr);
 		}
-		delete Cache;
+		else
+		{
+			GAllowApplePlatformSymbolication = false;
+		}
+	}
+	else if(!bEnable && GCoreSymbolicationHandle)
+	{
+		CSIsNull = nullptr;
+		CSRelease = nullptr;
+		CSRetain = nullptr;
+		CSSymbolicatorCreateWithPid = nullptr;
+		CSSymbolicatorCreateWithPathAndArchitecture = nullptr;
+		CSSymbolicatorGetSymbolWithAddressAtTime = nullptr;
+		CSSymbolicatorGetSourceInfoWithAddressAtTime = nullptr;
+		CSSymbolicatorGetSymbolOwnerWithUUIDAtTime = nullptr;
+		CSSymbolicatorGetSymbolOwner = nullptr;
+		CSSymbolicatorForeachSymbolAtTime = nullptr;
+		CSSymbolGetName = nullptr;
+		CSSymbolGetRange = nullptr;
+		CSSourceInfoGetSymbolOwner = nullptr;
+		CSSymbolGetSymbolOwner = nullptr;
+		CSSymbolForeachSourceInfo = nullptr;
+		CSSymbolOwnerGetName = nullptr;
+		CSSymbolOwnerGetUUID = nullptr;
+		CSSymbolOwnerGetBaseAddress = nullptr;
+		CSSourceInfoGetLineNumber = nullptr;
+		CSSourceInfoGetPath = nullptr;
+		CSSourceInfoGetRange = nullptr;
+		CSSourceInfoGetSymbol = nullptr;
+		FPlatformProcess::FreeDllHandle(GCoreSymbolicationHandle);
+		GCoreSymbolicationHandle = nullptr;
 	}
-#endif
 }
 
-bool FApplePlatformSymbolication::SymbolInfoForAddress(uint64 ProgramCounter, FProgramCounterSymbolInfo& out_SymbolInfo)
+bool FApplePlatformSymbolication::LoadSymbolDatabaseForBinary(FString SourceFolder, FString BinaryPath, FString BinarySignature, FApplePlatformSymbolDatabase& OutDatabase)
 {
-#if PLATFORM_MAC && IS_PROGRAM
-	bool bOK = false;
-	if (GAllowApplePlatformSymbolication)
+	bool bOK = FGenericPlatformSymbolication::LoadSymbolDatabaseForBinary(SourceFolder, BinaryPath, BinarySignature, OutDatabase.GenericDB);
+	if(!bOK && GAllowApplePlatformSymbolication && (IFileManager::Get().FileSize(*BinaryPath) > 0))
 	{
-		CSSymbolicatorRef Symbolicator = CSSymbolicatorCreateWithPid(FPlatformProcess::GetCurrentProcessId());
-		if(!CSIsNull(Symbolicator))
+		CSSymbolicatorRef Symbolicator = OutDatabase.AppleDB;
+		if( CSIsNull(Symbolicator) )
 		{
-			CSSourceInfoRef Symbol = CSSymbolicatorGetSourceInfoWithAddressAtTime(Symbolicator, (vm_address_t)ProgramCounter, kCSNow);
-			
-			if(!CSIsNull(Symbol))
+			Symbolicator = CSSymbolicatorCreateWithPathAndArchitecture(TCHAR_TO_UTF8(*BinaryPath), CPU_TYPE_X86_64);
+		}
+		if( !CSIsNull(Symbolicator) )
+		{
+			if(BinarySignature.Len())
 			{
-				out_SymbolInfo.LineNumber = CSSourceInfoGetLineNumber(Symbol);
-				FCStringAnsi::Sprintf(out_SymbolInfo.Filename, CSSourceInfoGetPath(Symbol));
-				FCStringAnsi::Sprintf(out_SymbolInfo.FunctionName, CSSymbolGetName(CSSourceInfoGetSymbol(Symbol)));
-				CSRange CodeRange = CSSourceInfoGetRange(Symbol);
-				out_SymbolInfo.SymbolDisplacement = (ProgramCounter - CodeRange.Location);
+				CFUUIDRef UUID = CFUUIDCreateFromString(kCFAllocatorDefault, (CFStringRef)BinarySignature.GetNSString());
+				check(UUID);
 				
-				CSSymbolOwnerRef Owner = CSSourceInfoGetSymbolOwner(Symbol);
-				if(!CSIsNull(Owner))
+				CSSymbolOwnerRef SymbolOwner = CSSymbolicatorGetSymbolOwnerWithUUIDAtTime(Symbolicator, UUID, kCSNow);
+				bOK = (!CSIsNull(SymbolOwner));
+				OutDatabase.GenericDB.Signature = BinarySignature;
+			}
+			else
+			{
+				CSSymbolOwnerRef SymbolOwner = CSSymbolicatorGetSymbolOwner(Symbolicator);
+				if(!CSIsNull(SymbolOwner))
 				{
-					ANSICHAR const* DylibName = CSSymbolOwnerGetName(Owner);
-					FCStringAnsi::Strcpy(out_SymbolInfo.ModuleName, DylibName);
-					
-					bOK = out_SymbolInfo.LineNumber != 0;
+					CFUUIDRef OwnerUUID = CSSymbolOwnerGetUUID(SymbolOwner);
+					CFStringRef String = CFUUIDCreateString(kCFAllocatorDefault, OwnerUUID);
+					check(String);
+					OutDatabase.GenericDB.Signature = FString((NSString*)String);
+					CFRelease(String);
+					bOK = true;
 				}
 			}
 			
-			CSRelease(Symbolicator);
+			if(bOK)
+			{
+				OutDatabase.AppleDB = Symbolicator;
+			}
 		}
 	}
-	
 	return bOK;
-#else
-	return false;
-#endif
 }
 
-bool FApplePlatformSymbolication::SymbolInfoForFunctionFromModule(ANSICHAR const* MangledName, ANSICHAR const* ModuleName, FProgramCounterSymbolInfo& Info)
+bool FApplePlatformSymbolication::SaveSymbolDatabaseForBinary(FString TargetFolder, FString Name, FApplePlatformSymbolDatabase& Database)
 {
-#if PLATFORM_MAC && IS_PROGRAM
-	bool bOK = false;
-	
-	if (GAllowApplePlatformSymbolication)
+	__block bool bOK = true;
+	if(GAllowApplePlatformSymbolication)
 	{
-		CSSymbolicatorRef Symbolicator = CSSymbolicatorCreateWithPid(FPlatformProcess::GetCurrentProcessId());
-		if(!CSIsNull(Symbolicator))
+		CSSymbolicatorRef Symbolicator = Database.AppleDB;
+		if( !CSIsNull(Symbolicator) )
 		{
-			CSSymbolRef Symbol = CSSymbolicatorGetSymbolWithMangledNameFromSymbolOwnerWithNameAtTime(Symbolicator, MangledName, ModuleName, kCSNow);
+			__block TMap<FString, int32> StringLookup;
 			
-			if(!CSIsNull(Symbol))
-			{
-				CSRange CodeRange = CSSymbolGetRange(Symbol);
+			CSSymbolOwnerRef SymbolOwner = CSSymbolicatorGetSymbolOwner(Symbolicator);
+			CFUUIDRef OwnerUUID = CSSymbolOwnerGetUUID(SymbolOwner);
+			const char* OwnerName = CSSymbolOwnerGetName(SymbolOwner);
+			vm_address_t BaseAddress = CSSymbolOwnerGetBaseAddress(SymbolOwner);
 			
-				bOK = SymbolInfoForAddress(CodeRange.Location, Info);
-			}
+			Database.GenericDB.Name = OwnerName;
+			Database.GenericDB.StringTable.Reset();
+			Database.GenericDB.Symbols.Reset();
 			
-			CSRelease(Symbolicator);
+			CFStringRef String = CFUUIDCreateString(kCFAllocatorDefault, OwnerUUID);
+			check(String);
+			Database.GenericDB.Signature = FString((NSString*)String);
+			CFRelease(String);
+			
+			CSSymbolicatorForeachSymbolAtTime(Symbolicator, kCSNow, ^(CSSymbolRef Symbol){
+				CSSymbolOwnerRef Owner = CSSymbolGetSymbolOwner(Symbol);
+				if (CSEqual(SymbolOwner, Owner))
+				{
+					__block FGenericPlatformSymbolData SymbolData;
+					CSRange Range = CSSymbolGetRange(Symbol);
+					const char* Name = CSSymbolGetName(Symbol);
+					SymbolData.Start = Range.Location - BaseAddress;
+					SymbolData.Length = Range.Length;
+					int32* NameIndex = StringLookup.Find(Name);
+					if(NameIndex)
+					{
+						SymbolData.NameIdx = *NameIndex;
+					}
+					else
+					{
+						SymbolData.NameIdx = Database.GenericDB.StringTable.Num();
+						StringLookup.Add(Name, SymbolData.NameIdx);
+						Database.GenericDB.StringTable.Push(Name);
+					}
+					
+					CSSymbolForeachSourceInfo(Symbol, ^(CSSourceInfoRef SourceInfo){
+						int Line = CSSourceInfoGetLineNumber(SourceInfo);
+						const char* Path = CSSourceInfoGetPath(SourceInfo);
+						CSRange InfoRange = CSSourceInfoGetRange(SourceInfo);
+						
+						FGenericPlatformSymbolInfo Info;
+						Info.Line = Line;
+						Info.Start = InfoRange.Location - BaseAddress;
+						Info.Length = InfoRange.Length;
+						
+						int32* PathIndex = StringLookup.Find(Path);
+						if(PathIndex)
+						{
+							Info.PathIdx = *PathIndex;
+						}
+						else
+						{
+							Info.PathIdx = Database.GenericDB.StringTable.Num();
+							StringLookup.Add(Path, Info.PathIdx);
+							Database.GenericDB.StringTable.Push(Path);
+						}
+						
+						SymbolData.SymbolInfo.Add(Info);
+						
+						return 0;
+					});
+					
+					Database.GenericDB.Symbols.Add(SymbolData);
+				}
+				else
+				{
+					UE_LOG(LogInit, Warning, TEXT("FApplePlatformSymbolication::SaveSymbolDatabaseForBinary doesn't handle Mach-O binaries/.dSYMs with multiple symbol owners!"));
+					bOK = false;
+					return 1;
+				}
+				return 0;
+			});
+			bOK &= (Database.GenericDB.Symbols.Num() >= 1);
 		}
 	}
 	
+	if(bOK)
+	{
+		bOK = FGenericPlatformSymbolication::SaveSymbolDatabaseForBinary(TargetFolder, Name, Database.GenericDB);
+	}
+	
 	return bOK;
-#else
-	return false;
-#endif
 }
 
-bool FApplePlatformSymbolication::SymbolInfoForStrippedSymbol(FApplePlatformSymbolCache* Cache, uint64 ProgramCounter, ANSICHAR const* ModulePath, ANSICHAR const* ModuleUUID, FProgramCounterSymbolInfo& Info)
+bool FApplePlatformSymbolication::SymbolInfoForStrippedSymbol(FApplePlatformSymbolDatabase const& Database, uint64 ProgramCounter, uint64 ModuleOffset, FString ModuleSignature, FProgramCounterSymbolInfo& Info)
 {
-#if PLATFORM_MAC && IS_PROGRAM
-	bool bOK = false;
-	
-	if(GAllowApplePlatformSymbolication && IFileManager::Get().FileSize(UTF8_TO_TCHAR(ModulePath)) > 0)
+	bool bOK = FGenericPlatformSymbolication::SymbolInfoForStrippedSymbol(Database.GenericDB, ProgramCounter, ModuleOffset, ModuleSignature, Info);
+	if(!bOK && GAllowApplePlatformSymbolication && !CSIsNull(Database.AppleDB) && (ModuleSignature.Len()))
 	{
-		CSSymbolicatorRef Symbolicator = { nullptr, nullptr };
-		if ( Cache )
-		{
-			Symbolicator = Cache->Symbolicators.FindOrAdd( FString(ModulePath) );
-		}
-		if( CSIsNull(Symbolicator) )
-		{
-			Symbolicator = CSSymbolicatorCreateWithPathAndArchitecture(ModulePath, CPU_TYPE_X86_64);
-		}
-		if( !CSIsNull(Symbolicator) )
+		CSSymbolicatorRef Symbolicator = Database.AppleDB;
+	
+		CFUUIDRef UUID = CFUUIDCreateFromString(kCFAllocatorDefault, (CFStringRef)ModuleSignature.GetNSString());
+		check(UUID);
+		
+		CSSymbolOwnerRef SymbolOwner = CSSymbolicatorGetSymbolOwnerWithUUIDAtTime(Symbolicator, UUID, kCSNow);
+		if(!CSIsNull(SymbolOwner))
 		{
-			if ( Cache )
-			{
-				Cache->Symbolicators[ FString(ModulePath) ] = Symbolicator;
-			}
+			ANSICHAR const* DylibName = CSSymbolOwnerGetName(SymbolOwner);
+			FCStringAnsi::Strcpy(Info.ModuleName, DylibName);
 			
-			FString ModuleID(ModuleUUID);
-			CFUUIDRef UUID = CFUUIDCreateFromString(nullptr, (CFStringRef)ModuleID.GetNSString());
-			check(UUID);
+			uint64 Address = ProgramCounter >= ModuleOffset ? (ProgramCounter - ModuleOffset) : ProgramCounter;
 			
-			CSSymbolOwnerRef SymbolOwner = CSSymbolicatorGetSymbolOwnerWithUUIDAtTime(Symbolicator, UUID, kCSNow);
-			if(!CSIsNull(SymbolOwner))
+			vm_address_t BaseAddress = CSSymbolOwnerGetBaseAddress(SymbolOwner);
+			CSSymbolRef Symbol = CSSymbolicatorGetSymbolWithAddressAtTime(Symbolicator, ((vm_address_t)Address) + BaseAddress, kCSNow);
+			if(!CSIsNull(Symbol))
 			{
-				vm_address_t BaseAddress = CSSymbolOwnerGetBaseAddress(SymbolOwner);
+				ANSICHAR const* FunctionName = CSSymbolGetName(Symbol);
+				if(FunctionName)
+				{
+					FCStringAnsi::Sprintf(Info.FunctionName, FunctionName);
+				}
+				
+				CSRange Range = CSSymbolGetRange(Symbol);
+				Info.SymbolDisplacement = (ProgramCounter - Range.Location);
+				Info.OffsetInModule = Range.Location;
+				Info.ProgramCounter = ProgramCounter;
 				
-				CSSymbolRef Symbol = CSSymbolicatorGetSymbolWithAddressAtTime(Symbolicator, ((vm_address_t)ProgramCounter) + BaseAddress, kCSNow);
-				if(!CSIsNull(Symbol))
+				CSSourceInfoRef SymbolInfo = CSSymbolicatorGetSourceInfoWithAddressAtTime(Symbolicator, ((vm_address_t)Address) + BaseAddress, kCSNow);
+				if(!CSIsNull(SymbolInfo))
 				{
-					ANSICHAR const* DylibName = CSSymbolOwnerGetName(SymbolOwner);
-					FCStringAnsi::Strcpy(Info.ModuleName, DylibName);
-					
-					ANSICHAR const* FunctionName = CSSymbolGetName(Symbol);
-					if(FunctionName)
-					{
-						FCStringAnsi::Sprintf(Info.FunctionName, FunctionName);
-					}
+					Info.LineNumber = CSSourceInfoGetLineNumber(SymbolInfo);
 					
-					CSRange Range = CSSymbolGetRange(Symbol);
-					Info.SymbolDisplacement = (ProgramCounter - Range.Location);
-					Info.OffsetInModule = Range.Location;
-					Info.ProgramCounter = ProgramCounter;
-					
-					CSSourceInfoRef SymbolInfo = CSSymbolicatorGetSourceInfoWithAddressAtTime(Symbolicator, ((vm_address_t)ProgramCounter) + BaseAddress, kCSNow);
-					if(!CSIsNull(SymbolInfo))
+					ANSICHAR const* FileName = CSSourceInfoGetPath(SymbolInfo);
+					if(FileName)
 					{
-						Info.LineNumber = CSSourceInfoGetLineNumber(SymbolInfo);
-						
-						ANSICHAR const* FileName = CSSourceInfoGetPath(SymbolInfo);
-						if(FileName)
-						{
-							FCStringAnsi::Sprintf(Info.Filename, FileName);
-						}
+						FCStringAnsi::Sprintf(Info.Filename, FileName);
 					}
-					
-					bOK = true;
 				}
+				
+				bOK = true;
 			}
+		}
+	}
+	
+	return bOK;
+}
+
+bool FApplePlatformSymbolication::SymbolInfoForAddress(uint64 ProgramCounter, FProgramCounterSymbolInfo& out_SymbolInfo)
+{
+	bool bOK = false;
+	if (GAllowApplePlatformSymbolication)
+	{
+		CSSymbolicatorRef Symbolicator = CSSymbolicatorCreateWithPid(FPlatformProcess::GetCurrentProcessId());
+		if(!CSIsNull(Symbolicator))
+		{
+			CSSourceInfoRef Symbol = CSSymbolicatorGetSourceInfoWithAddressAtTime(Symbolicator, (vm_address_t)ProgramCounter, kCSNow);
 			
-			CFRelease(UUID);
-			
-			if ( !Cache )
+			if(!CSIsNull(Symbol))
 			{
-				CSRelease(Symbolicator);
+				out_SymbolInfo.LineNumber = CSSourceInfoGetLineNumber(Symbol);
+				FCStringAnsi::Sprintf(out_SymbolInfo.Filename, CSSourceInfoGetPath(Symbol));
+				FCStringAnsi::Sprintf(out_SymbolInfo.FunctionName, CSSymbolGetName(CSSourceInfoGetSymbol(Symbol)));
+				CSRange CodeRange = CSSourceInfoGetRange(Symbol);
+				out_SymbolInfo.SymbolDisplacement = (ProgramCounter - CodeRange.Location);
+				
+				CSSymbolOwnerRef Owner = CSSourceInfoGetSymbolOwner(Symbol);
+				if(!CSIsNull(Owner))
+				{
+					ANSICHAR const* DylibName = CSSymbolOwnerGetName(Owner);
+					FCStringAnsi::Strcpy(out_SymbolInfo.ModuleName, DylibName);
+					
+					bOK = out_SymbolInfo.LineNumber != 0;
+				}
 			}
+			
+			CSRelease(Symbolicator);
 		}
 	}
 	
 	return bOK;
-#else
-	return false;
-#endif
 }
diff --git a/Engine/Source/Runtime/Core/Private/GenericPlatform/GenericPlatformSymbolication.cpp b/Engine/Source/Runtime/Core/Private/GenericPlatform/GenericPlatformSymbolication.cpp
new file mode 100644
--- /dev/null
+++ b/Engine/Source/Runtime/Core/Private/GenericPlatform/GenericPlatformSymbolication.cpp
@@ -0,0 +1,105 @@
+// Copyright 1998-2015 Epic Games, Inc. All Rights Reserved.
+
+#include "CorePrivatePCH.h"
+
+#include "GenericPlatformSymbolication.h"
+#include "Serialization/Archive.h"
+
+DEFINE_LOG_CATEGORY_STATIC(LogGenericPlatformSymbolication, Log, All);
+
+bool FGenericPlatformSymbolication::LoadSymbolDatabaseForBinary(FString SourceFolder, FString Binary, FString ModuleSignature, FGenericPlatformSymbolDatabase& OutDatabase)
+{
+	bool bOk = false;
+	
+	FString ModuleName = FPaths::GetBaseFilename(Binary);
+	FString InputFile = (SourceFolder / ModuleName) + TEXT(".udebugsymbols");
+	if (IFileManager::Get().FileSize(*InputFile) > 0)
+	{
+		TArray<uint8> DataBuffer;
+		if(FFileHelper::LoadFileToArray(DataBuffer, *InputFile))
+		{
+			FArchiveLoadCompressedProxy DataArchive(DataBuffer, (ECompressionFlags)(COMPRESS_Default));
+			DataArchive << OutDatabase;
+			if(!DataArchive.GetError())
+			{
+				if(OutDatabase.Signature == ModuleSignature)
+				{
+					bOk = true;
+				}
+			}
+		}
+	}
+	
+	return bOk;
+}
+
+bool FGenericPlatformSymbolication::SaveSymbolDatabaseForBinary(FString TargetFolder, FString Name, FGenericPlatformSymbolDatabase& Database)
+{
+	bool bOk = false;
+	
+	FString ModuleName = FPaths::GetBaseFilename(Name);
+	FString OutputFile = (TargetFolder / ModuleName) + TEXT(".udebugsymbols");
+
+	TArray<uint8> DataBuffer;
+	FArchiveSaveCompressedProxy DataArchive(DataBuffer, (ECompressionFlags)(COMPRESS_ZLIB|COMPRESS_BiasSpeed));
+	DataArchive << Database;
+	DataArchive.Flush();
+	if(!DataArchive.GetError())
+	{
+		if(FFileHelper::SaveArrayToFile(DataBuffer, *OutputFile))
+		{
+			bOk = true;
+		}
+		else
+		{
+			UE_LOG( LogGenericPlatformSymbolication, Error, TEXT( "Unable to write debug symbols to output file." ) );
+		}
+	}
+	else
+	{
+		UE_LOG( LogGenericPlatformSymbolication, Error, TEXT( "Unable to serialise debug symbols." ) );
+	}
+	
+	return bOk;
+}
+
+bool FGenericPlatformSymbolication::SymbolInfoForStrippedSymbol(FGenericPlatformSymbolDatabase const& Database, uint64 ProgramCounter, uint64 ModuleOffset, FString ModuleSignature, FProgramCounterSymbolInfo& Info)
+{
+	bool bOk = false;
+	
+	if(Database.Signature == ModuleSignature)
+	{
+		for(auto Symbol : Database.Symbols)
+		{
+			if((Symbol.Start <= (ProgramCounter - ModuleOffset)) && ((Symbol.Start + Symbol.Length) >= (ProgramCounter - ModuleOffset)))
+			{
+				FCStringAnsi::Strncpy(Info.ModuleName, TCHAR_TO_ANSI(*Database.Name), FProgramCounterSymbolInfo::MAX_NAME_LENGHT);
+				
+				FString SymbolName = Database.StringTable[Symbol.NameIdx];
+				FCStringAnsi::Strncpy(Info.FunctionName, TCHAR_TO_ANSI(*SymbolName), FProgramCounterSymbolInfo::MAX_NAME_LENGHT);
+				
+				Info.ProgramCounter = ProgramCounter;
+				Info.OffsetInModule = Symbol.Start;
+				
+				Info.LineNumber = -1;
+				
+				for(auto SymbolInfo : Symbol.SymbolInfo)
+				{
+					if((SymbolInfo.Start <= (ProgramCounter - ModuleOffset)) && ((SymbolInfo.Start + SymbolInfo.Length) >= (ProgramCounter - ModuleOffset)))
+					{
+						FString Path = Database.StringTable[SymbolInfo.PathIdx];
+						FCStringAnsi::Strncpy(Info.Filename, TCHAR_TO_ANSI(*Path), FProgramCounterSymbolInfo::MAX_NAME_LENGHT);
+						Info.LineNumber = SymbolInfo.Line;
+						break;
+					}
+				}
+				
+				bOk = true;
+			
+				break;
+			}
+		}
+	}
+	
+	return bOk;
+}
diff --git a/Engine/Source/Runtime/Core/Private/Mac/MacMallocZone.cpp b/Engine/Source/Runtime/Core/Private/Mac/MacMallocZone.cpp
--- a/Engine/Source/Runtime/Core/Private/Mac/MacMallocZone.cpp
+++ b/Engine/Source/Runtime/Core/Private/Mac/MacMallocZone.cpp
@@ -179,7 +179,9 @@ bool FMacMallocCrashHandler::IsOnCrashedThread( void )
 	}
 	else
 	{
-		FPlatformProcess::SleepInfinite();
+		do {
+			FPlatformProcess::SleepInfinite();
+		} while (true);
 		return false;
 	}
 }
diff --git a/Engine/Source/Runtime/Core/Private/Mac/MacPlatformMisc.cpp b/Engine/Source/Runtime/Core/Private/Mac/MacPlatformMisc.cpp
--- a/Engine/Source/Runtime/Core/Private/Mac/MacPlatformMisc.cpp
+++ b/Engine/Source/Runtime/Core/Private/Mac/MacPlatformMisc.cpp
@@ -178,6 +178,8 @@ struct FMacApplicationInfo
 		uint32 Status = notify_register_dispatch(kIOPSNotifyPowerSource, &PowerSourceNotification, dispatch_get_main_queue(), PowerSourceNotifyHandler);
 		check(Status == NOTIFY_STATUS_OK);
 		
+		NumCores = FPlatformMisc::NumberOfCores();
+		
 		NSString* PLCrashReportFile = [TemporaryCrashReportFolder().GetNSString() stringByAppendingPathComponent:TemporaryCrashReportName().GetNSString()];
 		[PLCrashReportFile getCString:PLCrashReportPath maxLength:PATH_MAX encoding:NSUTF8StringEncoding];
 	}
@@ -315,6 +317,8 @@ void FMacPlatformMisc::PlatformPreInit()
 	{
 		UE_LOG(LogInit, Warning, TEXT("Failed to change open file limit, UE4 may be unstable."));
 	}
+	
+	FApplePlatformSymbolication::EnableCoreSymbolication(!FPlatformProcess::IsSandboxedApplication() && IS_PROGRAM);
 }
 
 void FMacPlatformMisc::PlatformInit()
@@ -398,6 +402,7 @@ void FMacPlatformMisc::PlatformTearDown()
 		[[NSProcessInfo processInfo] endActivity:CommandletActivity];
 		CommandletActivity = nil;
 	}
+	FApplePlatformSymbolication::EnableCoreSymbolication(false);
 }
 
 void FMacPlatformMisc::UpdateWindowMenu()
@@ -1234,7 +1239,8 @@ static void DefaultCrashHandler(FMacCrashContext const& Context)
 /** True system-specific crash handler that gets called first */
 static void PlatformCrashHandler(siginfo_t* Info, ucontext_t* Uap, void* Context)
 {
-	FApplePlatformSymbolication::SetSymbolicationAllowed( false );
+	// Disable CoreSymbolication
+	FApplePlatformSymbolication::EnableCoreSymbolication( false );
 	
 	FMacCrashContext CrashContext;
 	CrashContext.InitFromSignal((int32)Info->si_signo, Info, Context);
@@ -1667,5 +1673,25 @@ void FMacCrashContext::GenerateCrashInfoAndLaunchReporter() const
 		// OS X uses internally to launch/wait on the CrashReportClient. It is simpler, easier & safer to just die here like a good little Mac.app.
 	}
 	
+	// Sandboxed applications re-raise the signal to trampoline into the system crash reporter as suppressing it may fall foul of Apple's Mac App Store rules.
+	// @todo Submit an application to the MAS & see whether Apple's reviewers highlight our crash reporting or trampolining to the system reporter.
+	if(GMacAppInfo.bIsSandboxed)
+	{
+		struct sigaction Action;
+		FMemory::Memzero(&Action, sizeof(struct sigaction));
+		Action.sa_handler = SIG_DFL;
+		sigemptyset(&Action.sa_mask);
+		sigaction(SIGQUIT, &Action, NULL);
+		sigaction(SIGILL, &Action, NULL);
+		sigaction(SIGEMT, &Action, NULL);
+		sigaction(SIGFPE, &Action, NULL);
+		sigaction(SIGBUS, &Action, NULL);
+		sigaction(SIGSEGV, &Action, NULL);
+		sigaction(SIGSYS, &Action, NULL);
+		sigaction(SIGABRT, &Action, NULL);
+	
+		raise(Signal);
+	}
+	
 	_Exit(0);
 }
diff --git a/Engine/Source/Runtime/Core/Private/Mac/MacPlatformSurvey.cpp b/Engine/Source/Runtime/Core/Private/Mac/MacPlatformSurvey.cpp
--- a/Engine/Source/Runtime/Core/Private/Mac/MacPlatformSurvey.cpp
+++ b/Engine/Source/Runtime/Core/Private/Mac/MacPlatformSurvey.cpp
@@ -253,9 +253,10 @@ void FMacPlatformSurvey::TickSurveyHardware( FHardwareSurveyResults& OutResults
 	}
 
 	// OS info
-	NSDictionary* SystemVersion = [NSDictionary dictionaryWithContentsOfFile: @"/System/Library/CoreServices/SystemVersion.plist"];
-	WriteFStringToResults(OutResults.OSVersion, TEXT("Mac OS X"));
-	FPlatformString::CFStringToTCHAR((CFStringRef)[SystemVersion objectForKey: @"ProductVersion"], OutResults.OSSubVersion);
+	FString OSXVersion, OSXBuild;
+	FMacPlatformMisc::GetOSVersions(OSXVersion, OSXBuild);
+	WriteFStringToResults(OutResults.OSVersion, FString(TEXT("Mac OS X ")) + OSXVersion);
+	WriteFStringToResults(OutResults.OSSubVersion, OSXBuild);
 	OutResults.OSBits = FPlatformMisc::Is64bitOperatingSystem() ? 64 : 32;
 
 	// OS language
diff --git a/Engine/Source/Runtime/Core/Private/Modules/ModuleManager.cpp b/Engine/Source/Runtime/Core/Private/Modules/ModuleManager.cpp
--- a/Engine/Source/Runtime/Core/Private/Modules/ModuleManager.cpp
+++ b/Engine/Source/Runtime/Core/Private/Modules/ModuleManager.cpp
@@ -847,7 +847,15 @@ void FModuleManager::FindModulePathsInDirectory(const FString& InDirectoryName,
 	for (int32 Idx = 0; Idx < FullFileNames.Num(); Idx++)
 	{
 		const FString &FullFileName = FullFileNames[Idx];
-
+	
+		// On Mac OS X the separate debug symbol format is the dSYM bundle, which is a bundle folder hierarchy containing a .dylib full of Mach-O formatted DWARF debug symbols, these are not loadable modules, so we mustn't ever try and use them. If we don't eliminate this here then it will appear in the module paths & cause errors later on which cannot be recovered from.
+	#if PLATFORM_MAC
+		if(FullFileName.Contains(".dSYM"))
+		{
+			continue;
+		}
+	#endif
+		
 		FString FileName = FPaths::GetCleanFilename(FullFileName);
 		if (FileName.StartsWith(ModulePrefix) && FileName.EndsWith(ModuleSuffix))
 		{
diff --git a/Engine/Source/Runtime/Core/Public/Apple/ApplePlatformSymbolication.h b/Engine/Source/Runtime/Core/Public/Apple/ApplePlatformSymbolication.h
--- a/Engine/Source/Runtime/Core/Public/Apple/ApplePlatformSymbolication.h
+++ b/Engine/Source/Runtime/Core/Public/Apple/ApplePlatformSymbolication.h
@@ -6,23 +6,82 @@
 
 #pragma once
 
+#include "GenericPlatformSymbolication.h"
 #include "ApplePlatformStackWalk.h"
+#include "Serialization/ArchiveBase.h"
 
 /**
  * Opaque symbol cache for improved symbolisation performance.
  */
-struct FApplePlatformSymbolCache;
+struct FApplePlatformSymbolCache
+{
+	void* Buffer0;
+	void* Buffer1;
+};
+
+/**
+ * Apple symbol database.
+ */
+struct FApplePlatformSymbolDatabase
+{
+	FApplePlatformSymbolDatabase();
+	FApplePlatformSymbolDatabase(FApplePlatformSymbolDatabase const& Other);
+	~FApplePlatformSymbolDatabase();
+	FApplePlatformSymbolDatabase& operator=(FApplePlatformSymbolDatabase const& Other);
+	
+	FGenericPlatformSymbolDatabase GenericDB;
+	FApplePlatformSymbolCache AppleDB;
+};
+
+/**
+ * Apple symbol database hash.
+ */
+struct FApplePlatformSymbolDatabaseKeyFuncs
+{
+	enum { bAllowDuplicateKeys = 0 };
+	typedef typename TCallTraits<FString>::ParamType KeyInitType;
+	typedef typename TCallTraits<FApplePlatformSymbolDatabase>::ParamType ElementInitType;
+
+	/**
+	 * @return The key used to index the given element.
+	 */
+	static FORCEINLINE KeyInitType GetSetKey(ElementInitType Element)
+	{
+		return Element.GenericDB.Signature;
+	}
+
+	/**
+	 * @return True if the keys match.
+	 */
+	static FORCEINLINE bool Matches(KeyInitType A,KeyInitType B)
+	{
+		return A == B;
+	}
+
+	/** Calculates a hash index for a key. */
+	static FORCEINLINE uint32 GetKeyHash(KeyInitType Key)
+	{
+		return GetTypeHash(Key);
+	}
+};
+
+typedef TSet<FApplePlatformSymbolDatabase, FApplePlatformSymbolDatabaseKeyFuncs> FApplePlatformSymbolDatabaseSet;
 
 /**
  * Apple platform implementation of symbolication - not async. handler safe, so don't call during crash reporting!
  */
 struct CORE_API FApplePlatformSymbolication
 {
-	static void SetSymbolicationAllowed(bool const bAllow);
-	static bool SymbolInfoForAddress(uint64 ProgramCounter, FProgramCounterSymbolInfo& Info);
-	static bool SymbolInfoForFunctionFromModule(ANSICHAR const* MangledName, ANSICHAR const* ModuleName, FProgramCounterSymbolInfo& Info);
+	static void EnableCoreSymbolication(bool const bEnable);
+	
+	static bool LoadSymbolDatabaseForBinary(FString SourceFolder, FString Binary, FString BinarySignature, FApplePlatformSymbolDatabase& OutDatabase);
+	static bool SaveSymbolDatabaseForBinary(FString TargetFolder, FString Name, FApplePlatformSymbolDatabase& Database);
+	
+	static bool SymbolInfoForStrippedSymbol(FApplePlatformSymbolDatabase const& Database, uint64 ProgramCounter, uint64 ModuleOffset, FString ModuleSignature, FProgramCounterSymbolInfo& Info);
 	
-	static FApplePlatformSymbolCache* CreateSymbolCache(void);
-	static bool SymbolInfoForStrippedSymbol(FApplePlatformSymbolCache* Cache, uint64 ProgramCounter, ANSICHAR const* ModulePath, ANSICHAR const* ModuleUUID, FProgramCounterSymbolInfo& Info);
-	static void DestroySymbolCache(FApplePlatformSymbolCache* Cache);
+	static bool SymbolInfoForAddress(uint64 ProgramCounter, FProgramCounterSymbolInfo& Info);
 };
+
+typedef FApplePlatformSymbolDatabase FPlatformSymbolDatabase;
+typedef FApplePlatformSymbolication FPlatformSymbolication;
+typedef TSet<FApplePlatformSymbolDatabase, FApplePlatformSymbolDatabaseKeyFuncs> FPlatformSymbolDatabaseSet;
diff --git a/Engine/Source/Runtime/Core/Public/GenericPlatform/GenericPlatformSymbolication.h b/Engine/Source/Runtime/Core/Public/GenericPlatform/GenericPlatformSymbolication.h
new file mode 100644
--- /dev/null
+++ b/Engine/Source/Runtime/Core/Public/GenericPlatform/GenericPlatformSymbolication.h
@@ -0,0 +1,117 @@
+// Copyright 1998-2015 Epic Games, Inc. All Rights Reserved.
+
+
+/*=============================================================================================
+	GenericPlatformSymbolication.h: Generic platform symbolication functions.
+==============================================================================================*/
+
+#pragma once
+#include "HAL/Platform.h"
+#include "Serialization/ArchiveBase.h"
+
+/**
+ * File & line info for a debug symbol region
+ */
+struct FGenericPlatformSymbolInfo
+{
+	uint32 Line;
+	uint64 Start;
+	uint64 Length;
+	int32 PathIdx;
+	
+	/**
+	 * Serializes a symbol info from or into an archive.
+	 *
+	 * @param Ar The archive to serialize from or into.
+	 * @param Info The info to serialize.
+	 */
+	friend FArchive& operator<<( FArchive& Ar, FGenericPlatformSymbolInfo& Info )
+	{
+		return Ar << Info.Line << Info.Start << Info.Length << Info.PathIdx;
+	}
+};
+
+/**
+ * Debug symbol information
+ */
+struct FGenericPlatformSymbolData
+{
+	uint64 Start;
+	uint64 Length;
+	int32 NameIdx;
+	TArray<FGenericPlatformSymbolInfo> SymbolInfo;
+	
+	/**
+	 * Serializes a symbol from or into an archive.
+	 *
+	 * @param Ar The archive to serialize from or into.
+	 * @param Info The info to serialize.
+	 */
+	friend FArchive& operator<<( FArchive& Ar, FGenericPlatformSymbolData& Info )
+	{
+		return Ar << Info.Start << Info.Length << Info.NameIdx << Info.SymbolInfo;
+	}
+};
+
+/**
+ * Container for debug symbols corresponding to a single binary file
+ */
+struct FGenericPlatformSymbolDatabase
+{
+	FString Signature;
+	FString Name;
+	TArray<FGenericPlatformSymbolData> Symbols;
+	TArray<FString> StringTable;
+	
+	/**
+	 * Serializes a symbol container from or into an archive.
+	 *
+	 * @param Ar The archive to serialize from or into.
+	 * @param Info The info to serialize.
+	 */
+	friend FArchive& operator<<( FArchive& Ar, FGenericPlatformSymbolDatabase& Info )
+	{
+		return Ar << Info.Signature << Info.Name << Info.Symbols << Info.StringTable;
+	}
+};
+
+struct FGenericPlatformSymbolDatabaseKeyFuncs
+{
+	enum { bAllowDuplicateKeys = 0 };
+	typedef TCallTraits<FString>::ParamType KeyInitType;
+	typedef TCallTraits<FGenericPlatformSymbolDatabase>::ParamType ElementInitType;
+
+	/**
+	 * @return The key used to index the given element.
+	 */
+	static FORCEINLINE KeyInitType GetSetKey(ElementInitType Element)
+	{
+		return Element.Signature;
+	}
+
+	/**
+	 * @return True if the keys match.
+	 */
+	static FORCEINLINE bool Matches(KeyInitType A,KeyInitType B)
+	{
+		return A == B;
+	}
+
+	/** Calculates a hash index for a key. */
+	static FORCEINLINE uint32 GetKeyHash(KeyInitType Key)
+	{
+		return GetTypeHash(Key);
+	}
+};
+
+typedef TSet<FGenericPlatformSymbolDatabase, FGenericPlatformSymbolDatabaseKeyFuncs> FGenericPlatformSymbolDatabaseSet;
+
+struct CORE_API FGenericPlatformSymbolication
+{
+	static bool LoadSymbolDatabaseForBinary(FString SourceFolder, FString BinaryPath, FString BinarySignature, FGenericPlatformSymbolDatabase& OutDatabase);
+	static bool SaveSymbolDatabaseForBinary(FString TargetFolder, FString Name, FGenericPlatformSymbolDatabase& Database);
+	
+	static bool SymbolInfoForStrippedSymbol(FGenericPlatformSymbolDatabase const& Database, uint64 ProgramCounter, uint64 ModuleOffset, FString ModuleSignature, FProgramCounterSymbolInfo& Info);
+	
+	static bool SymbolInfoForAddress(uint64 ProgramCounter, FProgramCounterSymbolInfo& Info) { return false; }
+};
