Allow access to local files in GExternalFilePath
#ue4
#android

diff --git a/Engine/Source/Runtime/Core/Private/Android/AndroidFile.cpp b/Engine/Source/Runtime/Core/Private/Android/AndroidFile.cpp
--- a/Engine/Source/Runtime/Core/Private/Android/AndroidFile.cpp
+++ b/Engine/Source/Runtime/Core/Private/Android/AndroidFile.cpp
@@ -1588,7 +1588,8 @@ private:
 		if (!AndroidPath.IsEmpty())
 		{
 			if ((AllowLocal && AndroidPath.StartsWith(TEXT("/"))) ||
-				AndroidPath.StartsWith(GFontPathBase))
+				AndroidPath.StartsWith(GFontPathBase) ||
+				AndroidPath.StartsWith(GExternalFilePath))
 			{
 				// Absolute paths are only local.
 				LocalPath = AndroidPath;
