Fixed behavior on FString::ParseIntoArray (muliple delimiters overload) functionality to support optionally culling empty strings. Greatly simplified implementation logic. Output parameter now properl

diff --git a/Engine/Source/Programs/CrashReporter/CrashReportClient/Private/CrashReportUtil.h b/Engine/Source/Programs/CrashReporter/CrashReportClient/Private/CrashReportUtil.h
--- a/Engine/Source/Programs/CrashReporter/CrashReportClient/Private/CrashReportUtil.h
+++ b/Engine/Source/Programs/CrashReporter/CrashReportClient/Private/CrashReportUtil.h
@@ -94,7 +94,7 @@ struct FCrashReportUtil
 		if( !Assertion.IsEmpty() )
 		{
 			TArray<FString> MultilineAssertion;
-			Assertion.ParseIntoArray( &MultilineAssertion, TEXT( "#" ), true );
+			Assertion.ParseIntoArray( MultilineAssertion, TEXT( "#" ), true );
 
 			for( const auto& AssertionLine : MultilineAssertion )
 			{
diff --git a/Engine/Source/Programs/CrashReporter/CrashReportClient/Private/Windows/WindowsErrorReport.cpp b/Engine/Source/Programs/CrashReporter/CrashReportClient/Private/Windows/WindowsErrorReport.cpp
--- a/Engine/Source/Programs/CrashReporter/CrashReportClient/Private/Windows/WindowsErrorReport.cpp
+++ b/Engine/Source/Programs/CrashReporter/CrashReportClient/Private/Windows/WindowsErrorReport.cpp
@@ -34,14 +34,14 @@ struct FWindowsReportParser
 		const FString FileAsString = reinterpret_cast<TCHAR*>(FileData.GetData());
 
 		TArray<FString> String;
-		FileAsString.ParseIntoArray( &String, TEXT( "\r\n" ), true );
+		FileAsString.ParseIntoArray( String, TEXT( "\r\n" ), true );
 
 		for( const auto& StringLine : String )
 		{
 			if( StringLine.Contains( Marker ) )
 			{
 				TArray<FString> SeparatedParameters;
-				StringLine.ParseIntoArray( &SeparatedParameters, Marker, true );
+				StringLine.ParseIntoArray( SeparatedParameters, Marker, true );
 
 				FString MatchedValue;
 				const bool bFound = FParse::Value( *StringLine, Marker, MatchedValue );
diff --git a/Engine/Source/Programs/UnrealFrontend/Private/Commands/StatsConvertCommand.cpp b/Engine/Source/Programs/UnrealFrontend/Private/Commands/StatsConvertCommand.cpp
--- a/Engine/Source/Programs/UnrealFrontend/Private/Commands/StatsConvertCommand.cpp
+++ b/Engine/Source/Programs/UnrealFrontend/Private/Commands/StatsConvertCommand.cpp
@@ -29,7 +29,7 @@ void FStatsConvertCommand::InternalRun()
 
 	// get the list of stats
 	TArray<FString> StatArrayString;
-	if (StatListString.ParseIntoArray(&StatArrayString, TEXT("+"), true) == 0)
+	if (StatListString.ParseIntoArray(StatArrayString, TEXT("+"), true) == 0)
 	{
 		StatArrayString.Add(TEXT("STAT_FrameTime"));
 	}
diff --git a/Engine/Source/Programs/UnrealFrontend/Private/Commands/StatsDumpMemoryCommand.cpp b/Engine/Source/Programs/UnrealFrontend/Private/Commands/StatsDumpMemoryCommand.cpp
--- a/Engine/Source/Programs/UnrealFrontend/Private/Commands/StatsDumpMemoryCommand.cpp
+++ b/Engine/Source/Programs/UnrealFrontend/Private/Commands/StatsDumpMemoryCommand.cpp
@@ -165,7 +165,7 @@ static FString EncodeCallstack( const TArray<FName>& Callstack )
 
 static void DecodeCallstack( const FName& EncodedCallstack, TArray<FString>& out_DecodedCallstack )
 {
-	EncodedCallstack.ToString().ParseIntoArray( &out_DecodedCallstack, CallstackSeparator, true );
+	EncodedCallstack.ToString().ParseIntoArray( out_DecodedCallstack, CallstackSeparator, true );
 }
 
 /** Simple allocation info. Assumes the sequence tag are unique and doesn't wrap. */
diff --git a/Engine/Source/Programs/UnrealHeaderTool/Private/CodeGenerator.cpp b/Engine/Source/Programs/UnrealHeaderTool/Private/CodeGenerator.cpp
--- a/Engine/Source/Programs/UnrealHeaderTool/Private/CodeGenerator.cpp
+++ b/Engine/Source/Programs/UnrealHeaderTool/Private/CodeGenerator.cpp
@@ -102,9 +102,9 @@ static struct FFlagAudit
 			{
 				FStringOutputDevice MisMatches;
 				TArray<FString> VerifyLines;
-				Verify.ParseIntoArray(&VerifyLines, TEXT("\n"), true);
+				Verify.ParseIntoArray(VerifyLines, TEXT("\n"), true);
 				TArray<FString> RefLines;
-				Ref.ParseIntoArray(&RefLines, TEXT("\n"), true);
+				Ref.ParseIntoArray(RefLines, TEXT("\n"), true);
 				check(VerifyLines.Num() == RefLines.Num()); // we aren't doing a sophisticated diff
 				for (int32 Index = 0; Index < RefLines.Num(); Index++)
 				{
@@ -129,7 +129,7 @@ static FString Tabify(const TCHAR *Input)
 	Result.ReplaceInline(TEXT("\r\n"), TEXT("\n"));
 
 	TArray<FString> Lines;
-	Result.ParseIntoArray( &Lines, TEXT( "\n" ), false );
+	Result.ParseIntoArray( Lines, TEXT( "\n" ), false );
 
 	Result = TEXT("");
 	for (int32 Index = 0; Index < Lines.Num(); Index++)
diff --git a/Engine/Source/Programs/UnrealHeaderTool/Private/HeaderParser.cpp b/Engine/Source/Programs/UnrealHeaderTool/Private/HeaderParser.cpp
--- a/Engine/Source/Programs/UnrealHeaderTool/Private/HeaderParser.cpp
+++ b/Engine/Source/Programs/UnrealHeaderTool/Private/HeaderParser.cpp
@@ -965,7 +965,7 @@ FString FHeaderParser::FormatCommentForToolTip(const FString& Input)
 
 	// get rid of uniform leading whitespace and all trailing whitespace, on each line
 	TArray<FString> Lines;
-	Result.ParseIntoArray(&Lines, TEXT("\n"), false);
+	Result.ParseIntoArray(Lines, TEXT("\n"), false);
 
 	for (auto& Line : Lines)
 	{
@@ -4382,7 +4382,7 @@ void FHeaderParser::CompileClassDeclaration(FClasses& AllClasses)
 				if (HideCategories.Remove(Value) == 0)
 				{
 					TArray<FString> SubCategoryList;
-					Value.ParseIntoArray(&SubCategoryList, TEXT("|"), true);
+					Value.ParseIntoArray(SubCategoryList, TEXT("|"), true);
 
 					FString SubCategoryPath;
 					// look to see if any of the parent paths are excluded in the HideCategories list
diff --git a/Engine/Source/Programs/UnrealHeaderTool/Private/ParserClass.cpp b/Engine/Source/Programs/UnrealHeaderTool/Private/ParserClass.cpp
--- a/Engine/Source/Programs/UnrealHeaderTool/Private/ParserClass.cpp
+++ b/Engine/Source/Programs/UnrealHeaderTool/Private/ParserClass.cpp
@@ -78,7 +78,7 @@ void FClass::GetHideCategories(TArray<FString>& OutHideCategories) const
 	if (HasMetaData(NAME_HideCategories))
 	{
 		const FString& HideCategories = GetMetaData(NAME_HideCategories);
-		HideCategories.ParseIntoArray(&OutHideCategories, TEXT(" "), true);
+		HideCategories.ParseIntoArray(OutHideCategories, TEXT(" "), true);
 	}
 }
 
@@ -88,6 +88,6 @@ void FClass::GetShowCategories(TArray<FString>& OutShowCategories) const
 	if (HasMetaData(NAME_ShowCategories))
 	{
 		const FString& ShowCategories = GetMetaData(NAME_ShowCategories);
-		ShowCategories.ParseIntoArray(&OutShowCategories, TEXT(" "), true);
+		ShowCategories.ParseIntoArray(OutShowCategories, TEXT(" "), true);
 	}
 }
diff --git a/Engine/Source/Programs/UnrealHeaderTool/Private/ParserHelper.cpp b/Engine/Source/Programs/UnrealHeaderTool/Private/ParserHelper.cpp
--- a/Engine/Source/Programs/UnrealHeaderTool/Private/ParserHelper.cpp
+++ b/Engine/Source/Programs/UnrealHeaderTool/Private/ParserHelper.cpp
@@ -202,7 +202,7 @@ FAdvancedDisplayParameterHandler::FAdvancedDisplayParameterHandler(const TMap<FN
 		const FString* FoundString = MetaData->Find(FName(TEXT("AdvancedDisplay")));
 		if(FoundString)
 		{
-			FoundString->ParseIntoArray(&ParametersNames, TEXT(","), true);
+			FoundString->ParseIntoArray(ParametersNames, TEXT(","), true);
 			for(int32 NameIndex = 0; NameIndex < ParametersNames.Num();)
 			{
 				FString& ParameterName = ParametersNames[NameIndex];
diff --git a/Engine/Source/Programs/UnrealPak/Private/KeyGenerator.cpp b/Engine/Source/Programs/UnrealPak/Private/KeyGenerator.cpp
--- a/Engine/Source/Programs/UnrealPak/Private/KeyGenerator.cpp
+++ b/Engine/Source/Programs/UnrealPak/Private/KeyGenerator.cpp
@@ -454,7 +454,7 @@ bool ReadKeysFromFile(const TCHAR* KeyFilename, FKeyPair& OutKeys)
 	if (FFileHelper::LoadFileToString(KeyFileContents, KeyFilename))
 	{
 		TArray<FString> KeyValues;
-		KeyFileContents.ParseIntoArrayWS(&KeyValues);
+		KeyFileContents.ParseIntoArrayWS(KeyValues);
 		if (KeyValues.Num() != 3)
 		{
 			UE_LOG(LogPakFile, Error, TEXT("Expecting 3 values in %s, got %d."), KeyFilename, KeyValues.Num());
diff --git a/Engine/Source/Programs/UnrealPak/Private/UnrealPak.cpp b/Engine/Source/Programs/UnrealPak/Private/UnrealPak.cpp
--- a/Engine/Source/Programs/UnrealPak/Private/UnrealPak.cpp
+++ b/Engine/Source/Programs/UnrealPak/Private/UnrealPak.cpp
@@ -337,7 +337,7 @@ void ProcessOrderFile(int32 ArgC, TCHAR* ArgV[], TMap<FString, uint64>& OrderMap
 		{
 			// Read all lines
 			TArray<FString> Lines;
-			Text.ParseIntoArray(&Lines, TEXT("\n"), true);
+			Text.ParseIntoArray(Lines, TEXT("\n"), true);
 			for (int32 EntryIndex = 0; EntryIndex < Lines.Num(); EntryIndex++)
 			{
 				Lines[EntryIndex].ReplaceInline(TEXT("\r"), TEXT(""));
@@ -434,7 +434,7 @@ void ProcessCommandLine(int32 ArgC, TCHAR* ArgV[], TArray<FPakInputPair>& Entrie
 				// Remove all carriage return characters.
 				Text.ReplaceInline(TEXT("\r"), TEXT(""));
 				// Read all lines
-				Text.ParseIntoArray(&Lines, TEXT("\n"), true);
+				Text.ParseIntoArray(Lines, TEXT("\n"), true);
 			}
 			else
 			{
@@ -697,7 +697,7 @@ FArchive* CreatePakWriter(const TCHAR* Filename)
 		if (KeyFilename.StartsWith(TEXT("0x")))
 		{
 			TArray<FString> KeyValueText;
-			if (KeyFilename.ParseIntoArray(&KeyValueText, TEXT("+"), true) == 3)
+			if (KeyFilename.ParseIntoArray(KeyValueText, TEXT("+"), true) == 3)
 			{
 				Pair.PrivateKey.Exponent.Parse(KeyValueText[0]);
 				Pair.PrivateKey.Modulus.Parse(KeyValueText[1]);
diff --git a/Engine/Source/Runtime/Analytics/AnalyticsMulticast/Private/AnalyticsMulticast.cpp b/Engine/Source/Runtime/Analytics/AnalyticsMulticast/Private/AnalyticsMulticast.cpp
--- a/Engine/Source/Runtime/Analytics/AnalyticsMulticast/Private/AnalyticsMulticast.cpp
+++ b/Engine/Source/Runtime/Analytics/AnalyticsMulticast/Private/AnalyticsMulticast.cpp
@@ -91,7 +91,7 @@ FAnalyticsProviderMulticast::FAnalyticsProviderMulticast(const FAnalyticsMultica
 	if (GetConfigValue.IsBound())
 	{
 		TArray<FString> ModuleNamesArray;
-		ConfigValues.ProviderModuleNames.ParseIntoArray(&ModuleNamesArray, TEXT(","), true);
+		ConfigValues.ProviderModuleNames.ParseIntoArray(ModuleNamesArray, TEXT(","), true);
 		for (TArray<FString>::TConstIterator it(ModuleNamesArray);it;++it)
 		{
 			TSharedPtr<IAnalyticsProvider> NewProvider = FAnalytics::Get().CreateAnalyticsProvider(FName(**it), GetConfigValue);
@@ -170,7 +170,7 @@ bool FAnalyticsProviderMulticast::SetSessionID(const FString& InSessionID)
 	// parse out the format from GetSessionID and set the SessionID in each provider.
 	TArray<FString> SessionIDs;
 	// parse out the module/session pairs
-	InSessionID.ParseIntoArray(&SessionIDs, TEXT("##"), true);
+	InSessionID.ParseIntoArray(SessionIDs, TEXT("##"), true);
 	for (TArray<FString>::TConstIterator it(SessionIDs);it;++it)
 	{
 		FString ModuleName, SessionID;
diff --git a/Engine/Source/Runtime/Analytics/AnalyticsSwrve/Private/AnalyticsSwrve.cpp b/Engine/Source/Runtime/Analytics/AnalyticsSwrve/Private/AnalyticsSwrve.cpp
--- a/Engine/Source/Runtime/Analytics/AnalyticsSwrve/Private/AnalyticsSwrve.cpp
+++ b/Engine/Source/Runtime/Analytics/AnalyticsSwrve/Private/AnalyticsSwrve.cpp
@@ -143,7 +143,7 @@ static void ApplyABTestDiffs(const FString& ResourceDiffPayload)
 			TSharedPtr<FJsonObject> TestObj = (*it)->AsObject();
 			FString IniKey = TestObj->GetStringField(TEXT("uid"));
 			TArray<FString> IniTokens;
-			IniKey.ParseIntoArray(&IniTokens, TEXT("."), false);
+			IniKey.ParseIntoArray(IniTokens, TEXT("."), false);
 			if (IniTokens.Num() < 2)
 			{
 				UE_LOG(LogAnalytics, Warning, TEXT("Failed to parse resource name %s into an INI file and section"), *IniKey);
diff --git a/Engine/Source/Runtime/AssetRegistry/Private/PathTreeNode.cpp b/Engine/Source/Runtime/AssetRegistry/Private/PathTreeNode.cpp
--- a/Engine/Source/Runtime/AssetRegistry/Private/PathTreeNode.cpp
+++ b/Engine/Source/Runtime/AssetRegistry/Private/PathTreeNode.cpp
@@ -33,7 +33,7 @@ bool FPathTreeNode::CachePath(const FString& Path)
 	if (IsUsingExperimentalFNameTreeNodes())
 	{
 		TArray<FString> PathElementStrings;
-		Path.ParseIntoArray(&PathElementStrings, TEXT("/"), /*InCullEmpty=*/true);
+		Path.ParseIntoArray(PathElementStrings, TEXT("/"), /*InCullEmpty=*/true);
 
 		TArray<FName> PathElements;
 		PathElements.Reserve(PathElementStrings.Num());
@@ -47,7 +47,7 @@ bool FPathTreeNode::CachePath(const FString& Path)
 	else
 	{
 		TArray<FString> PathElements;
-		Path.ParseIntoArray(&PathElements, TEXT("/"), /*InCullEmpty=*/true);
+		Path.ParseIntoArray(PathElements, TEXT("/"), /*InCullEmpty=*/true);
 
 		return CachePath_Recursive(PathElements);
 	}
@@ -58,7 +58,7 @@ bool FPathTreeNode::RemoveFolder(const FString& Path)
 	if (IsUsingExperimentalFNameTreeNodes())
 	{
 		TArray<FString> PathElementStrings;
-		Path.ParseIntoArray(&PathElementStrings, TEXT("/"), /*InCullEmpty=*/true);
+		Path.ParseIntoArray(PathElementStrings, TEXT("/"), /*InCullEmpty=*/true);
 
 		TArray<FName> PathElements;
 		PathElements.Reserve(PathElementStrings.Num());
@@ -72,7 +72,7 @@ bool FPathTreeNode::RemoveFolder(const FString& Path)
 	else
 	{
 		TArray<FString> PathElements;
-		Path.ParseIntoArray(&PathElements, TEXT("/"), /*InCullEmpty=*/true);
+		Path.ParseIntoArray(PathElements, TEXT("/"), /*InCullEmpty=*/true);
 
 		return RemoveFolder_Recursive(PathElements);
 	}
@@ -83,7 +83,7 @@ bool FPathTreeNode::GetSubPaths(const FString& BasePath, TSet<FName>& OutPaths,
 	if (IsUsingExperimentalFNameTreeNodes())
 	{
 		TArray<FString> PathElementStrings;
-		BasePath.ParseIntoArray(&PathElementStrings, TEXT("/"), /*InCullEmpty=*/true);
+		BasePath.ParseIntoArray(PathElementStrings, TEXT("/"), /*InCullEmpty=*/true);
 
 		TArray<FName> PathElements;
 		PathElements.Reserve(PathElementStrings.Num());
@@ -113,7 +113,7 @@ bool FPathTreeNode::GetSubPaths(const FString& BasePath, TSet<FName>& OutPaths,
 	else
 	{
 		TArray<FString> PathElements;
-		BasePath.ParseIntoArray(&PathElements, TEXT("/"), /*InCullEmpty=*/true);
+		BasePath.ParseIntoArray(PathElements, TEXT("/"), /*InCullEmpty=*/true);
 
 		const FPathTreeNode* BasePathNode = FindNode_Recursive(PathElements);
 
diff --git a/Engine/Source/Runtime/AutomationWorker/Private/AutomationWorkerModule.cpp b/Engine/Source/Runtime/AutomationWorker/Private/AutomationWorkerModule.cpp
--- a/Engine/Source/Runtime/AutomationWorker/Private/AutomationWorkerModule.cpp
+++ b/Engine/Source/Runtime/AutomationWorker/Private/AutomationWorkerModule.cpp
@@ -523,7 +523,7 @@ bool GenerateTestNamesFromCommandLine(const TCHAR* InTestCommands, TArray<FStrin
 	//Split the test names up
 	TArray<FString> TestList;
 	const FString StringCommand = InTestCommands;
-	StringCommand.ParseIntoArray(&TestList, TEXT(","), true);
+	StringCommand.ParseIntoArray(TestList, TEXT(","), true);
 
 	//Get the list of valid names
 	TArray<FAutomationTestInfo> TestInfo;
diff --git a/Engine/Source/Runtime/Core/Private/Android/AndroidFile.cpp b/Engine/Source/Runtime/Core/Private/Android/AndroidFile.cpp
--- a/Engine/Source/Runtime/Core/Private/Android/AndroidFile.cpp
+++ b/Engine/Source/Runtime/Core/Private/Android/AndroidFile.cpp
@@ -381,7 +381,7 @@ public:
 		close( Handle );
 
 		TArray<FString> Lines;
-		EntireFile.ParseIntoArrayLines(&Lines);
+		EntireFile.ParseIntoArrayLines(Lines);
 
 #if LOG_ANDROID_FILE_MANIFEST
 		FPlatformMisc::LowLevelOutputDebugStringf(TEXT("Loaded manifest file %s"), *ManifestFileName);
diff --git a/Engine/Source/Runtime/Core/Private/Containers/String.cpp b/Engine/Source/Runtime/Core/Private/Containers/String.cpp
--- a/Engine/Source/Runtime/Core/Private/Containers/String.cpp
+++ b/Engine/Source/Runtime/Core/Private/Containers/String.cpp
@@ -501,12 +501,11 @@ bool FString::IsNumeric() const
  *
  * @return	The number of elements in InArray
  */
-int32 FString::ParseIntoArray( TArray<FString>* InArray, const TCHAR* pchDelim, bool InCullEmpty ) const
+int32 FString::ParseIntoArray( TArray<FString>& OutArray, const TCHAR* pchDelim, bool InCullEmpty ) const
 {
 	// Make sure the delimit string is not null or empty
 	check(pchDelim);
-	check(InArray);
-	InArray->Empty();
+	OutArray.Empty();
 	const TCHAR *Start = Data.GetData();
 	int32 DelimLength = FCString::Strlen(pchDelim);
 	if (Start && DelimLength)
@@ -515,101 +514,17 @@ int32 FString::ParseIntoArray( TArray<FString>* InArray, const TCHAR* pchDelim,
 		{
 			if (!InCullEmpty || At-Start)
 			{
-				new (*InArray) FString(At-Start,Start);
+				new (OutArray) FString(At-Start,Start);
 			}
 			Start += DelimLength + (At-Start);
 		}
 		if (!InCullEmpty || *Start)
 		{
-			new(*InArray) FString(Start);
+			new(OutArray) FString(Start);
 		}
 
 	}
-	return InArray->Num();
-}
-
-/**
- * Takes a string, and skips over all instances of white space and returns the new string
- *
- * @param	WhiteSpace		An array of white space strings
- * @param	NumWhiteSpaces	The length of the WhiteSpace array
- * @param	S				The input and output string
- */
-static void SkipOver(const TCHAR** WhiteSpace, int32 NumWhiteSpaces, FString& S)
-{
-	bool bStop = false;
-
-	// keep going until we hit non-white space
-	while (!bStop)
-	{
-		// we stop it we don't find any white space
-		bStop = true;
-		// loop over all possible white spaces to search for
-		for (int32 iWS = 0; iWS < NumWhiteSpaces; iWS++)
-		{
-			// get the length (tiny optimization)
-			int32 WSLen = FCString::Strlen(WhiteSpace[iWS]);
-
-			// if we start with this bit of whitespace, chop it off, and keep looking for more whitespace
-			if (FCString::Strnicmp(*S, WhiteSpace[iWS], WSLen) == 0)
-			{
-				// chop it off
-				S = S.Mid(WSLen);
-				// keep looking!
-				bStop = false;
-				break;
-			}
-		}
-	}
-}
-
-/**
- * Splits the input string on the first bit of white space, and returns the initial token
- * (to the left of the white space), and the rest (white space and to the right)
- *
- * @param	WhiteSpace		An array of white space strings
- * @param	NumWhiteSpaces	The length of the WhiteSpace array
- * @param	Token			The first token before any white space
- * @param	S				The input and outputted remainder string
- *
- * @return	Was there a valid token before the end of the string?
- */
-static bool SplitOn( const TCHAR** WhiteSpace, int32 NumWhiteSpaces, FString& Token, FString& S, TCHAR& InCh )
-{
-	// this is the index of the first instance of whitespace
-	int32 SmallestToken = MAX_int32;
-	InCh = TEXT(' ');
-
-	// loop through all possible white spaces
-	for (int32 iWS = 0; iWS < NumWhiteSpaces; iWS++)
-	{
-		// look for the first instance of it
-		int32 NextWS = S.Find(WhiteSpace[iWS]);
-
-		// if shouldn't be at the start of the string, because SkipOver should have been called
-		check(NextWS != 0);
-
-		// if we found this white space, and it is before any other white spaces, remember it
-		if (NextWS > 0 && NextWS < SmallestToken)
-		{
-			SmallestToken = NextWS;
-			InCh = *WhiteSpace[iWS];
-		}
-	}
-
-	// if we found some white space, SmallestToken is pointing to the the first one
-	if (SmallestToken != MAX_int32)
-	{
-		// get the token before the white space
-		Token = S.Left(SmallestToken);
-		// update out string with the remainder
-		S = S.Mid(SmallestToken);
-		// we found a token
-		return true;
-	}
-
-	// we failed to find a token
-	return false;
+	return OutArray.Num();
 }
 
 bool FString::MatchesWildcard(const FString& InWildcard, ESearchCase::Type SearchCase) const
@@ -690,7 +605,7 @@ bool FString::MatchesWildcard(const FString& InWildcard, ESearchCase::Type Searc
 
 
 /** Caution!! this routine is O(N^2) allocations...use it for parsing very short text or not at all */
-int32 FString::ParseIntoArrayWS( TArray<FString>* InArray, const TCHAR* pchExtraDelim ) const
+int32 FString::ParseIntoArrayWS( TArray<FString>& OutArray, const TCHAR* pchExtraDelim, bool InCullEmpty ) const
 {
 	// default array of White Spaces, the last entry can be replaced with the optional pchExtraDelim string
 	// (if you want to split on white space and another character)
@@ -711,86 +626,85 @@ int32 FString::ParseIntoArrayWS( TArray<FString>* InArray, const TCHAR* pchExtra
 		WhiteSpace[NumWhiteSpaces++] = pchExtraDelim;
 	}
 
-	return ParseIntoArray(InArray, WhiteSpace, NumWhiteSpaces);
+	return ParseIntoArray(OutArray, WhiteSpace, NumWhiteSpaces, InCullEmpty);
 }
 
-int32 FString::ParseIntoArrayLines(TArray<FString>* InArray) const
+int32 FString::ParseIntoArrayLines(TArray<FString>& OutArray, bool InCullEmpty) const
 {
 	// default array of LineEndings
 	static const TCHAR* LineEndings[] =
 	{				
+		TEXT("\r\n"),
 		TEXT("\r"),
 		TEXT("\n"),	
 	};
 
 	// start with just the standard line endings
 	int32 NumLineEndings = ARRAY_COUNT(LineEndings);	
-	return ParseIntoArray(InArray, LineEndings, NumLineEndings);
+	return ParseIntoArray(OutArray, LineEndings, NumLineEndings, InCullEmpty);
 }
 
-int32 FString::ParseIntoArray(TArray<FString>* InArray, const TCHAR** DelimArray, int32 NumDelims) const
+int32 FString::ParseIntoArray(TArray<FString>& OutArray, const TCHAR** DelimArray, int32 NumDelims, bool InCullEmpty) const
 {
+	// Make sure the delimit string is not null or empty
 	check(DelimArray);
-	check(InArray);
-	InArray->Empty();
-
-	// this is our temporary workhorse string
-	FString S = *this;
-
-	bool bStop = false;
-	// keep going until we run out of tokens
-	while (!bStop)
+	OutArray.Empty();
+	const TCHAR *Start = Data.GetData();
+	const int32 Length = Len();
+	if (Start)
 	{
-		// skip over any white space at the beginning of the string
-		SkipOver(DelimArray, NumDelims, S);
+		int32 SubstringBeginIndex = 0;
 
-		// find the first token in the string, and if we get one, add it to the output array of tokens
-		FString Token;
-		TCHAR ch;
-		if (SplitOn(DelimArray, NumDelims, Token, S, ch))
+		// Iterate through string.
+		for(int32 i = 0; i < Len();)
 		{
-			if (Token[0] == TEXT('"'))
+			int32 SubstringEndIndex = INDEX_NONE;
+			int32 DelimiterLength = 0;
+
+			// Attempt each delimiter.
+			for(int32 DelimIndex = 0; DelimIndex < NumDelims; ++DelimIndex)
 			{
-				int32 SaveSz = Token.Len();
+				DelimiterLength = FCString::Strlen(DelimArray[DelimIndex]);
 
-				FString Wk = FString::Printf(TEXT("%s%c"), *Token, ch);
-				for (int32 x = 1; x < S.Len(); ++x)
+				// If we found a delimiter...
+				if (FCString::Strncmp(Start + i, DelimArray[DelimIndex], DelimiterLength) == 0)
 				{
-					if (S[x] == TEXT('"'))
-					{
-						Wk += TEXT("\"");
-						break;
-					}
-					else
-					{
-						Wk = Wk + S.Mid(x, 1);
-					}
+					// Mark the end of the substring.
+					SubstringEndIndex = i;
+					break;
 				}
-
-				Token = Wk;
-
-				int32 DiffSz = Token.Len() - SaveSz;
-				S = S.Mid(DiffSz);
 			}
 
-			// stick it on the end
-			new(*InArray)FString(Token);
-		}
-		else
-		{
-			// if the remaining string is not empty, then we need to add the last token
-			if (S.Len())
+			if (SubstringEndIndex != INDEX_NONE)
 			{
-				new(*InArray)FString(S);
+				const int32 SubstringLength = SubstringEndIndex - SubstringBeginIndex;
+				// If we're not culling empty strings or if we are but the string isn't empty anyways...
+				if(!InCullEmpty || SubstringLength != 0)
+				{
+					// ... add new string from substring beginning up to the beginning of this delimiter.
+					new (OutArray) FString(SubstringEndIndex - SubstringBeginIndex, Start + SubstringBeginIndex);
+				}
+				// Next substring begins at the end of the discovered delimiter.
+				SubstringBeginIndex = SubstringEndIndex + DelimiterLength;
+				i = SubstringBeginIndex;
+			}
+			else
+			{
+				++i;
 			}
+		}
 
-			// and, we're done this crazy ride
-			bStop = true;
+		// Add any remaining characters after the last delimiter.
+		const int32 SubstringLength = Length - SubstringBeginIndex;
+		// If we're not culling empty strings or if we are but the string isn't empty anyways...
+		if(!InCullEmpty || SubstringLength != 0)
+		{
+			// ... add new string from substring beginning up to the beginning of this delimiter.
+			new (OutArray) FString(Start + SubstringBeginIndex);
 		}
 	}
 
-	// simply return the number of elements in the output array
-	return InArray->Num();
+	return OutArray.Num();
 }
 
 FString FString::Replace(const TCHAR* From, const TCHAR* To, ESearchCase::Type SearchCase) const
diff --git a/Engine/Source/Runtime/Core/Private/HAL/ConsoleManager.cpp b/Engine/Source/Runtime/Core/Private/HAL/ConsoleManager.cpp
--- a/Engine/Source/Runtime/Core/Private/HAL/ConsoleManager.cpp
+++ b/Engine/Source/Runtime/Core/Private/HAL/ConsoleManager.cpp
@@ -973,7 +973,7 @@ bool FConsoleManager::ProcessUserConsoleInput(const TCHAR* InInput, FOutputDevic
 		// Process command
 		// Build up argument list
 		TArray< FString > Args;
-		FString( It ).ParseIntoArrayWS( &Args );
+		FString( It ).ParseIntoArrayWS( Args );
 
 		const bool bShowHelp = Args.Num() == 1 && Args[0] == TEXT("?");
 		if( bShowHelp )
diff --git a/Engine/Source/Runtime/Core/Private/HAL/OutputDevices.cpp b/Engine/Source/Runtime/Core/Private/HAL/OutputDevices.cpp
--- a/Engine/Source/Runtime/Core/Private/HAL/OutputDevices.cpp
+++ b/Engine/Source/Runtime/Core/Private/HAL/OutputDevices.cpp
@@ -85,7 +85,7 @@ class FLogSuppressionImplementation: public FLogSuppressionInterface, private FS
 		Cmds = Cmds.Trim().TrimQuotes();
 		Cmds.Trim();
 		TArray<FString> SubCmds;
-		Cmds.ParseIntoArray(&SubCmds, TEXT(","), true);
+		Cmds.ParseIntoArray(SubCmds, TEXT(","), true);
 		for (int32 Index = 0; Index < SubCmds.Num(); Index++)
 		{
 			static FString LogString(TEXT("Log "));
@@ -95,7 +95,7 @@ class FLogSuppressionImplementation: public FLogSuppressionInterface, private FS
 				Command = Command.Right(Command.Len() - LogString.Len());
 			}
 			TArray<FString> CommandParts;
-			Command.ParseIntoArrayWS(&CommandParts);
+			Command.ParseIntoArrayWS(CommandParts);
 			if (CommandParts.Num() < 1)
 			{
 				continue;
diff --git a/Engine/Source/Runtime/Core/Private/Linux/LinuxPlatformProcess.cpp b/Engine/Source/Runtime/Core/Private/Linux/LinuxPlatformProcess.cpp
--- a/Engine/Source/Runtime/Core/Private/Linux/LinuxPlatformProcess.cpp
+++ b/Engine/Source/Runtime/Core/Private/Linux/LinuxPlatformProcess.cpp
@@ -438,7 +438,7 @@ FProcHandle FLinuxPlatformProcess::CreateProc(const TCHAR* URL, const TCHAR* Par
 	UE_LOG(LogHAL, Verbose, TEXT("FLinuxPlatformProcess::CreateProc: '%s'"), *Commandline);
 
 	TArray<FString> ArgvArray;
-	int Argc = Commandline.ParseIntoArray(&ArgvArray, TEXT(" "), true);
+	int Argc = Commandline.ParseIntoArray(ArgvArray, TEXT(" "), true);
 	char* Argv[PlatformProcessLimits::MaxArgvParameters + 1] = { NULL };	// last argument is NULL, hence +1
 	struct CleanupArgvOnExit
 	{
@@ -819,7 +819,7 @@ bool FLinuxPlatformProcess::IsApplicationRunning( const TCHAR* ProcName )
 bool FLinuxPlatformProcess::ExecProcess( const TCHAR* URL, const TCHAR* Params, int32* OutReturnCode, FString* OutStdOut, FString* OutStdErr )
 {
 	TArray<FString> ArgsArray;
-	FString(Params).ParseIntoArray(&ArgsArray, TEXT(" "), true);
+	FString(Params).ParseIntoArray(ArgsArray, TEXT(" "), true);
 	char *args[ArgsArray.Num()];
 
 	for(int i = 0; i < ArgsArray.Num(); i++)
diff --git a/Engine/Source/Runtime/Core/Private/Mac/MacPlatformProcess.cpp b/Engine/Source/Runtime/Core/Private/Mac/MacPlatformProcess.cpp
--- a/Engine/Source/Runtime/Core/Private/Mac/MacPlatformProcess.cpp
+++ b/Engine/Source/Runtime/Core/Private/Mac/MacPlatformProcess.cpp
@@ -237,7 +237,7 @@ bool FMacPlatformProcess::ExecProcess( const TCHAR* URL, const TCHAR* Params, in
 		[ProcessHandle setLaunchPath: LaunchPath];
 		
 		TArray<FString> ArgsArray;
-		FString(Params).ParseIntoArray(&ArgsArray, TEXT(" "), true);
+		FString(Params).ParseIntoArray(ArgsArray, TEXT(" "), true);
 		
 		NSMutableArray *Arguments = [[NSMutableArray new] autorelease];
 		
@@ -390,7 +390,7 @@ FProcHandle FMacPlatformProcess::CreateProc( const TCHAR* URL, const TCHAR* Parm
 		else
 		{
 			TArray<FString> ArgsArray;
-			FString(Parms).ParseIntoArray(&ArgsArray, TEXT(" "), true);
+			FString(Parms).ParseIntoArray(ArgsArray, TEXT(" "), true);
 
 			FString MultiPartArg;
 			for (int32 Index = 0; Index < ArgsArray.Num(); Index++)
diff --git a/Engine/Source/Runtime/Core/Private/Math/Transform.cpp b/Engine/Source/Runtime/Core/Private/Math/Transform.cpp
--- a/Engine/Source/Runtime/Core/Private/Math/Transform.cpp
+++ b/Engine/Source/Runtime/Core/Private/Math/Transform.cpp
@@ -52,7 +52,7 @@ FString FTransform::ToString() const
 bool FTransform::InitFromString( const FString& Source )
 {
 	TArray<FString> ComponentStrings;
-	Source.ParseIntoArray(&ComponentStrings, TEXT("|"), true);
+	Source.ParseIntoArray(ComponentStrings, TEXT("|"), true);
 	const int32 NumComponents = ComponentStrings.Num();
 	if(3 != NumComponents)
 	{
diff --git a/Engine/Source/Runtime/Core/Private/Math/TransformVectorized.cpp b/Engine/Source/Runtime/Core/Private/Math/TransformVectorized.cpp
--- a/Engine/Source/Runtime/Core/Private/Math/TransformVectorized.cpp
+++ b/Engine/Source/Runtime/Core/Private/Math/TransformVectorized.cpp
@@ -52,7 +52,7 @@ FString FTransform::ToString() const
 bool FTransform::InitFromString( const FString& Source )
 {
 	TArray<FString> ComponentStrings;
-	Source.ParseIntoArray(&ComponentStrings, TEXT("|"), true);
+	Source.ParseIntoArray(ComponentStrings, TEXT("|"), true);
 	const int32 NumComponents = ComponentStrings.Num();
 	if(3 != NumComponents)
 	{
diff --git a/Engine/Source/Runtime/Core/Private/Misc/ConfigCacheIni.cpp b/Engine/Source/Runtime/Core/Private/Misc/ConfigCacheIni.cpp
--- a/Engine/Source/Runtime/Core/Private/Misc/ConfigCacheIni.cpp
+++ b/Engine/Source/Runtime/Core/Private/Misc/ConfigCacheIni.cpp
@@ -623,7 +623,7 @@ static void OverrideFromCommandline(FConfigFile* File, const FString& Filename)
 	{
 		// break apart on the commas
 		TArray<FString> SettingPairs;
-		Settings.ParseIntoArray(&SettingPairs, &CommandlineOverrideSpecifiers::PropertySeperator, true);
+		Settings.ParseIntoArray(SettingPairs, &CommandlineOverrideSpecifiers::PropertySeperator, true);
 		for (int32 Index = 0; Index < SettingPairs.Num(); Index++)
 		{
 			// set each one, by splitting on the =
diff --git a/Engine/Source/Runtime/Core/Private/Misc/DateTime.cpp b/Engine/Source/Runtime/Core/Private/Misc/DateTime.cpp
--- a/Engine/Source/Runtime/Core/Private/Misc/DateTime.cpp
+++ b/Engine/Source/Runtime/Core/Private/Misc/DateTime.cpp
@@ -283,7 +283,7 @@ bool FDateTime::Parse( const FString& DateTimeString, FDateTime& OutDateTime )
 	TArray<FString> Tokens;
 
 	// split up on a single delimiter
-	FixedString.ParseIntoArray(&Tokens, TEXT(" "), true);
+	FixedString.ParseIntoArray(Tokens, TEXT(" "), true);
 
 	// make sure it parsed it properly (within reason)
 	if ((Tokens.Num() < 6) || (Tokens.Num() > 7))
diff --git a/Engine/Source/Runtime/Core/Private/Misc/MonitoredProcess.cpp b/Engine/Source/Runtime/Core/Private/Misc/MonitoredProcess.cpp
--- a/Engine/Source/Runtime/Core/Private/Misc/MonitoredProcess.cpp
+++ b/Engine/Source/Runtime/Core/Private/Misc/MonitoredProcess.cpp
@@ -82,7 +82,7 @@ void FMonitoredProcess::ProcessOutput( const FString& Output )
 {
 	TArray<FString> LogLines;
 
-	Output.ParseIntoArray(&LogLines, TEXT("\n"), false);
+	Output.ParseIntoArray(LogLines, TEXT("\n"), false);
 
 	for (int32 LogIndex = 0; LogIndex < LogLines.Num(); ++LogIndex)
 	{
diff --git a/Engine/Source/Runtime/Core/Private/Misc/Paths.cpp b/Engine/Source/Runtime/Core/Private/Misc/Paths.cpp
--- a/Engine/Source/Runtime/Core/Private/Misc/Paths.cpp
+++ b/Engine/Source/Runtime/Core/Private/Misc/Paths.cpp
@@ -721,9 +721,9 @@ bool FPaths::MakePathRelativeTo( FString& InPath, const TCHAR* InRelativeTo )
 	Target.ReplaceInline(TEXT("\\"), TEXT("/"), ESearchCase::CaseSensitive);
 
 	TArray<FString> TargetArray;
-	Target.ParseIntoArray(&TargetArray, TEXT("/"), true);
+	Target.ParseIntoArray(TargetArray, TEXT("/"), true);
 	TArray<FString> SourceArray;
-	Source.ParseIntoArray(&SourceArray, TEXT("/"), true);
+	Source.ParseIntoArray(SourceArray, TEXT("/"), true);
 
 	if (TargetArray.Num() && SourceArray.Num())
 	{
diff --git a/Engine/Source/Runtime/Core/Private/Misc/Timespan.cpp b/Engine/Source/Runtime/Core/Private/Misc/Timespan.cpp
--- a/Engine/Source/Runtime/Core/Private/Misc/Timespan.cpp
+++ b/Engine/Source/Runtime/Core/Private/Misc/Timespan.cpp
@@ -129,7 +129,7 @@ bool FTimespan::Parse( const FString& TimespanString, FTimespan& OutTimespan )
 	TokenString.ReplaceInline(TEXT("-"), TEXT(":"), ESearchCase::CaseSensitive);
 
 	TArray<FString> Tokens;
-	TokenString.ParseIntoArray(&Tokens, TEXT(":"), true);
+	TokenString.ParseIntoArray(Tokens, TEXT(":"), true);
 
 	if (Tokens.Num() == 4)
 	{
diff --git a/Engine/Source/Runtime/Core/Private/Stats/Stats2.cpp b/Engine/Source/Runtime/Core/Private/Stats/Stats2.cpp
--- a/Engine/Source/Runtime/Core/Private/Stats/Stats2.cpp
+++ b/Engine/Source/Runtime/Core/Private/Stats/Stats2.cpp
@@ -1012,7 +1012,7 @@ void FThreadStats::CheckForCollectingStartupStats()
 			break;
 		}
 		TArray<FString> CmdsArray;
-		Cmds.ParseIntoArray(&CmdsArray, TEXT( "," ), true);
+		Cmds.ParseIntoArray(CmdsArray, TEXT( "," ), true);
 		for (int32 Index = 0; Index < CmdsArray.Num(); Index++)
 		{
 			FString StatCmd = FString("stat ") + CmdsArray[Index].Trim();
diff --git a/Engine/Source/Runtime/Core/Private/Stats/StatsData.cpp b/Engine/Source/Runtime/Core/Private/Stats/StatsData.cpp
--- a/Engine/Source/Runtime/Core/Private/Stats/StatsData.cpp
+++ b/Engine/Source/Runtime/Core/Private/Stats/StatsData.cpp
@@ -169,7 +169,7 @@ void FRawStatStackNode::AddNameHierarchy(int32 CurrentPrefixDepth)
 				else
 				{
 					Name.ReplaceInline(TEXT("/"), TEXT("."));
-					Name.ParseIntoArray(&Parts, TEXT("."), true);
+					Name.ParseIntoArray(Parts, TEXT("."), true);
 					check(Parts.Num());
 					ParsedNames.Empty(Parts.Num());
 					for (int32 PartIndex = 0; PartIndex < Parts.Num(); PartIndex++)
diff --git a/Engine/Source/Runtime/Core/Public/Containers/UnrealString.h b/Engine/Source/Runtime/Core/Public/Containers/UnrealString.h
--- a/Engine/Source/Runtime/Core/Public/Containers/UnrealString.h
+++ b/Engine/Source/Runtime/Core/Public/Containers/UnrealString.h
@@ -1202,7 +1202,7 @@ public:
 	 *
 	 * @return	The number of elements in InArray
 	 */
-	int32 ParseIntoArray( TArray<FString>* InArray, const TCHAR* pchDelim, bool InCullEmpty ) const;
+	int32 ParseIntoArray( TArray<FString>& OutArray, const TCHAR* pchDelim, bool InCullEmpty = true ) const;
 
 	/**
 	 * Breaks up a delimited string into elements of a string array, using any whitespace and an 
@@ -1214,7 +1214,7 @@ public:
 	 *
 	 * @return	The number of elements in InArray
 	 */
-	int32 ParseIntoArrayWS( TArray<FString>* InArray, const TCHAR* pchExtraDelim = nullptr ) const;
+	int32 ParseIntoArrayWS( TArray<FString>& OutArray, const TCHAR* pchExtraDelim = nullptr, bool InCullEmpty = true ) const;
 
 	/**
 	* Breaks up a delimited string into elements of a string array, using line ending characters
@@ -1224,7 +1224,7 @@ public:
 	*
 	* @return	The number of elements in InArray
 	*/
-	int32 ParseIntoArrayLines(TArray<FString>* InArray) const;
+	int32 ParseIntoArrayLines(TArray<FString>& OutArray, bool InCullEmpty = true) const;
 
 	/**
 	* Breaks up a delimited string into elements of a string array, using the given delimiters
@@ -1236,7 +1236,7 @@ public:
 	*
 	* @return	The number of elements in InArray
 	*/
-	int32 ParseIntoArray(TArray<FString>* InArray, const TCHAR** DelimArray, int32 NumDelims) const;
+	int32 ParseIntoArray(TArray<FString>& OutArray, const TCHAR** DelimArray, int32 NumDelims, bool InCullEmpty = true) const;
 
 	/**
 	 * Takes an array of strings and removes any zero length entries.
diff --git a/Engine/Source/Runtime/Core/Public/HAL/IPlatformFileOpenLogWrapper.h b/Engine/Source/Runtime/Core/Public/HAL/IPlatformFileOpenLogWrapper.h
--- a/Engine/Source/Runtime/Core/Public/HAL/IPlatformFileOpenLogWrapper.h
+++ b/Engine/Source/Runtime/Core/Public/HAL/IPlatformFileOpenLogWrapper.h
@@ -45,7 +45,7 @@ public:
 			TArray<FString> PlatformNames;
 			if (!(PlatformStr == TEXT("None") || PlatformStr == TEXT("All")))
 			{
-				PlatformStr.ParseIntoArray(&PlatformNames, TEXT("+"), true);
+				PlatformStr.ParseIntoArray(PlatformNames, TEXT("+"), true);
 			}
 
 			for (int32 Platform = 0;Platform < PlatformNames.Num(); ++Platform)
diff --git a/Engine/Source/Runtime/Core/Public/Misc/AutomationTest.h b/Engine/Source/Runtime/Core/Public/Misc/AutomationTest.h
--- a/Engine/Source/Runtime/Core/Public/Misc/AutomationTest.h
+++ b/Engine/Source/Runtime/Core/Public/Misc/AutomationTest.h
@@ -256,7 +256,7 @@ private:
 	{
 		//Split New Test Name into string array
 		TArray<FString> Pieces;
-		InTestInfo.ParseIntoArray(&Pieces, TEXT(","), true);
+		InTestInfo.ParseIntoArray(Pieces, TEXT(","), true);
 
 		// We should always have at least 3 parameters
 		check( Pieces.Num() >= 3 )
diff --git a/Engine/Source/Runtime/CoreUObject/Private/UObject/Class.cpp b/Engine/Source/Runtime/CoreUObject/Private/UObject/Class.cpp
--- a/Engine/Source/Runtime/CoreUObject/Private/UObject/Class.cpp
+++ b/Engine/Source/Runtime/CoreUObject/Private/UObject/Class.cpp
@@ -3607,7 +3607,7 @@ void UClass::GetHideFunctions(TArray<FString>& OutHideFunctions) const
 	if (HasMetaData(NAME_HideFunctions))
 	{
 		const FString& HideFunctions = GetMetaData(NAME_HideFunctions);
-		HideFunctions.ParseIntoArray(&OutHideFunctions, TEXT(" "), true);
+		HideFunctions.ParseIntoArray(OutHideFunctions, TEXT(" "), true);
 	}
 }
 
@@ -3628,7 +3628,7 @@ void UClass::GetAutoExpandCategories(TArray<FString>& OutAutoExpandCategories) c
 	if (HasMetaData(NAME_AutoExpandCategories))
 	{
 		const FString& AutoExpandCategories = GetMetaData(NAME_AutoExpandCategories);
-		AutoExpandCategories.ParseIntoArray(&OutAutoExpandCategories, TEXT(" "), true);
+		AutoExpandCategories.ParseIntoArray(OutAutoExpandCategories, TEXT(" "), true);
 	}
 }
 
@@ -3649,7 +3649,7 @@ void UClass::GetAutoCollapseCategories(TArray<FString>& OutAutoCollapseCategorie
 	if (HasMetaData(NAME_AutoCollapseCategories))
 	{
 		const FString& AutoCollapseCategories = GetMetaData(NAME_AutoCollapseCategories);
-		AutoCollapseCategories.ParseIntoArray(&OutAutoCollapseCategories, TEXT(" "), true);
+		AutoCollapseCategories.ParseIntoArray(OutAutoCollapseCategories, TEXT(" "), true);
 	}
 }
 
@@ -3670,7 +3670,7 @@ void UClass::GetClassGroupNames(TArray<FString>& OutClassGroupNames) const
 	if (HasMetaData(NAME_ClassGroupNames))
 	{
 		const FString& ClassGroupNames = GetMetaData(NAME_ClassGroupNames);
-		ClassGroupNames.ParseIntoArray(&OutClassGroupNames, TEXT(" "), true);
+		ClassGroupNames.ParseIntoArray(OutClassGroupNames, TEXT(" "), true);
 	}
 }
 
diff --git a/Engine/Source/Runtime/CoreUObject/Private/UObject/LazyObjectPtr.cpp b/Engine/Source/Runtime/CoreUObject/Private/UObject/LazyObjectPtr.cpp
--- a/Engine/Source/Runtime/CoreUObject/Private/UObject/LazyObjectPtr.cpp
+++ b/Engine/Source/Runtime/CoreUObject/Private/UObject/LazyObjectPtr.cpp
@@ -51,7 +51,7 @@ void FUniqueObjectGuid::FromString(const FString& From)
 {
 	TArray<FString> Split;
 	Split.Empty(4);
-	if( From.ParseIntoArray( &Split, TEXT("-"), false ) == 4 )
+	if( From.ParseIntoArray( Split, TEXT("-"), false ) == 4 )
 	{
 		Guid.A=FParse::HexNumber(*Split[0]);
 		Guid.B=FParse::HexNumber(*Split[1]);
diff --git a/Engine/Source/Runtime/CoreUObject/Private/UObject/Property.cpp b/Engine/Source/Runtime/CoreUObject/Private/UObject/Property.cpp
--- a/Engine/Source/Runtime/CoreUObject/Private/UObject/Property.cpp
+++ b/Engine/Source/Runtime/CoreUObject/Private/UObject/Property.cpp
@@ -897,7 +897,7 @@ const TCHAR* UProperty::ImportSingleProperty( const TCHAR* Str, void* DestData,
 					if (ImportError.Len() > 0)
 					{
 						TArray<FString> ImportErrors;
-						ImportError.ParseIntoArray(&ImportErrors, LINE_TERMINATOR, true);
+						ImportError.ParseIntoArray(ImportErrors, LINE_TERMINATOR, true);
 
 						for ( int32 ErrorIndex = 0; ErrorIndex < ImportErrors.Num(); ErrorIndex++ )
 						{
@@ -1100,7 +1100,7 @@ const TCHAR* UProperty::ImportSingleProperty( const TCHAR* Str, void* DestData,
 					if (ImportError.Len() > 0)
 					{
 						TArray<FString> ImportErrors;
-						ImportError.ParseIntoArray(&ImportErrors,LINE_TERMINATOR,true);
+						ImportError.ParseIntoArray(ImportErrors,LINE_TERMINATOR,true);
 
 						for ( int32 ErrorIndex = 0; ErrorIndex < ImportErrors.Num(); ErrorIndex++ )
 						{
@@ -1132,7 +1132,7 @@ const TCHAR* UProperty::ImportSingleProperty( const TCHAR* Str, void* DestData,
 					if (ImportError.Len() > 0)
 					{
 						TArray<FString> ImportErrors;
-						ImportError.ParseIntoArray(&ImportErrors, LINE_TERMINATOR, true);
+						ImportError.ParseIntoArray(ImportErrors, LINE_TERMINATOR, true);
 
 						for ( int32 ErrorIndex = 0; ErrorIndex < ImportErrors.Num(); ErrorIndex++ )
 						{
diff --git a/Engine/Source/Runtime/Engine/Classes/Commandlets/Commandlet.h b/Engine/Source/Runtime/Engine/Classes/Commandlets/Commandlet.h
--- a/Engine/Source/Runtime/Engine/Classes/Commandlets/Commandlet.h
+++ b/Engine/Source/Runtime/Engine/Classes/Commandlets/Commandlet.h
@@ -129,7 +129,7 @@ class UCommandlet : public UObject
 		{
 			FString& Switch = Switches[SwitchIdx];
 			TArray<FString> SplitSwitch;
-			if (2 == Switch.ParseIntoArray(&SplitSwitch, TEXT("="), true))
+			if (2 == Switch.ParseIntoArray(SplitSwitch, TEXT("="), true))
 			{
 				Params.Add(SplitSwitch[0], SplitSwitch[1].TrimQuotes());
 				Switches.RemoveAt(SwitchIdx);
diff --git a/Engine/Source/Runtime/Engine/Private/Animation/Skeleton.cpp b/Engine/Source/Runtime/Engine/Private/Animation/Skeleton.cpp
--- a/Engine/Source/Runtime/Engine/Private/Animation/Skeleton.cpp
+++ b/Engine/Source/Runtime/Engine/Private/Animation/Skeleton.cpp
@@ -694,7 +694,7 @@ void USkeleton::CollectAnimationNotifies()
 			if (const FString* Value = Asset.TagsAndValues.Find(USkeleton::AnimNotifyTag))
 			{
 				TArray<FString> NotifyList;
-				Value->ParseIntoArray(&NotifyList, *USkeleton::AnimNotifyTagDelimiter, true);
+				Value->ParseIntoArray(NotifyList, *USkeleton::AnimNotifyTagDelimiter, true);
 				for (auto NotifyIter = NotifyList.CreateConstIterator(); NotifyIter; ++NotifyIter)
 				{
 					FString NotifyName = *NotifyIter;
diff --git a/Engine/Source/Runtime/Engine/Private/EdGraph/EdGraphSchema.cpp b/Engine/Source/Runtime/Engine/Private/EdGraph/EdGraphSchema.cpp
--- a/Engine/Source/Runtime/Engine/Private/EdGraph/EdGraphSchema.cpp
+++ b/Engine/Source/Runtime/Engine/Private/EdGraph/EdGraphSchema.cpp
@@ -67,14 +67,14 @@ void FGraphActionListBuilderBase::ActionGroup::GetCategoryChain(TArray<FString>&
 {
 #if WITH_EDITOR
 	static FString const CategoryDelim("|");
-	FEditorCategoryUtils::GetCategoryDisplayString(RootCategory).ParseIntoArray(&HierarchyOut, *CategoryDelim, true);
+	FEditorCategoryUtils::GetCategoryDisplayString(RootCategory).ParseIntoArray(HierarchyOut, *CategoryDelim, true);
 
 	if (Actions.Num() > 0)
 	{
 		TArray<FString> SubCategoryChain;
 
 		FString SubCategory = FEditorCategoryUtils::GetCategoryDisplayString(Actions[0]->Category);
-		SubCategory.ParseIntoArray(&SubCategoryChain, *CategoryDelim, true);
+		SubCategory.ParseIntoArray(SubCategoryChain, *CategoryDelim, true);
 
 		HierarchyOut.Append(SubCategoryChain);
 	}
diff --git a/Engine/Source/Runtime/Engine/Private/KismetStringLibrary.cpp b/Engine/Source/Runtime/Engine/Private/KismetStringLibrary.cpp
--- a/Engine/Source/Runtime/Engine/Private/KismetStringLibrary.cpp
+++ b/Engine/Source/Runtime/Engine/Private/KismetStringLibrary.cpp
@@ -281,7 +281,7 @@ int32 UKismetStringLibrary::GetCharacterAsNumber(const FString& SourceString, in
 TArray<FString> UKismetStringLibrary::ParseIntoArray(const FString& SourceString, const FString& Delimiter, const bool CullEmptyStrings)
 {
 	TArray<FString> SeparatedStrings;
-	const int32 nArraySize = SourceString.ParseIntoArray(&SeparatedStrings, *Delimiter, CullEmptyStrings);
+	const int32 nArraySize = SourceString.ParseIntoArray(SeparatedStrings, *Delimiter, CullEmptyStrings);
 	return SeparatedStrings;
 }
 
diff --git a/Engine/Source/Runtime/Engine/Private/LocalPlayer.cpp b/Engine/Source/Runtime/Engine/Private/LocalPlayer.cpp
--- a/Engine/Source/Runtime/Engine/Private/LocalPlayer.cpp
+++ b/Engine/Source/Runtime/Engine/Private/LocalPlayer.cpp
@@ -1318,7 +1318,7 @@ bool ULocalPlayer::Exec(UWorld* InWorld, const TCHAR* Cmd,FOutputDevice& Ar)
 		int32 NewLineIndex;
 		if (CmdString.FindChar(TEXT(';'),NewLineIndex))
 		{
-			CmdString.ParseIntoArray(&Lines,TEXT(";"),true);
+			CmdString.ParseIntoArray(Lines,TEXT(";"),true);
 		}
 		else
 		{
@@ -1328,7 +1328,7 @@ bool ULocalPlayer::Exec(UWorld* InWorld, const TCHAR* Cmd,FOutputDevice& Ar)
 		for (int32 LineIndex = 0; LineIndex < Lines.Num(); ++LineIndex)
 		{
 			TArray<FString> Args;
-			Lines[LineIndex].ParseIntoArrayWS(&Args);
+			Lines[LineIndex].ParseIntoArrayWS(Args);
 			FLockedViewState::Get().LockView(this,Args);
 		}
 		if (Lines.Num() > 1)
diff --git a/Engine/Source/Runtime/Engine/Private/Tests/ExternalToolTest.cpp b/Engine/Source/Runtime/Engine/Private/Tests/ExternalToolTest.cpp
--- a/Engine/Source/Runtime/Engine/Private/Tests/ExternalToolTest.cpp
+++ b/Engine/Source/Runtime/Engine/Private/Tests/ExternalToolTest.cpp
@@ -41,7 +41,7 @@ bool FWaitForProcessToCompleteLatentCommand::Update()
 		if( ProcessOutput.Len() > 0 )
 		{
 			TArray<FString> OutputLines;
-			ProcessOutput.ParseIntoArray(&OutputLines, TEXT("\n"), false);
+			ProcessOutput.ParseIntoArray(OutputLines, TEXT("\n"), false);
 			for( int32 i = 0; i < OutputLines.Num(); i++ )
 			{
 				UE_LOG(LogBaseAutomationTests, Log, TEXT("%s"), *OutputLines[i]);
@@ -130,7 +130,7 @@ void FRunExternalToolTest::GetTests(TArray<FString>& OutBeautifiedNames, TArray
 bool FRunExternalToolTest::RunTest(const FString& Parameters)
 {
 	TArray<FString> Commands;
-	Parameters.ParseIntoArray(&Commands,TEXT(";"),false);
+	Parameters.ParseIntoArray(Commands,TEXT(";"),false);
 
 	if( Commands.Num() != 3 )
 	{
diff --git a/Engine/Source/Runtime/Engine/Private/UnrealEngine.cpp b/Engine/Source/Runtime/Engine/Private/UnrealEngine.cpp
--- a/Engine/Source/Runtime/Engine/Private/UnrealEngine.cpp
+++ b/Engine/Source/Runtime/Engine/Private/UnrealEngine.cpp
@@ -773,7 +773,7 @@ void UEngine::Init(IEngineLoop* InEngineLoop)
 	if( FParse::Value(FCommandLine::Get(), TEXT("ExecCmds="), ExecCmds, false) )
 	{
 		TArray<FString> CommandArray;
-		ExecCmds.ParseIntoArray( &CommandArray, TEXT(","), true );
+		ExecCmds.ParseIntoArray( CommandArray, TEXT(","), true );
 
 		for( int32 Cx = 0; Cx < CommandArray.Num(); ++Cx )
 		{
@@ -5115,11 +5115,11 @@ bool UEngine::HandleObjCommand( const TCHAR* Cmd, FOutputDevice& Ar )
 				//@todo: add support to FParse::Value() for specifying characters that should be ignored
 				if ( FParse::Value(Cmd, TEXT("HIDE="), Value/*, TEXT(",")*/) )
 				{
-					Value.ParseIntoArray(&HiddenCategories,TEXT(","),1);
+					Value.ParseIntoArray(HiddenCategories,TEXT(","),1);
 				}
 				else if ( FParse::Value(Cmd, TEXT("SHOW="), Value/*, TEXT(",")*/) )
 				{
-					Value.ParseIntoArray(&ShowingCategories,TEXT(","),1);
+					Value.ParseIntoArray(ShowingCategories,TEXT(","),1);
 				}
 #endif
 				UClass* LastOwnerClass = NULL;
@@ -9104,7 +9104,7 @@ bool UEngine::LoadMap( FWorldContext& WorldContext, FURL URL, class UPendingNetG
 	if (MutatorString)
 	{
 		TArray<FString> Mutators;
-		FString(MutatorString).ParseIntoArray(&Mutators, TEXT(","), true);
+		FString(MutatorString).ParseIntoArray(Mutators, TEXT(","), true);
 
 		for (int32 MutatorIndex = 0; MutatorIndex < Mutators.Num(); MutatorIndex++)
 		{
diff --git a/Engine/Source/Runtime/GameMenuBuilder/Private/GameMenuPage.cpp b/Engine/Source/Runtime/GameMenuBuilder/Private/GameMenuPage.cpp
--- a/Engine/Source/Runtime/GameMenuBuilder/Private/GameMenuPage.cpp
+++ b/Engine/Source/Runtime/GameMenuBuilder/Private/GameMenuPage.cpp
@@ -32,7 +32,7 @@ bool FGameMenuPage::InitialiseRootMenu(APlayerController* InPCOwner, const FStri
 	}	
 
 	TArray<FString> SplitString;
-	InStyleName.ParseIntoArray(&SplitString, TEXT("/"), true);
+	InStyleName.ParseIntoArray(SplitString, TEXT("/"), true);
 	FString StylePath = InStyleName.Replace(*SplitString[SplitString.Num() - 1],TEXT(""));
 	StylePath.RemoveFromEnd("/");
 	FGameMenuBuilderStyle::Initialize(StylePath);
diff --git a/Engine/Source/Runtime/GameplayTags/Private/GameplayTagsManager.cpp b/Engine/Source/Runtime/GameplayTags/Private/GameplayTagsManager.cpp
--- a/Engine/Source/Runtime/GameplayTags/Private/GameplayTagsManager.cpp
+++ b/Engine/Source/Runtime/GameplayTags/Private/GameplayTagsManager.cpp
@@ -79,7 +79,7 @@ void UGameplayTagsManager::GetAllNodesForTag( const FString& Tag, TArray< TShare
 	TArray<FString> Tags;
 	TArray<TSharedPtr<FGameplayTagNode>>& GameplayRootTags = GameplayRootTag->GetChildTagNodes();
 	OutTagArray.Empty();
-	if(Tag.ParseIntoArray( &Tags, TEXT( "." ), true ) > 0)
+	if(Tag.ParseIntoArray( Tags, TEXT( "." ), true ) > 0)
 	{
 		int32 CurrentTagDepth = 0;
 		// find the first node of the tag
@@ -324,7 +324,7 @@ void UGameplayTagsManager::AddTagTableRow(const FGameplayTagTableRow& TagRow)
 	// Split the tag text on the "." delimiter to establish tag depth and then insert each tag into the
 	// gameplay tag tree
 	TArray<FString> SubTags;
-	TagRow.Tag.ParseIntoArray(&SubTags, TEXT("."), true);
+	TagRow.Tag.ParseIntoArray(SubTags, TEXT("."), true);
 
 	if (SubTags.Num() > 0)
 	{
@@ -445,7 +445,7 @@ void UGameplayTagsManager::GetFilteredGameplayRootTags( const FString& InFilterS
 	TArray<TSharedPtr<FGameplayTagNode>>& GameplayRootTags = GameplayRootTag->GetChildTagNodes();
 
 	OutTagArray.Empty();
-	if( InFilterString.ParseIntoArray( &Filters, TEXT( "," ), true ) > 0 )
+	if( InFilterString.ParseIntoArray( Filters, TEXT( "," ), true ) > 0 )
 	{
 		// Check all filters in the list
 		for (int32 iFilter = 0; iFilter < Filters.Num(); ++iFilter)
diff --git a/Engine/Source/Runtime/Internationalization/Private/Serialization/JsonInternationalizationArchiveSerializer.cpp b/Engine/Source/Runtime/Internationalization/Private/Serialization/JsonInternationalizationArchiveSerializer.cpp
--- a/Engine/Source/Runtime/Internationalization/Private/Serialization/JsonInternationalizationArchiveSerializer.cpp
+++ b/Engine/Source/Runtime/Internationalization/Private/Serialization/JsonInternationalizationArchiveSerializer.cpp
@@ -330,7 +330,7 @@ void FJsonInternationalizationArchiveSerializer::GenerateStructuredData( TShared
 		TArray< FString > NamespaceTokens;
 
 		// Tokenize the namespace by using '.' as a delimiter
-		int32 NamespaceTokenCount = UnstructuredArchiveEntry->Namespace.ParseIntoArray( &NamespaceTokens, *NAMESPACE_DELIMITER, true );
+		int32 NamespaceTokenCount = UnstructuredArchiveEntry->Namespace.ParseIntoArray( NamespaceTokens, *NAMESPACE_DELIMITER, true );
 
 		TSharedPtr< FStructuredArchiveEntry > StructuredArchiveEntry = RootElement;
 		//Loop through all the namespace tokens and find the appropriate structured entry, if it does not exist add it.  At the end StructuredArchiveEntry
diff --git a/Engine/Source/Runtime/Internationalization/Private/Serialization/JsonInternationalizationManifestSerializer.cpp b/Engine/Source/Runtime/Internationalization/Private/Serialization/JsonInternationalizationManifestSerializer.cpp
--- a/Engine/Source/Runtime/Internationalization/Private/Serialization/JsonInternationalizationManifestSerializer.cpp
+++ b/Engine/Source/Runtime/Internationalization/Private/Serialization/JsonInternationalizationManifestSerializer.cpp
@@ -338,7 +338,7 @@ void FJsonInternationalizationManifestSerializer::GenerateStructuredData( TShare
 		TArray< FString > NamespaceTokens;
 
 		// Tokenize the namespace by using '.' as a delimiter
-		int32 NamespaceTokenCount = UnstructuredManifestEntry->Namespace.ParseIntoArray( &NamespaceTokens, *NAMESPACE_DELIMITER, true );
+		int32 NamespaceTokenCount = UnstructuredManifestEntry->Namespace.ParseIntoArray( NamespaceTokens, *NAMESPACE_DELIMITER, true );
 
 		TSharedPtr< FStructuredEntry > StructuredManifestEntry = RootElement;
 		//Loop through all the namespace tokens and find the appropriate structured entry, if it does not exist add it.  At the end StructuredManifestEntry
diff --git a/Engine/Source/Runtime/Launch/Private/LaunchEngineLoop.cpp b/Engine/Source/Runtime/Launch/Private/LaunchEngineLoop.cpp
--- a/Engine/Source/Runtime/Launch/Private/LaunchEngineLoop.cpp
+++ b/Engine/Source/Runtime/Launch/Private/LaunchEngineLoop.cpp
@@ -714,7 +714,7 @@ int32 FEngineLoop::PreInit( const TCHAR* CmdLine )
 		if (FParse::Value(FCommandLine::Get(), TEXT("testexit="), ExitPhrases))
 		{
 			TArray<FString> ExitPhrasesList;
-			if (ExitPhrases.ParseIntoArray(&ExitPhrasesList, TEXT("+"), true) > 0)
+			if (ExitPhrases.ParseIntoArray(ExitPhrasesList, TEXT("+"), true) > 0)
 			{
 				GScopedTestExit = new FOutputDeviceTestExit(ExitPhrasesList);
 				GLog->AddOutputDevice(GScopedTestExit.GetOwnedPointer());
diff --git a/Engine/Source/Runtime/MovieSceneCore/Private/MovieSceneCommonHelpers.cpp b/Engine/Source/Runtime/MovieSceneCore/Private/MovieSceneCommonHelpers.cpp
--- a/Engine/Source/Runtime/MovieSceneCore/Private/MovieSceneCommonHelpers.cpp
+++ b/Engine/Source/Runtime/MovieSceneCore/Private/MovieSceneCommonHelpers.cpp
@@ -132,7 +132,7 @@ FTrackInstancePropertyBindings::FPropertyAddress FTrackInstancePropertyBindings:
 {
 	TArray<FString> PropertyNames;
 
-	InPropertyPath.ParseIntoArray(&PropertyNames, TEXT("."), true);
+	InPropertyPath.ParseIntoArray(PropertyNames, TEXT("."), true);
 
 	if( PropertyNames.Num() > 0 )
 	{
diff --git a/Engine/Source/Runtime/NetworkFile/Private/NetworkPlatformFile.cpp b/Engine/Source/Runtime/NetworkFile/Private/NetworkPlatformFile.cpp
--- a/Engine/Source/Runtime/NetworkFile/Private/NetworkPlatformFile.cpp
+++ b/Engine/Source/Runtime/NetworkFile/Private/NetworkPlatformFile.cpp
@@ -61,7 +61,7 @@ bool FNetworkPlatformFile::Initialize(IPlatformFile* Inner, const TCHAR* CmdLine
 	if (FParse::Value(CmdLine, TEXT("-FileHostIP="), HostIpString))
 	{
 		TArray<FString> HostIpList;
-		if (HostIpString.ParseIntoArray(&HostIpList, TEXT("+"), true) > 0)
+		if (HostIpString.ParseIntoArray(HostIpList, TEXT("+"), true) > 0)
 		{
 			for (int32 HostIpIndex = 0; !bResult && HostIpIndex < HostIpList.Num(); ++HostIpIndex)
 			{
diff --git a/Engine/Source/Runtime/Networking/Private/IPv4/IPv4Address.cpp b/Engine/Source/Runtime/Networking/Private/IPv4/IPv4Address.cpp
--- a/Engine/Source/Runtime/Networking/Private/IPv4/IPv4Address.cpp
+++ b/Engine/Source/Runtime/Networking/Private/IPv4/IPv4Address.cpp
@@ -29,7 +29,7 @@ bool FIPv4Address::Parse( const FString& AddressString, FIPv4Address& OutAddress
 {
 	TArray<FString> Tokens;
 
-	if (AddressString.ParseIntoArray(&Tokens, TEXT("."), false) == 4)
+	if (AddressString.ParseIntoArray(Tokens, TEXT("."), false) == 4)
 	{
 		OutAddress.Bytes[0] = FCString::Atoi(*Tokens[3]);
 		OutAddress.Bytes[1] = FCString::Atoi(*Tokens[2]);
diff --git a/Engine/Source/Runtime/Networking/Private/IPv4/IPv4Endpoint.cpp b/Engine/Source/Runtime/Networking/Private/IPv4/IPv4Endpoint.cpp
--- a/Engine/Source/Runtime/Networking/Private/IPv4/IPv4Endpoint.cpp
+++ b/Engine/Source/Runtime/Networking/Private/IPv4/IPv4Endpoint.cpp
@@ -32,7 +32,7 @@ bool FIPv4Endpoint::Parse( const FString& EndpointString, FIPv4Endpoint& OutEndp
 {
 	TArray<FString> Tokens;
 
-	if (EndpointString.ParseIntoArray(&Tokens, TEXT(":"), false) == 2)
+	if (EndpointString.ParseIntoArray(Tokens, TEXT(":"), false) == 2)
 	{
 		if (FIPv4Address::Parse(Tokens[0], OutEndpoint.Address))
 		{
diff --git a/Engine/Source/Runtime/Networking/Private/IPv4/IPv4Subnet.cpp b/Engine/Source/Runtime/Networking/Private/IPv4/IPv4Subnet.cpp
--- a/Engine/Source/Runtime/Networking/Private/IPv4/IPv4Subnet.cpp
+++ b/Engine/Source/Runtime/Networking/Private/IPv4/IPv4Subnet.cpp
@@ -19,7 +19,7 @@ bool FIPv4Subnet::Parse( const FString& SubnetString, FIPv4Subnet& OutSubnet )
 {
 	TArray<FString> Tokens;
 
-	if (SubnetString.ParseIntoArray(&Tokens, TEXT("/"), false) == 2)
+	if (SubnetString.ParseIntoArray(Tokens, TEXT("/"), false) == 2)
 	{
 		return (FIPv4Address::Parse(Tokens[0], OutSubnet.Address) && FIPv4SubnetMask::Parse(Tokens[1], OutSubnet.Mask));
 	}
diff --git a/Engine/Source/Runtime/Networking/Private/IPv4/IPv4SubnetMask.cpp b/Engine/Source/Runtime/Networking/Private/IPv4/IPv4SubnetMask.cpp
--- a/Engine/Source/Runtime/Networking/Private/IPv4/IPv4SubnetMask.cpp
+++ b/Engine/Source/Runtime/Networking/Private/IPv4/IPv4SubnetMask.cpp
@@ -19,7 +19,7 @@ bool FIPv4SubnetMask::Parse( const FString& MaskString, FIPv4SubnetMask& OutMask
 {
 	TArray<FString> Tokens;
 
-	if (MaskString.ParseIntoArray(&Tokens, TEXT("."), false) == 4)
+	if (MaskString.ParseIntoArray(Tokens, TEXT("."), false) == 4)
 	{
 		OutMask.Bytes[0] = FCString::Atoi(*Tokens[3]);
 		OutMask.Bytes[1] = FCString::Atoi(*Tokens[2]);
diff --git a/Engine/Source/Runtime/Online/BuildPatchServices/Private/BuildPatchFileConstructor.cpp b/Engine/Source/Runtime/Online/BuildPatchServices/Private/BuildPatchFileConstructor.cpp
--- a/Engine/Source/Runtime/Online/BuildPatchServices/Private/BuildPatchFileConstructor.cpp
+++ b/Engine/Source/Runtime/Online/BuildPatchServices/Private/BuildPatchFileConstructor.cpp
@@ -58,7 +58,7 @@ public:
 			FString PrevResumeData;
 			TArray< FString > ResumeDataLines;
 			FFileHelper::LoadFileToString( PrevResumeData, *ResumeDataFile );
-			PrevResumeData.ParseIntoArray( &ResumeDataLines, TEXT( "\n" ), true );
+			PrevResumeData.ParseIntoArray( ResumeDataLines, TEXT( "\n" ), true );
 			// Line 1 will be the previously attempted version
 			bHasResumeData = ( ResumeDataLines.Num() >= 1 ) && ResumeDataLines[0] == PatchVersion;
 			if( bHasResumeData )
diff --git a/Engine/Source/Runtime/Online/BuildPatchServices/Private/BuildPatchGeneration.cpp b/Engine/Source/Runtime/Online/BuildPatchServices/Private/BuildPatchGeneration.cpp
--- a/Engine/Source/Runtime/Online/BuildPatchServices/Private/BuildPatchGeneration.cpp
+++ b/Engine/Source/Runtime/Online/BuildPatchServices/Private/BuildPatchGeneration.cpp
@@ -1948,7 +1948,7 @@ void FBuildDataGenerator::StripIgnoredFiles( TArray< FString >& AllFiles, const
 	FString IgnoreFileList = TEXT( "" );
 	FFileHelper::LoadFileToString( IgnoreFileList, *IgnoreListFile );
 	TArray< FString > IgnoreFiles;
-	IgnoreFileList.ParseIntoArray( &IgnoreFiles, TEXT( "\r\n" ), true );
+	IgnoreFileList.ParseIntoArray( IgnoreFiles, TEXT( "\r\n" ), true );
 
 	// Normalize all paths first
 	for (FString& Filename : AllFiles)
diff --git a/Engine/Source/Runtime/Online/HTTP/Private/Curl/CurlHttp.cpp b/Engine/Source/Runtime/Online/HTTP/Private/Curl/CurlHttp.cpp
--- a/Engine/Source/Runtime/Online/HTTP/Private/Curl/CurlHttp.cpp
+++ b/Engine/Source/Runtime/Online/HTTP/Private/Curl/CurlHttp.cpp
@@ -99,7 +99,7 @@ FString FCurlHttpRequest::GetURLParameter(const FString& ParameterName)
 {
 	TArray<FString> StringElements;
 
-	int32 NumElems = URL.ParseIntoArray(&StringElements, TEXT("&"), true);
+	int32 NumElems = URL.ParseIntoArray(StringElements, TEXT("&"), true);
 	check(NumElems == StringElements.Num());
 	
 	FString ParamValDelimiter(TEXT("="));
diff --git a/Engine/Source/Runtime/Online/HTTP/Private/HTML5/HTML5HTTP.cpp b/Engine/Source/Runtime/Online/HTTP/Private/HTML5/HTML5HTTP.cpp
--- a/Engine/Source/Runtime/Online/HTTP/Private/HTML5/HTML5HTTP.cpp
+++ b/Engine/Source/Runtime/Online/HTTP/Private/HTML5/HTML5HTTP.cpp
@@ -64,7 +64,7 @@ FString FHTML5HttpRequest::GetURLParameter(const FString& ParameterName)
 
 	TArray<FString> StringElements;
 
-	int32 NumElems = URL.ParseIntoArray(&StringElements, TEXT("&"), true);
+	int32 NumElems = URL.ParseIntoArray(StringElements, TEXT("&"), true);
 	check(NumElems == StringElements.Num());
 
 	FString ParamValDelimiter(TEXT("="));
diff --git a/Engine/Source/Runtime/Online/OnlineSubsystemAmazon/Private/OnlineIdentityAmazon.cpp b/Engine/Source/Runtime/Online/OnlineSubsystemAmazon/Private/OnlineIdentityAmazon.cpp
--- a/Engine/Source/Runtime/Online/OnlineSubsystemAmazon/Private/OnlineIdentityAmazon.cpp
+++ b/Engine/Source/Runtime/Online/OnlineSubsystemAmazon/Private/OnlineIdentityAmazon.cpp
@@ -135,7 +135,7 @@ bool FOnlineIdentityAmazon::ParseLoginResults(const FString& Results,FUserOnline
 	// The format is var=val,var=val,,
 	// First split the string into a list of var=val entries
 	TArray<FString> Fields;
-	Results.ParseIntoArray(&Fields, TEXT(","), true);
+	Results.ParseIntoArray(Fields, TEXT(","), true);
 	// Loop through the fields, setting the proper account field with the data
 	for (int Index = 0; Index < Fields.Num(); Index++)
 	{
diff --git a/Engine/Source/Runtime/Online/OnlineSubsystemSteam/Private/OnlineSessionAsyncServerSteam.cpp b/Engine/Source/Runtime/Online/OnlineSubsystemSteam/Private/OnlineSessionAsyncServerSteam.cpp
--- a/Engine/Source/Runtime/Online/OnlineSubsystemSteam/Private/OnlineSessionAsyncServerSteam.cpp
+++ b/Engine/Source/Runtime/Online/OnlineSubsystemSteam/Private/OnlineSessionAsyncServerSteam.cpp
@@ -943,7 +943,7 @@ void FOnlineAsyncTaskSteamFindServerBase::ParseSearchResult(class gameserveritem
 		int32 BuildUniqueId = GetBuildUniqueId();
 
 		TArray<FString> TagArray;
-		GameTags.ParseIntoArray(&TagArray, TEXT(","), true);
+		GameTags.ParseIntoArray(TagArray, TEXT(","), true);
 		if (TagArray.Num() > 0 && TagArray[0].StartsWith(STEAMKEY_BUILDUNIQUEID))
 		{
 			ServerBuildId = FCString::Atoi(*TagArray[0].Mid(ARRAY_COUNT(STEAMKEY_BUILDUNIQUEID)));
diff --git a/Engine/Source/Runtime/OpenGLDrv/Private/OpenGLDevice.cpp b/Engine/Source/Runtime/OpenGLDrv/Private/OpenGLDevice.cpp
--- a/Engine/Source/Runtime/OpenGLDrv/Private/OpenGLDevice.cpp
+++ b/Engine/Source/Runtime/OpenGLDrv/Private/OpenGLDevice.cpp
@@ -519,7 +519,7 @@ static void InitRHICapabilitiesForGL()
 		// Log supported GL extensions
 		UE_LOG(LogRHI, Log, TEXT("OpenGL Extensions:"));
 		TArray<FString> GLExtensionArray;
-		ExtensionsString.ParseIntoArray(&GLExtensionArray, TEXT(" "), true);
+		ExtensionsString.ParseIntoArray(GLExtensionArray, TEXT(" "), true);
 		for (int ExtIndex = 0; ExtIndex < GLExtensionArray.Num(); ExtIndex++)
 		{
 			UE_LOG(LogRHI, Log, TEXT("  %s"), *GLExtensionArray[ExtIndex]);
diff --git a/Engine/Source/Runtime/OpenGLDrv/Private/OpenGLShaders.cpp b/Engine/Source/Runtime/OpenGLDrv/Private/OpenGLShaders.cpp
--- a/Engine/Source/Runtime/OpenGLDrv/Private/OpenGLShaders.cpp
+++ b/Engine/Source/Runtime/OpenGLDrv/Private/OpenGLShaders.cpp
@@ -1941,7 +1941,7 @@ static void BindShaderStage(FOpenGLLinkedProgramConfiguration::ShaderInfo& Shade
 					
 					TArray<FString> PrevLines;
 					FString PrevSource = GetShaderStageSource<TOpenGLStage1>(PrevStage);
-					PrevSource.ParseIntoArrayLines(&PrevLines);
+					PrevSource.ParseIntoArrayLines(PrevLines);
 					bool const bOutputElimination = OutputElimination.Num() > 0;
 					for(FOpenGLShaderVarying Output : OutputElimination)
 					{
diff --git a/Engine/Source/Runtime/PakFile/Private/IPlatformFilePak.cpp b/Engine/Source/Runtime/PakFile/Private/IPlatformFilePak.cpp
--- a/Engine/Source/Runtime/PakFile/Private/IPlatformFilePak.cpp
+++ b/Engine/Source/Runtime/PakFile/Private/IPlatformFilePak.cpp
@@ -597,7 +597,7 @@ void FPakPlatformFile::GetPakFolders(const TCHAR* CmdLine, TArray<FString>& OutP
 	if (FParse::Value(CmdLine, TEXT("-pakdir="), PakDirs))
 	{
 		TArray<FString> CmdLineFolders;
-		PakDirs.ParseIntoArray(&CmdLineFolders, TEXT("*"), true);
+		PakDirs.ParseIntoArray(CmdLineFolders, TEXT("*"), true);
 		OutPakFolders.Append(CmdLineFolders);
 	}
 #endif
@@ -658,7 +658,7 @@ bool FPakPlatformFile::Initialize(IPlatformFile* Inner, const TCHAR* CmdLine)
 	FString CmdLinePaksToLoad;
 	if (FParse::Value(CmdLine, TEXT("-paklist="), CmdLinePaksToLoad))
 	{
-		CmdLinePaksToLoad.ParseIntoArray(&PaksToLoad, TEXT("+"), true);
+		CmdLinePaksToLoad.ParseIntoArray(PaksToLoad, TEXT("+"), true);
 	}
 #endif
 
diff --git a/Engine/Source/Runtime/SlateCore/Private/Fonts/FontCache.cpp b/Engine/Source/Runtime/SlateCore/Private/Fonts/FontCache.cpp
--- a/Engine/Source/Runtime/SlateCore/Private/Fonts/FontCache.cpp
+++ b/Engine/Source/Runtime/SlateCore/Private/Fonts/FontCache.cpp
@@ -806,7 +806,7 @@ private:
 				{
 					// Parse out the font attributes
 					TArray<FString> Styles;
-					FString(FaceAndMemory->Face->style_name).ParseIntoArray(&Styles, TEXT(" "), true);
+					FString(FaceAndMemory->Face->style_name).ParseIntoArray(Styles, TEXT(" "), true);
 
 					for (const FString& Style : Styles)
 					{
diff --git a/Engine/Source/Runtime/Sockets/Private/BSDSockets/IPAddressBSD.h b/Engine/Source/Runtime/Sockets/Private/BSDSockets/IPAddressBSD.h
--- a/Engine/Source/Runtime/Sockets/Private/BSDSockets/IPAddressBSD.h
+++ b/Engine/Source/Runtime/Sockets/Private/BSDSockets/IPAddressBSD.h
@@ -48,7 +48,7 @@ public:
 		FString AddressString = InAddr;
 
 		TArray<FString> PortTokens;
-		AddressString.ParseIntoArray(&PortTokens, TEXT(":"), true);
+		AddressString.ParseIntoArray(PortTokens, TEXT(":"), true);
 		
 		// look for a port number
 		if (PortTokens.Num() > 1)
@@ -58,7 +58,7 @@ public:
 
 		// now split the part before the : into a.b.c.d
 		TArray<FString> AddrTokens;
-		PortTokens[0].ParseIntoArray(&AddrTokens, TEXT("."), true);
+		PortTokens[0].ParseIntoArray(AddrTokens, TEXT("."), true);
 
 		if (AddrTokens.Num() < 4)
 		{
diff --git a/Engine/Source/Runtime/Sockets/Private/HTML5/Simulator/IPAddressX.cpp b/Engine/Source/Runtime/Sockets/Private/HTML5/Simulator/IPAddressX.cpp
--- a/Engine/Source/Runtime/Sockets/Private/HTML5/Simulator/IPAddressX.cpp
+++ b/Engine/Source/Runtime/Sockets/Private/HTML5/Simulator/IPAddressX.cpp
@@ -23,7 +23,7 @@ void FInternetAddrX::SetIp(const TCHAR* InAddr, bool& bIsValid)
 	FString AddressString = InAddr;
 
 	TArray<FString> PortTokens;
-	AddressString.ParseIntoArray(&PortTokens, TEXT(":"), true);
+	AddressString.ParseIntoArray(PortTokens, TEXT(":"), true);
 
 	// look for a port number
 	if (PortTokens.Num() > 1)
@@ -33,7 +33,7 @@ void FInternetAddrX::SetIp(const TCHAR* InAddr, bool& bIsValid)
 
 	// now split the part before the : into a.b.c.d
 	TArray<FString> AddrTokens;
-	PortTokens[0].ParseIntoArray(&AddrTokens, TEXT("."), true);
+	PortTokens[0].ParseIntoArray(AddrTokens, TEXT("."), true);
 
 	if (AddrTokens.Num() < 4)
 	{
diff --git a/Engine/Source/Runtime/WebBrowser/Private/WebBrowserSingleton.cpp b/Engine/Source/Runtime/WebBrowser/Private/WebBrowserSingleton.cpp
--- a/Engine/Source/Runtime/WebBrowser/Private/WebBrowserSingleton.cpp
+++ b/Engine/Source/Runtime/WebBrowser/Private/WebBrowserSingleton.cpp
@@ -24,7 +24,7 @@ FWebBrowserSingleton::FWebBrowserSingleton()
 	CefMainArgs MainArgs(hInstance);
 #elif PLATFORM_MAC || PLATFORM_LINUX
     //TArray<FString> Args;
-    //int ArgCount = GSavedCommandLine.ParseIntoArray(&Args, TEXT(" "), true);
+    //int ArgCount = GSavedCommandLine.ParseIntoArray(Args, TEXT(" "), true);
     CefMainArgs MainArgs(0, nullptr);
 #endif
 
diff --git a/Engine/Source/Runtime/XmlParser/Private/XmlFile.cpp b/Engine/Source/Runtime/XmlParser/Private/XmlFile.cpp
--- a/Engine/Source/Runtime/XmlParser/Private/XmlFile.cpp
+++ b/Engine/Source/Runtime/XmlParser/Private/XmlFile.cpp
@@ -610,7 +610,7 @@ void FXmlFile::AddAttribute(const FString& InToken, TArray<FXmlAttribute>& OutAt
 	TArray<FString> Tokens;
 	if(InToken.Contains(TEXT("=")))
 	{
-		InToken.ParseIntoArray(&Tokens, TEXT("="), true);
+		InToken.ParseIntoArray(Tokens, TEXT("="), true);
 		if(Tokens.Num() == 2)
 		{
 			bool bQuotesRemoved = false;
