Minor perf: cache thread id in TLS on Linux as syscall is too heavy.

(Merging CL 2465322 from my branch).

diff --git a/Engine/Source/Runtime/Core/Public/Linux/LinuxPlatformTLS.h b/Engine/Source/Runtime/Core/Public/Linux/LinuxPlatformTLS.h
--- a/Engine/Source/Runtime/Core/Public/Linux/LinuxPlatformTLS.h
+++ b/Engine/Source/Runtime/Core/Public/Linux/LinuxPlatformTLS.h
@@ -25,9 +25,16 @@ struct CORE_API FLinuxTLS : public FGenericPlatformTLS
 	{
 		// note: cannot use pthread_self() without updating the rest of API to opaque (or at least 64-bit) thread handles
 #if defined(_GNU_SOURCE)
-		pid_t ThreadId = static_cast<pid_t>(syscall(SYS_gettid));
-		static_assert(sizeof(pid_t) <= sizeof(uint32), "pid_t is larger than uint32, reconsider implementation of GetCurrentThreadId()");
-		return static_cast<pid_t>(ThreadId);
+		// syscall() is relatively heavy and shows up in the profiler, given that IsInGameThread() is used quite often. Cache thread id in TLS.
+		static __thread uint32 ThreadIdTLS = 0;
+		if (ThreadIdTLS == 0)
+		{
+			pid_t ThreadId = static_cast<pid_t>(syscall(SYS_gettid));
+			static_assert(sizeof(pid_t) <= sizeof(uint32), "pid_t is larger than uint32, reconsider implementation of GetCurrentThreadId()");
+			ThreadIdTLS = static_cast<uint32>(ThreadId);
+			checkf(ThreadIdTLS != 0, TEXT("ThreadId is 0 - reconsider implementation of GetCurrentThreadId() (syscall changed?)"));
+		}
+		return ThreadIdTLS;
 #else
 		// better than nothing...
 		return static_cast< uint32 >(pthread_self());
