Make sure FThreadStats::Shutdown() checks if stats for a thread were actually initialized when getting a pointer from TLS

diff --git a/Engine/Source/Runtime/Core/Public/Stats/Stats2.h b/Engine/Source/Runtime/Core/Public/Stats/Stats2.h
--- a/Engine/Source/Runtime/Core/Public/Stats/Stats2.h
+++ b/Engine/Source/Runtime/Core/Public/Stats/Stats2.h
@@ -1272,7 +1272,7 @@ public:
 	/** This should be called when a thread exits, this deletes FThreadStats from the heap and TLS. **/
 	static void Shutdown()
 	{
-		FThreadStats* Stats = (FThreadStats*)FPlatformTLS::GetTlsValue(TlsSlot);
+		FThreadStats* Stats = IsThreadingReady() ? (FThreadStats*)FPlatformTLS::GetTlsValue(TlsSlot) : nullptr;
 		if (Stats)
 		{
 			// Send all remaining messages.
