LogUnion log category restored.

diff --git a/Engine/Source/Runtime/Core/Private/Containers/Union.cpp b/Engine/Source/Runtime/Core/Private/Containers/Union.cpp
--- a/Engine/Source/Runtime/Core/Private/Containers/Union.cpp
+++ b/Engine/Source/Runtime/Core/Private/Containers/Union.cpp
@@ -1,8 +1,6 @@
 // Copyright 1998-2015 Epic Games, Inc. All Rights Reserved.
 
 #include "CorePrivatePCH.h"
+#include "Union.h"
 
-// Seemingly not used, because it doesn't compile!
-//#include "Union.h"
-
-//DEFINE_LOG_CATEGORY(LogUnion);
+DEFINE_LOG_CATEGORY(LogUnion);
