Added handling of window reordering on Mac when the editor is in the background

diff --git a/Engine/Source/Runtime/Core/Private/Mac/CocoaTextView.cpp b/Engine/Source/Runtime/Core/Private/Mac/CocoaTextView.cpp
--- a/Engine/Source/Runtime/Core/Private/Mac/CocoaTextView.cpp
+++ b/Engine/Source/Runtime/Core/Private/Mac/CocoaTextView.cpp
@@ -21,14 +21,6 @@
     return self;
 }
 
-- (void)drawRect:(NSRect)dirtyRect
-{
-	SCOPED_AUTORELEASE_POOL;
-	[super drawRect:dirtyRect];
-	
-    // Drawing code here.
-}
-
 - (bool)imkKeyDown:(NSEvent *)theEvent
 {
 	if (IMMContext.IsValid())
@@ -43,6 +35,14 @@
 	}
 }
 
+/**
+ * Forward mouse events up to the window rather than through the responder chain.
+ */
+- (BOOL)acceptsFirstMouse:(NSEvent *)Event
+{
+	return YES;
+}
+
 - (void)mouseDown:(NSEvent *)theEvent
 {
 	SCOPED_AUTORELEASE_POOL;
diff --git a/Engine/Source/Runtime/Core/Private/Mac/MacApplication.cpp b/Engine/Source/Runtime/Core/Private/Mac/MacApplication.cpp
--- a/Engine/Source/Runtime/Core/Private/Mac/MacApplication.cpp
+++ b/Engine/Source/Runtime/Core/Private/Mac/MacApplication.cpp
@@ -139,95 +139,12 @@ FMacApplication::FMacApplication()
 
 	AppActivationObserver = [[NSNotificationCenter defaultCenter] addObserverForName:NSApplicationDidBecomeActiveNotification object:[NSApplication sharedApplication] queue:[NSOperationQueue mainQueue] usingBlock:^(NSNotification* Notification)
 							{
-								for (int32 Index = 0; Index < SavedWindowsOrder.Num(); Index++)
-								{
-									const FSavedWindowOrderInfo& SavedWindowLevel = SavedWindowsOrder[Index];
-									NSWindow* Window = [NSApp windowWithWindowNumber:SavedWindowLevel.WindowNumber];
-									if (Window)
-									{
-										[Window setLevel:SavedWindowLevel.Level];
-									}
-								}
-
-								if (SavedWindowsOrder.Num() > 0)
-								{
-									[[NSApp windowWithWindowNumber:SavedWindowsOrder[0].WindowNumber] orderWindow:NSWindowAbove relativeTo:0];
-									for (int32 Index = 1; Index < SavedWindowsOrder.Num(); Index++)
-									{
-										const FSavedWindowOrderInfo& SavedWindowLevel = SavedWindowsOrder[Index];
-										NSWindow* Window = [NSApp windowWithWindowNumber:SavedWindowLevel.WindowNumber];
-										if (Window)
-										{
-											[Window orderWindow:NSWindowBelow relativeTo:SavedWindowsOrder[Index - 1].WindowNumber];
-										}
-									}
-								}
-
-								// If editor thread doesn't have the focus, don't suck up too much CPU time.
-								if (GIsEditor && !IsRunningCommandlet())
-								{
-									// Boost our priority back to normal.
-									struct sched_param Sched;
-									FMemory::Memzero(&Sched, sizeof(struct sched_param));
-									Sched.sched_priority = 15;
-									pthread_setschedparam(pthread_self(), SCHED_RR, &Sched);
-								}
-
-								// app is active, allow sound
-								FApp::SetVolumeMultiplier( 1.0f );
-
-								GameThreadCall(^{
-									MessageHandler->OnApplicationActivationChanged(true);
-								}, @[ NSDefaultRunLoopMode ], false);
+								OnApplicationDidBecomeActive();
 							 }];
 
 	AppDeactivationObserver = [[NSNotificationCenter defaultCenter] addObserverForName:NSApplicationWillResignActiveNotification object:[NSApplication sharedApplication] queue:[NSOperationQueue mainQueue] usingBlock:^(NSNotification* Notification)
 							{
-								SavedWindowsOrder.Empty();
-
-								NSArray* OrderedWindows = [NSApp orderedWindows];
-								for (NSWindow* Window in OrderedWindows)
-								{
-									if ([Window isKindOfClass:[FCocoaWindow class]] && [Window isVisible] && ![Window hidesOnDeactivate])
-									{
-										SavedWindowsOrder.Add(FSavedWindowOrderInfo([Window windowNumber], [Window level]));
-										[Window setLevel:NSNormalWindowLevel];
-									}
-								}
-
-								if (SavedWindowsOrder.Num() > 0)
-								{
-									[[NSApp windowWithWindowNumber:SavedWindowsOrder[0].WindowNumber] orderWindow:NSWindowAbove relativeTo:0];
-									for (int32 Index = 1; Index < SavedWindowsOrder.Num(); Index++)
-									{
-										const FSavedWindowOrderInfo& SavedWindowLevel = SavedWindowsOrder[Index];
-										NSWindow* Window = [NSApp windowWithWindowNumber:SavedWindowLevel.WindowNumber];
-										if (Window)
-										{
-											[Window orderWindow:NSWindowBelow relativeTo:SavedWindowsOrder[Index - 1].WindowNumber];
-										}
-									}
-								}
-
-								// If editor thread doesn't have the focus, don't suck up too much CPU time.
-								if (GIsEditor && !IsRunningCommandlet())
-								{
-									// Drop our priority to speed up whatever is in the foreground.
-									struct sched_param Sched;
-									FMemory::Memzero(&Sched, sizeof(struct sched_param));
-									Sched.sched_priority = 5;
-									pthread_setschedparam(pthread_self(), SCHED_RR, &Sched);
-
-									// Sleep for a bit to not eat up all CPU time.
-									FPlatformProcess::Sleep(0.005f);
-								}
-
-								// app is inactive, apply multiplier
-								FApp::SetVolumeMultiplier(FApp::GetUnfocusedVolumeMultiplier());
-
-								GameThreadCall(^{
-									MessageHandler->OnApplicationActivationChanged(false);
-								}, @[ NSDefaultRunLoopMode ], false);
+								OnApplicationWillResignActive();
 							}];
 
 	WorkspaceActivationObserver = [[[NSWorkspace sharedWorkspace] notificationCenter] addObserverForName:NSWorkspaceSessionDidBecomeActiveNotification object:[NSWorkspace sharedWorkspace] queue:[NSOperationQueue mainQueue] usingBlock:^(NSNotification* Notification){
@@ -1029,6 +946,122 @@ TSharedPtr<FMacWindow> FMacApplication::LocateWindowUnderCursor( const NSPoint P
 	return NULL;
 }
 
+void FMacApplication::OnApplicationDidBecomeActive()
+{
+	for (int32 Index = 0; Index < SavedWindowsOrder.Num(); Index++)
+	{
+		const FSavedWindowOrderInfo& Info = SavedWindowsOrder[Index];
+		NSWindow* Window = [NSApp windowWithWindowNumber:Info.WindowNumber];
+		if (Window)
+		{
+			[Window setLevel:Info.Level];
+		}
+	}
+
+	if (SavedWindowsOrder.Num() > 0)
+	{
+		[[NSApp windowWithWindowNumber:SavedWindowsOrder[0].WindowNumber] orderWindow:NSWindowAbove relativeTo:0];
+		for (int32 Index = 1; Index < SavedWindowsOrder.Num(); Index++)
+		{
+			const FSavedWindowOrderInfo& Info = SavedWindowsOrder[Index];
+			NSWindow* Window = [NSApp windowWithWindowNumber:Info.WindowNumber];
+			if (Window)
+			{
+				[Window orderWindow:NSWindowBelow relativeTo:SavedWindowsOrder[Index - 1].WindowNumber];
+			}
+		}
+	}
+
+	// If editor thread doesn't have the focus, don't suck up too much CPU time.
+	if (GIsEditor && !IsRunningCommandlet())
+	{
+		// Boost our priority back to normal.
+		struct sched_param Sched;
+		FMemory::Memzero(&Sched, sizeof(struct sched_param));
+		Sched.sched_priority = 15;
+		pthread_setschedparam(pthread_self(), SCHED_RR, &Sched);
+	}
+
+	// app is active, allow sound
+	FApp::SetVolumeMultiplier(1.0f);
+
+	GameThreadCall(^{
+		MessageHandler->OnApplicationActivationChanged(true);
+	}, @[ NSDefaultRunLoopMode ], false);
+}
+
+void FMacApplication::OnApplicationWillResignActive()
+{
+	OnWindowsReordered(true);
+
+	if (SavedWindowsOrder.Num() > 0)
+	{
+		[[NSApp windowWithWindowNumber:SavedWindowsOrder[0].WindowNumber] orderWindow:NSWindowAbove relativeTo:0];
+		for (int32 Index = 1; Index < SavedWindowsOrder.Num(); Index++)
+		{
+			const FSavedWindowOrderInfo& Info = SavedWindowsOrder[Index];
+			NSWindow* Window = [NSApp windowWithWindowNumber:Info.WindowNumber];
+			if (Window)
+			{
+				if (SavedWindowsOrder[Index - 1].Level == Info.Level)
+				{
+					[Window orderWindow:NSWindowBelow relativeTo:SavedWindowsOrder[Index - 1].WindowNumber];
+				}
+				else
+				{
+					[Window orderWindow:NSWindowAbove relativeTo:0];
+				}
+			}
+		}
+	}
+
+	// If editor thread doesn't have the focus, don't suck up too much CPU time.
+	if (GIsEditor && !IsRunningCommandlet())
+	{
+		// Drop our priority to speed up whatever is in the foreground.
+		struct sched_param Sched;
+		FMemory::Memzero(&Sched, sizeof(struct sched_param));
+		Sched.sched_priority = 5;
+		pthread_setschedparam(pthread_self(), SCHED_RR, &Sched);
+
+		// Sleep for a bit to not eat up all CPU time.
+		FPlatformProcess::Sleep(0.005f);
+	}
+
+	// app is inactive, apply multiplier
+	FApp::SetVolumeMultiplier(FApp::GetUnfocusedVolumeMultiplier());
+
+	GameThreadCall(^{
+		MessageHandler->OnApplicationActivationChanged(false);
+	}, @[ NSDefaultRunLoopMode ], false);
+}
+
+void FMacApplication::OnWindowsReordered(bool bIsAppInBackground)
+{
+	if (bIsAppInBackground)
+	{
+		TMap<int32, int32> Levels;
+
+		for (int32 Index = 0; Index < SavedWindowsOrder.Num(); Index++)
+		{
+			const FSavedWindowOrderInfo& Info = SavedWindowsOrder[Index];
+			Levels.Add(Info.WindowNumber, Info.Level);
+		}
+
+		SavedWindowsOrder.Empty();
+
+		NSArray* OrderedWindows = [NSApp orderedWindows];
+		for (NSWindow* Window in OrderedWindows)
+		{
+			if ([Window isKindOfClass:[FCocoaWindow class]] && [Window isVisible] && ![Window hidesOnDeactivate])
+			{
+				SavedWindowsOrder.Add(FSavedWindowOrderInfo([Window windowNumber], Levels.Contains([Window windowNumber]) ? Levels[[Window windowNumber]] : [Window level]));
+				[Window setLevel:NSNormalWindowLevel];
+			}
+		}
+	}
+}
+
 void FMacApplication::PollGameDeviceState( const float TimeDelta )
 {
 	// Poll game device state and send new events
@@ -1205,6 +1238,8 @@ void FMacApplication::OnWindowDidBecomeMain( FCocoaWindow* Window )
 	{
 		MessageHandler->OnWindowActivationChanged( EventWindow.ToSharedRef(), EWindowActivation::Activate );
 	}
+
+	OnWindowsReordered(![NSApp isActive]);
 }
 
 void FMacApplication::OnWindowDidResignMain( FCocoaWindow* Window )
@@ -1214,6 +1249,8 @@ void FMacApplication::OnWindowDidResignMain( FCocoaWindow* Window )
 	{
 		MessageHandler->OnWindowActivationChanged( EventWindow.ToSharedRef(), EWindowActivation::Deactivate );
 	}
+
+	OnWindowsReordered(![NSApp isActive]);
 }
 
 void FMacApplication::OnWindowWillMove( FCocoaWindow* Window )
diff --git a/Engine/Source/Runtime/Core/Private/Mac/MacWindow.cpp b/Engine/Source/Runtime/Core/Private/Mac/MacWindow.cpp
--- a/Engine/Source/Runtime/Core/Private/Mac/MacWindow.cpp
+++ b/Engine/Source/Runtime/Core/Private/Mac/MacWindow.cpp
@@ -94,8 +94,7 @@ void FMacWindow::Initialize( FMacApplication* const Application, const TSharedRe
 			// @todo: We really need a window type in tab manager to know what window level to use and whether or not a window should hide on deactivate
 			if (Definition->IsModalWindow)
 			{
-				NSInteger Level = [NSApp isActive] ? NSModalPanelWindowLevel : NSNormalWindowLevel;
-				[WindowHandle setLevel: Level];
+				[WindowHandle setLevel: NSModalPanelWindowLevel];
 			}
 			else if (Definition->IsRegularWindow)
 			{
diff --git a/Engine/Source/Runtime/Core/Public/Mac/MacApplication.h b/Engine/Source/Runtime/Core/Public/Mac/MacApplication.h
--- a/Engine/Source/Runtime/Core/Public/Mac/MacApplication.h
+++ b/Engine/Source/Runtime/Core/Public/Mac/MacApplication.h
@@ -122,7 +122,9 @@ public:
 
 	/** Invalidates all queued windows requiring text layout changes */
 	void InvalidateTextLayouts();
+
 private:
+
 	enum FMacApplicationEventTypes
 	{
 		ResentEvent = 0x0f00
@@ -141,6 +143,10 @@ private:
 	FCocoaWindow* FindEventWindow( NSEvent* CocoaEvent );
 	TSharedPtr<FMacWindow> LocateWindowUnderCursor( const NSPoint Position );
 
+	void OnApplicationWillResignActive();
+	void OnApplicationDidBecomeActive();
+	void OnWindowsReordered(bool bIsAppInBackground);
+
 	NSScreen* FindScreenByPoint( int32 X, int32 Y ) const;
 
 	void HandleModifierChange(NSUInteger NewModifierFlags, NSUInteger FlagsShift, NSUInteger UE4Shift, EMacModifierKeys TranslatedCode);
