Fix for ranged-for string iteration returning the null terminator.
Redundant (now) test removed from FUnrealSourceFile::GetFileId.
PLATFORM_COMPILER_HAS_RANGED_FOR_LOOP macro removed.

diff --git a/Engine/Source/Editor/ClassViewer/Private/SClassViewer.cpp b/Engine/Source/Editor/ClassViewer/Private/SClassViewer.cpp
--- a/Engine/Source/Editor/ClassViewer/Private/SClassViewer.cpp
+++ b/Engine/Source/Editor/ClassViewer/Private/SClassViewer.cpp
@@ -2563,15 +2563,8 @@ void SClassViewer::OnFilterTextChanged( const FText& InFilterText )
 	{
 		bool bIsWithinQuotedSection = false;
 		FString NewSearchTerm;
-		#if PLATFORM_COMPILER_HAS_RANGED_FOR_LOOP
 		for( auto CurChar : CurrentFilterText )
 		{
-		#else
-		for( auto CurCharIndex = 0; CurCharIndex < CurrentFilterText.Len(); ++CurCharIndex )
-		{
-			const auto CurChar = CurrentFilterText.GetCharArray()[ CurCharIndex ];
-		#endif
-
 			// Keep an eye out for double-quotes.  We want to retain whitespace within a search term if
 			// it has double-quotes around it
 			if( CurChar == TCHAR('\"') )
diff --git a/Engine/Source/Editor/StatsViewer/Private/StatsPages/PrimitiveStatsPage.cpp b/Engine/Source/Editor/StatsViewer/Private/StatsPages/PrimitiveStatsPage.cpp
--- a/Engine/Source/Editor/StatsViewer/Private/StatsPages/PrimitiveStatsPage.cpp
+++ b/Engine/Source/Editor/StatsViewer/Private/StatsPages/PrimitiveStatsPage.cpp
@@ -228,15 +228,8 @@ struct PrimitiveStatsGenerator
 					// If Count represents the Model itself, we do NOT want to increment it now.
 					StatsEntry->Count--;
 
-					#if PLATFORM_COMPILER_HAS_RANGED_FOR_LOOP
 					for (const auto& Element : ModelComponent->GetElements())
 					{
-					#else
-					TIndirectArray<FModelElement> Elements = ModelComponent->GetElements();
-					for( int32 ElementIndex=0; ElementIndex<Elements.Num(); ElementIndex++ )
-					{
-						const FModelElement& Element = Elements[ElementIndex];
-					#endif
 						StatsEntry->Triangles += Element.NumTriangles;
 						StatsEntry->Sections++;
 					}
diff --git a/Engine/Source/Editor/UnrealEd/Private/Factories/SkeletalMeshImport.cpp b/Engine/Source/Editor/UnrealEd/Private/Factories/SkeletalMeshImport.cpp
--- a/Engine/Source/Editor/UnrealEd/Private/Factories/SkeletalMeshImport.cpp
+++ b/Engine/Source/Editor/UnrealEd/Private/Factories/SkeletalMeshImport.cpp
@@ -717,15 +717,8 @@ ExistingSkelMeshData* SaveExistingSkelMeshData(USkeletalMesh* ExistingSkelMesh)
 				LODModel.LegacyRawPointIndices.Lock( LOCK_READ_ONLY );
 			}
 			ExistingMeshDataPtr->ExistingLODModels = ImportedResource->LODModels;
-			#if PLATFORM_COMPILER_HAS_RANGED_FOR_LOOP
 			for ( auto& LODModel : ImportedResource->LODModels )
 			{
-			#else
-			for ( int32 LODModelIndex = 0 ; LODModelIndex < ImportedResource->LODModels.Num() ; ++LODModelIndex )
-			{
-
-				FStaticLODModel& LODModel = ImportedResource->LODModels[LODModelIndex];
-			#endif
 				LODModel.RawPointIndices.Unlock();
 				LODModel.LegacyRawPointIndices.Unlock();
 
diff --git a/Engine/Source/Programs/UnrealHeaderTool/Private/UnrealSourceFile.cpp b/Engine/Source/Programs/UnrealHeaderTool/Private/UnrealSourceFile.cpp
--- a/Engine/Source/Programs/UnrealHeaderTool/Private/UnrealSourceFile.cpp
+++ b/Engine/Source/Programs/UnrealHeaderTool/Private/UnrealSourceFile.cpp
@@ -21,11 +21,7 @@ FString FUnrealSourceFile::GetFileId() const
 	StdFilename = StdFilename.Mid(9);
 	for (auto Char : StdFilename)
 	{
-		if (Char == 0)
-		{
-			break;
-		}
-		else if (FChar::IsAlnum(Char))
+		if (FChar::IsAlnum(Char))
 		{
 			Out.AppendChar(Char);
 		}
diff --git a/Engine/Source/Runtime/Core/Public/Containers/UnrealString.h b/Engine/Source/Runtime/Core/Public/Containers/UnrealString.h
--- a/Engine/Source/Runtime/Core/Public/Containers/UnrealString.h
+++ b/Engine/Source/Runtime/Core/Public/Containers/UnrealString.h
@@ -239,8 +239,8 @@ private:
 	 */
 	FORCEINLINE friend TIterator      begin(      FString& Str) { return begin(Str.Data); }
 	FORCEINLINE friend TConstIterator begin(const FString& Str) { return begin(Str.Data); }
-	FORCEINLINE friend TIterator      end  (      FString& Str) { return end  (Str.Data); }
-	FORCEINLINE friend TConstIterator end  (const FString& Str) { return end  (Str.Data); }
+	FORCEINLINE friend TIterator      end  (      FString& Str) { TIterator      Result = end(Str.Data); if (Str.Data.Num()) { --Result; } return Result; }
+	FORCEINLINE friend TConstIterator end  (const FString& Str) { TConstIterator Result = end(Str.Data); if (Str.Data.Num()) { --Result; } return Result; }
 
 public:
 	FORCEINLINE uint32 GetAllocatedSize() const
diff --git a/Engine/Source/Runtime/Core/Public/HAL/Platform.h b/Engine/Source/Runtime/Core/Public/HAL/Platform.h
--- a/Engine/Source/Runtime/Core/Public/HAL/Platform.h
+++ b/Engine/Source/Runtime/Core/Public/HAL/Platform.h
@@ -162,9 +162,6 @@
 #ifndef PLATFORM_COMPILER_HAS_DEFAULTED_FUNCTIONS
 	#define PLATFORM_COMPILER_HAS_DEFAULTED_FUNCTIONS	1
 #endif
-#ifndef PLATFORM_COMPILER_HAS_RANGED_FOR_LOOP
-	#define PLATFORM_COMPILER_HAS_RANGED_FOR_LOOP	1
-#endif
 #ifndef PLATFORM_COMPILER_HAS_VARIADIC_TEMPLATES
 	#define PLATFORM_COMPILER_HAS_VARIADIC_TEMPLATES	1
 #endif
diff --git a/Engine/Source/Runtime/Core/Public/Windows/WIndowsPlatform.h b/Engine/Source/Runtime/Core/Public/Windows/WIndowsPlatform.h
--- a/Engine/Source/Runtime/Core/Public/Windows/WIndowsPlatform.h
+++ b/Engine/Source/Runtime/Core/Public/Windows/WIndowsPlatform.h
@@ -71,10 +71,6 @@ typedef FWindowsPlatformTypes FPlatformTypes;
 #define PLATFORM_HAS_128BIT_ATOMICS							(!HACK_HEADER_GENERATOR && PLATFORM_64BITS && (WINVER >= 0x602))
 #define PLATFORM_USES_ANSI_STRING_FOR_EXTERNAL_PROFILING	0
 
-#if _MSC_VER < 1700
-	#define PLATFORM_COMPILER_HAS_RANGED_FOR_LOOP			0
-#endif
-
 
 // Function type macros.
 #define VARARGS     __cdecl					/* Functions with variable arguments */
