Added FCString::Strrstr which finds the last occurrence of a substring.
Changed dynamic delegate binding to avoid temporary FString creation/conversion, which also improves the debugging experience.

diff --git a/Engine/Source/Runtime/Core/Public/Delegates/Delegate.h b/Engine/Source/Runtime/Core/Public/Delegates/Delegate.h
--- a/Engine/Source/Runtime/Core/Public/Delegates/Delegate.h
+++ b/Engine/Source/Runtime/Core/Public/Delegates/Delegate.h
@@ -280,19 +280,19 @@ class DynamicMulticastDelegateName : public FUNC_DECLARE_DYNAMIC_DELEGATE_BASE(T
 
 
 // Helper macro for calling BindDynamic() on dynamic delegates.  Automatically generates the function name string.
-#define BindDynamic( UserObject, FuncName ) __Internal_BindDynamic( UserObject, FuncName, FString( #FuncName ) )
+#define BindDynamic( UserObject, FuncName ) __Internal_BindDynamic( UserObject, FuncName, TEXT( #FuncName ) )
 
 // Helper macro for calling AddDynamic() on dynamic multi-cast delegates.  Automatically generates the function name string.
-#define AddDynamic( UserObject, FuncName ) __Internal_AddDynamic( UserObject, FuncName, FString( #FuncName ) )
+#define AddDynamic( UserObject, FuncName ) __Internal_AddDynamic( UserObject, FuncName, TEXT( #FuncName ) )
 
 // Helper macro for calling AddUniqueDynamic() on dynamic multi-cast delegates.  Automatically generates the function name string.
-#define AddUniqueDynamic( UserObject, FuncName ) __Internal_AddUniqueDynamic( UserObject, FuncName, FString( #FuncName ) )
+#define AddUniqueDynamic( UserObject, FuncName ) __Internal_AddUniqueDynamic( UserObject, FuncName, TEXT( #FuncName ) )
 
 // Helper macro for calling RemoveDynamic() on dynamic multi-cast delegates.  Automatically generates the function name string.
-#define RemoveDynamic( UserObject, FuncName ) __Internal_RemoveDynamic( UserObject, FuncName, FString( #FuncName ) )
+#define RemoveDynamic( UserObject, FuncName ) __Internal_RemoveDynamic( UserObject, FuncName, TEXT( #FuncName ) )
 
 // Helper macro for calling IsAlreadyBound() on dynamic multi-cast delegates.  Automatically generates the function name string.
-#define IsAlreadyBound( UserObject, FuncName ) __Internal_IsAlreadyBound( UserObject, FuncName, FString( #FuncName ) )
+#define IsAlreadyBound( UserObject, FuncName ) __Internal_IsAlreadyBound( UserObject, FuncName, TEXT( #FuncName ) )
 
 
 /*********************************************************************************************************************/
diff --git a/Engine/Source/Runtime/Core/Public/Delegates/DelegateSignatureImpl.inl b/Engine/Source/Runtime/Core/Public/Delegates/DelegateSignatureImpl.inl
--- a/Engine/Source/Runtime/Core/Public/Delegates/DelegateSignatureImpl.inl
+++ b/Engine/Source/Runtime/Core/Public/Delegates/DelegateSignatureImpl.inl
@@ -2196,9 +2196,9 @@ public:
 	 *        automatically sets the function name string for the caller.
 	 */
 	template< class UserClass >
-	void __Internal_BindDynamic( UserClass* InUserObject, typename TMethodPtrResolver< UserClass >::FMethodPtr InMethodPtr, const FString& InMacroFunctionName )
+	void __Internal_BindDynamic( UserClass* InUserObject, typename TMethodPtrResolver< UserClass >::FMethodPtr InMethodPtr, const TCHAR* InMacroFunctionName )
 	{
-		check( InUserObject != nullptr && InMethodPtr != nullptr && !InMacroFunctionName.IsEmpty() );
+		check( InUserObject != nullptr && InMethodPtr != nullptr && InMacroFunctionName[0] );
 
 		// NOTE: We're not actually storing the incoming method pointer or calling it.  We simply require it for type-safety reasons.
 
@@ -2208,9 +2208,11 @@ public:
 
 		// Store the function name.  The incoming function name was generated by a macro and includes the method's class name.
 		// We'll need to strip off that class prefix and just store the function name by itself.
-		this->FunctionName = *InMacroFunctionName.Mid( InMacroFunctionName.Find( TEXT( "::" ), ESearchCase::CaseSensitive, ESearchDir::FromEnd) + 2 );
+		const TCHAR* TrimmedMacroFunctionName = FCString::Strrstr( InMacroFunctionName, TEXT( "::" ) );
+		checkf(TrimmedMacroFunctionName, TEXT("'%s' does not look like a member function"), InMacroFunctionName);
+		this->FunctionName = TrimmedMacroFunctionName + 2;
 
-		ensureMsgf(this->IsBound(), TEXT("Unable to bind delegate to '%s' (function might not be marked as a UFUNCTION)"), *InMacroFunctionName);
+		ensureMsgf(this->IsBound(), TEXT("Unable to bind delegate to '%s' (function might not be marked as a UFUNCTION)"), InMacroFunctionName);
 	}
 
 	friend uint32 GetTypeHash(const DYNAMIC_DELEGATE_CLASS& Key)
@@ -2265,9 +2267,9 @@ public:
 	 *        automatically sets the function name string for the caller.
 	 */
 	template< class UserClass >
-	bool __Internal_IsAlreadyBound( UserClass* InUserObject, typename FDelegate::template TMethodPtrResolver< UserClass >::FMethodPtr InMethodPtr, const FString& InMacroFunctionName ) const
+	bool __Internal_IsAlreadyBound( UserClass* InUserObject, typename FDelegate::template TMethodPtrResolver< UserClass >::FMethodPtr InMethodPtr, const TCHAR* InMacroFunctionName ) const
 	{
-		check( InUserObject != nullptr && InMethodPtr != nullptr && !InMacroFunctionName.IsEmpty() );
+		check( InUserObject != nullptr && InMethodPtr != nullptr && InMacroFunctionName[0] );
 
 		// NOTE: We're not actually using the incoming method pointer or calling it.  We simply require it for type-safety reasons.
 
@@ -2288,9 +2290,9 @@ public:
 	 *        automatically sets the function name string for the caller.
 	 */
 	template< class UserClass >
-	void __Internal_AddDynamic( UserClass* InUserObject, typename FDelegate::template TMethodPtrResolver< UserClass >::FMethodPtr InMethodPtr, const FString& InMacroFunctionName )
+	void __Internal_AddDynamic( UserClass* InUserObject, typename FDelegate::template TMethodPtrResolver< UserClass >::FMethodPtr InMethodPtr, const TCHAR* InMacroFunctionName )
 	{
-		check( InUserObject != nullptr && InMethodPtr != nullptr && !InMacroFunctionName.IsEmpty() );
+		check( InUserObject != nullptr && InMethodPtr != nullptr && InMacroFunctionName[0] );
 
 		// NOTE: We're not actually storing the incoming method pointer or calling it.  We simply require it for type-safety reasons.
 
@@ -2311,9 +2313,9 @@ public:
 	 *        automatically sets the function name string for the caller.
 	 */
 	template< class UserClass >
-	void __Internal_AddUniqueDynamic( UserClass* InUserObject, typename FDelegate::template TMethodPtrResolver< UserClass >::FMethodPtr InMethodPtr, const FString& InMacroFunctionName )
+	void __Internal_AddUniqueDynamic( UserClass* InUserObject, typename FDelegate::template TMethodPtrResolver< UserClass >::FMethodPtr InMethodPtr, const TCHAR* InMacroFunctionName )
 	{
-		check( InUserObject != nullptr && InMethodPtr != nullptr && !InMacroFunctionName.IsEmpty() );
+		check( InUserObject != nullptr && InMethodPtr != nullptr && InMacroFunctionName[0] );
 
 		// NOTE: We're not actually storing the incoming method pointer or calling it.  We simply require it for type-safety reasons.
 
@@ -2334,9 +2336,9 @@ public:
 	 *        automatically sets the function name string for the caller.
 	 */
 	template< class UserClass >
-	void __Internal_RemoveDynamic( UserClass* InUserObject, typename FDelegate::template TMethodPtrResolver< UserClass >::FMethodPtr InMethodPtr, const FString& InMacroFunctionName )
+	void __Internal_RemoveDynamic( UserClass* InUserObject, typename FDelegate::template TMethodPtrResolver< UserClass >::FMethodPtr InMethodPtr, const TCHAR* InMacroFunctionName )
 	{
-		check( InUserObject != nullptr && InMethodPtr != nullptr && !InMacroFunctionName.IsEmpty() );
+		check( InUserObject != nullptr && InMethodPtr != nullptr && InMacroFunctionName[0] );
 
 		// NOTE: We're not actually storing the incoming method pointer or calling it.  We simply require it for type-safety reasons.
 
diff --git a/Engine/Source/Runtime/Core/Public/Delegates/DelegateSignatureImpl_Variadics.inl b/Engine/Source/Runtime/Core/Public/Delegates/DelegateSignatureImpl_Variadics.inl
--- a/Engine/Source/Runtime/Core/Public/Delegates/DelegateSignatureImpl_Variadics.inl
+++ b/Engine/Source/Runtime/Core/Public/Delegates/DelegateSignatureImpl_Variadics.inl
@@ -1176,9 +1176,9 @@ public:
 	 *        automatically sets the function name string for the caller.
 	 */
 	template< class UserClass >
-	void __Internal_BindDynamic( UserClass* InUserObject, typename TMethodPtrResolver< UserClass >::FMethodPtr InMethodPtr, const FString& InMacroFunctionName )
+	void __Internal_BindDynamic( UserClass* InUserObject, typename TMethodPtrResolver< UserClass >::FMethodPtr InMethodPtr, const TCHAR* InMacroFunctionName )
 	{
-		check( InUserObject != nullptr && InMethodPtr != nullptr && !InMacroFunctionName.IsEmpty() );
+		check( InUserObject != nullptr && InMethodPtr != nullptr && InMacroFunctionName[0] );
 
 		// NOTE: We're not actually storing the incoming method pointer or calling it.  We simply require it for type-safety reasons.
 
@@ -1188,9 +1188,11 @@ public:
 
 		// Store the function name.  The incoming function name was generated by a macro and includes the method's class name.
 		// We'll need to strip off that class prefix and just store the function name by itself.
-		this->FunctionName = *InMacroFunctionName.Mid( InMacroFunctionName.Find( TEXT( "::" ), ESearchCase::CaseSensitive, ESearchDir::FromEnd) + 2 );
+		const TCHAR* TrimmedMacroFunctionName = FCString::Strrstr( InMacroFunctionName, TEXT( "::" ) );
+		checkf(TrimmedMacroFunctionName, TEXT("'%s' does not look like a member function"), InMacroFunctionName);
+		this->FunctionName = TrimmedMacroFunctionName + 2;
 
-		ensureMsgf(this->IsBound(), TEXT("Unable to bind delegate to '%s' (function might not be marked as a UFUNCTION)"), *InMacroFunctionName);
+		ensureMsgf(this->IsBound(), TEXT("Unable to bind delegate to '%s' (function might not be marked as a UFUNCTION)"), InMacroFunctionName);
 	}
 
 	friend uint32 GetTypeHash(const TBaseDynamicDelegate& Key)
@@ -1243,9 +1245,9 @@ public:
 	 *        automatically sets the function name string for the caller.
 	 */
 	template< class UserClass >
-	bool __Internal_IsAlreadyBound( UserClass* InUserObject, typename FDelegate::template TMethodPtrResolver< UserClass >::FMethodPtr InMethodPtr, const FString& InMacroFunctionName ) const
+	bool __Internal_IsAlreadyBound( UserClass* InUserObject, typename FDelegate::template TMethodPtrResolver< UserClass >::FMethodPtr InMethodPtr, const TCHAR* InMacroFunctionName ) const
 	{
-		check( InUserObject != nullptr && InMethodPtr != nullptr && !InMacroFunctionName.IsEmpty() );
+		check( InUserObject != nullptr && InMethodPtr != nullptr && InMacroFunctionName[0] );
 
 		// NOTE: We're not actually using the incoming method pointer or calling it.  We simply require it for type-safety reasons.
 
@@ -1266,9 +1268,9 @@ public:
 	 *        automatically sets the function name string for the caller.
 	 */
 	template< class UserClass >
-	void __Internal_AddDynamic( UserClass* InUserObject, typename FDelegate::template TMethodPtrResolver< UserClass >::FMethodPtr InMethodPtr, const FString& InMacroFunctionName )
+	void __Internal_AddDynamic( UserClass* InUserObject, typename FDelegate::template TMethodPtrResolver< UserClass >::FMethodPtr InMethodPtr, const TCHAR* InMacroFunctionName )
 	{
-		check( InUserObject != nullptr && InMethodPtr != nullptr && !InMacroFunctionName.IsEmpty() );
+		check( InUserObject != nullptr && InMethodPtr != nullptr && InMacroFunctionName[0] );
 
 		// NOTE: We're not actually storing the incoming method pointer or calling it.  We simply require it for type-safety reasons.
 
@@ -1289,9 +1291,9 @@ public:
 	 *        automatically sets the function name string for the caller.
 	 */
 	template< class UserClass >
-	void __Internal_AddUniqueDynamic( UserClass* InUserObject, typename FDelegate::template TMethodPtrResolver< UserClass >::FMethodPtr InMethodPtr, const FString& InMacroFunctionName )
+	void __Internal_AddUniqueDynamic( UserClass* InUserObject, typename FDelegate::template TMethodPtrResolver< UserClass >::FMethodPtr InMethodPtr, const TCHAR* InMacroFunctionName )
 	{
-		check( InUserObject != nullptr && InMethodPtr != nullptr && !InMacroFunctionName.IsEmpty() );
+		check( InUserObject != nullptr && InMethodPtr != nullptr && InMacroFunctionName[0] );
 
 		// NOTE: We're not actually storing the incoming method pointer or calling it.  We simply require it for type-safety reasons.
 
@@ -1312,9 +1314,9 @@ public:
 	 *        automatically sets the function name string for the caller.
 	 */
 	template< class UserClass >
-	void __Internal_RemoveDynamic( UserClass* InUserObject, typename FDelegate::template TMethodPtrResolver< UserClass >::FMethodPtr InMethodPtr, const FString& InMacroFunctionName )
+	void __Internal_RemoveDynamic( UserClass* InUserObject, typename FDelegate::template TMethodPtrResolver< UserClass >::FMethodPtr InMethodPtr, const TCHAR* InMacroFunctionName )
 	{
-		check( InUserObject != nullptr && InMethodPtr != nullptr && !InMacroFunctionName.IsEmpty() );
+		check( InUserObject != nullptr && InMethodPtr != nullptr && InMacroFunctionName[0] );
 
 		// NOTE: We're not actually storing the incoming method pointer or calling it.  We simply require it for type-safety reasons.
 
diff --git a/Engine/Source/Runtime/Core/Public/Misc/CString.h b/Engine/Source/Runtime/Core/Public/Misc/CString.h
--- a/Engine/Source/Runtime/Core/Public/Misc/CString.h
+++ b/Engine/Source/Runtime/Core/Public/Misc/CString.h
@@ -243,6 +243,12 @@ struct TCString
 	static FORCEINLINE const CharType* Strrchr( const CharType* String, CharType c );
 	static FORCEINLINE CharType* Strrchr( CharType* String, CharType c );
 
+	/**
+	 * strrstr wrapper
+	 */
+	static FORCEINLINE const CharType* Strrstr( const CharType* String, const CharType* Find );
+	static FORCEINLINE CharType* Strrstr( CharType* String, const CharType* Find );
+
 	/**
 	 * atoi wrapper
 	 */
@@ -589,6 +595,34 @@ typename TCString<T>::CharType* TCString<T>::Strrchr( CharType* String, CharType
 	return (CharType*)FPlatformString::Strrchr(String, c);
 }
 
+template <typename T> FORCEINLINE
+const typename TCString<T>::CharType* TCString<T>::Strrstr( const CharType* String, const CharType* Find )
+{
+	return Strrstr((CharType*)String, Find);
+}
+
+template <typename T> FORCEINLINE
+typename TCString<T>::CharType* TCString<T>::Strrstr( CharType* String, const CharType* Find )
+{
+	if (*Find == (CharType)0)
+	{
+		return String + Strlen(String);
+	}
+
+	CharType* Result = nullptr;
+	for (;;)
+	{
+		CharType* Found = Strstr(String, Find);
+		if (!Found)
+		{
+			return Result;
+		}
+
+		Result = Found;
+		String = Found + 1;
+	}
+}
+
 template <typename T> FORCEINLINE 
 int32 TCString<T>::Atoi( const CharType* String ) 
 {
