Made tooltips and context menus visible on Mac when the app is in the background

diff --git a/Engine/Source/Runtime/Core/Private/Mac/MacApplication.cpp b/Engine/Source/Runtime/Core/Private/Mac/MacApplication.cpp
--- a/Engine/Source/Runtime/Core/Private/Mac/MacApplication.cpp
+++ b/Engine/Source/Runtime/Core/Private/Mac/MacApplication.cpp
@@ -580,7 +580,7 @@ void FMacApplication::ProcessNSEvent(NSEvent* const Event)
 				MacCursor->WarpCursor(HighPrecisionMousePos.X, HighPrecisionMousePos.Y);
 				MessageHandler->OnRawMouseMove(Delta.X, Delta.Y);
 			}
-			else if (CurrentEventWindow.IsValid())
+			else
 			{
 				NSPoint CursorPos = [NSEvent mouseLocation];
 				CursorPos.y--; // The y coordinate of the point returned by mouseLocation starts from a base of 1
@@ -596,18 +596,21 @@ void FMacApplication::ProcessNSEvent(NSEvent* const Event)
 					MacCursor->UpdateCurrentPosition(CurrentPosition);
 				}
 
-				// Cocoa does not update NSWindow's frame until user stops dragging the window, so while window is being dragged, we calculate
-				// its position based on mouse move delta
-				if (DraggedWindow && DraggedWindow == NativeWindow)
+				if (CurrentEventWindow.IsValid())
 				{
-					const int32 X = FMath::TruncToInt(CurrentEventWindow->PositionX + MouseDelta.X);
-					const int32 Y = FMath::TruncToInt(CurrentEventWindow->PositionY + MouseDelta.Y);
-					MessageHandler->OnMovedWindow(CurrentEventWindow.ToSharedRef(), X, Y);
-					CurrentEventWindow->PositionX = X;
-					CurrentEventWindow->PositionY = Y;
-				}
+					// Cocoa does not update NSWindow's frame until user stops dragging the window, so while window is being dragged, we calculate
+					// its position based on mouse move delta
+					if (DraggedWindow && DraggedWindow == NativeWindow)
+					{
+						const int32 X = FMath::TruncToInt(CurrentEventWindow->PositionX + MouseDelta.X);
+						const int32 Y = FMath::TruncToInt(CurrentEventWindow->PositionY + MouseDelta.Y);
+						MessageHandler->OnMovedWindow(CurrentEventWindow.ToSharedRef(), X, Y);
+						CurrentEventWindow->PositionX = X;
+						CurrentEventWindow->PositionY = Y;
+					}
 
-				MessageHandler->OnMouseMove();
+					MessageHandler->OnMouseMove();
+				}
 			}
 
 			if (CurrentEventWindow.IsValid() && !DraggedWindow && !GetCapture())
diff --git a/Engine/Source/Runtime/Core/Private/Mac/MacWindow.cpp b/Engine/Source/Runtime/Core/Private/Mac/MacWindow.cpp
--- a/Engine/Source/Runtime/Core/Private/Mac/MacWindow.cpp
+++ b/Engine/Source/Runtime/Core/Private/Mac/MacWindow.cpp
@@ -114,7 +114,6 @@ void FMacWindow::Initialize( FMacApplication* const Application, const TSharedRe
 			else
 			{
 				[WindowHandle setLevel: NSModalPanelWindowLevel];
-				[WindowHandle setHidesOnDeactivate: YES];
 			}
 			
 			if( !Definition->HasOSWindowBorder )
diff --git a/Engine/Source/Runtime/Slate/Private/Framework/Application/SlateApplication.cpp b/Engine/Source/Runtime/Slate/Private/Framework/Application/SlateApplication.cpp
--- a/Engine/Source/Runtime/Slate/Private/Framework/Application/SlateApplication.cpp
+++ b/Engine/Source/Runtime/Slate/Private/Framework/Application/SlateApplication.cpp
@@ -4003,12 +4003,18 @@ bool FSlateApplication::ProcessMouseButtonDownEvent( const TSharedPtr< FGenericW
 				}
 
 #if PLATFORM_MAC
-				if (MouseEvent.GetEffectingButton() == EKeys::LeftMouseButton && WidgetsUnderCursor.TopLevelWindow.IsValid() && !DragDetector.DetectDragForWidget.IsValid() && ActiveWindow == [NSApp keyWindow])
+				if (WidgetsUnderCursor.TopLevelWindow.IsValid())
 				{
-					MouseCaptorHelper Captor = MouseCaptor;
-					FPlatformMisc::ActivateApplication();
-					WidgetsUnderCursor.TopLevelWindow->BringToFront(true);
-					MouseCaptor = Captor;
+					// Clicking on a context menu should not activate anything
+					// @todo: This needs to be updated when we have window type in SWindow and we no longer have to guess if WidgetsUnderCursor.TopLevelWindow is a menu
+					const bool bIsContextMenu = !WidgetsUnderCursor.TopLevelWindow->IsRegularWindow() && WidgetsUnderCursor.TopLevelWindow->HasMinimizeBox() && WidgetsUnderCursor.TopLevelWindow->HasMaximizeBox();
+					if (!bIsContextMenu && MouseEvent.GetEffectingButton() == EKeys::LeftMouseButton && !DragDetector.DetectDragForWidget.IsValid() && ActiveWindow == [NSApp keyWindow])
+					{
+						MouseCaptorHelper Captor = MouseCaptor;
+						FPlatformMisc::ActivateApplication();
+						WidgetsUnderCursor.TopLevelWindow->BringToFront(true);
+						MouseCaptor = Captor;
+					}
 				}
 #endif
 			}
