Removed handling of negative exit codes from UBT in MacPlatformProcess, as we no longer use them

diff --git a/Engine/Source/Runtime/Core/Private/Mac/MacPlatformProcess.cpp b/Engine/Source/Runtime/Core/Private/Mac/MacPlatformProcess.cpp
--- a/Engine/Source/Runtime/Core/Private/Mac/MacPlatformProcess.cpp
+++ b/Engine/Source/Runtime/Core/Private/Mac/MacPlatformProcess.cpp
@@ -367,11 +367,10 @@ FProcHandle FMacPlatformProcess::CreateProc( const TCHAR* URL, const TCHAR* Parm
 	
 	if(LaunchPath == NULL)
 	{
-		return FProcHandle(NULL, false);
+		return FProcHandle(NULL);
 	}
 
 	NSTask* ProcessHandle = [[NSTask alloc] init];
-	bool bIsShellScript = false;
 
 	if (ProcessHandle)
 	{
@@ -385,7 +384,6 @@ FProcHandle FMacPlatformProcess::CreateProc( const TCHAR* URL, const TCHAR* Parm
 			[Arguments addObject: @"-c"];
 			[Arguments addObject: Arg];
 			CFRelease((CFStringRef)Arg);
-			bIsShellScript = true;
 		}
 		else
 		{
@@ -474,7 +472,6 @@ FProcHandle FMacPlatformProcess::CreateProc( const TCHAR* URL, const TCHAR* Parm
 		{
 			[ProcessHandle release];
 			ProcessHandle = nil;
-			bIsShellScript = false;
 		}
 
 		[Arguments release];
@@ -487,7 +484,7 @@ FProcHandle FMacPlatformProcess::CreateProc( const TCHAR* URL, const TCHAR* Parm
 
 	CFRelease((CFStringRef)LaunchPath);
 
-	return FProcHandle(ProcessHandle, bIsShellScript);
+	return FProcHandle(ProcessHandle);
 }
 
 bool FMacPlatformProcess::IsProcRunning( FProcHandle& ProcessHandle )
@@ -550,10 +547,6 @@ bool FMacPlatformProcess::GetProcReturnCode( FProcHandle& ProcessHandle, int32*
 	}
 
 	*ReturnCode = [(NSTask*)ProcessHandle.Get() terminationStatus];
-	if (ProcessHandle.IsShellScript && *ReturnCode > 128)
-	{
-		*ReturnCode = *ReturnCode - 256;
-	}
 	return true;
 }
 
diff --git a/Engine/Source/Runtime/Core/Public/Mac/MacPlatformProcess.h b/Engine/Source/Runtime/Core/Public/Mac/MacPlatformProcess.h
--- a/Engine/Source/Runtime/Core/Public/Mac/MacPlatformProcess.h
+++ b/Engine/Source/Runtime/Core/Public/Mac/MacPlatformProcess.h
@@ -15,16 +15,12 @@ public:
 	/** Default constructor. */
 	FORCEINLINE FProcHandle()
 		: TProcHandle()
-		, IsShellScript( false )
 	{}
 
 	/** Initialization constructor. */
-	FORCEINLINE explicit FProcHandle( HandleType Other, bool InIsShellScript = false )
+	FORCEINLINE explicit FProcHandle( HandleType Other )
 		: TProcHandle( Other )
-		, IsShellScript( InIsShellScript )
 	{}
-
-	bool IsShellScript;
 };
 
 /**
