Rework Mac OS X high-precision mouse handling so that we better emulate Windows behaviour.
- By default high-precision disables mouse event coalescing so that we get all the inputs as soon as possible to reduce apparent lag (r.Mac.HighPrecisionDisablesMouseCoalescing).
- An invisible, high-precision cursor is no longer warped to the current location on screen since that results in slight inaccuracies due to having to correct for the mouse warp delta event which occurs with a disassociated mouse/cursor.
- For mice only we can also optionally disable mouse acceleration when the cursor is invisible & high-precision mode is enabled (r.Mac.HighPrecisionDisablesMouseAcceleration) - this is more akin to Windows RawInput, but the some mice with very high DPI might not work correctly with it & it may not be Mac AppStore permissable as it disables all mouse acceleration & scaling system-wide. It certainly 'feels' wrong when the cursor is visible, so it isn't applied then.
- Mouse acceleration will be toggled as necessary on application focus change & on exit so that we don't appear to affect other applications. If we crash the setting will only persist until logout or the user changes the Mouse tracking speed setting in System Preferences.

diff --git a/Engine/Source/Runtime/Core/Private/Mac/MacApplication.cpp b/Engine/Source/Runtime/Core/Private/Mac/MacApplication.cpp
--- a/Engine/Source/Runtime/Core/Private/Mac/MacApplication.cpp
+++ b/Engine/Source/Runtime/Core/Private/Mac/MacApplication.cpp
@@ -539,29 +539,35 @@ void FMacApplication::ProcessNSEvent(NSEvent* const Event)
 
 			FMacCursor* MacCursor = (FMacCursor*)Cursor.Get();
 
+			FVector2D const MouseScaling = MacCursor->GetMouseScaling();
+			
 			if (bUsingHighPrecisionMouseInput)
 			{
 				// Under OS X we disassociate the cursor and mouse position during hi-precision mouse input.
 				// The game snaps the mouse cursor back to the starting point when this is disabled, which
 				// accumulates mouse delta that we want to ignore.
-				const FVector2D AccumDelta = MacCursor->GetMouseWarpDelta(true);
-
+				const FVector2D AccumDelta = MacCursor->GetMouseWarpDelta() * MouseScaling;
+				
+				// Get the mouse position
+				FVector2D HighPrecisionMousePos = MacCursor->GetPosition();
+				
 				// Find the screen the cursor is currently on.
 				NSEnumerator *ScreenEnumerator = [[NSScreen screens] objectEnumerator];
 				NSScreen *Screen;
-				while ((Screen = [ScreenEnumerator nextObject]) && !NSMouseInRect(NSMakePoint(HighPrecisionMousePos.X, HighPrecisionMousePos.Y), Screen.frame, NO))
+				while ((Screen = [ScreenEnumerator nextObject]) && !NSMouseInRect(NSMakePoint(HighPrecisionMousePos.X / MouseScaling.X, HighPrecisionMousePos.Y / MouseScaling.Y), Screen.frame, NO))
 					;
-
-				// Clamp to no more than the reported delta - a single event of no mouse movement won't be noticed
-				// but going in the wrong direction will.
-				const FVector2D FullDelta([Event deltaX], [Event deltaY]);
-				const FVector2D WarpDelta(FMath::Abs(AccumDelta.X)<FMath::Abs(FullDelta.X) ? AccumDelta.X : FullDelta.X, FMath::Abs(AccumDelta.Y)<FMath::Abs(FullDelta.Y) ? AccumDelta.Y : FullDelta.Y);
-
-				FVector2D Delta = ((FullDelta - WarpDelta) / 2.f) * MacCursor->GetMouseScaling();
-
-				HighPrecisionMousePos = MacCursor->GetPosition() + Delta;
+				
+				// Account for warping delta's
+				FVector2D Delta = FVector2D([Event deltaX], [Event deltaY]) * MouseScaling;
+				const FVector2D WarpDelta(FMath::Abs(AccumDelta.X)<FMath::Abs(Delta.X) ? AccumDelta.X : Delta.X, FMath::Abs(AccumDelta.Y)<FMath::Abs(Delta.Y) ? AccumDelta.Y : Delta.Y);
+				Delta -= WarpDelta;
+				
+				// Update to latest position
+				HighPrecisionMousePos += Delta;
+				
+				// Clip to lock rect
 				MacCursor->UpdateCursorClipping(HighPrecisionMousePos);
-
+				
 				// Clamp to the current screen and avoid the menu bar and dock to prevent popups and other
 				// assorted potential for mouse abuse.
 				NSRect VisibleFrame = [Screen visibleFrame];
@@ -573,11 +579,14 @@ void FMacApplication::ProcessNSEvent(NSEvent* const Event)
 				}
 				NSRect FullFrame = [Screen frame];
 				VisibleFrame.origin.y = (FullFrame.origin.y+FullFrame.size.height) - (VisibleFrame.origin.y + VisibleFrame.size.height);
-
-				HighPrecisionMousePos.X = FMath::Clamp(HighPrecisionMousePos.X / MacCursor->GetMouseScaling().X, (float)VisibleFrame.origin.x, (float)(VisibleFrame.origin.x + VisibleFrame.size.width)-1.f);
-				HighPrecisionMousePos.Y = FMath::Clamp(HighPrecisionMousePos.Y / MacCursor->GetMouseScaling().Y, (float)VisibleFrame.origin.y, (float)(VisibleFrame.origin.y + VisibleFrame.size.height)-1.f);
-
-				MacCursor->WarpCursor(HighPrecisionMousePos.X, HighPrecisionMousePos.Y);
+				
+				int32 ScaledX = (int32)(HighPrecisionMousePos.X / MouseScaling.X);
+				int32 ScaledY = (int32)(HighPrecisionMousePos.Y / MouseScaling.Y);
+				int32 ClampedPosX = FMath::Clamp(ScaledX, (int32)VisibleFrame.origin.x, (int32)(VisibleFrame.origin.x + VisibleFrame.size.width)-1);
+				int32 ClampedPosY = FMath::Clamp(ScaledY, (int32)VisibleFrame.origin.y, (int32)(VisibleFrame.origin.y + VisibleFrame.size.height)-1);
+				MacCursor->SetPosition(ClampedPosX * MouseScaling.X, ClampedPosY * MouseScaling.Y);
+				
+				// Forward the delta on to Slate
 				MessageHandler->OnRawMouseMove(Delta.X, Delta.Y);
 			}
 			else
@@ -585,7 +594,7 @@ void FMacApplication::ProcessNSEvent(NSEvent* const Event)
 				NSPoint CursorPos = [NSEvent mouseLocation];
 				CursorPos.y--; // The y coordinate of the point returned by mouseLocation starts from a base of 1
 				const FVector2D MousePosition = FVector2D(CursorPos.x, FPlatformMisc::ConvertSlateYPositionToCocoa(CursorPos.y));
-				FVector2D CurrentPosition = MousePosition * MacCursor->GetMouseScaling();
+				FVector2D CurrentPosition = MousePosition * MouseScaling;
 				const FVector2D MouseDelta = CurrentPosition - MacCursor->GetPosition();
 				if (MacCursor->UpdateCursorClipping(CurrentPosition))
 				{
@@ -974,6 +983,8 @@ void FMacApplication::OnApplicationDidBecomeActive()
 			}
 		}
 	}
+	
+	static_cast<FMacCursor*>( Cursor.Get() )->UpdateVisibility();
 
 	// If editor thread doesn't have the focus, don't suck up too much CPU time.
 	if (GIsEditor && !IsRunningCommandlet())
@@ -1017,6 +1028,8 @@ void FMacApplication::OnApplicationWillResignActive()
 			}
 		}
 	}
+	
+	static_cast<FMacCursor*>( Cursor.Get() )->UpdateVisibility();
 
 	// If editor thread doesn't have the focus, don't suck up too much CPU time.
 	if (GIsEditor && !IsRunningCommandlet())
@@ -1074,8 +1087,7 @@ void FMacApplication::PollGameDeviceState( const float TimeDelta )
 void FMacApplication::SetHighPrecisionMouseMode( const bool Enable, const TSharedPtr< FGenericWindow >& InWindow )
 {
 	bUsingHighPrecisionMouseInput = Enable;
-	HighPrecisionMousePos = static_cast<FMacCursor*>( Cursor.Get() )->GetPosition();
-	static_cast<FMacCursor*>( Cursor.Get() )->AssociateMouseAndCursorPosition( !Enable );
+	static_cast<FMacCursor*>( Cursor.Get() )->SetHighPrecisionMouseMode( Enable );
 }
 
 FModifierKeysState FMacApplication::GetModifierKeys() const
diff --git a/Engine/Source/Runtime/Core/Private/Mac/MacCursor.cpp b/Engine/Source/Runtime/Core/Private/Mac/MacCursor.cpp
--- a/Engine/Source/Runtime/Core/Private/Mac/MacCursor.cpp
+++ b/Engine/Source/Runtime/Core/Private/Mac/MacCursor.cpp
@@ -5,12 +5,30 @@
 #include "MacWindow.h"
 #include "MacApplication.h"
 
+#import <IOKit/hid/IOHIDKeys.h>
+#import <IOKit/hidsystem/IOHIDShared.h>
+
+int32 GMacDisableMouseCoalescing = 1;
+static FAutoConsoleVariableRef CVarMacDisableMouseCoalescing(
+	TEXT("r.Mac.HighPrecisionDisablesMouseCoalescing"),
+	GMacDisableMouseCoalescing,
+	TEXT("If set to true then OS X mouse event coalescing will be disabled while using high-precision mouse mode, to send all mouse events to UE4's event handling routines to reduce apparent mouse lag. (Default: True)"));
+
+int32 GMacDisableMouseAcceleration = 0;
+static FAutoConsoleVariableRef CVarMacDisableMouseAcceleration(
+	TEXT("r.Mac.HighPrecisionDisablesMouseAcceleration"),
+	GMacDisableMouseAcceleration,
+	TEXT("If set to true then OS X's mouse acceleration curve will be disabled while using high-precision mouse mode (typically used when games capture the mouse) resulting in a linear relationship between mouse movement & on-screen cursor movement. For some pointing devices this will make the cursor very slow. (Default: False)"));
+
 FMacCursor::FMacCursor()
 :	bIsVisible(true)
-,	bAssociateMouseCursor(false)
+,	bUseHighPrecisionMode(false)
 ,	CurrentPosition(FVector2D::ZeroVector)
 ,	MouseWarpDelta(FVector2D::ZeroVector)
 ,	MouseScale(1.0f, 1.0f)
+,	bIsPositionInitialised(false)
+,	HIDInterface(0)
+,	SavedAcceleration(0)
 {
 	SCOPED_AUTORELEASE_POOL;
 
@@ -79,11 +97,11 @@ FMacCursor::FMacCursor()
 			case EMouseCursor::GrabHand:
 				CursorHandle = [NSCursor openHandCursor];
 				break;
-				
+
 			case EMouseCursor::GrabHandClosed:
 				CursorHandle = [NSCursor closedHandCursor];
 				break;
-				
+
 			case EMouseCursor::SlashedCircle:
 				CursorHandle = [NSCursor operationNotAllowedCursor];
 				break;
@@ -109,15 +127,48 @@ FMacCursor::FMacCursor()
 	// Set the default cursor
 	SetType( EMouseCursor::Default );
 
-	NSPoint CursorPos = [NSEvent mouseLocation];
-	CursorPos.y--; // The y coordinate of the point returned by mouseLocation starts from a base of 1
-	CurrentPosition = FVector2D(FMath::TruncToFloat(CursorPos.x), FMath::TruncToFloat(FPlatformMisc::ConvertSlateYPositionToCocoa(CursorPos.y)));
+	CGEventRef Event = CGEventCreate(NULL);
+	CGPoint CursorPos = CGEventGetLocation(Event);
+	CFRelease(Event);
+
+	CurrentPosition = FVector2D(FMath::TruncToFloat(CursorPos.x), FMath::TruncToFloat(CursorPos.y));
+
+	// Get the IOHIDSystem so we can disable mouse acceleration
+	mach_port_t MasterPort;
+	kern_return_t KernResult = IOMasterPort(MACH_PORT_NULL, &MasterPort);
+	if (KERN_SUCCESS == KernResult)
+	{
+		CFMutableDictionaryRef ClassesToMatch = IOServiceMatching("IOHIDSystem");
+		if (ClassesToMatch)
+		{
+			io_iterator_t MatchingServices;
+			KernResult = IOServiceGetMatchingServices(MasterPort, ClassesToMatch, &MatchingServices);
+			if (KERN_SUCCESS == KernResult)
+			{
+				io_object_t IntfService;
+				if ((IntfService = IOIteratorNext(MatchingServices)))
+				{
+					KernResult = IOServiceOpen(IntfService, mach_task_self(), kIOHIDParamConnectType, &HIDInterface);
+					if (KernResult == KERN_SUCCESS)
+					{
+						KernResult = IOHIDGetAccelerationWithKey(HIDInterface, CFSTR(kIOHIDMouseAccelerationType), &SavedAcceleration);
+					}
+					if (KERN_SUCCESS != KernResult)
+					{
+						HIDInterface = 0;
+					}
+				}
+			}
+		}
+	}
 }
 
 FMacCursor::~FMacCursor()
 {
 	SCOPED_AUTORELEASE_POOL;
-
+	
+	SetHighPrecisionMouseMode( false );
+	
 	// Release cursors
 	// NOTE: Shared cursors will automatically be destroyed when the application is destroyed.
 	//       For dynamically created cursors, use [CursorHandles[CurCursorIndex] release];
@@ -155,16 +206,36 @@ FMacCursor::~FMacCursor()
 
 FVector2D FMacCursor::GetPosition() const
 {
-	return FVector2D(FMath::TruncToFloat(CurrentPosition.X * MouseScale.X), FMath::TruncToFloat(CurrentPosition.Y * MouseScale.Y));
+	FVector2D CurrentPos = CurrentPosition;
+	if ( !bIsPositionInitialised )
+	{
+		SCOPED_AUTORELEASE_POOL;
+		NSPoint CursorPos = [NSEvent mouseLocation];
+		CursorPos.y--; // The y coordinate of the point returned by mouseLocation starts from a base of 1
+		CurrentPos = FVector2D(CursorPos.x, FPlatformMisc::ConvertSlateYPositionToCocoa(CursorPos.y));
+	}
+	return FVector2D(FMath::TruncToFloat(CurrentPos.X * MouseScale.X), FMath::TruncToFloat(CurrentPos.Y * MouseScale.Y));
 }
 
 void FMacCursor::SetPosition( const int32 X, const int32 Y )
 {
-	FVector2D CurrentPos = GetPosition();
-	FVector2D NewPos(X, Y);
-	MouseWarpDelta += (NewPos - CurrentPos) / MouseScale;
-
-	WarpCursor(X / MouseScale.X, Y / MouseScale.Y);
+	FVector2D NewPos(X / MouseScale.X, Y / MouseScale.Y);
+	
+	UpdateCursorClipping(NewPos);
+	
+	MouseWarpDelta += (NewPos - CurrentPosition);
+	
+	if ( !bIsPositionInitialised || FIntVector(NewPos.X, NewPos.Y, 0) != FIntVector(CurrentPosition.X, CurrentPosition.Y, 0) )
+	{
+		if ( !bUseHighPrecisionMode || (CurrentCursor && bIsVisible) )
+		{
+			WarpCursor(NewPos.X, NewPos.Y);
+		}
+		else
+		{
+			UpdateCurrentPosition(NewPos);
+		}
+	}
 }
 
 void FMacCursor::SetType( const EMouseCursor::Type InNewCursor )
@@ -251,11 +322,12 @@ bool FMacCursor::UpdateCursorClipping( FVector2D& CursorPosition )
 
 void FMacCursor::UpdateVisibility()
 {
-	// @TODO: Remove usage of deprecated CGCursorIsVisible function
-#pragma clang diagnostic push
-#pragma clang diagnostic ignored "-Wdeprecated-declarations"
+	SCOPED_AUTORELEASE_POOL;
 	if([NSApp isActive])
 	{
+		// @TODO: Remove usage of deprecated CGCursorIsVisible function
+#pragma clang diagnostic push
+#pragma clang diagnostic ignored "-Wdeprecated-declarations"
 		if (CurrentCursor && bIsVisible)
 		{
 			// Enable the cursor.
@@ -263,6 +335,10 @@ void FMacCursor::UpdateVisibility()
 			{
 				CGDisplayShowCursor(kCGDirectMainDisplay);
 			}
+			if ( GMacDisableMouseAcceleration && HIDInterface && bUseHighPrecisionMode )
+			{
+				IOHIDSetAccelerationWithKey(HIDInterface, CFSTR(kIOHIDMouseAccelerationType), SavedAcceleration);
+			}
 		}
 		else
 		{
@@ -271,14 +347,23 @@ void FMacCursor::UpdateVisibility()
 			{
 				CGDisplayHideCursor(kCGDirectMainDisplay);
 			}
+			if ( GMacDisableMouseAcceleration && HIDInterface && bUseHighPrecisionMode && (!CurrentCursor || !bIsVisible) )
+			{
+				IOHIDSetAccelerationWithKey(HIDInterface, CFSTR(kIOHIDMouseAccelerationType), -1);
+			}
 		}
-	}
 #pragma clang diagnostic pop
+	}
+	else if ( GMacDisableMouseAcceleration && HIDInterface && bUseHighPrecisionMode && (!CurrentCursor || !bIsVisible) )
+	{
+		IOHIDSetAccelerationWithKey(HIDInterface, CFSTR(kIOHIDMouseAccelerationType), SavedAcceleration);
+	}
 }
 
 void FMacCursor::UpdateCurrentPosition(const FVector2D &Position)
 {
 	CurrentPosition = Position;
+	bIsPositionInitialised = true;
 }
 
 void FMacCursor::WarpCursor( const int32 X, const int32 Y )
@@ -287,7 +372,7 @@ void FMacCursor::WarpCursor( const int32 X, const int32 Y )
 	// Previously there was CGSetLocalEventsSuppressionInterval to explicitly control this behaviour but that is deprecated.
 	// The replacement CGEventSourceSetLocalEventsSuppressionInterval isn't useful because it is unclear how to obtain the correct event source.
 	// Instead, when we want the warp to be visible we need to disassociate mouse & cursor...
-	if( bAssociateMouseCursor )
+	if( !bUseHighPrecisionMode )
 	{
 		CGAssociateMouseAndMouseCursorPosition( false );
 	}
@@ -296,27 +381,65 @@ void FMacCursor::WarpCursor( const int32 X, const int32 Y )
 	CGWarpMouseCursorPosition( NSMakePoint( FMath::TruncToInt( X ), FMath::TruncToInt( Y ) ) );
 	
 	// And then reassociate the mouse cursor, which forces the mouse events to come through.
-	if( bAssociateMouseCursor )
+	if( !bUseHighPrecisionMode )
 	{
 		CGAssociateMouseAndMouseCursorPosition( true );
 	}
-
-	CurrentPosition = FVector2D(X, Y);
+	
+	UpdateCurrentPosition( FVector2D(X, Y) );
 }
 
-FVector2D FMacCursor::GetMouseWarpDelta( bool const bClearAccumulatedDelta )
+FVector2D FMacCursor::GetMouseWarpDelta()
 {
-	FVector2D Result = MouseWarpDelta;
-	MouseWarpDelta = bClearAccumulatedDelta ? FVector2D::ZeroVector : MouseWarpDelta;
+	FVector2D Result = ( !bUseHighPrecisionMode || (CurrentCursor && bIsVisible) ) ? MouseWarpDelta : FVector2D::ZeroVector;
+	MouseWarpDelta = FVector2D::ZeroVector;
 	return Result;
 }
 
-void FMacCursor::AssociateMouseAndCursorPosition( bool const bEnable )
+void FMacCursor::SetHighPrecisionMouseMode( bool const bEnable )
 {
-	if( bEnable != bAssociateMouseCursor )
+	if( bUseHighPrecisionMode != bEnable )
 	{
-		bAssociateMouseCursor = bEnable;
-		CGAssociateMouseAndMouseCursorPosition( bEnable );
+		bUseHighPrecisionMode = bEnable;
+		
+		CGAssociateMouseAndMouseCursorPosition( !bUseHighPrecisionMode );
+		
+		if ( GMacDisableMouseCoalescing )
+		{
+			SCOPED_AUTORELEASE_POOL;
+			[NSEvent setMouseCoalescingEnabled: !bUseHighPrecisionMode];
+		}
+		
+		if ( HIDInterface && GMacDisableMouseAcceleration && (!CurrentCursor || !bIsVisible) )
+		{
+			if ( !bUseHighPrecisionMode )
+			{
+				IOHIDSetAccelerationWithKey(HIDInterface, CFSTR(kIOHIDMouseAccelerationType), SavedAcceleration);
+			}
+			else
+			{
+				// Update the current saved acceleration.
+				double CurrentSetting = 0;
+				IOHIDGetAccelerationWithKey(HIDInterface, CFSTR(kIOHIDMouseAccelerationType), &CurrentSetting);
+				
+				// Need to check that we aren't picking up an invalid setting from ourselves.
+				if ( CurrentSetting >= 0 && CurrentSetting <= 3.0001 )
+				{
+					SavedAcceleration = CurrentSetting;
+				}
+				
+				IOHIDSetAccelerationWithKey(HIDInterface, CFSTR(kIOHIDMouseAccelerationType), -1);
+			}
+		}
+		
+		UpdateVisibility();
+		
+		// On disable put the cursor where the user would expect it
+		if ( !bEnable && (!CurrentCursor || !bIsVisible) )
+		{
+			FVector2D Pos = GetPosition();
+			SetPosition(Pos.X, Pos.Y);
+		}
 	}
 }
 
diff --git a/Engine/Source/Runtime/Core/Public/Mac/MacApplication.h b/Engine/Source/Runtime/Core/Public/Mac/MacApplication.h
--- a/Engine/Source/Runtime/Core/Public/Mac/MacApplication.h
+++ b/Engine/Source/Runtime/Core/Public/Mac/MacApplication.h
@@ -163,8 +163,6 @@ private:
 
 	bool bUsingTrackpad;
 
-	FVector2D HighPrecisionMousePos;
-
 	EMouseButtons::Type LastPressedMouseButton;
 
 	FCriticalSection WindowsMutex;
diff --git a/Engine/Source/Runtime/Core/Public/Mac/MacCursor.h b/Engine/Source/Runtime/Core/Public/Mac/MacCursor.h
--- a/Engine/Source/Runtime/Core/Public/Mac/MacCursor.h
+++ b/Engine/Source/Runtime/Core/Public/Mac/MacCursor.h
@@ -37,23 +37,22 @@ public:
 public:
 
 	bool UpdateCursorClipping( FVector2D& CursorPosition );
-	
+
 	void WarpCursor( const int32 X, const int32 Y );
-	
-	FVector2D GetMouseWarpDelta( bool const bClearAccumulatedDelta );
 
-	void AssociateMouseAndCursorPosition( bool const bEnable );
-	
+	FVector2D GetMouseWarpDelta();
+
+	void SetHighPrecisionMouseMode( bool const bEnable );
+
 	void SetMouseScaling( FVector2D Scale );
-	
+
 	FVector2D GetMouseScaling( void );
 
 	void UpdateCurrentPosition( const FVector2D& Position );
 
-private:
-
 	void UpdateVisibility();
 
+private:
 	EMouseCursor::Type CurrentType;
 
 	/** Cursors */
@@ -62,10 +61,14 @@ private:
 	FIntRect CusorClipRect;
 
 	bool bIsVisible;
-	bool bAssociateMouseCursor;
+	bool bUseHighPrecisionMode;
 	NSCursor* CurrentCursor;
 
 	FVector2D CurrentPosition;
 	FVector2D MouseWarpDelta;
 	FVector2D MouseScale;
+	bool bIsPositionInitialised;
+
+	io_object_t HIDInterface;
+	double SavedAcceleration;
 };
