Removed FMacEvent

diff --git a/Engine/Source/Runtime/Core/Private/Mac/CocoaWindow.cpp b/Engine/Source/Runtime/Core/Private/Mac/CocoaWindow.cpp
--- a/Engine/Source/Runtime/Core/Private/Mac/CocoaWindow.cpp
+++ b/Engine/Source/Runtime/Core/Private/Mac/CocoaWindow.cpp
@@ -4,7 +4,6 @@
 #include "CocoaWindow.h"
 #include "MacApplication.h"
 #include "CocoaTextView.h"
-#include "MacEvent.h"
 #include "CocoaThread.h"
 #include "MacCursor.h"
 
@@ -210,7 +209,7 @@ NSString* NSPerformDragOperation = @"NSPerformDragOperation";
 {
 	WindowMode = self.TargetWindowMode;
 
-	FMacEvent::SendToGameRunLoop(Notification, EMacEventSendMethod::Async, @[ NSDefaultRunLoopMode, UE4FullscreenEventMode ]);
+	GameThreadCall(^{ FMacApplication::ProcessEvent([Notification retain]); }, @[ NSDefaultRunLoopMode, UE4FullscreenEventMode ], false);
 
 	FMacCursor* MacCursor = (FMacCursor*)MacApplication->Cursor.Get();
 	if (MacCursor)
@@ -236,7 +235,7 @@ NSString* NSPerformDragOperation = @"NSPerformDragOperation";
 	WindowMode = EWindowMode::Windowed;
 	self.TargetWindowMode = EWindowMode::Windowed;
 
-	FMacEvent::SendToGameRunLoop(Notification, EMacEventSendMethod::Async, @[ NSDefaultRunLoopMode, UE4FullscreenEventMode ]);
+	GameThreadCall(^{ FMacApplication::ProcessEvent([Notification retain]); }, @[ NSDefaultRunLoopMode, UE4FullscreenEventMode ], false);
 
 	[self setFrame:self.PreFullScreenRect display:YES];
 	FMacCursor* MacCursor = (FMacCursor*)MacApplication->Cursor.Get();
@@ -254,7 +253,7 @@ NSString* NSPerformDragOperation = @"NSPerformDragOperation";
 		[self orderFrontAndMakeMain:false andKey:false];
 	}
 
-	FMacEvent::SendToGameRunLoop(Notification, EMacEventSendMethod::Async, @[ NSDefaultRunLoopMode, UE4ShowEventMode, UE4CloseEventMode, UE4FullscreenEventMode ]);
+	GameThreadCall(^{ FMacApplication::ProcessEvent([Notification retain]); }, @[ NSDefaultRunLoopMode, UE4ShowEventMode, UE4CloseEventMode, UE4FullscreenEventMode ], false);
 }
 
 - (void)windowDidResignMain:(NSNotification*)Notification
@@ -263,12 +262,12 @@ NSString* NSPerformDragOperation = @"NSPerformDragOperation";
 	[self setMovable: YES];
 	[self setMovableByWindowBackground: NO];
 
-	FMacEvent::SendToGameRunLoop(Notification, EMacEventSendMethod::Async, @[ NSDefaultRunLoopMode, UE4ShowEventMode, UE4CloseEventMode, UE4FullscreenEventMode ]);
+	GameThreadCall(^{ FMacApplication::ProcessEvent([Notification retain]); }, @[ NSDefaultRunLoopMode, UE4ShowEventMode, UE4CloseEventMode, UE4FullscreenEventMode ], false);
 }
 
 - (void)windowWillMove:(NSNotification*)Notification
 {
-	FMacEvent::SendToGameRunLoop(Notification, EMacEventSendMethod::Async, @[ NSDefaultRunLoopMode, UE4ResizeEventMode, UE4ShowEventMode, UE4FullscreenEventMode ]);
+	GameThreadCall(^{ FMacApplication::ProcessEvent([Notification retain]); }, @[ NSDefaultRunLoopMode, UE4ResizeEventMode, UE4ShowEventMode, UE4FullscreenEventMode ], false);
 }
 
 - (void)windowDidMove:(NSNotification*)Notification
@@ -279,7 +278,7 @@ NSString* NSPerformDragOperation = @"NSPerformDragOperation";
 	NSView* OpenGLView = [self openGLView];
 	[[NSNotificationCenter defaultCenter] postNotificationName:NSViewGlobalFrameDidChangeNotification object:OpenGLView];
 	
-	FMacEvent::SendToGameRunLoop(Notification, EMacEventSendMethod::Async, @[ NSDefaultRunLoopMode, UE4ResizeEventMode, UE4ShowEventMode, UE4FullscreenEventMode ]);
+	GameThreadCall(^{ FMacApplication::ProcessEvent([Notification retain]); }, @[ NSDefaultRunLoopMode, UE4ResizeEventMode, UE4ShowEventMode, UE4FullscreenEventMode ], false);
 }
 
 - (void)windowDidChangeScreen:(NSNotification*)Notification
@@ -347,32 +346,32 @@ NSString* NSPerformDragOperation = @"NSPerformDragOperation";
 - (void)windowWillStartLiveResize:(NSNotification*)Notification
 {
 	SCOPED_AUTORELEASE_POOL;
-	FMacEvent::SendToGameRunLoop(Notification, EMacEventSendMethod::Sync, @[ NSDefaultRunLoopMode, UE4ResizeEventMode, UE4ShowEventMode, UE4FullscreenEventMode ]);
+	GameThreadCall(^{ FMacApplication::ProcessEvent([Notification retain]); }, @[ NSDefaultRunLoopMode, UE4ResizeEventMode, UE4ShowEventMode, UE4FullscreenEventMode ], false);
 }
 
 - (void)windowDidEndLiveResize:(NSNotification*)Notification
 {
 	SCOPED_AUTORELEASE_POOL;
-	FMacEvent::SendToGameRunLoop(Notification, EMacEventSendMethod::Sync, @[ NSDefaultRunLoopMode, UE4ResizeEventMode, UE4ShowEventMode, UE4FullscreenEventMode ]);
+	GameThreadCall(^{ FMacApplication::ProcessEvent([Notification retain]); }, @[ NSDefaultRunLoopMode, UE4ResizeEventMode, UE4ShowEventMode, UE4FullscreenEventMode ], false);
 }
 
 - (void)windowDidResize:(NSNotification*)Notification
 {
 	SCOPED_AUTORELEASE_POOL;
 	bZoomed = [self isZoomed];
-	FMacEvent::SendToGameRunLoop(Notification, EMacEventSendMethod::Sync, @[ NSDefaultRunLoopMode, UE4ResizeEventMode, UE4ShowEventMode, UE4FullscreenEventMode ]);
+	GameThreadCall(^{ FMacApplication::ProcessEvent([Notification retain]); }, @[ NSDefaultRunLoopMode, UE4ResizeEventMode, UE4ShowEventMode, UE4FullscreenEventMode ], true);
 }
 
 - (void)windowWillClose:(NSNotification*)Notification
 {
 	SCOPED_AUTORELEASE_POOL;
-	FMacEvent::SendToGameRunLoop(Notification, EMacEventSendMethod::Async, @[ NSDefaultRunLoopMode, UE4CloseEventMode ]);
+	GameThreadCall(^{ FMacApplication::ProcessEvent([Notification retain]); }, @[ NSDefaultRunLoopMode, UE4CloseEventMode ], false);
 	[self setDelegate:nil];
 }
 
 - (void)mouseDown:(NSEvent*)Event
 {
-	FMacEvent::SendToGameRunLoop(Event, EMacEventSendMethod::Async);
+	GameThreadCall(^{ FMacApplication::ProcessEvent([Event retain]); }, @[ NSDefaultRunLoopMode ], false);
 }
 
 - (void)rightMouseDown:(NSEvent*)Event
@@ -385,27 +384,27 @@ NSString* NSPerformDragOperation = @"NSPerformDragOperation";
 		[self makeKeyWindow];
 	}
 	
-	FMacEvent::SendToGameRunLoop(Event, EMacEventSendMethod::Async);
+	GameThreadCall(^{ FMacApplication::ProcessEvent([Event retain]); }, @[ NSDefaultRunLoopMode ], false);
 }
 
 - (void)otherMouseDown:(NSEvent*)Event
 {
-	FMacEvent::SendToGameRunLoop(Event, EMacEventSendMethod::Async);
+	GameThreadCall(^{ FMacApplication::ProcessEvent([Event retain]); }, @[ NSDefaultRunLoopMode ], false);
 }
 
 - (void)mouseUp:(NSEvent*)Event
 {
-	FMacEvent::SendToGameRunLoop(Event, EMacEventSendMethod::Async);
+	GameThreadCall(^{ FMacApplication::ProcessEvent([Event retain]); }, @[ NSDefaultRunLoopMode ], false);
 }
 
 - (void)rightMouseUp:(NSEvent*)Event
 {
-	FMacEvent::SendToGameRunLoop(Event, EMacEventSendMethod::Async);
+	GameThreadCall(^{ FMacApplication::ProcessEvent([Event retain]); }, @[ NSDefaultRunLoopMode ], false);
 }
 
 - (void)otherMouseUp:(NSEvent*)Event
 {
-	FMacEvent::SendToGameRunLoop(Event, EMacEventSendMethod::Async);
+	GameThreadCall(^{ FMacApplication::ProcessEvent([Event retain]); }, @[ NSDefaultRunLoopMode ], false);
 }
 
 - (NSDragOperation)draggingEntered:(id <NSDraggingInfo>)Sender
@@ -417,14 +416,14 @@ NSString* NSPerformDragOperation = @"NSPerformDragOperation";
 {
 	SCOPED_AUTORELEASE_POOL;
 	NSNotification* Notification = [NSNotification notificationWithName:NSDraggingExited object:Sender];
-	FMacEvent::SendToGameRunLoop(Notification, EMacEventSendMethod::Async);
+	GameThreadCall(^{ FMacApplication::ProcessEvent([Notification retain]); }, @[ NSDefaultRunLoopMode ], false);
 }
 
 - (NSDragOperation)draggingUpdated:(id <NSDraggingInfo>)Sender
 {
 	SCOPED_AUTORELEASE_POOL;
 	NSNotification* Notification = [NSNotification notificationWithName:NSDraggingUpdated object:Sender];
-	FMacEvent::SendToGameRunLoop(Notification, EMacEventSendMethod::Async);
+	GameThreadCall(^{ FMacApplication::ProcessEvent([Notification retain]); }, @[ NSDefaultRunLoopMode ], false);
 	return NSDragOperationGeneric;
 }
 
@@ -432,7 +431,7 @@ NSString* NSPerformDragOperation = @"NSPerformDragOperation";
 {
 	SCOPED_AUTORELEASE_POOL;
 	NSNotification* Notification = [NSNotification notificationWithName:NSPrepareForDragOperation object:Sender];
-	FMacEvent::SendToGameRunLoop(Notification, EMacEventSendMethod::Async);
+	GameThreadCall(^{ FMacApplication::ProcessEvent([Notification retain]); }, @[ NSDefaultRunLoopMode ], false);
 	return YES;
 }
 
@@ -440,7 +439,7 @@ NSString* NSPerformDragOperation = @"NSPerformDragOperation";
 {
 	SCOPED_AUTORELEASE_POOL;
 	NSNotification* Notification = [NSNotification notificationWithName:NSPerformDragOperation object:Sender];
-	FMacEvent::SendToGameRunLoop(Notification, EMacEventSendMethod::Async);
+	GameThreadCall(^{ FMacApplication::ProcessEvent([Notification retain]); }, @[ NSDefaultRunLoopMode ], false);
 	return YES;
 }
 
diff --git a/Engine/Source/Runtime/Core/Private/Mac/MacApplication.cpp b/Engine/Source/Runtime/Core/Private/Mac/MacApplication.cpp
--- a/Engine/Source/Runtime/Core/Private/Mac/MacApplication.cpp
+++ b/Engine/Source/Runtime/Core/Private/Mac/MacApplication.cpp
@@ -5,7 +5,6 @@
 #include "MacApplication.h"
 #include "MacWindow.h"
 #include "MacCursor.h"
-#include "MacEvent.h"
 #include "CocoaMenu.h"
 #include "GenericApplicationMessageHandler.h"
 #include "HIDInputInterface.h"
@@ -15,8 +14,6 @@
 #include "ModuleManager.h"
 #include "CocoaTextView.h"
 
-static NSString* NSWindowDraggingFinished = @"NSWindowDraggingFinished";
-
 FMacApplication* MacApplication = NULL;
 
 #if WITH_EDITOR
@@ -61,19 +58,13 @@ NSEvent* FMacApplication::HandleNSEvent(NSEvent* Event)
 
 		if (!bIsResentEvent && (!bIsMouseClickOrKeyEvent || [Event window] == NULL))
 		{
-			FMacEvent::SendToGameRunLoop(Event, EMacEventSendMethod::Async);
+			GameThreadCall(^{ FMacApplication::ProcessEvent([Event retain]); }, @[ NSDefaultRunLoopMode ], false);
 			
 			if ([Event type] == NSKeyDown || [Event type] == NSKeyUp)
 			{
 				ReturnEvent = nil;
 			}
 		}
-
-		if ([Event type] == NSLeftMouseUp)
-		{
-			NSNotification* Notification = [NSNotification notificationWithName:NSWindowDraggingFinished object:[Event window]];
-			FMacEvent::SendToGameRunLoop(Notification, EMacEventSendMethod::Async);
-		}
 	}
 
 	return ReturnEvent;
@@ -113,7 +104,7 @@ FMacApplication::FMacApplication()
 	CGDisplayRegisterReconfigurationCallback(FMacApplication::OnDisplayReconfiguration, this);
 
 	[NSEvent addGlobalMonitorForEventsMatchingMask:NSMouseMovedMask handler:^(NSEvent* Event){
-		FMacEvent::SendToGameRunLoop(Event, EMacEventSendMethod::Async);
+		GameThreadCall(^{ FMacApplication::ProcessEvent([Event retain]); }, @[ NSDefaultRunLoopMode ], false);
 	}];
 
 	EventMonitor = [NSEvent addLocalMonitorForEventsMatchingMask:NSAnyEventMask handler:^(NSEvent* IncomingEvent)
@@ -243,24 +234,21 @@ bool FMacApplication::IsCursorDirectlyOverSlateWindow() const
 	return Window && [Window isKindOfClass:[FCocoaWindow class]] && Window != DraggedWindow;
 }
 
-void FMacApplication::ProcessEvent(FMacEvent const* const Event)
+void FMacApplication::ProcessEvent(NSObject const* const Event)
 {
 	// Must have an event
 	check(Event);
 	
-	// Determine the type of the event to forward to the right handling routine
-	NSEvent* AppKitEvent = Event->GetEvent();
-	NSNotification* Notification = Event->GetNotification();
-
-	if (AppKitEvent)
+	if ([Event isKindOfClass:[NSEvent class]])
 	{
 		// Process a standard NSEvent as we always did
-		MacApplication->ProcessNSEvent(AppKitEvent);
+		MacApplication->ProcessNSEvent((NSEvent*)Event);
 	}
-	else if (Notification)
+	else if ([Event isKindOfClass:[NSNotification class]])
 	{
 		// Notifications need to mapped to the right handler function, that's no longer handled in the
 		// call location.
+		NSNotification* Notification = (NSNotification*)Event;
 		NSString* const NotificationName = [Notification name];
 		FCocoaWindow* CocoaWindow = [[Notification object] isKindOfClass:[FCocoaWindow class]] ? (FCocoaWindow*)[Notification object] : nullptr;
 		if (CocoaWindow)
@@ -305,10 +293,6 @@ void FMacApplication::ProcessEvent(FMacEvent const* const Event)
 			{
 				MacApplication->OnWindowDidClose(CocoaWindow);
 			}
-			else if (NotificationName == NSWindowDraggingFinished)
-			{
-				MacApplication->OnWindowDraggingFinished();
-			}
 			else
 			{
 				check(false);
@@ -341,7 +325,7 @@ void FMacApplication::ProcessEvent(FMacEvent const* const Event)
 	}
 	
 	// We are responsible for deallocating the event
-	delete Event;
+	[Event release];
 }
 
 void FMacApplication::RedrawWindow(FCocoaWindow* Window)
@@ -419,15 +403,6 @@ void FMacApplication::InvalidateTextLayouts()
 
 }
 
-void FMacApplication::OnWindowDraggingFinished()
-{
-	if ( DraggedWindow )
-	{
-		SCOPED_AUTORELEASE_POOL;
-		DraggedWindow = NULL;
-	}
-}
-
 bool FMacApplication::IsWindowMovable(FCocoaWindow* Win, bool* OutMovableByBackground)
 {
 	if (OutMovableByBackground)
@@ -720,6 +695,7 @@ void FMacApplication::ProcessNSEvent(NSEvent* const Event)
 			}
 
 			FPlatformMisc::bChachedMacMenuStateNeedsUpdate = true;
+			DraggedWindow = nullptr;
 			break;
 		}
 
@@ -851,7 +827,6 @@ void FMacApplication::ProcessNSEvent(NSEvent* const Event)
 		}
 	}
 
-	[Event ResetWindow];
 	bIsProcessingNSEvent = bWasProcessingNSEvent;
 }
 
@@ -886,7 +861,7 @@ FCocoaWindow* FMacApplication::FindEventWindow( NSEvent* Event )
 			break;
 	}
 
-	FCocoaWindow* EventWindow = (FCocoaWindow*)[Event GetWindow];
+	FCocoaWindow* EventWindow = (FCocoaWindow*)[Event window];
 
 	if ([Event type] == NSMouseMoved)
 	{
diff --git a/Engine/Source/Runtime/Core/Private/Mac/MacEvent.cpp b/Engine/Source/Runtime/Core/Private/Mac/MacEvent.cpp
deleted file mode 100644
diff --git a/Engine/Source/Runtime/Core/Private/Mac/MacWindow.cpp b/Engine/Source/Runtime/Core/Private/Mac/MacWindow.cpp
--- a/Engine/Source/Runtime/Core/Private/Mac/MacWindow.cpp
+++ b/Engine/Source/Runtime/Core/Private/Mac/MacWindow.cpp
@@ -4,7 +4,6 @@
 #include "MacWindow.h"
 #include "MacApplication.h"
 #include "MacCursor.h"
-#include "MacEvent.h"
 #include "CocoaTextView.h"
 #include "CocoaThread.h"
 
@@ -253,7 +252,7 @@ void FMacWindow::ReshapeWindow( int32 X, int32 Y, int32 Width, int32 Height )
 						MacCursor->SetMouseScaling(FVector2D(WidthScale, HeightScale));
 					}
 				}, UE4ResizeEventMode, true);
-				FMacEvent::SendToGameRunLoop([NSNotification notificationWithName:NSWindowDidResizeNotification object:WindowHandle], EMacEventSendMethod::Sync, @[ NSDefaultRunLoopMode, UE4ResizeEventMode, UE4ShowEventMode, UE4FullscreenEventMode ]);
+				GameThreadCall(^{ FMacApplication::ProcessEvent([[NSNotification notificationWithName:NSWindowDidResizeNotification object:WindowHandle] retain]); }, @[ NSDefaultRunLoopMode, UE4ResizeEventMode, UE4ShowEventMode, UE4FullscreenEventMode ], true);
 			}
 		}
 	}
diff --git a/Engine/Source/Runtime/Core/Public/Mac/MacApplication.h b/Engine/Source/Runtime/Core/Public/Mac/MacApplication.h
--- a/Engine/Source/Runtime/Core/Public/Mac/MacApplication.h
+++ b/Engine/Source/Runtime/Core/Public/Mac/MacApplication.h
@@ -6,8 +6,6 @@
 #include "MacWindow.h"
 #include "MacTextInputMethodSystem.h"
 
-class FMacEvent;
-
 /**
  * Mac-specific application implementation.
  */
@@ -65,8 +63,6 @@ public:
 
 	void ProcessNSEvent(NSEvent* const Event);
 
-	void OnWindowDraggingFinished();
-
 	bool IsWindowMovable(FCocoaWindow* Win, bool* OutMovableByBackground);
 
 	void ResetModifierKeys() { ModifierKeysFlags = 0; }
@@ -108,7 +104,7 @@ public:
 
 	void OnMouseCursorLock(bool bLockEnabled);
 
-	static void ProcessEvent(FMacEvent const* const Event);
+	static void ProcessEvent(NSObject const* const Event);
 
 	void RedrawWindow(FCocoaWindow* Window);
 
@@ -222,7 +218,6 @@ private:
 	id EventMonitor;
 
 	friend class FMacWindow;
-	friend class FMacEvent;
 };
 
 extern FMacApplication* MacApplication;
diff --git a/Engine/Source/Runtime/Core/Public/Mac/MacEvent.h b/Engine/Source/Runtime/Core/Public/Mac/MacEvent.h
deleted file mode 100644
