implementation of executable name for iOS
#ios

diff --git a/Engine/Source/Runtime/Core/Private/IOS/IOSPlatformProcess.cpp b/Engine/Source/Runtime/Core/Private/IOS/IOSPlatformProcess.cpp
--- a/Engine/Source/Runtime/Core/Private/IOS/IOSPlatformProcess.cpp
+++ b/Engine/Source/Runtime/Core/Private/IOS/IOSPlatformProcess.cpp
@@ -102,6 +102,17 @@ void FIOSPlatformProcess::SetThreadAffinityMask(uint64 AffinityMask)
 	}
 }
 
+const TCHAR* FIOSPlatformProcess::ExecutableName(bool bRemoveExtension)
+{
+	static TCHAR Result[512] = TEXT("");
+	if (!Result[0])
+	{
+		NSString *NSExeName = [[[NSBundle mainBundle] executablePath] lastPathComponent];
+		FPlatformString::CFStringToTCHAR((CFStringRef)NSExeName, Result);
+	}
+	return Result;
+}
+
 const uint64 FIOSPlatformAffinity::GetMainGameMask()
 {
 	return MAKEAFFINITYMASK1(0);
diff --git a/Engine/Source/Runtime/Core/Public/IOS/IOSPlatformProcess.h b/Engine/Source/Runtime/Core/Public/IOS/IOSPlatformProcess.h
--- a/Engine/Source/Runtime/Core/Public/IOS/IOSPlatformProcess.h
+++ b/Engine/Source/Runtime/Core/Public/IOS/IOSPlatformProcess.h
@@ -34,6 +34,7 @@ struct CORE_API FIOSPlatformProcess : public FGenericPlatformProcess
 	static void SetRealTimeMode();
 	static void SetupGameOrRenderThread(bool bIsRenderThread);
 	static void SetThreadAffinityMask(uint64 AffinityMask);
+	static const TCHAR* ExecutableName(bool bRemoveExtension = true);
 };
 
 typedef FIOSPlatformProcess FPlatformProcess;
