Fixed checkfSlow macro for Clang

diff --git a/Engine/Source/Runtime/Core/Public/Misc/AssertionMacros.h b/Engine/Source/Runtime/Core/Public/Misc/AssertionMacros.h
--- a/Engine/Source/Runtime/Core/Public/Misc/AssertionMacros.h
+++ b/Engine/Source/Runtime/Core/Public/Misc/AssertionMacros.h
@@ -73,7 +73,7 @@
 //
 #if DO_GUARD_SLOW
 	#define checkSlow(expr)					{ if(!(expr)) { FDebug::AssertFailed( #expr, __FILE__, __LINE__ ); CA_ASSUME(expr); } }
-	#define checkfSlow(expr, format, ...)	{ if(!(expr)) { FDebug::AssertFailed( #expr, __FILE__, __LINE__, format, __VA_ARGS__ ); CA_ASSUME(expr); } }
+	#define checkfSlow(expr, format, ...)	{ if(!(expr)) { FDebug::AssertFailed( #expr, __FILE__, __LINE__, format, ##__VA_ARGS__ ); CA_ASSUME(expr); } }
 	#define verifySlow(expr)				{ if(!(expr)) { FDebug::AssertFailed( #expr, __FILE__, __LINE__ );} }
 #else
 	#define checkSlow(expr)
