Added FPaths method for getting the path to the Engine's source code directory.

diff --git a/Engine/Source/Runtime/Core/Private/Misc/Paths.cpp b/Engine/Source/Runtime/Core/Private/Misc/Paths.cpp
--- a/Engine/Source/Runtime/Core/Private/Misc/Paths.cpp
+++ b/Engine/Source/Runtime/Core/Private/Misc/Paths.cpp
@@ -353,6 +353,11 @@ FString FPaths::GameAgnosticSavedDir()
 	return EngineSavedDir();
 }
 
+FString FPaths::EngineSourceDir()
+{
+	return FPaths::EngineDir() + TEXT("Source/");
+}
+
 FString FPaths::GameSourceDir()
 {
 	return FPaths::GameDir() + TEXT("Source/");
diff --git a/Engine/Source/Runtime/Core/Public/Misc/Paths.h b/Engine/Source/Runtime/Core/Public/Misc/Paths.h
--- a/Engine/Source/Runtime/Core/Public/Misc/Paths.h
+++ b/Engine/Source/Runtime/Core/Public/Misc/Paths.h
@@ -260,6 +260,11 @@ public:
 	 */
 	static FString GameAgnosticSavedDir();
 
+	/**
+	 * @return The directory where engine source code files are kept
+	 */
+	static FString EngineSourceDir();
+
 	/**
 	 * @return The directory where game source code files are kept
 	 */
