Consistent whitespaces.

diff --git a/Engine/Source/Runtime/Core/Public/Async/Async.h b/Engine/Source/Runtime/Core/Public/Async/Async.h
--- a/Engine/Source/Runtime/Core/Public/Async/Async.h
+++ b/Engine/Source/Runtime/Core/Public/Async/Async.h
@@ -220,16 +220,20 @@ private:
 	TPromise<ResultType> Promise;
 };
 
-/** Helper struct used to generate unique ids for the stats. */
+
+/**
+ * Helper struct used to generate unique ids for the stats.
+ */
 struct FAsyncIndex
 {
 	CORE_API static int32 GetNext()
 	{
 		static FThreadSafeCounter TAsyncThreadIndex;
-		return TAsyncThreadIndex.Add( 1 );
+		return TAsyncThreadIndex.Add(1);
 	}
 };
 
+
 /**
  * Executes a given function asynchronously.
  *
@@ -287,7 +291,7 @@ TFuture<ResultType> Async(EAsyncExecution Execution, TFunction<ResultType()> Fun
 			TPromise<FRunnableThread*> ThreadPromise;
 			TAsyncRunnable<ResultType>* Runnable = new TAsyncRunnable<ResultType>(MoveTemp(Function), MoveTemp(Promise), ThreadPromise.GetFuture());
 			
-			const FString TAsyncThreadName = FString::Printf( TEXT( "TAsync %d" ), FAsyncIndex::GetNext() );
+			const FString TAsyncThreadName = FString::Printf(TEXT("TAsync %d"), FAsyncIndex::GetNext());
 			FRunnableThread* RunnableThread = FRunnableThread::Create(Runnable, *TAsyncThreadName);
 
 			check(RunnableThread != nullptr);
@@ -317,11 +321,11 @@ template<typename ResultType> uint32 TAsyncRunnable<ResultType>::Run()
 	FRunnableThread* Thread = ThreadFuture.Get();
 
 	// Enqueue deletion of the thread to a different thread.
-	Async<void>( EAsyncExecution::TaskGraph, [=]()
-	{
-		delete Thread;
-		delete this;
-	} );
+	Async<void>(EAsyncExecution::TaskGraph, [=]() {
+			delete Thread;
+			delete this;
+		}
+	);
 
 	return 0;
 }
