Linux: Make open menus behave more like on Windows.

- PR #852 contributed by [REDACTED_PERSON].

https://github.com/EpicGames/UnrealEngine/pull/852

diff --git a/Engine/Source/Runtime/Core/Private/Linux/LinuxApplication.cpp b/Engine/Source/Runtime/Core/Private/Linux/LinuxApplication.cpp
--- a/Engine/Source/Runtime/Core/Private/Linux/LinuxApplication.cpp
+++ b/Engine/Source/Runtime/Core/Private/Linux/LinuxApplication.cpp
@@ -339,7 +339,8 @@ void FLinuxApplication::ProcessDeferredMessage( SDL_Event Event )
 				{
 					if(CurrentlyActiveWindow != CurrentEventWindow)
 					{
-						CurrentEventWindow->SetWindowFocus();
+						SDL_RaiseWindow( CurrentEventWindow->GetHWnd() );
+						SDL_SetWindowInputFocus( CurrentEventWindow->GetHWnd() );
 
 						SDL_PushEvent(&Event);
 					}
@@ -635,41 +636,6 @@ void FLinuxApplication::ProcessDeferredMessage( SDL_Event Event )
 					}
 					break;
 
-				case SDL_WINDOWEVENT_DESTROY:
-					{
-						RemoveEventWindow(CurrentEventWindow->GetHWnd());
-						RemoveRevertFocusWindow(CurrentEventWindow->GetHWnd());
-
-						// After destroying the window we have to decide what to do with the
-						// focus. We should check if the current window has a parent so we can give
-						// the focus to the parent. If not we get a window from the Windows list
-						// and set the focus to the most top one.
- 						if (CurrentEventWindow->IsRegularWindow() && CurrentEventWindow->GetParent().IsValid())
- 						{
- 							CurrentEventWindow->GetParent()->BringToFront();
- 							CurrentEventWindow->GetParent()->SetWindowFocus();
- 						} 
- 						else if (CurrentEventWindow->IsRegularWindow() && RevertFocusStack.Num() > 0)
-						{
- 							TSharedPtr< FLinuxWindow > TopmostWindow = RevertFocusStack.Top();
-							if (TopmostWindow.IsValid())
-							{
-								TopmostWindow->BringToFront();
-								TopmostWindow->SetWindowFocus();
-							}
-						}
-						else if (CurrentEventWindow->IsPopupMenuWindow() && CurrentEventWindow->GetParent().IsValid() && !bIsDragWindowButtonPressed && bActivateApp)
-						{
-							CurrentEventWindow->GetParent()->BringToFront();
-							CurrentEventWindow->GetParent()->SetWindowFocus();
-						}
-
-						SDL_DestroyWindow(CurrentEventWindow->GetHWnd());
-
-						MessageHandler->OnWindowClose(CurrentEventWindow.ToSharedRef());
-					}
-					break;
-
 				case SDL_WINDOWEVENT_SHOWN:
 					{
 						if (CurrentEventWindow->IsRegularWindow())
@@ -723,6 +689,8 @@ void FLinuxApplication::ProcessDeferredMessage( SDL_Event Event )
 					{
 						if (CurrentEventWindow.IsValid())
 						{
+							CurrentEventWindow->OnPointerEnteredWindow(true);
+							
 							MessageHandler->OnCursorSet();
 
 							bInsideOwnWindow = true;
@@ -735,6 +703,8 @@ void FLinuxApplication::ProcessDeferredMessage( SDL_Event Event )
 					{
 						if (CurrentEventWindow.IsValid())
 						{
+							CurrentEventWindow->OnPointerEnteredWindow(false);
+							
 							if (GetCapture() != nullptr)
 							{
 								UpdateMouseCaptureWindow((SDL_HWindow)GetCapture());
@@ -1228,22 +1198,6 @@ void FDisplayMetrics::GetDisplayMetrics(FDisplayMetrics& OutDisplayMetrics)
 	}
 }
 
-void FLinuxApplication::SendDestroyEvent(SDL_Window * HWnd)
-{
-	const TSharedPtr< FLinuxWindow > Window = FindWindowBySDLWindow(HWnd);
-	if (!Window.IsValid())
-	{
-		return;
-	}
-
-	SDL_Event event;
-	SDL_zero(event);
-	event.type = SDL_WINDOWEVENT;
-	event.window.event = SDL_WINDOWEVENT_DESTROY;
-	event.window.windowID = SDL_GetWindowID(Window->GetHWnd());
-	SDL_PushEvent(&event);
-}
-
 void FLinuxApplication::RemoveRevertFocusWindow(SDL_HWindow HWnd)
 {
 	for (int32 WindowIndex=0; WindowIndex < RevertFocusStack.Num(); ++WindowIndex)
@@ -1253,6 +1207,35 @@ void FLinuxApplication::RemoveRevertFocusWindow(SDL_HWindow HWnd)
 		if (Window->GetHWnd() == HWnd)
 		{
 			RevertFocusStack.RemoveAt(WindowIndex);
+
+			// After destroying the window we have to decide what to do with the
+			// focus. We should check if the current window has a parent so we can give
+			// the focus to the parent. If not we get a window from the Windows list
+			// and set the focus to the most top one.
+			if (Window->IsRegularWindow() && Window->GetParent().IsValid())
+			{
+				SDL_RaiseWindow( Window->GetParent()->GetHWnd() );
+				SDL_SetWindowInputFocus( Window->GetParent()->GetHWnd() );
+			} 
+			else if (Window->IsRegularWindow() && RevertFocusStack.Num() > 0)
+			{
+				TSharedPtr< FLinuxWindow > TopmostWindow = RevertFocusStack.Top();
+				if (TopmostWindow.IsValid())
+				{
+					SDL_RaiseWindow( TopmostWindow->GetHWnd() );
+					SDL_SetWindowInputFocus( TopmostWindow->GetHWnd() );
+				}
+			}
+			// Check if a popup window got destroyed.
+			else if (Window->IsPopupMenuWindow() && Window->GetParent().IsValid() && bActivateApp )
+			{
+				if ( (!bIsDragWindowButtonPressed && bInsideOwnWindow) ||
+					(bIsDragWindowButtonPressed && Window->IsPointerInsideWindow()) )
+				{
+					SDL_RaiseWindow(Window->GetParent()->GetHWnd() );
+					SDL_SetWindowInputFocus(Window->GetParent()->GetHWnd() );
+				}
+			}
 			return;
 		}
 	}	
diff --git a/Engine/Source/Runtime/Core/Private/Linux/LinuxWindow.cpp b/Engine/Source/Runtime/Core/Private/Linux/LinuxWindow.cpp
--- a/Engine/Source/Runtime/Core/Private/Linux/LinuxWindow.cpp
+++ b/Engine/Source/Runtime/Core/Private/Linux/LinuxWindow.cpp
@@ -31,6 +31,7 @@ FLinuxWindow::FLinuxWindow()
 	, bWasFullscreen( false )
 	, bIsPopupWindow(false)
 	, bIsTooltipWindow(false)
+	, bIsPointerInsideWindow(false)
 {
 	PreFullscreenWindowRect.left = PreFullscreenWindowRect.top = PreFullscreenWindowRect.right = PreFullscreenWindowRect.bottom = 0;
 }
@@ -228,7 +229,10 @@ void FLinuxWindow::BringToFront( bool bForce )
 /** Native windows should implement this function by asking the OS to destroy OS-specific resource associated with the window (e.g. Win32 window handle) */
 void FLinuxWindow::Destroy()
 {
-	OwningApplication->SendDestroyEvent(HWnd);
+	OwningApplication->RemoveEventWindow( HWnd );
+	OwningApplication->RemoveRevertFocusWindow( HWnd );
+
+	SDL_DestroyWindow( HWnd );
 }
 
 /** Native window should implement this function by performing the equivalent of the Win32 minimize-to-taskbar operation */
@@ -561,3 +565,13 @@ const TSharedPtr< FLinuxWindow >& FLinuxWindow::GetParent() const
 {
 	return ParentWindow;
 }
+
+void FLinuxWindow::OnPointerEnteredWindow(bool PointerEnteredWindow)
+{
+	bIsPointerInsideWindow = PointerEnteredWindow;
+}
+
+bool FLinuxWindow::IsPointerInsideWindow() const
+{
+	return bIsPointerInsideWindow;
+}
diff --git a/Engine/Source/Runtime/Core/Public/Linux/LinuxApplication.h b/Engine/Source/Runtime/Core/Public/Linux/LinuxApplication.h
--- a/Engine/Source/Runtime/Core/Public/Linux/LinuxApplication.h
+++ b/Engine/Source/Runtime/Core/Public/Linux/LinuxApplication.h
@@ -73,8 +73,6 @@ public:
 
 	virtual bool IsCursorDirectlyOverSlateWindow() const override;
 
-	void SendDestroyEvent(SDL_Window * HWnd);
-
 private:
 
 	FLinuxApplication();
diff --git a/Engine/Source/Runtime/Core/Public/Linux/LinuxWindow.h b/Engine/Source/Runtime/Core/Public/Linux/LinuxWindow.h
--- a/Engine/Source/Runtime/Core/Public/Linux/LinuxWindow.h
+++ b/Engine/Source/Runtime/Core/Public/Linux/LinuxWindow.h
@@ -138,6 +138,20 @@ public:
 	/** @return	Gives the native window a chance to adjust our stored window size before we cache it off */
 	virtual void AdjustCachedSize( FVector2D& Size ) const override;
 
+	/**
+	 * Sets the state if the pointer entered/left the window.
+	 * 
+	 * @param PointerEnteredWindow If the pointer entered the window set 'true' else 'false'.
+	 */
+	void OnPointerEnteredWindow(bool PointerEnteredWindow);
+
+	/**
+	 * Is the pointer currently inside the window?
+	 * 
+	 * @return Returns 'true' if the pointer is inside the window else 'false'.
+	 */
+	bool IsPointerInsideWindow() const;
+	
 private:
 
 	/**
@@ -187,6 +201,9 @@ private:
 	/** TODO: describe */
 	bool bIsConsoleWindow;
 
+	/** Whether the is inside this window */
+	bool bIsPointerInsideWindow;
+	
 	/** SDL ID (for debugging purposes) */
 	uint32 WindowSDLID;
 
