Revert focus after destruction of popup windows.

- A commit from PR #852 by yaakuro that I missed.
- Also minor cosmetic code changes.

diff --git a/Engine/Source/Runtime/Core/Private/Linux/LinuxApplication.cpp b/Engine/Source/Runtime/Core/Private/Linux/LinuxApplication.cpp
--- a/Engine/Source/Runtime/Core/Private/Linux/LinuxApplication.cpp
+++ b/Engine/Source/Runtime/Core/Private/Linux/LinuxApplication.cpp
@@ -1234,9 +1234,17 @@ void FLinuxApplication::RemoveRevertFocusWindow(SDL_HWindow HWnd)
 				{
 					SDL_RaiseWindow(Window->GetParent()->GetHWnd() );
 					SDL_SetWindowInputFocus(Window->GetParent()->GetHWnd() );
+
+					// TODO If we have a popup menu and a sub popup menu open, we have
+					// to fake the following. After getting destructed the condition above
+					// 'Window->IsPointerInsideWindow()' will be false for the parent popup window
+					// and the focus will be not reset. This is rather hackery and should be
+					// removed later if possible.
+					Window->GetParent()->OnPointerEnteredWindow(true);
 				}
 			}
-			return;
+
+			break;
 		}
-	}	
+	}
 }
