Simplifying the Default Config Array saving code.

No longer manually remove all the previous elements manually, we use the '!' syntax specifier to clear the array.

diff --git a/Engine/Source/Runtime/Core/Private/Misc/ConfigCacheIni.cpp b/Engine/Source/Runtime/Core/Private/Misc/ConfigCacheIni.cpp
--- a/Engine/Source/Runtime/Core/Private/Misc/ConfigCacheIni.cpp
+++ b/Engine/Source/Runtime/Core/Private/Misc/ConfigCacheIni.cpp
@@ -1278,34 +1278,11 @@ void FConfigFile::ProcessSourceAndCheckAgainstBackup()
 
 void FConfigFile::ProcessPropertyAndWriteForDefaults( const TArray< FString >& InCompletePropertyToProcess, FString& OutText, const FString& SectionName, const FString& PropertyName )
 {
-	// Only process against a hierarchy if this config file has one.
-	if (SourceIniHierarchy.Num() > 0)
+	if (PropertyName.StartsWith(TEXT("+")))
 	{
-		// Handle array elements from the configs hierarchy.
-		if (PropertyName.StartsWith(TEXT("+")) || InCompletePropertyToProcess.Num() > 1)
-		{
-			// Build a config file out of this default configs hierarchy.
-			FConfigCacheIni Hierarchy(EConfigCacheType::Temporary);
-			const FString& LastFileInHierarchy = SourceIniHierarchy.Last().Filename;
-
-			FConfigFile& DefaultConfigFile = Hierarchy.Add(LastFileInHierarchy, FConfigFile());
-			for (TArray<FIniFilename>::TIterator HierarchyFileIt(SourceIniHierarchy); HierarchyFileIt; ++HierarchyFileIt)
-			{
-				DefaultConfigFile.Combine(HierarchyFileIt->Filename);
-			}
-
-			// Remove any array elements from the default configs hierearchy, we will add these in below
-			// Note.	This compensates for an issue where strings in the hierarchy have a slightly different format
-			//			to how the config system wishes to serialize them.
-			TArray<FString> ArrayProperties;
-			Hierarchy.GetArray(*SectionName, *PropertyName.Replace(TEXT("+"), TEXT("")), ArrayProperties, *LastFileInHierarchy);
-
-			for (const FString& NextElement : ArrayProperties)
-			{
-				FString PropertyNameWithRemoveOp = PropertyName.Replace(TEXT("+"), TEXT("-"));
-				OutText += FString::Printf(TEXT("%s=%s") LINE_TERMINATOR, *PropertyNameWithRemoveOp, *NextElement);
-			}
-		}
+		// if we are writing to defaults, we will be writing all the array values to this config.
+		// Therefore, we should clear the array of all entries prior to this file in the hierarchy.
+		OutText += FString::Printf(TEXT("%s=%s") LINE_TERMINATOR, *PropertyName.Replace(TEXT("+"), TEXT("!")), TEXT("CLEAR_ARRAY"));
 	}
 
 	// Write the properties out to a file.
