Unicast delegate lookup/removal no longer creates a delegate to do that.
static_asserts to ensure that TScriptDelegate and TWeakObjectPtr are constructed with the expected base classes.
Missing FWeakObjectPtr::GetEvenIfUnreachable implementation added.
TWeakObjectPtr::Get and TWeakObjectPtr::GetEvenIfUnreachable forwarded instead of being reimplemented.
New TScriptDelegate::IsBoundToObjectEvenIfUnreachable and TScriptDelegate::GetUObjectEvenIfUnreachable functions.
Some TScriptDelegate implementation changed to use TArray algorithms.

diff --git a/Engine/Source/Runtime/Core/Public/Delegates/DelegateSignatureImpl.inl b/Engine/Source/Runtime/Core/Public/Delegates/DelegateSignatureImpl.inl
--- a/Engine/Source/Runtime/Core/Public/Delegates/DelegateSignatureImpl.inl
+++ b/Engine/Source/Runtime/Core/Public/Delegates/DelegateSignatureImpl.inl
@@ -2076,6 +2076,25 @@ protected:
 };
 
 
+namespace UE4Delegates_Private
+{
+	/**
+	 * Returns the root function name from a string representing a member function pointer.
+	 * Note: this function only returns a pointer to the substring and doesn't create a new string.
+	 *
+	 * @param  InMacroFunctionName  The string containing the member function name.
+	 * @return A pointer to the substring containing the member function name.
+	 */
+	inline const TCHAR* GetTrimmedMemberFunctionName(const TCHAR* InMacroFunctionName)
+	{
+		// We strip off the class prefix and just return the function name by itself.
+		const TCHAR* Result = FCString::Strrstr( InMacroFunctionName, TEXT( "::" ) );
+		checkf(Result, TEXT("'%s' does not look like a member function"), InMacroFunctionName);
+		return Result + 2;
+	}
+}
+
+
 /**
  * Implements a multicast delegate.
  *
@@ -2207,10 +2226,7 @@ public:
 		this->Object = InUserObject;
 
 		// Store the function name.  The incoming function name was generated by a macro and includes the method's class name.
-		// We'll need to strip off that class prefix and just store the function name by itself.
-		const TCHAR* TrimmedMacroFunctionName = FCString::Strrstr( InMacroFunctionName, TEXT( "::" ) );
-		checkf(TrimmedMacroFunctionName, TEXT("'%s' does not look like a member function"), InMacroFunctionName);
-		this->FunctionName = TrimmedMacroFunctionName + 2;
+		this->FunctionName = UE4Delegates_Private::GetTrimmedMemberFunctionName( InMacroFunctionName );
 
 		ensureMsgf(this->IsBound(), TEXT("Unable to bind delegate to '%s' (function might not be marked as a UFUNCTION)"), InMacroFunctionName);
 	}
@@ -2273,10 +2289,8 @@ public:
 
 		// NOTE: We're not actually using the incoming method pointer or calling it.  We simply require it for type-safety reasons.
 
-		FDelegate NewDelegate;
-		NewDelegate.__Internal_BindDynamic( InUserObject, InMethodPtr, InMacroFunctionName );
-
-		return this->Contains( NewDelegate );
+		FName TrimmedName = UE4Delegates_Private::GetTrimmedMemberFunctionName( InMacroFunctionName );
+		return this->Contains( InUserObject, TrimmedName );
 	}
 
 	/**
@@ -2342,10 +2356,8 @@ public:
 
 		// NOTE: We're not actually storing the incoming method pointer or calling it.  We simply require it for type-safety reasons.
 
-		FDelegate NewDelegate;
-		NewDelegate.__Internal_BindDynamic( InUserObject, InMethodPtr, InMacroFunctionName );
-
-		this->Remove( NewDelegate );
+		FName TrimmedName = UE4Delegates_Private::GetTrimmedMemberFunctionName( InMacroFunctionName );
+		this->Remove( InUserObject, TrimmedName );
 	}
 
 	// NOTE:  Broadcast() method must be defined in derived classes
diff --git a/Engine/Source/Runtime/Core/Public/Delegates/DelegateSignatureImpl_Variadics.inl b/Engine/Source/Runtime/Core/Public/Delegates/DelegateSignatureImpl_Variadics.inl
--- a/Engine/Source/Runtime/Core/Public/Delegates/DelegateSignatureImpl_Variadics.inl
+++ b/Engine/Source/Runtime/Core/Public/Delegates/DelegateSignatureImpl_Variadics.inl
@@ -1115,6 +1115,25 @@ protected:
 };
 
 
+namespace UE4Delegates_Private
+{
+	/**
+	 * Returns the root function name from a string representing a member function pointer.
+	 * Note: this function only returns a pointer to the substring and doesn't create a new string.
+	 *
+	 * @param  InMacroFunctionName  The string containing the member function name.
+	 * @return A pointer to the substring containing the member function name.
+	 */
+	inline const TCHAR* GetTrimmedMemberFunctionName(const TCHAR* InMacroFunctionName)
+	{
+		// We strip off the class prefix and just return the function name by itself.
+		const TCHAR* Result = FCString::Strrstr( InMacroFunctionName, TEXT( "::" ) );
+		checkf(Result, TEXT("'%s' does not look like a member function"), InMacroFunctionName);
+		return Result + 2;
+	}
+}
+
+
 /**
  * Implements a multicast delegate.
  *
@@ -1187,10 +1206,7 @@ public:
 		this->Object = InUserObject;
 
 		// Store the function name.  The incoming function name was generated by a macro and includes the method's class name.
-		// We'll need to strip off that class prefix and just store the function name by itself.
-		const TCHAR* TrimmedMacroFunctionName = FCString::Strrstr( InMacroFunctionName, TEXT( "::" ) );
-		checkf(TrimmedMacroFunctionName, TEXT("'%s' does not look like a member function"), InMacroFunctionName);
-		this->FunctionName = TrimmedMacroFunctionName + 2;
+		this->FunctionName = UE4Delegates_Private::GetTrimmedMemberFunctionName( InMacroFunctionName );
 
 		ensureMsgf(this->IsBound(), TEXT("Unable to bind delegate to '%s' (function might not be marked as a UFUNCTION)"), InMacroFunctionName);
 	}
@@ -1251,10 +1267,8 @@ public:
 
 		// NOTE: We're not actually using the incoming method pointer or calling it.  We simply require it for type-safety reasons.
 
-		FDelegate NewDelegate;
-		NewDelegate.__Internal_BindDynamic( InUserObject, InMethodPtr, InMacroFunctionName );
-
-		return this->Contains( NewDelegate );
+		FName TrimmedName = UE4Delegates_Private::GetTrimmedMemberFunctionName( InMacroFunctionName );
+		return this->Contains( InUserObject, TrimmedName );
 	}
 
 	/**
@@ -1320,10 +1334,8 @@ public:
 
 		// NOTE: We're not actually storing the incoming method pointer or calling it.  We simply require it for type-safety reasons.
 
-		FDelegate NewDelegate;
-		NewDelegate.__Internal_BindDynamic( InUserObject, InMethodPtr, InMacroFunctionName );
-
-		this->Remove( NewDelegate );
+		FName TrimmedName = UE4Delegates_Private::GetTrimmedMemberFunctionName( InMacroFunctionName );
+		this->Remove( InUserObject, TrimmedName );
 	}
 
 	// NOTE:  Broadcast() method must be defined in derived classes
diff --git a/Engine/Source/Runtime/Core/Public/UObject/ScriptDelegates.h b/Engine/Source/Runtime/Core/Public/UObject/ScriptDelegates.h
--- a/Engine/Source/Runtime/Core/Public/UObject/ScriptDelegates.h
+++ b/Engine/Source/Runtime/Core/Public/UObject/ScriptDelegates.h
@@ -5,12 +5,19 @@
 #include "WeakObjectPtrTemplates.h"
 
 
+class UObject;
+
+
 /**
  * Script delegate base class.
  */
 template <typename TWeakPtr = FWeakObjectPtr>
 class TScriptDelegate
 {
+	// Although templated, the parameter is not intended to be anything other than the default,
+	// and is only a template for module organization reasons.
+	static_assert(TAreTypesEqual<TWeakPtr, FWeakObjectPtr>::Value, "TWeakPtr should not be overridden");
+
 public:
 
 	/** Default constructor. */
@@ -43,7 +50,7 @@ public:
 	 * @param InObject The object to call the function on.
 	 * @param InFunctionName The name of the function to call.
 	 */
-	void BindUFunction( class UObject* InObject, const FName& InFunctionName )
+	void BindUFunction( UObject* InObject, const FName& InFunctionName )
 	{
 		Object = InObject;
 		FunctionName = InFunctionName;
@@ -69,6 +76,16 @@ public:
 		return InUserObject && (InUserObject == GetUObject());
 	}
 
+	/** 
+	 * Checks to see if this delegate is bound to the given user object, even if the object is unreachable.
+	 *
+	 * @return	True if this delegate is bound to InUserObject, false otherwise.
+	 */
+	bool IsBoundToObjectEvenIfUnreachable(void const* InUserObject) const
+	{
+		return InUserObject && InUserObject == GetUObjectEvenIfUnreachable();
+	}
+
 	/** 
 	 * Checks to see if the user object bound to this delegate will ever be valid again
 	 *
@@ -140,7 +157,7 @@ public:
 	 *
 	 * @return	The object
 	 */
-	class UObject* GetUObject()
+	UObject* GetUObject()
 	{
 		// Downcast UObjectBase to UObject
 		return static_cast< UObject* >( Object.Get() );
@@ -151,12 +168,34 @@ public:
 	 *
 	 * @return	The object
 	 */
-	const class UObject* GetUObject() const
+	const UObject* GetUObject() const
 	{
 		// Downcast UObjectBase to UObject
 		return static_cast< const UObject* >( Object.Get() );
 	}
 
+	/** 
+	 * Gets the object bound to this delegate, even if the object is unreachable
+	 *
+	 * @return	The object
+	 */
+	UObject* GetUObjectEvenIfUnreachable()
+	{
+		// Downcast UObjectBase to UObject
+		return static_cast< UObject* >( Object.GetEvenIfUnreachable() );
+	}
+
+	/**
+	 * Gets the object bound to this delegate (const), even if the object is unreachable
+	 *
+	 * @return	The object
+	 */
+	const UObject* GetUObjectEvenIfUnreachable() const
+	{
+		// Downcast UObjectBase to UObject
+		return static_cast< const UObject* >( Object.GetEvenIfUnreachable() );
+	}
+
 	/**
 	 * Gets the name of the function to call on the bound object
 	 *
@@ -244,10 +283,23 @@ public:
 	 */
 	bool Contains( const TScriptDelegate<TWeakPtr>& InDelegate ) const
 	{
-		// Add the delegate
 		return InvocationList.Contains( InDelegate );
 	}
 
+	/**
+	 * Checks whether a function delegate is already a member of this multi-cast delegate's invocation list
+	 *
+	 * @param	InObject		Object of the delegate to check
+	 * @param	InFunctionName	Function name of the delegate to check
+	 * @return	True if the delegate is already in the list.
+	 */
+	bool Contains( const UObject* InObject, FName InFunctionName ) const
+	{
+		return InvocationList.ContainsByPredicate( [=]( const TScriptDelegate<TWeakPtr>& Delegate ){
+			return Delegate.GetFunctionName() == InFunctionName && Delegate.IsBoundToObjectEvenIfUnreachable(InObject);
+		} );
+	}
+
 	/**
 	 * Adds a function delegate to this multi-cast delegate's invocation list
 	 *
@@ -292,6 +344,22 @@ public:
 		CompactInvocationList();
 	}
 
+	/**
+	 * Removes a function from this multi-cast delegate's invocation list (performance is O(N)).  Note that the
+	 * order of the delegates may not be preserved!
+	 *
+	 * @param	InObject		Object of the delegate to remove
+	 * @param	InFunctionName	Function name of the delegate to remove
+	 */
+	void Remove( const UObject* InObject, FName InFunctionName )
+	{
+		// Remove the delegate
+		RemoveInternal( InObject, InFunctionName );
+
+		// Check for any delegates that may have expired
+		CompactInvocationList();
+	}
+
 	/**
 	 * Removes all delegate bindings from this multicast delegate's
 	 * invocation list that are bound to the specified object.
@@ -404,7 +472,7 @@ public:
 	 * need call this function in normal circumstances.
  	 * @return	List of objects bound to this delegate
 	*/
-	TArray< class UObject* > GetAllObjects()
+	TArray< UObject* > GetAllObjects()
 	{
 		TArray< UObject* > OutputList;
 		for( typename FInvocationList::TIterator CurDelegate( InvocationList ); CurDelegate; ++CurDelegate )
@@ -458,30 +526,34 @@ protected:
 	*/
 	void RemoveInternal( const TScriptDelegate<TWeakPtr>& InDelegate ) const
 	{
-		const int32 NumFunctions = InvocationList.Num();
-		for( int32 CurFunctionIndex = 0; CurFunctionIndex < NumFunctions; ++CurFunctionIndex )
-		{
-			if( InvocationList[ CurFunctionIndex ] == InDelegate )
-			{
-				InvocationList.RemoveAtSwap( CurFunctionIndex );
+		InvocationList.RemoveSingleSwap(InDelegate);
+	}
 
-				// No need to continue, as we never allow the same delegate to be bound twice
-					break;
-			}
+	/**
+	 * Removes a function from this multi-cast delegate's invocation list (performance is O(N)).  Note that the
+	 * order of the delegates may not be preserved!
+	 *
+	 * @param	InObject		Object of the delegate to remove
+	 * @param	InFunctionName	Function name of the delegate to remove
+	*/
+	void RemoveInternal( const UObject* InObject, FName InFunctionName ) const
+	{
+		int32 FoundDelegate = InvocationList.IndexOfByPredicate([=](const TScriptDelegate<TWeakPtr>& Delegate) {
+			return Delegate.GetFunctionName() == InFunctionName && Delegate.IsBoundToObjectEvenIfUnreachable(InObject);
+		});
+
+		if (FoundDelegate != INDEX_NONE)
+		{
+			InvocationList.RemoveAtSwap(FoundDelegate, 1, false);
 		}
 	}
 
 	/** Cleans up any delegates in our invocation list that have expired (performance is O(N)) */
 	void CompactInvocationList() const
 	{
-		for( int32 CurFunctionIndex = 0; CurFunctionIndex < InvocationList.Num(); ++CurFunctionIndex )
-		{
-			if( InvocationList[ CurFunctionIndex ].IsCompactable() )
-			{
-				// Remove it, and decrement loop counter so that we don't skip over an element
-				InvocationList.RemoveAtSwap( CurFunctionIndex-- );
-			}
-		}
+		InvocationList.RemoveAllSwap([](const TScriptDelegate<TWeakPtr>& Delegate){
+			return Delegate.IsCompactable();
+		});
 	}
 
 protected:
diff --git a/Engine/Source/Runtime/Core/Public/UObject/WeakObjectPtrTemplates.h b/Engine/Source/Runtime/Core/Public/UObject/WeakObjectPtrTemplates.h
--- a/Engine/Source/Runtime/Core/Public/UObject/WeakObjectPtrTemplates.h
+++ b/Engine/Source/Runtime/Core/Public/UObject/WeakObjectPtrTemplates.h
@@ -26,6 +26,11 @@ struct FIndexToObject;
 template<class T=UObject, class TWeakObjectPtrBase=FWeakObjectPtr, class TUObjectArray=FIndexToObject>
 struct TWeakObjectPtr : private TWeakObjectPtrBase
 {
+	// Although templated, these parameters are not intended to be anything other than the default,
+	// and are only templates for module organization reasons.
+	static_assert(TAreTypesEqual<TWeakObjectPtrBase, FWeakObjectPtr>::Value, "TWeakObjectPtrBase should not be overridden");
+	static_assert(TAreTypesEqual<TUObjectArray,      FIndexToObject>::Value, "TUObjectArray should not be overridden");
+
 public:
 
 	/** Default constructor (no initialization). **/
@@ -111,25 +116,15 @@ public:
 	 * @param bEvenIfPendingKill, if this is true, pendingkill objects are considered valid	
 	 * @return NULL if this object is gone or the weak pointer was NULL, otherwise a valid uobject pointer
 	**/
-	FORCEINLINE T *Get(bool bEvenIfPendingKill = false) const
+	FORCEINLINE T* Get(bool bEvenIfPendingKill = false) const
 	{
-		UObjectBase *Result = NULL;
-		if (IsValid(true))
-		{
-			Result = static_cast<UObjectBase *>(TUObjectArray::IndexToObject(this->GetObjectIndex(), bEvenIfPendingKill ));
-		}
-		return (T *)Result;
+		return (T*)TWeakObjectPtrBase::Get(bEvenIfPendingKill);
 	}
 
 	/** Deferences the weak pointer even if its marked RF_Unreachable. This is needed to resolve weak pointers during GC (such as ::AddReferenceObjects) */
-	FORCEINLINE T *GetEvenIfUnreachable() const
+	FORCEINLINE T* GetEvenIfUnreachable() const
 	{
-		UObjectBase *Result = NULL;
-		if (IsValid(true, true))
-		{
-			Result = static_cast<UObjectBase *>(TUObjectArray::IndexToObject(this->GetObjectIndex(), true));
-		}
-		return (T *)Result;
+		return (T*)TWeakObjectPtrBase::GetEvenIfUnreachable();
 	}
 
 	/**  
diff --git a/Engine/Source/Runtime/CoreUObject/Private/UObject/WeakObjectPtr.cpp b/Engine/Source/Runtime/CoreUObject/Private/UObject/WeakObjectPtr.cpp
--- a/Engine/Source/Runtime/CoreUObject/Private/UObject/WeakObjectPtr.cpp
+++ b/Engine/Source/Runtime/CoreUObject/Private/UObject/WeakObjectPtr.cpp
@@ -267,7 +267,7 @@ bool FWeakObjectPtr::IsStale(bool bEvenIfPendingKill, bool bThreadsafeTest) cons
 
 UObject* FWeakObjectPtr::Get(bool bEvenIfPendingKill) const
 {
-	UObject *Result = NULL;
+	UObject* Result = nullptr;
 	if (IsValid(true))
 	{
 		Result = (UObject*)(GUObjectArray.IndexToObject(GetObjectIndex(), bEvenIfPendingKill));
@@ -275,6 +275,16 @@ UObject* FWeakObjectPtr::Get(bool bEvenIfPendingKill) const
 	return Result;
 }
 
+UObject* FWeakObjectPtr::GetEvenIfUnreachable() const
+{
+	UObject* Result = nullptr;
+	if (IsValid(true, true))
+	{
+		Result = static_cast<UObject*>(GUObjectArray.IndexToObject(GetObjectIndex(), true));
+	}
+	return Result;
+}
+
 
 /**
  * Weak object pointer serialization.  Weak object pointers only have weak references to objects and
