Stats - Fixes

diff --git a/Engine/Source/Runtime/Core/Private/Stats/Stats2.cpp b/Engine/Source/Runtime/Core/Private/Stats/Stats2.cpp
--- a/Engine/Source/Runtime/Core/Private/Stats/Stats2.cpp
+++ b/Engine/Source/Runtime/Core/Private/Stats/Stats2.cpp
@@ -202,6 +202,7 @@ class FStatGroupEnableManager : public IStatGroupEnableManager
 	struct FGroupEnable
 	{
 		TMap<FName, TStatIdData *> NamesInThisGroup;
+		TMap<FName, TStatIdData *> AlwaysEnabledNamesInThisGroup;
 		bool DefaultEnable; 
 		bool CurrentEnable; 
 
@@ -256,9 +257,9 @@ class FStatGroupEnableManager : public IStatGroupEnableManager
 
 	void DisableStat(TStatIdData* DisablePtr)
 	{
-		static_assert(sizeof(FMinimalName) == sizeof(uint64), "FMinimalName should have the same size of uint64.");
-		check(UPTRINT(&DisablePtr->Name) % sizeof(FMinimalName) == 0);
-		*(uint64*)&DisablePtr->Name = 0;
+		static_assert(sizeof( FMinimalName ) == sizeof( uint64 ), "FMinimalName should have the same size of uint64.");
+		check( UPTRINT( &DisablePtr->Name ) % sizeof( FMinimalName ) == 0 );
+		*(uint64*)&DisablePtr->Name = *(uint64 const*)TStatId::GetStatNone();
 	}
 
 public:
@@ -370,14 +371,28 @@ public:
 				UE_LOG(LogStatGroupEnableManager, Fatal, TEXT("Stat group %s was was defined both on and off by default."), *Group.ToString());
 			}
 			TStatIdData** StatFound = Found->NamesInThisGroup.Find(Stat);
+			TStatIdData** StatFoundAlways = Found->AlwaysEnabledNamesInThisGroup.Find( Stat );
 			if (StatFound)
 			{
-				return TStatId(*StatFound);
+				TStatIdData** StatFound = Found->NamesInThisGroup.Find( Stat );
+				TStatIdData** StatFoundAlways = Found->AlwaysEnabledNamesInThisGroup.Find( Stat );
+				if( StatFound )
+				{
+					if( StatFoundAlways )
+					{
+						UE_LOG( LogStatGroupEnableManager, Fatal, TEXT( "Stat %s is both always enabled and not always enabled, so it was used for two different things." ), *Stat.ToString() );
+					}
+					return TStatId( *StatFound );
+				}
+				else if( StatFoundAlways )
+				{
+					return TStatId( *StatFoundAlways );
+				}
 			}
 		}
 		else
 		{
-			Found = &HighPerformanceEnable.Add(Group, FGroupEnable(bDefaultEnable));
+			Found = &HighPerformanceEnable.Add( Group, FGroupEnable( bDefaultEnable || !bShouldClearEveryFrame ) );
 
 			// this was set up before we saw the group, so set the enable now
 			if (EnableForNewGroup.Contains(Group))
@@ -416,13 +431,18 @@ public:
 		
 		++PendingStatIds;
 
-		if (Found->CurrentEnable)
+		if( Found->CurrentEnable )
 		{
-			EnableStat(Stat, Result);
+			EnableStat( Stat, Result );
+		}
+
+		if( bShouldClearEveryFrame )
+		{
+			Found->NamesInThisGroup.Add( Stat, Result );
 		}
 		else
 		{
-			Found->NamesInThisGroup.Add(Stat, Result);
+			Found->AlwaysEnabledNamesInThisGroup.Add( Stat, Result );
 		}
 		return TStatId(Result);
 	}
@@ -447,6 +467,10 @@ public:
 				{
 					UE_LOG(LogStatGroupEnableManager, Display, TEXT("      %d %s"), !ItInner.Value()->IsNone(), *ItInner.Key().ToString());
 				}
+				for( auto ItInner = It.Value().AlwaysEnabledNamesInThisGroup.CreateConstIterator(); ItInner; ++ItInner )
+				{
+					UE_LOG( LogStatGroupEnableManager, Display, TEXT( "      (always enabled) %s" ), *ItInner.Key().ToString() );
+				}
 			}
 		}
 	}
diff --git a/Engine/Source/Runtime/Core/Public/Stats/Stats2.h b/Engine/Source/Runtime/Core/Public/Stats/Stats2.h
--- a/Engine/Source/Runtime/Core/Public/Stats/Stats2.h
+++ b/Engine/Source/Runtime/Core/Public/Stats/Stats2.h
@@ -66,6 +66,7 @@ struct CORE_API FStats
 
 #if STATS
 
+MS_ALIGN(8)
 struct TStatIdData
 {
 	FORCEINLINE TStatIdData()
@@ -77,7 +78,7 @@ struct TStatIdData
 
 	FORCEINLINE bool IsNone() const
 	{
-		return MinimalNameToName(Name).IsNone();
+		return Name.Index == 0 && Name.Number == 0;
 	}
 
 	/** Name of the active stat; stored as a minimal name to minimize the data size */
@@ -88,7 +89,7 @@ struct TStatIdData
 
 	/** const WIDECHAR* pointer to a string; stored as a uint64 so it doesn't change size and affect TStatIdData alignment between 32 and 64-bit builds) */
 	uint64 WideString;
-};
+} GCC_ALIGN(8);
 
 struct TStatId
 {
@@ -102,7 +103,7 @@ struct TStatId
 	}
 	FORCEINLINE bool IsValidStat() const
 	{
-		return StatIdPtr != &TStatId_NAME_None;
+		return !IsNone();
 	}
 	FORCEINLINE bool IsNone() const
 	{
@@ -118,6 +119,11 @@ struct TStatId
 		return MinimalNameToName(StatIdPtr->Name);
 	}
 
+	FORCEINLINE static const FMinimalName* GetStatNone()
+	{
+		return &TStatId_NAME_None.Name;
+	}
+
 	/**
 	 * @return a stat description as an ansi string.
 	 * StatIdPtr must point to a valid FName pointer.
diff --git a/Engine/Source/Runtime/CoreUObject/Private/UObject/LinkerLoad.cpp b/Engine/Source/Runtime/CoreUObject/Private/UObject/LinkerLoad.cpp
--- a/Engine/Source/Runtime/CoreUObject/Private/UObject/LinkerLoad.cpp
+++ b/Engine/Source/Runtime/CoreUObject/Private/UObject/LinkerLoad.cpp
@@ -16,8 +16,10 @@ DECLARE_CYCLE_STAT(TEXT("Linker Preload"),STAT_LinkerPreload,STATGROUP_LinkerLoa
 DECLARE_CYCLE_STAT(TEXT("Linker Precache"),STAT_LinkerPrecache,STATGROUP_LinkerLoad);
 DECLARE_CYCLE_STAT(TEXT("Linker Serialize"),STAT_LinkerSerialize,STATGROUP_LinkerLoad);
 DECLARE_CYCLE_STAT(TEXT("Linker Load Deferred"), STAT_LinkerLoadDeferred, STATGROUP_LinkerLoad);
-DECLARE_DWORD_ACCUMULATOR_STAT(TEXT("Linker Count"), STAT_LinkerCount, STATGROUP_LinkerLoad);
-DECLARE_DWORD_ACCUMULATOR_STAT(TEXT("Live Linker Count"), STAT_LiveLinkerCount, STATGROUP_LinkerLoad);
+
+DECLARE_STATS_GROUP( TEXT( "Linker Count" ), STATGROUP_LinkerCount, STATCAT_Advanced );
+DECLARE_DWORD_ACCUMULATOR_STAT(TEXT("Linker Count"), STAT_LinkerCount, STATGROUP_LinkerCount);
+DECLARE_DWORD_ACCUMULATOR_STAT(TEXT("Live Linker Count"), STAT_LiveLinkerCount, STATGROUP_LinkerCount);
 
 /** Points to the main PackageLinker currently being serialized (Defined in Linker.cpp) */
 extern ULinkerLoad* GSerializedPackageLinker;
