Cache the correct origin on entering fullscreen so that we can correctly restore position on exiting fullscreen on OS X [UE-11504].

diff --git a/Engine/Source/Runtime/Core/Private/Mac/CocoaWindow.cpp b/Engine/Source/Runtime/Core/Private/Mac/CocoaWindow.cpp
--- a/Engine/Source/Runtime/Core/Private/Mac/CocoaWindow.cpp
+++ b/Engine/Source/Runtime/Core/Private/Mac/CocoaWindow.cpp
@@ -201,7 +201,8 @@ NSString* NSPerformDragOperation = @"NSPerformDragOperation";
 			self.TargetWindowMode = EWindowMode::WindowedFullscreen;
 		}
 #endif
-		self.PreFullScreenRect = [self openGLFrame];
+		self.PreFullScreenRect.origin = [self frame].origin;
+		self.PreFullScreenRect.size = [self openGLFrame].size;
 	}
 }
 
diff --git a/Engine/Source/Runtime/Core/Private/Mac/MacWindow.cpp b/Engine/Source/Runtime/Core/Private/Mac/MacWindow.cpp
--- a/Engine/Source/Runtime/Core/Private/Mac/MacWindow.cpp
+++ b/Engine/Source/Runtime/Core/Private/Mac/MacWindow.cpp
@@ -390,7 +390,8 @@ void FMacWindow::SetWindowMode( EWindowMode::Type NewWindowMode )
 		
 		if(!bIsFullscreen)
 		{
-			PreFullscreenWindowRect = [WindowHandle openGLFrame];
+			PreFullscreenWindowRect.origin = [WindowHandle frame].origin;
+			PreFullscreenWindowRect.size = [WindowHandle openGLFrame].size;
 			WindowHandle.PreFullScreenRect = PreFullscreenWindowRect;
 		}
 		
