Core: Non-invasive type name declaration for better TTypeContainer support

#UpgradeNotes:
This is an API breaking change. Instead of using the DECLARE_REGISTRABLE_TYPE(Type) macro inside a struct or class use the Expose_TNameOf(Type) macro outside a struct or class. This will allow for supporting existing types in type containers without having to modify those types.

diff --git a/Engine/Source/Runtime/Core/Private/Tests/Misc/TypeContainerTest.cpp b/Engine/Source/Runtime/Core/Private/Tests/Misc/TypeContainerTest.cpp
--- a/Engine/Source/Runtime/Core/Private/Tests/Misc/TypeContainerTest.cpp
+++ b/Engine/Source/Runtime/Core/Private/Tests/Misc/TypeContainerTest.cpp
@@ -19,21 +19,19 @@ IMPLEMENT_SIMPLE_AUTOMATION_TEST(FTypeContainerTest, "Core.Misc.TypeContainer",
 
 struct IFruit
 {
-	DECLARE_REGISTRABLE_TYPE(IFruit)
 	virtual ~IFruit() { }
 	virtual FString Name() = 0;
 };
 
 struct IBerry : public IFruit
 {
-	DECLARE_REGISTRABLE_TYPE(IBerry)
 	virtual ~IBerry() { }
 };
 
 struct FBanana : public IFruit
 {
 	virtual ~FBanana() { }
-	virtual FString Name() override { return TEXT("Cherry"); }
+	virtual FString Name() override { return TEXT("Banana"); }
 };
 
 struct FStrawberry : public IBerry
@@ -44,7 +42,6 @@ struct FStrawberry : public IBerry
 
 struct ISmoothie
 {
-	DECLARE_REGISTRABLE_TYPE(ISmoothie)
 	virtual ~ISmoothie() { }
 	virtual TSharedRef<IBerry> GetBerry() = 0;
 	virtual TSharedRef<IFruit> GetFruit() = 0;
@@ -68,6 +65,10 @@ struct FTwoSmoothies
 
 DECLARE_DELEGATE_RetVal(TSharedPtr<IFruit>, FFruitFactoryDelegate);
 
+Expose_TNameOf(IBerry)
+Expose_TNameOf(IFruit)
+Expose_TNameOf(ISmoothie)
+
 
 /* Tests
  *****************************************************************************/
diff --git a/Engine/Source/Runtime/Core/Public/Misc/TypeContainer.h b/Engine/Source/Runtime/Core/Public/Misc/TypeContainer.h
--- a/Engine/Source/Runtime/Core/Public/Misc/TypeContainer.h
+++ b/Engine/Source/Runtime/Core/Public/Misc/TypeContainer.h
@@ -6,18 +6,6 @@
 #if PLATFORM_COMPILER_HAS_VARIADIC_TEMPLATES
 
 
-/**
- * Convenience macro for declaring type traits for FTypeContainer compatible classes.
- */
-#define DECLARE_REGISTRABLE_TYPE(Type) \
-	public: \
-		static FName GetTypeName() \
-		{ \
-			static FName TypeName = TEXT(#Type); \
-			return TypeName; \
-		}
-
-
 /**
  * Enumerates the scopes for instance creation in type containers.
  */
@@ -230,11 +218,9 @@ public:
 	template<class R>
 	TSharedPtr<R> GetInstance()
 	{
-		const FName TypeName = R::GetTypeName();
-
 		FScopeLock Lock(&CriticalSection);
 		{
-			const TSharedPtr<IInstanceProvider>& Provider = Providers.FindRef(TypeName);
+			const TSharedPtr<IInstanceProvider>& Provider = Providers.FindRef(TNameOf<R>::GetName());
 
 			if (Provider.IsValid())
 			{
@@ -308,7 +294,7 @@ public:
 
 			FScopeLock Lock(&CriticalSection);
 			{
-				Providers.Add(R::GetTypeName(), Provider);
+				Providers.Add(TNameOf<R>::GetName(), Provider);
 			}
 		}
 	}
@@ -324,7 +310,7 @@ public:
 		FScopeLock Lock(&CriticalSection);
 		{
 			Providers.Add(
-				R::GetTypeName(),
+				TNameOf<R>::GetName(),
 				MakeShareable(new TDelegateInstanceProvider<D>(Delegate))
 			);
 		}
@@ -343,7 +329,7 @@ public:
 		FScopeLock Lock(&CriticalSection);
 		{
 			Providers.Add(
-				R::GetTypeName(),
+				TNameOf<R>::GetName(),
 				MakeShareable(new TFactoryInstanceProvider<R>(CreateFunc))
 			);
 		}
@@ -365,7 +351,7 @@ public:
 		FScopeLock Lock(&CriticalSection);
 		{
 			Providers.Add(
-				R::GetTypeName(),
+				TNameOf<R>::GetName(),
 				MakeShareable(new TSharedInstanceProvider<R>(Instance))
 			);
 		}
@@ -380,11 +366,9 @@ public:
 	template<class R>
 	void Unregister()
 	{
-		const FName TypeName = R::GetTypeName();
-
 		FScopeLock Lock(&CriticalSection);
 		{
-			Providers.Remove(TypeName);
+			Providers.Remove(TNameOf<R>::GetName());
 		}
 	}
 
@@ -394,7 +378,7 @@ private:
 	FCriticalSection CriticalSection;
 
 	/** Maps class names to instance providers. */
-	TMap<FName, TSharedPtr<IInstanceProvider>> Providers;
+	TMap<const TCHAR*, TSharedPtr<IInstanceProvider>> Providers;
 };
 
 
