Core: Updated TypeContainer documentation

diff --git a/Engine/Source/Runtime/Core/Public/Misc/TypeContainer.h b/Engine/Source/Runtime/Core/Public/Misc/TypeContainer.h
--- a/Engine/Source/Runtime/Core/Public/Misc/TypeContainer.h
+++ b/Engine/Source/Runtime/Core/Public/Misc/TypeContainer.h
@@ -31,8 +31,8 @@ enum class ETypeContainerScope
  * methods, type containers can facilitate the Inversion of Control (IoC) pattern.
  *
  * Since UE4 neither uses run-time type information nor pre-processes plain old C++ classes,
- * a special macro - DECLARE_REGISTRABLE_TYPE - needs to be inserted into a class declaration
- * in order to make it registrable with type containers.
+ * type names need to be exposed using the Expose_TNameOf macro in order to make them
+ * registrable with type containers, i.e. Expose_TNameOf(FMyClass).
  *
  * Once a type is registered with a container, instances of that type can be retrieved from it.
  * There are currently three life time scopes available for instance creation:
@@ -44,10 +44,11 @@ enum class ETypeContainerScope
  *   3. Unique instance per call (aka. class factory),
  *      using RegisterClass(ETypeContainerScope::Instance) or RegisterFactory()
  *
- * Note: Type containers depend on variadic templates and is therefore not available
- * on XboxOne at this time. It should therefore only be used for desktop applications.
+ * See the file TypeContainerTest.cpp for detailed examples on how to use each of these
+ * type registration methods.
  *
- * @todo gmp: better documentation
+ * Note: Type containers depend on variadic templates and are therefore not available
+ * on XboxOne at this time, which means they should only be used for desktop applications.
  */
 class FTypeContainer
 {
