Fix for Xbox One.

diff --git a/Engine/Source/Runtime/Core/Public/Delegates/Delegate.h b/Engine/Source/Runtime/Core/Public/Delegates/Delegate.h
--- a/Engine/Source/Runtime/Core/Public/Delegates/Delegate.h
+++ b/Engine/Source/Runtime/Core/Public/Delegates/Delegate.h
@@ -295,6 +295,25 @@ class DynamicMulticastDelegateName : public FUNC_DECLARE_DYNAMIC_DELEGATE_BASE(T
 #define IsAlreadyBound( UserObject, FuncName ) __Internal_IsAlreadyBound( UserObject, FuncName, TEXT( #FuncName ) )
 
 
+namespace UE4Delegates_Private
+{
+	/**
+	 * Returns the root function name from a string representing a member function pointer.
+	 * Note: this function only returns a pointer to the substring and doesn't create a new string.
+	 *
+	 * @param  InMacroFunctionName  The string containing the member function name.
+	 * @return A pointer to the substring containing the member function name.
+	 */
+	inline const TCHAR* GetTrimmedMemberFunctionName(const TCHAR* InMacroFunctionName)
+	{
+		// We strip off the class prefix and just return the function name by itself.
+		const TCHAR* Result = FCString::Strrstr( InMacroFunctionName, TEXT( "::" ) );
+		checkf(Result, TEXT("'%s' does not look like a member function"), InMacroFunctionName);
+		return Result + 2;
+	}
+}
+
+
 /*********************************************************************************************************************/
 
 #define DELEGATE_DEPRECATED(message) DEPRECATED(4.7, message)
diff --git a/Engine/Source/Runtime/Core/Public/Delegates/DelegateSignatureImpl.inl b/Engine/Source/Runtime/Core/Public/Delegates/DelegateSignatureImpl.inl
--- a/Engine/Source/Runtime/Core/Public/Delegates/DelegateSignatureImpl.inl
+++ b/Engine/Source/Runtime/Core/Public/Delegates/DelegateSignatureImpl.inl
@@ -2076,25 +2076,6 @@ protected:
 };
 
 
-namespace UE4Delegates_Private
-{
-	/**
-	 * Returns the root function name from a string representing a member function pointer.
-	 * Note: this function only returns a pointer to the substring and doesn't create a new string.
-	 *
-	 * @param  InMacroFunctionName  The string containing the member function name.
-	 * @return A pointer to the substring containing the member function name.
-	 */
-	inline const TCHAR* GetTrimmedMemberFunctionName(const TCHAR* InMacroFunctionName)
-	{
-		// We strip off the class prefix and just return the function name by itself.
-		const TCHAR* Result = FCString::Strrstr( InMacroFunctionName, TEXT( "::" ) );
-		checkf(Result, TEXT("'%s' does not look like a member function"), InMacroFunctionName);
-		return Result + 2;
-	}
-}
-
-
 /**
  * Implements a multicast delegate.
  *
diff --git a/Engine/Source/Runtime/Core/Public/Delegates/DelegateSignatureImpl_Variadics.inl b/Engine/Source/Runtime/Core/Public/Delegates/DelegateSignatureImpl_Variadics.inl
--- a/Engine/Source/Runtime/Core/Public/Delegates/DelegateSignatureImpl_Variadics.inl
+++ b/Engine/Source/Runtime/Core/Public/Delegates/DelegateSignatureImpl_Variadics.inl
@@ -1115,25 +1115,6 @@ protected:
 };
 
 
-namespace UE4Delegates_Private
-{
-	/**
-	 * Returns the root function name from a string representing a member function pointer.
-	 * Note: this function only returns a pointer to the substring and doesn't create a new string.
-	 *
-	 * @param  InMacroFunctionName  The string containing the member function name.
-	 * @return A pointer to the substring containing the member function name.
-	 */
-	inline const TCHAR* GetTrimmedMemberFunctionName(const TCHAR* InMacroFunctionName)
-	{
-		// We strip off the class prefix and just return the function name by itself.
-		const TCHAR* Result = FCString::Strrstr( InMacroFunctionName, TEXT( "::" ) );
-		checkf(Result, TEXT("'%s' does not look like a member function"), InMacroFunctionName);
-		return Result + 2;
-	}
-}
-
-
 /**
  * Implements a multicast delegate.
  *
