Core: Added support for variadic template parameters in TypeContainer for TFunction and Delegate registrations.

The syntax is not ideal yet as explicit specification of template parameters is still required. Automatic deduction of these will require much more meta-programming, but is probably doable. Also, the delegate registration with variadic parameters currently requires the delegate type to be specified as well. Other than that, high fives all around!

diff --git a/Engine/Source/Runtime/Core/Private/Tests/Misc/TypeContainerTest.cpp b/Engine/Source/Runtime/Core/Private/Tests/Misc/TypeContainerTest.cpp
--- a/Engine/Source/Runtime/Core/Private/Tests/Misc/TypeContainerTest.cpp
+++ b/Engine/Source/Runtime/Core/Private/Tests/Misc/TypeContainerTest.cpp
@@ -63,7 +63,10 @@ struct FTwoSmoothies
 	TSharedPtr<ISmoothie> Two;
 };
 
+
+DECLARE_DELEGATE_RetVal(TSharedPtr<IBerry>, FBerryFactoryDelegate);
 DECLARE_DELEGATE_RetVal(TSharedPtr<IFruit>, FFruitFactoryDelegate);
+DECLARE_DELEGATE_RetVal_TwoParams(TSharedPtr<ISmoothie>, FSmoothieFactoryDelegate, TSharedRef<IFruit>, TSharedRef<IBerry>);
 
 Expose_TNameOf(IBerry)
 Expose_TNameOf(IFruit)
@@ -194,6 +197,11 @@ bool FTypeContainerTest::RunTest(const FString& Parameters)
 			{
 				return MakeShareable(new FStrawberry());
 			}
+
+			static TSharedPtr<ISmoothie> MakeSmoothie(TSharedRef<IFruit> Fruit, TSharedRef<IBerry> Berry)
+			{
+				return MakeShareable(new FSmoothie(Fruit, Berry));
+			}
 		};
 
 		FTypeContainer Container;
@@ -204,19 +212,27 @@ bool FTypeContainerTest::RunTest(const FString& Parameters)
 					return MakeShareable(new FBanana());
 				}
 			);
+			Container.RegisterFactory<ISmoothie, IFruit, IBerry>(&FLocal::MakeSmoothie);
 		}
 
 		auto Berry = Container.GetInstance<IBerry>();
 		auto Fruit = Container.GetInstance<IFruit>();
+		auto Smoothie = Container.GetInstance<ISmoothie>();
 
 		TestTrue(TEXT("For static factory functions, an instance must be returned"), Berry.IsValid());
 		TestTrue(TEXT("For lambda factory functions, an instance must be returned"), Fruit.IsValid());
+		TestTrue(TEXT("For factory functions with dependencies, an instance must be returned"), Smoothie.IsValid());
 	}
 
 	// delegate test
 	{
 		struct FLocal
 		{
+			static TSharedPtr<IBerry> MakeBerry()
+			{
+				return MakeShareable(new FStrawberry());
+			}
+
 			static TSharedPtr<IFruit> MakeFruit(bool Banana)
 			{
 				if (Banana)
@@ -226,16 +242,25 @@ bool FTypeContainerTest::RunTest(const FString& Parameters)
 
 				return MakeShareable(new FStrawberry());
 			}
+
+			static TSharedPtr<ISmoothie> MakeSmoothie(TSharedRef<IFruit> Fruit, TSharedRef<IBerry> Berry)
+			{
+				return MakeShareable(new FSmoothie(Fruit, Berry));
+			}
 		};
 
 		FTypeContainer Container;
 		{
+			Container.RegisterDelegate<IBerry>(FBerryFactoryDelegate::CreateStatic(&FLocal::MakeBerry));
 			Container.RegisterDelegate<IFruit>(FFruitFactoryDelegate::CreateStatic(&FLocal::MakeFruit, true));
+			Container.RegisterDelegate<ISmoothie, FSmoothieFactoryDelegate, IFruit, IBerry>(FSmoothieFactoryDelegate::CreateStatic(&FLocal::MakeSmoothie));
 		}
 
 		auto Fruit = Container.GetInstance<IFruit>();
+		auto Smoothie = Container.GetInstance<ISmoothie>();
 
 		TestTrue(TEXT("For factory delegates, an instance must be returned"), Fruit.IsValid());
+		TestTrue(TEXT("For factory delegates with dependencies, an instance must be returned"), Smoothie.IsValid());
 	}
 
 	return true;
diff --git a/Engine/Source/Runtime/Core/Public/Misc/TypeContainer.h b/Engine/Source/Runtime/Core/Public/Misc/TypeContainer.h
--- a/Engine/Source/Runtime/Core/Public/Misc/TypeContainer.h
+++ b/Engine/Source/Runtime/Core/Public/Misc/TypeContainer.h
@@ -68,38 +68,17 @@ class FTypeContainer
 		virtual TSharedPtr<void> GetInstance() = 0;
 	};
 
-	/** Implements an instance provider that forwards instance requests to a factory delegate. */
-	template<class D>
-	struct TDelegateInstanceProvider
-		: public IInstanceProvider
-	{
-		/** Constructor. */
-		TDelegateInstanceProvider(D InDelegate)
-			: Delegate(InDelegate)
-		{ }
-
-		virtual ~TDelegateInstanceProvider() override { }
-
-		virtual TSharedPtr<void> GetInstance() override
-		{
-			return Delegate.Execute();
-		}
-
-		/** Factory function that creates the instances. */
-		D Delegate;
-	};
-
 	/** Implements an instance provider that forwards instance requests to a factory function. */
 	template<class T>
-	struct TFactoryInstanceProvider
+	struct TFunctionInstanceProvider
 		: public IInstanceProvider
 	{
 		/** Constructor. */
-		TFactoryInstanceProvider(TFunction<TSharedPtr<T>()> InCreateFunc)
+		TFunctionInstanceProvider(TFunction<TSharedPtr<void>()> InCreateFunc)
 			: CreateFunc(InCreateFunc)
 		{ }
 
-		virtual ~TFactoryInstanceProvider() override { }
+		virtual ~TFunctionInstanceProvider() override { }
 
 		virtual TSharedPtr<void> GetInstance() override
 		{
@@ -107,7 +86,7 @@ class FTypeContainer
 		}
 
 		/** Factory function that creates the instances. */
-		TFunction<TSharedPtr<T>()> CreateFunc;
+		TFunction<TSharedPtr<void>()> CreateFunc;
 	};
 
 
@@ -262,78 +241,114 @@ public:
 	{
 		static_assert(TPointerIsConvertibleFromTo<T, R>::Value, "T must implement R");
 
-		if (Scope == ETypeContainerScope::Instance)
+		TSharedPtr<IInstanceProvider> Provider;
+		
+		switch(Scope)
 		{
-			RegisterFactory<R>(
-				[this]() -> TSharedPtr<R> {
-					return MakeShareable(new T(GetInstanceRef<P>()...));
-				}
+		case ETypeContainerScope::Instance:
+			Provider = MakeShareable(
+				new TFunctionInstanceProvider<T>(
+					[this]() -> TSharedPtr<void> {
+						return MakeShareable(new T(GetInstanceRef<P>()...));
+					}
+				)
 			);
-		}
-		else
-		{
-			TSharedPtr<IInstanceProvider> Provider;
-			
-			if (Scope == ETypeContainerScope::Thread)
-			{
-				Provider = MakeShareable(
-					new TThreadInstanceProvider<T>(
-						[this]() -> TSharedPtr<void> {
-							return MakeShareable(new T(GetInstanceRef<P>()...));
-						}
-					)
-				);
-			}
-			else
-			{
-				Provider = MakeShareable(
-					new TSharedInstanceProvider<T>(
-						MakeShareable(new T(GetInstanceRef<P>()...))
-					)
-				);
-			}
+			break;
 
-			FScopeLock Lock(&CriticalSection);
-			{
-				Providers.Add(TNameOf<R>::GetName(), Provider);
-			}
+		case ETypeContainerScope::Thread:
+			Provider = MakeShareable(
+				new TThreadInstanceProvider<T>(
+					[this]() -> TSharedPtr<void> {
+						return MakeShareable(new T(GetInstanceRef<P>()...));
+					}
+				)
+			);
+			break;
+
+		default:
+			Provider = MakeShareable(
+				new TSharedInstanceProvider<T>(
+					MakeShareable(new T(GetInstanceRef<P>()...))
+				)
+			);
+			break;
 		}
+
+		AddProvider(TNameOf<R>::GetName(), Provider);
 	}
 
 	/**
 	 * Register a factory delegate for the specified class.
+	 *
+	 * @param R The type of class to register an instance class for.
+	 * @param D Tye type of factory delegate to register.
+	 * @param P The parameters that the delegate requires.
+	 * @param Delegate The delegate to register.
 	 */
-	template<class R, class D>
+	template<class R, class D, typename... P>
 	void RegisterDelegate(D Delegate)
 	{
 		static_assert(TAreTypesEqual<TSharedPtr<R>, typename D::RetValType>::Value, "Delegate return type must be TSharedPtr<R>");
 
-		FScopeLock Lock(&CriticalSection);
-		{
-			Providers.Add(
-				TNameOf<R>::GetName(),
-				MakeShareable(new TDelegateInstanceProvider<D>(Delegate))
-			);
-		}
+		TSharedPtr<IInstanceProvider> Provider = MakeShareable(
+			new TFunctionInstanceProvider<R>(
+				[=]() -> TSharedPtr<void> {
+					return Delegate.Execute(GetInstanceRef<P>()...);
+				}
+			)
+		);
+
+		AddProvider(TNameOf<R>::GetName(), Provider);
 	}
 
 	/**
 	 * Register a factory function for the specified class.
 	 *
+	 * Note: VS2013 did not like a purely variadic template function, so we split
+	 * these up into two overloads, one without and one with variadic parameters.
+	 *
 	 * @param R The type of class to register the instance for.
-	 * @param CreateFunc A factory function that will create instances.
+	 * @param CreateFunc The factory function that will create instances.
 	 * @see RegisterClass, RegisterInstance, Unregister
 	 */
 	template<class R>
 	void RegisterFactory(TFunction<TSharedPtr<R>()> CreateFunc)
 	{
-		FScopeLock Lock(&CriticalSection);
-		{
-			Providers.Add(
-				TNameOf<R>::GetName(),
-				MakeShareable(new TFactoryInstanceProvider<R>(CreateFunc))
-			);
-		}
+		TSharedPtr<IInstanceProvider> Provider = MakeShareable(
+			new TFunctionInstanceProvider<R>(
+				[=]() -> TSharedPtr<void> {
+					return CreateFunc();
+				}
+			)
+		);
+
+		AddProvider(TNameOf<R>::GetName(), Provider);
+	}
+
+	/**
+	 * Register a factory function for the specified class.
+	 *
+	 * Note: VS2013 did not like a purely variadic template function, so we split
+	 * these up into two overloads, one without and one with variadic parameters.
+	 *
+	 * @param R The type of class to register the instance for.
+	 * @param P0 The first parameter that the factory function requires.
+	 * @param P Additional parameters that the factory function requires.
+	 * @param CreateFunc The factory function that will create instances.
+	 * @see RegisterClass, RegisterInstance, Unregister
+	 */
+	template<class R, typename P0, typename... P>
+	void RegisterFactory(TFunction<TSharedPtr<R>(TSharedRef<P0>, TSharedRef<P>...)> CreateFunc)
+	{
+		TSharedPtr<IInstanceProvider> Provider = MakeShareable(
+			new TFunctionInstanceProvider<R>(
+				[=]() -> TSharedPtr<void> {
+					return CreateFunc(GetInstanceRef<P0>(), GetInstanceRef<P>()...);
+				}
+			)
+		);
+
+		AddProvider(TNameOf<R>::GetName(), Provider);
 	}
 
 	/**
@@ -349,13 +364,7 @@ public:
 	{
 		static_assert(TPointerIsConvertibleFromTo<T, R>::Value, "T must implement R");
 
-		FScopeLock Lock(&CriticalSection);
-		{
-			Providers.Add(
-				TNameOf<R>::GetName(),
-				MakeShareable(new TSharedInstanceProvider<R>(Instance))
-			);
-		}
+		AddProvider(TNameOf<R>::GetName(), MakeShareable(new TSharedInstanceProvider<R>(Instance)));
 	}
 
 	/**
@@ -373,6 +382,22 @@ public:
 		}
 	}
 
+protected:
+
+	/**
+	 * Add an instance provider to the container.
+	 *
+	 * @param Name The name of the type to add the provider for.
+	 * @param Provider The instance provider.
+	 */
+	void AddProvider(const TCHAR* Name, const TSharedPtr<IInstanceProvider>& Provider)
+	{
+		FScopeLock Lock(&CriticalSection);
+		{
+			Providers.Add(Name, Provider);
+		}
+	}
+
 private:
 
 	/** Critical section for synchronizing access. */
