Core: TypeContainer now always returns shared references, or asserts if no instance was registered for the requested type.

#UpgradeNotes: This is an API breaking change. Your factory functions/methods/delegates/lambdas now have to return TSharedRef instead of TSharedPtr. Requesting an instance for an unregistered type now constitutes a critical run-time error.

diff --git a/Engine/Source/Runtime/Core/Private/Tests/Misc/TypeContainerTest.cpp b/Engine/Source/Runtime/Core/Private/Tests/Misc/TypeContainerTest.cpp
--- a/Engine/Source/Runtime/Core/Private/Tests/Misc/TypeContainerTest.cpp
+++ b/Engine/Source/Runtime/Core/Private/Tests/Misc/TypeContainerTest.cpp
@@ -64,9 +64,9 @@ struct FTwoSmoothies
 };
 
 
-DECLARE_DELEGATE_RetVal(TSharedPtr<IBerry>, FBerryFactoryDelegate);
-DECLARE_DELEGATE_RetVal(TSharedPtr<IFruit>, FFruitFactoryDelegate);
-DECLARE_DELEGATE_RetVal_TwoParams(TSharedPtr<ISmoothie>, FSmoothieFactoryDelegate, TSharedRef<IFruit>, TSharedRef<IBerry>);
+DECLARE_DELEGATE_RetVal(TSharedRef<IBerry>, FBerryFactoryDelegate);
+DECLARE_DELEGATE_RetVal(TSharedRef<IFruit>, FFruitFactoryDelegate);
+DECLARE_DELEGATE_RetVal_TwoParams(TSharedRef<ISmoothie>, FSmoothieFactoryDelegate, TSharedRef<IFruit>, TSharedRef<IBerry>);
 
 Expose_TNameOf(IBerry)
 Expose_TNameOf(IFruit)
@@ -91,13 +91,8 @@ bool FTypeContainerTest::RunTest(const FString& Parameters)
 
 		auto Instance = Container.GetInstance<IFruit>();
 
-		TestTrue(TEXT("Existing registered instance must be returned"), Instance.IsValid());
-
-		if (Instance.IsValid())
-		{
-			TestEqual(TEXT("Correct instance must be returned"), Instance.ToSharedRef(), Fruit);
-			TestNotEqual(TEXT("Incorrect instance must not be returned"), Instance.ToSharedRef(), StaticCastSharedRef<IFruit>(Berry));
-		}
+		TestEqual(TEXT("Correct instance must be returned"), Instance, Fruit);
+		TestNotEqual(TEXT("Incorrect instance must not be returned"), Instance, StaticCastSharedRef<IFruit>(Berry));
 	}
 
 	// per instance test
@@ -112,8 +107,6 @@ bool FTypeContainerTest::RunTest(const FString& Parameters)
 		auto Smoothie1 = Container.GetInstance<ISmoothie>();
 		auto Smoothie2 = Container.GetInstance<ISmoothie>();
 
-		TestTrue(TEXT("For per-instance classes, an instance must be returned [1]"), Smoothie1.IsValid());
-		TestTrue(TEXT("For per-instance classes, an instance must be returned [2]"), Smoothie2.IsValid());
 		TestNotEqual(TEXT("For per-instances classes, a unique instance must be returned each time"), Smoothie1, Smoothie2);
 		TestNotEqual(TEXT("For per-instances dependencies, a unique instance must be returned each time [1]"), Smoothie1->GetBerry(), Smoothie2->GetBerry());
 		TestNotEqual(TEXT("For per-instances dependencies, a unique instance must be returned each time [2]"), Smoothie1->GetFruit(), Smoothie2->GetFruit());
@@ -144,10 +137,6 @@ bool FTypeContainerTest::RunTest(const FString& Parameters)
 		TSharedPtr<ISmoothie> One2 = Smoothies2.Get().One;
 		TSharedPtr<ISmoothie> Two2 = Smoothies2.Get().Two;
 
-		TestTrue(TEXT("For per-thread classes, an instance must be returned [1]"), One1.IsValid());
-		TestTrue(TEXT("For per-thread classes, an instance must be returned [2]"), Two1.IsValid());
-		TestTrue(TEXT("For per-thread classes, an instance must be returned [3]"), One2.IsValid());
-		TestTrue(TEXT("For per-thread classes, an instance must be returned [4]"), Two2.IsValid());
 		TestEqual(TEXT("For per-thread classes, the same instance must be returned from the same thread [1]"), One1, Two1);
 		TestEqual(TEXT("For per-thread classes, the same instance must be returned from the same thread [2]"), One2, Two2);
 		TestNotEqual(TEXT("For per-thread classes, different instances must be returned from different threads [1]"), One1, One2);
@@ -179,10 +168,6 @@ bool FTypeContainerTest::RunTest(const FString& Parameters)
 		TSharedPtr<ISmoothie> One2 = Smoothies2.Get().One;
 		TSharedPtr<ISmoothie> Two2 = Smoothies2.Get().Two;
 
-		TestTrue(TEXT("For per-process classes, an instance must be returned [1]"), One1.IsValid());
-		TestTrue(TEXT("For per-process classes, an instance must be returned [2]"), Two1.IsValid());
-		TestTrue(TEXT("For per-process classes, an instance must be returned [3]"), One2.IsValid());
-		TestTrue(TEXT("For per-process classes, an instance must be returned [4]"), Two2.IsValid());
 		TestEqual(TEXT("For per-process classes, the same instance must be returned from the same thread [1]"), One1, Two1);
 		TestEqual(TEXT("For per-process classes, the same instance must be returned from the same thread [2]"), One2, Two2);
 		TestEqual(TEXT("For per-process classes, the same instance must be returned from different threads [1]"), One1, One2);
@@ -193,12 +178,12 @@ bool FTypeContainerTest::RunTest(const FString& Parameters)
 	{
 		struct FLocal
 		{
-			static TSharedPtr<IBerry> MakeStrawberry()
+			static TSharedRef<IBerry> MakeStrawberry()
 			{
 				return MakeShareable(new FStrawberry());
 			}
 
-			static TSharedPtr<ISmoothie> MakeSmoothie(TSharedRef<IFruit> Fruit, TSharedRef<IBerry> Berry)
+			static TSharedRef<ISmoothie> MakeSmoothie(TSharedRef<IFruit> Fruit, TSharedRef<IBerry> Berry)
 			{
 				return MakeShareable(new FSmoothie(Fruit, Berry));
 			}
@@ -208,7 +193,7 @@ bool FTypeContainerTest::RunTest(const FString& Parameters)
 		{
 			Container.RegisterFactory<IBerry>(&FLocal::MakeStrawberry);
 			Container.RegisterFactory<IFruit>(
-				[]() -> TSharedPtr<IFruit> {
+				[]() -> TSharedRef<IFruit> {
 					return MakeShareable(new FBanana());
 				}
 			);
@@ -218,22 +203,18 @@ bool FTypeContainerTest::RunTest(const FString& Parameters)
 		auto Berry = Container.GetInstance<IBerry>();
 		auto Fruit = Container.GetInstance<IFruit>();
 		auto Smoothie = Container.GetInstance<ISmoothie>();
-
-		TestTrue(TEXT("For static factory functions, an instance must be returned"), Berry.IsValid());
-		TestTrue(TEXT("For lambda factory functions, an instance must be returned"), Fruit.IsValid());
-		TestTrue(TEXT("For factory functions with dependencies, an instance must be returned"), Smoothie.IsValid());
 	}
 
 	// delegate test
 	{
 		struct FLocal
 		{
-			static TSharedPtr<IBerry> MakeBerry()
+			static TSharedRef<IBerry> MakeBerry()
 			{
 				return MakeShareable(new FStrawberry());
 			}
 
-			static TSharedPtr<IFruit> MakeFruit(bool Banana)
+			static TSharedRef<IFruit> MakeFruit(bool Banana)
 			{
 				if (Banana)
 				{
@@ -243,7 +224,7 @@ bool FTypeContainerTest::RunTest(const FString& Parameters)
 				return MakeShareable(new FStrawberry());
 			}
 
-			static TSharedPtr<ISmoothie> MakeSmoothie(TSharedRef<IFruit> Fruit, TSharedRef<IBerry> Berry)
+			static TSharedRef<ISmoothie> MakeSmoothie(TSharedRef<IFruit> Fruit, TSharedRef<IBerry> Berry)
 			{
 				return MakeShareable(new FSmoothie(Fruit, Berry));
 			}
@@ -258,9 +239,6 @@ bool FTypeContainerTest::RunTest(const FString& Parameters)
 
 		auto Fruit = Container.GetInstance<IFruit>();
 		auto Smoothie = Container.GetInstance<ISmoothie>();
-
-		TestTrue(TEXT("For factory delegates, an instance must be returned"), Fruit.IsValid());
-		TestTrue(TEXT("For factory delegates with dependencies, an instance must be returned"), Smoothie.IsValid());
 	}
 
 	return true;
diff --git a/Engine/Source/Runtime/Core/Public/Misc/TypeContainer.h b/Engine/Source/Runtime/Core/Public/Misc/TypeContainer.h
--- a/Engine/Source/Runtime/Core/Public/Misc/TypeContainer.h
+++ b/Engine/Source/Runtime/Core/Public/Misc/TypeContainer.h
@@ -192,39 +192,19 @@ public:
 	 * Gets a shared pointer to an instance of the specified class.
 	 *
 	 * @param T The type of class to get an instance for.
-	 * @param A shared pointer to the instance, or nullptr if no instance was registered.
-	 * @see GetInstanceRef, RegisterClass, RegisterInstance
+	 * @param A shared reference to the instance.
+	 * @see RegisterClass, RegisterDelegate, RegisterFactory, RegisterInstance
 	 */
 	template<class R>
-	TSharedPtr<R> GetInstance()
+	TSharedRef<R> GetInstance()
 	{
 		FScopeLock Lock(&CriticalSection);
 		{
 			const TSharedPtr<IInstanceProvider>& Provider = Providers.FindRef(TNameOf<R>::GetName());
+			check(Provider.IsValid());
 
-			if (Provider.IsValid())
-			{
-				return StaticCastSharedPtr<R>(Provider->GetInstance());
-			}
+			return StaticCastSharedPtr<R>(Provider->GetInstance()).ToSharedRef();
 		}
-
-		return TSharedPtr<R>();
-	}
-
-	/**
-	 * Gets a shared reference to an instance of the specified class.
-	 *
-	 * Unlike GetInstance(), this function will assert if no instance was registered for
-	 * the requested type of class using either RegisterClass() or RegisterInstance().
-	 *
-	 * @param R The type of class that an instance is being requested for.
-	 * @param A shared pointer to the instance.
-	 * @see GetInstance, RegisterClass, RegisterInstance
-	 */
-	template<class R>
-	TSharedRef<R> GetInstanceRef()
-	{
-		return GetInstance<R>().ToSharedRef();
 	}
 
 	/**
@@ -249,7 +229,7 @@ public:
 			Provider = MakeShareable(
 				new TFunctionInstanceProvider<T>(
 					[this]() -> TSharedPtr<void> {
-						return MakeShareable(new T(GetInstanceRef<P>()...));
+						return MakeShareable(new T(GetInstance<P>()...));
 					}
 				)
 			);
@@ -259,7 +239,7 @@ public:
 			Provider = MakeShareable(
 				new TThreadInstanceProvider<T>(
 					[this]() -> TSharedPtr<void> {
-						return MakeShareable(new T(GetInstanceRef<P>()...));
+						return MakeShareable(new T(GetInstance<P>()...));
 					}
 				)
 			);
@@ -268,7 +248,7 @@ public:
 		default:
 			Provider = MakeShareable(
 				new TSharedInstanceProvider<T>(
-					MakeShareable(new T(GetInstanceRef<P>()...))
+					MakeShareable(new T(GetInstance<P>()...))
 				)
 			);
 			break;
@@ -288,12 +268,12 @@ public:
 	template<class R, class D, typename... P>
 	void RegisterDelegate(D Delegate)
 	{
-		static_assert(TAreTypesEqual<TSharedPtr<R>, typename D::RetValType>::Value, "Delegate return type must be TSharedPtr<R>");
+		static_assert(TAreTypesEqual<TSharedRef<R>, typename D::RetValType>::Value, "Delegate return type must be TSharedPtr<R>");
 
 		TSharedPtr<IInstanceProvider> Provider = MakeShareable(
 			new TFunctionInstanceProvider<R>(
 				[=]() -> TSharedPtr<void> {
-					return Delegate.Execute(GetInstanceRef<P>()...);
+					return Delegate.Execute(GetInstance<P>()...);
 				}
 			)
 		);
@@ -312,7 +292,7 @@ public:
 	 * @see RegisterClass, RegisterInstance, Unregister
 	 */
 	template<class R>
-	void RegisterFactory(TFunction<TSharedPtr<R>()> CreateFunc)
+	void RegisterFactory(TFunction<TSharedRef<R>()> CreateFunc)
 	{
 		TSharedPtr<IInstanceProvider> Provider = MakeShareable(
 			new TFunctionInstanceProvider<R>(
@@ -338,12 +318,12 @@ public:
 	 * @see RegisterClass, RegisterInstance, Unregister
 	 */
 	template<class R, typename P0, typename... P>
-	void RegisterFactory(TFunction<TSharedPtr<R>(TSharedRef<P0>, TSharedRef<P>...)> CreateFunc)
+	void RegisterFactory(TFunction<TSharedRef<R>(TSharedRef<P0>, TSharedRef<P>...)> CreateFunc)
 	{
 		TSharedPtr<IInstanceProvider> Provider = MakeShareable(
 			new TFunctionInstanceProvider<R>(
 				[=]() -> TSharedPtr<void> {
-					return CreateFunc(GetInstanceRef<P0>(), GetInstanceRef<P>()...);
+					return CreateFunc(GetInstance<P0>(), GetInstance<P>()...);
 				}
 			)
 		);
