Merged GitHub PR# 907 (fix for TComPtr) by OptiPlanet

diff --git a/Engine/Source/Runtime/Core/Public/Windows/COMPointer.h b/Engine/Source/Runtime/Core/Public/Windows/COMPointer.h
--- a/Engine/Source/Runtime/Core/Public/Windows/COMPointer.h
+++ b/Engine/Source/Runtime/Core/Public/Windows/COMPointer.h
@@ -41,15 +41,15 @@ public:
 	{
 		if(RawPointer != Source)
 		{
-			if(RawPointer)
+			if (Source)
 			{
-				RawPointer->Release();
+				Source->AddRef();
 			}
-			RawPointer = Source;
-			if(RawPointer)
+			if (RawPointer)
 			{
-				RawPointer->AddRef();
+				RawPointer->Release();
 			}
+			RawPointer = Source;
 		}
 		return *this;
 	}
@@ -58,15 +58,15 @@ public:
 	{
 		if(RawPointer != Source.RawPointer)
 		{
-			if(RawPointer)
+			if (Source.RawPointer)
 			{
-				RawPointer->Release();
+				Source->AddRef();
 			}
-			RawPointer = Source.RawPointer;
-			if(RawPointer)
+			if (RawPointer)
 			{
-				RawPointer->AddRef();
+				RawPointer->Release();
 			}
+			RawPointer = Source.RawPointer;
 		}
 		return *this;
 	}
