Fix for dragging objects in the viewport.

- PR #889 contributed by [REDACTED_PERSON].

https://github.com/EpicGames/UnrealEngine/pull/889

diff --git a/Engine/Source/Runtime/Core/Private/Linux/LinuxApplication.cpp b/Engine/Source/Runtime/Core/Private/Linux/LinuxApplication.cpp
--- a/Engine/Source/Runtime/Core/Private/Linux/LinuxApplication.cpp
+++ b/Engine/Source/Runtime/Core/Private/Linux/LinuxApplication.cpp
@@ -68,6 +68,7 @@ FLinuxApplication::FLinuxApplication()
 	,	bInsideOwnWindow(false)
 	,	bIsDragWindowButtonPressed(false)
 	,	bActivateApp(false)
+	,	bLockToCurrentMouseType(false)
 {
 	bUsingHighPrecisionMouseInput = false;
 	bAllowedToDeferMessageProcessing = true;
@@ -245,19 +246,23 @@ void FLinuxApplication::ProcessDeferredMessage( SDL_Event Event )
 			SDL_MouseMotionEvent motionEvent = Event.motion;
 			FLinuxCursor *LinuxCursor = (FLinuxCursor*)Cursor.Get();
 
-			if(SDL_ShowCursor(-1) == 0)
+			if (LinuxCursor->IsHidden())
 			{
-				int width, height;
-				SDL_GetWindowSize( NativeWindow, &width, &height );
-				if( motionEvent.x != (width / 2) || motionEvent.y != (height / 2) )
+				// Check if the mouse got locked for dragging in viewport.
+				if (bLockToCurrentMouseType == false)
 				{
-					int xOffset, yOffset;
-					SDL_GetWindowPosition( NativeWindow, &xOffset, &yOffset );
-					LinuxCursor->SetPosition( width / 2 + xOffset, height / 2 + yOffset );
-				}
-				else
-				{
-					break;
+					int width, height;
+					SDL_GetWindowSize(NativeWindow, &width, &height);
+					if (motionEvent.x != (width / 2) || motionEvent.y != (height / 2))
+					{
+						int xOffset, yOffset;
+						SDL_GetWindowPosition(NativeWindow, &xOffset, &yOffset);
+						LinuxCursor->SetPosition(width / 2 + xOffset, height / 2 + yOffset);
+					}
+					else
+					{
+						break;
+					}
 				}
 			}
 			else
@@ -281,12 +286,7 @@ void FLinuxApplication::ProcessDeferredMessage( SDL_Event Event )
 
 			if(bUsingHighPrecisionMouseInput)
 			{
-					// maintain "shadow" global position
-					if (LinuxCursor->IsHidden())
-					{
-						LinuxCursor->AddOffset(motionEvent.xrel, motionEvent.yrel);
-					}
- 					MessageHandler->OnRawMouseMove(motionEvent.xrel, motionEvent.yrel);
+ 				MessageHandler->OnRawMouseMove(motionEvent.xrel, motionEvent.yrel);
 			}
 			else
 			{
@@ -326,11 +326,32 @@ void FLinuxApplication::ProcessDeferredMessage( SDL_Event Event )
 			if (buttonEvent.type == SDL_MOUSEBUTTONUP)
 			{
 				MessageHandler->OnMouseUp(button);
-				bIsDragWindowButtonPressed = false;
+				
+				if (buttonEvent.button == SDL_BUTTON_LEFT)
+				{
+					// Unlock the mouse dragging type.
+					bLockToCurrentMouseType = false;
+
+					bIsDragWindowButtonPressed = false;
+				}
 			}
 			else
 			{
-				bIsDragWindowButtonPressed = true;
+				if (buttonEvent.button == SDL_BUTTON_LEFT)
+				{
+					// The user clicked an object and wants to drag maybe. We can use that to disable 
+					// the resetting of the cursor. Before the user can drag objects, the pointer will change.
+					// Usually it will be EMouseCursor::CardinalCross (Default added after IRC discussion how to fix selection in Front/Top/Side views). 
+                    // If that happends and the user clicks the left mouse button, we know they want to move something.
+					// TODO Is this always true? Need more checks.
+					if (((FLinuxCursor*)Cursor.Get())->GetType() == EMouseCursor::CardinalCross || ((FLinuxCursor*)Cursor.Get())->GetType() == EMouseCursor::Default)
+					{
+						bLockToCurrentMouseType = true;
+					}
+
+					bIsDragWindowButtonPressed = true;
+				}
+
 				if (buttonEvent.clicks == 2)
 				{
 					MessageHandler->OnMouseDoubleClick(CurrentEventWindow, button);
@@ -1075,7 +1096,6 @@ void FLinuxApplication::SetHighPrecisionMouseMode( const bool Enable, const TSha
 {
 	MessageHandler->OnCursorSet();
 	bUsingHighPrecisionMouseInput = Enable;
-	((FLinuxCursor*)Cursor.Get())->ResetOffset();
 }
 
 
diff --git a/Engine/Source/Runtime/Core/Public/Linux/LinuxApplication.h b/Engine/Source/Runtime/Core/Public/Linux/LinuxApplication.h
--- a/Engine/Source/Runtime/Core/Public/Linux/LinuxApplication.h
+++ b/Engine/Source/Runtime/Core/Public/Linux/LinuxApplication.h
@@ -4,6 +4,7 @@
 
 #include "GenericApplication.h"
 #include "LinuxWindow.h"
+#include "LinuxCursor.h"
 
 typedef SDL_GameController* SDL_HController;
 
@@ -144,6 +145,10 @@ private:
 
 	/** Used to check if the application is active or not. */
 	bool bActivateApp;
+
+	/** Used to check with cursor type is current and set to true if left button is pressed.*/
+	bool bLockToCurrentMouseType;
+
 };
 
 
