Treat type truncation warnings like errors.

diff --git a/Engine/Source/Runtime/Core/Public/Windows/WindowsPlatformCompilerSetup.h b/Engine/Source/Runtime/Core/Public/Windows/WindowsPlatformCompilerSetup.h
--- a/Engine/Source/Runtime/Core/Public/Windows/WindowsPlatformCompilerSetup.h
+++ b/Engine/Source/Runtime/Core/Public/Windows/WindowsPlatformCompilerSetup.h
@@ -125,7 +125,6 @@ static_assert(_MSC_VER >= 1700, "Visual Studio 2012 or newer required.");
 #pragma warning(disable : 4267) // 'argument' : conversion from 'type1' to 'type2', possible loss of data
 #pragma warning(default : 4287) // 'operator' : unsigned/negative constant mismatch
 #pragma warning(default : 4289) // nonstandard extension used : 'var' : loop control variable declared in the for-loop is used outside the for-loop scope
-#pragma warning(default : 4302) // 'conversion' : truncation from 'type 1' to 'type 2'
 #pragma warning(disable : 4315) // 'this' pointer for member may not be aligned 8 as expected by the constructor
 #pragma warning(disable : 4316) // 'type' : object allocated on the heap may not be aligned 16
 //#pragma warning(disable : 4339) // 'type' : use of undefined type detected in CLR meta-data - use of this type may lead to a runtime exception
