Log debug message before assertion macros.

diff --git a/Engine/Source/Runtime/Core/Private/Misc/OutputDevice.cpp b/Engine/Source/Runtime/Core/Private/Misc/OutputDevice.cpp
--- a/Engine/Source/Runtime/Core/Private/Misc/OutputDevice.cpp
+++ b/Engine/Source/Runtime/Core/Private/Misc/OutputDevice.cpp
@@ -172,8 +172,7 @@ void StaticFailDebug( const TCHAR* Error, const ANSICHAR* File, int32 Line, cons
 	}
 
 	FScopeLock Lock( &FailDebugCriticalSection );
-
-	FPlatformMisc::LowLevelOutputDebugStringf( TEXT( "%s" ) FILE_LINE_DESC TEXT("\n%s\n" ), Error, ANSI_TO_TCHAR( File ), Line, Description );
+	FPlatformMisc::LowLevelOutputDebugStringf(TEXT("%s") FILE_LINE_DESC TEXT("\n%s\n"), Error, ANSI_TO_TCHAR(File), Line, Description);
 
 	// Copy the detailed error into the error message.
 	FCString::Snprintf( GErrorMessage, ARRAY_COUNT( GErrorMessage ), TEXT( "%s" ) FILE_LINE_DESC TEXT( "\n%s\n" ), Error, ANSI_TO_TCHAR( File ), Line, Description );
@@ -181,13 +180,6 @@ void StaticFailDebug( const TCHAR* Error, const ANSICHAR* File, int32 Line, cons
 	// Copy the error message to the error history.
 	FCString::Strncpy( GErrorHist, GErrorMessage, ARRAY_COUNT( GErrorHist ) );
 	FCString::Strncat( GErrorHist, TEXT( "\r\n\r\n" ), ARRAY_COUNT( GErrorHist ) );
-
-	if( !FPlatformMisc::IsDebuggerPresent() )
-	{
-		FPlatformMisc::PromptForRemoteDebugging( false );
-	}
-
-	FPlatformMisc::DebugBreak();
 }
 
 #if DO_CHECK || DO_GUARD_SLOW
@@ -221,7 +213,6 @@ void VARARGS FDebug::AssertFailed( const ANSICHAR* Expr, const ANSICHAR* File, i
 		FCString::Sprintf( ErrorString, TEXT( "Assertion failed: %s" ), ANSI_TO_TCHAR( Expr ) );
 
 		StaticFailDebug( ErrorString, File, Line, DescriptionString );
-		GError->Logf( TEXT( "Assertion failed: %s" ) FILE_LINE_DESC TEXT( "\n%s\n" ), ErrorString, ANSI_TO_TCHAR( File ), Line, DescriptionString );
 	}
 }
 
@@ -241,108 +232,100 @@ void FDebug::EnsureFailed( const ANSICHAR* Expr, const ANSICHAR* File, int32 Lin
 	{
 		// Just trigger a regular assertion which will crash via GError->Logf()
 		FDebug::AssertFailed( Expr, File, Line, Msg );
+		return;
 	}
-	else
+
+	// Print initial debug message for this error
+	TCHAR ErrorString[MAX_SPRINTF];
+	FCString::Sprintf(ErrorString,TEXT("Ensure condition failed: %s"),ANSI_TO_TCHAR(Expr));
+
+	StaticFailDebug( ErrorString, File, Line, Msg, true );
+
+	// Is there a debugger attached?  If not we'll submit an error report.
+	if (FPlatformMisc::IsDebuggerPresent())
 	{
-		// Print initial debug message for this error
-		TCHAR ErrorString[MAX_SPRINTF];
-		FCString::Sprintf(ErrorString,TEXT("Ensure condition failed: %s"),ANSI_TO_TCHAR(Expr));
+		return;
+	}
 
-		StaticFailDebug( ErrorString, File, Line, Msg, true );
+	// If we determine that we have not sent a report for this ensure yet, send the report below.
+	bool bShouldSendNewReport = false;
 
-		// Is there a debugger attached?  If so we'll just break, otherwise we'll submit an error report.
-		if( FPlatformMisc::IsDebuggerPresent() )
-		{
-			// ensure() assertion failed.  Break into the attached debugger.  Presumedly, this assertion
-			// is handled and you can safely "set next statement" and resume execution if you want.
-			// NOTE: In Release builds FPlatformMisc::DebugBreak currently initiates a crash
-			FPlatformMisc::DebugBreak();
-		}
-		else
-		{
-			FPlatformMisc::PromptForRemoteDebugging(true);
+	// No debugger attached, so generate a call stack and submit a crash report
+	// Walk the stack and dump it to the allocated memory.
+	const SIZE_T StackTraceSize = 65535;
+	ANSICHAR* StackTrace = (ANSICHAR*) FMemory::SystemMalloc( StackTraceSize );
+	if( StackTrace != NULL )
+	{
+		StackTrace[0] = 0;
+		FPlatformStackWalk::StackWalkAndDump( StackTrace, StackTraceSize, CALLSTACK_IGNOREDEPTH );
 
-			// If we determine that we have not sent a report for this ensure yet, send the report below.
-			bool bShouldSendNewReport = false;
+		// Create a final string that we'll output to the log (and error history buffer)
+		TCHAR ErrorMsg[16384];
+		FCString::Snprintf( ErrorMsg, ARRAY_COUNT( ErrorMsg ), TEXT( "Ensure condition failed: %s [File:%s] [Line: %i]" ) LINE_TERMINATOR TEXT( "%s" ) LINE_TERMINATOR TEXT( "Stack: " ) LINE_TERMINATOR, ANSI_TO_TCHAR( Expr ), ANSI_TO_TCHAR( File ), Line, Msg );
 
-			// No debugger attached, so generate a call stack and submit a crash report
-			// Walk the stack and dump it to the allocated memory.
-			const SIZE_T StackTraceSize = 65535;
-			ANSICHAR* StackTrace = (ANSICHAR*) FMemory::SystemMalloc( StackTraceSize );
-			if( StackTrace != NULL )
-			{
-				StackTrace[0] = 0;
-				FPlatformStackWalk::StackWalkAndDump( StackTrace, StackTraceSize, CALLSTACK_IGNOREDEPTH );
+		// Also append the stack trace
+		FCString::Strncat( ErrorMsg, ANSI_TO_TCHAR(StackTrace), ARRAY_COUNT(ErrorMsg) - 1 );
+		FMemory::SystemFree( StackTrace );
 
-				// Create a final string that we'll output to the log (and error history buffer)
-				TCHAR ErrorMsg[16384];
-				FCString::Snprintf( ErrorMsg, ARRAY_COUNT( ErrorMsg ), TEXT( "Ensure condition failed: %s [File:%s] [Line: %i]" ) LINE_TERMINATOR TEXT( "%s" ) LINE_TERMINATOR TEXT( "Stack: " ) LINE_TERMINATOR, ANSI_TO_TCHAR( Expr ), ANSI_TO_TCHAR( File ), Line, Msg );
+		// Dump the error and flush the log.
+		UE_LOG(LogOutputDevice, Log, TEXT("=== Handled error: ===") LINE_TERMINATOR LINE_TERMINATOR TEXT("%s"), ErrorMsg);
+		GLog->Flush();
 
-				// Also append the stack trace
-				FCString::Strncat( ErrorMsg, ANSI_TO_TCHAR(StackTrace), ARRAY_COUNT(ErrorMsg) - 1 );
-				FMemory::SystemFree( StackTrace );
+		// Submit the error report to the server! (and display a balloon in the system tray)
+		{
+			// How many unique previous errors we should keep track of
+			const uint32 MaxPreviousErrorsToTrack = 4;
+			static uint32 StaticPreviousErrorCount = 0;
+			if( StaticPreviousErrorCount < MaxPreviousErrorsToTrack )
+			{
+				// Check to see if we've already reported this error.  No point in blasting the server with
+				// the same error over and over again in a single application session.
+				bool bHasErrorAlreadyBeenReported = false;
+
+				// Static: Array of previous unique error message CRCs
+				static uint32 StaticPreviousErrorCRCs[ MaxPreviousErrorsToTrack ];
 
-				// Dump the error and flush the log.
-				UE_LOG(LogOutputDevice, Log, TEXT("=== Handled error: ===") LINE_TERMINATOR LINE_TERMINATOR TEXT("%s"), ErrorMsg);
-				GLog->Flush();
+				// Compute CRC of error string.  Note that along with the call stack, this includes the message
+				// string passed to the macro, so only truly redundant errors will go unreported.  Though it also
+				// means you shouldn't pass loop counters to ensureMsgf(), otherwise failures may spam the server!
+				const uint32 ErrorStrCRC = FCrc::StrCrc_DEPRECATED( ErrorMsg );
 
-				// Submit the error report to the server! (and display a balloon in the system tray)
+				for( uint32 CurErrorIndex = 0; CurErrorIndex < StaticPreviousErrorCount; ++CurErrorIndex )
 				{
-					// How many unique previous errors we should keep track of
-					const uint32 MaxPreviousErrorsToTrack = 4;
-					static uint32 StaticPreviousErrorCount = 0;
-					if( StaticPreviousErrorCount < MaxPreviousErrorsToTrack )
+					if( StaticPreviousErrorCRCs[ CurErrorIndex ] == ErrorStrCRC )
 					{
-						// Check to see if we've already reported this error.  No point in blasting the server with
-						// the same error over and over again in a single application session.
-						bool bHasErrorAlreadyBeenReported = false;
-
-						// Static: Array of previous unique error message CRCs
-						static uint32 StaticPreviousErrorCRCs[ MaxPreviousErrorsToTrack ];
-
-						// Compute CRC of error string.  Note that along with the call stack, this includes the message
-						// string passed to the macro, so only truly redundant errors will go unreported.  Though it also
-						// means you shouldn't pass loop counters to ensureMsgf(), otherwise failures may spam the server!
-						const uint32 ErrorStrCRC = FCrc::StrCrc_DEPRECATED( ErrorMsg );
-
-						for( uint32 CurErrorIndex = 0; CurErrorIndex < StaticPreviousErrorCount; ++CurErrorIndex )
-						{
-							if( StaticPreviousErrorCRCs[ CurErrorIndex ] == ErrorStrCRC )
-							{
-								// Found it!  This is a redundant error message.
-								bHasErrorAlreadyBeenReported = true;
-								break;
-							}
-						}
-
-						// Add the element to the list and bump the count
-						StaticPreviousErrorCRCs[ StaticPreviousErrorCount++ ] = ErrorStrCRC;
-
-						if( !bHasErrorAlreadyBeenReported )
-						{
-							FCoreDelegates::OnHandleSystemEnsure.Broadcast();
-
-							FPlatformMisc::SubmitErrorReport( ErrorMsg, EErrorReportMode::Balloon );
-
-							bShouldSendNewReport = true;
-						}
+						// Found it!  This is a redundant error message.
+						bHasErrorAlreadyBeenReported = true;
+						break;
 					}
 				}
+
+				// Add the element to the list and bump the count
+				StaticPreviousErrorCRCs[ StaticPreviousErrorCount++ ] = ErrorStrCRC;
+
+				if( !bHasErrorAlreadyBeenReported )
+				{
+					FCoreDelegates::OnHandleSystemEnsure.Broadcast();
+
+					FPlatformMisc::SubmitErrorReport( ErrorMsg, EErrorReportMode::Balloon );
+
+					bShouldSendNewReport = true;
+				}
 			}
-			else
-			{
-				// If we fail to generate the string to identify the crash we don't know if we should skip sending the report,
-				// so we will just send the report anyway.
-				bShouldSendNewReport = true;
-			}
+		}
+	}
+	else
+	{
+		// If we fail to generate the string to identify the crash we don't know if we should skip sending the report,
+		// so we will just send the report anyway.
+		bShouldSendNewReport = true;
+	}
 
-			if ( bShouldSendNewReport )
-			{
+	if ( bShouldSendNewReport )
+	{
 #if PLATFORM_DESKTOP
-				NewReportEnsure( GErrorMessage );
+		NewReportEnsure( GErrorMessage );
 #endif
-			}
-		}
 	}
 }
 
@@ -354,7 +337,21 @@ void VARARGS FError::LowLevelFatal(const ANSICHAR* File, int32 Line, const TCHAR
 	GET_VARARGS( DescriptionString, ARRAY_COUNT(DescriptionString), ARRAY_COUNT(DescriptionString)-1, Format, Format );
 
 	StaticFailDebug(TEXT("LowLevelFatalError"),File,Line,DescriptionString);
-	GError->Log(DescriptionString);
+}
+
+void VARARGS FError::FinalLogf(const ANSICHAR* Expr, const ANSICHAR* File, int32 Line, const TCHAR* Format/*=TEXT("")*/, ...)
+{
+	if (GIsCriticalError)
+	{
+		return;
+	}
+
+	TCHAR DescriptionString[4096];
+	GET_VARARGS(DescriptionString, ARRAY_COUNT(DescriptionString), ARRAY_COUNT(DescriptionString) - 1, Format, Format);
+
+	TCHAR ErrorString[MAX_SPRINTF];
+	FCString::Sprintf(ErrorString, TEXT("Assertion failed: %s"), ANSI_TO_TCHAR(Expr));
+	GError->Logf(TEXT("Assertion failed: %s") FILE_LINE_DESC TEXT("\n%s\n"), ErrorString, ANSI_TO_TCHAR(File), Line, DescriptionString);
 }
 
 #if DO_CHECK || DO_GUARD_SLOW
@@ -520,7 +517,6 @@ VARARG_BODY(void, FMsg::Logf, const TCHAR*, VARARG_EXTRA(const ANSICHAR* File) V
 		}
 
 		StaticFailDebug(TEXT("Fatal error:"), File, Line, Message);
-		GError->Log(Category, Verbosity, Message);
 	}
 #endif
 }
diff --git a/Engine/Source/Runtime/Core/Public/Android/AndroidMisc.h b/Engine/Source/Runtime/Core/Public/Android/AndroidMisc.h
--- a/Engine/Source/Runtime/Core/Public/Android/AndroidMisc.h
+++ b/Engine/Source/Runtime/Core/Public/Android/AndroidMisc.h
@@ -81,6 +81,21 @@ struct CORE_API FAndroidMisc : public FGenericPlatformMisc
 		return false;
 	}
 
+	/** Prompts for remote debugging if debugger is not attached. Regardless of result, breaks into debugger afterwards. Returns false for use in conditionals. */
+	static FORCEINLINE bool DebugBreakAndPromptForRemoteReturningFalse()
+	{
+#if !UE_BUILD_SHIPPING
+		if (!IsDebuggerPresent())
+		{
+			PromptForRemoteDebugging(false);
+		}
+
+		DebugBreak();
+#endif
+
+		return false;
+	}
+
 	FORCEINLINE static void MemoryBarrier()
 	{
 		__sync_synchronize();
diff --git a/Engine/Source/Runtime/Core/Public/HTML5/HTML5PlatformMisc.h b/Engine/Source/Runtime/Core/Public/HTML5/HTML5PlatformMisc.h
--- a/Engine/Source/Runtime/Core/Public/HTML5/HTML5PlatformMisc.h
+++ b/Engine/Source/Runtime/Core/Public/HTML5/HTML5PlatformMisc.h
@@ -60,7 +60,24 @@ struct CORE_API FHTML5Misc : public FGenericPlatformMisc
 	/** Break into debugger. Returning false allows this function to be used in conditionals. */
 	FORCEINLINE static bool DebugBreakReturningFalse()
 	{
+#if !UE_BUILD_SHIPPING
 		DebugBreak();
+#endif
+		return false;
+	}
+
+	/** Prompts for remote debugging if debugger is not attached. Regardless of result, breaks into debugger afterwards. Returns false for use in conditionals. */
+	static FORCEINLINE bool DebugBreakAndPromptForRemoteReturningFalse()
+	{
+#if !UE_BUILD_SHIPPING
+		if (!IsDebuggerPresent())
+		{
+			PromptForRemoteDebugging(false);
+		}
+
+		DebugBreak();
+#endif
+
 		return false;
 	}
 
diff --git a/Engine/Source/Runtime/Core/Public/IOS/IOSPlatformMisc.h b/Engine/Source/Runtime/Core/Public/IOS/IOSPlatformMisc.h
--- a/Engine/Source/Runtime/Core/Public/IOS/IOSPlatformMisc.h
+++ b/Engine/Source/Runtime/Core/Public/IOS/IOSPlatformMisc.h
@@ -61,6 +61,21 @@ struct CORE_API FIOSPlatformMisc : public FGenericPlatformMisc
 		return false;
 	}
 
+	/** Prompts for remote debugging if debugger is not attached. Regardless of result, breaks into debugger afterwards. Returns false for use in conditionals. */
+	static FORCEINLINE bool DebugBreakAndPromptForRemoteReturningFalse()
+	{
+#if !UE_BUILD_SHIPPING
+		if (!IsDebuggerPresent())
+		{
+			PromptForRemoteDebugging(false);
+		}
+
+		DebugBreak();
+#endif
+
+		return false;
+	}
+
 	FORCEINLINE static void MemoryBarrier()
 	{
 		OSMemoryBarrier();
diff --git a/Engine/Source/Runtime/Core/Public/Linux/LinuxPlatformMisc.h b/Engine/Source/Runtime/Core/Public/Linux/LinuxPlatformMisc.h
--- a/Engine/Source/Runtime/Core/Public/Linux/LinuxPlatformMisc.h
+++ b/Engine/Source/Runtime/Core/Public/Linux/LinuxPlatformMisc.h
@@ -43,6 +43,21 @@ struct CORE_API FLinuxPlatformMisc : public FGenericPlatformMisc
 		return false;
 	}
 
+	/** Prompts for remote debugging if debugger is not attached. Regardless of result, breaks into debugger afterwards. Returns false for use in conditionals. */
+	static FORCEINLINE bool DebugBreakAndPromptForRemoteReturningFalse()
+	{
+#if !UE_BUILD_SHIPPING
+		if (!IsDebuggerPresent())
+		{
+			PromptForRemoteDebugging(false);
+		}
+
+		DebugBreak();
+#endif
+
+		return false;
+	}
+
 	static void PumpMessages(bool bFromMainLoop);
 	static uint32 GetKeyMap( uint16* KeyCodes, FString* KeyNames, uint32 MaxMappings );
 	static uint32 GetCharKeyMap(uint16* KeyCodes, FString* KeyNames, uint32 MaxMappings);
diff --git a/Engine/Source/Runtime/Core/Public/Mac/MacPlatformMisc.h b/Engine/Source/Runtime/Core/Public/Mac/MacPlatformMisc.h
--- a/Engine/Source/Runtime/Core/Public/Mac/MacPlatformMisc.h
+++ b/Engine/Source/Runtime/Core/Public/Mac/MacPlatformMisc.h
@@ -62,6 +62,21 @@ struct CORE_API FMacPlatformMisc : public FGenericPlatformMisc
 		return false;
 	}
 
+	/** Prompts for remote debugging if debugger is not attached. Regardless of result, breaks into debugger afterwards. Returns false for use in conditionals. */
+	static FORCEINLINE bool DebugBreakAndPromptForRemoteReturningFalse()
+	{
+#if !UE_BUILD_SHIPPING
+		if (!IsDebuggerPresent())
+		{
+			PromptForRemoteDebugging(false);
+		}
+
+		DebugBreak();
+#endif
+
+		return false;
+	}
+
 	FORCEINLINE static void MemoryBarrier()
 	{
 		OSMemoryBarrier();
diff --git a/Engine/Source/Runtime/Core/Public/Misc/AssertionMacros.h b/Engine/Source/Runtime/Core/Public/Misc/AssertionMacros.h
--- a/Engine/Source/Runtime/Core/Public/Misc/AssertionMacros.h
+++ b/Engine/Source/Runtime/Core/Public/Misc/AssertionMacros.h
@@ -15,21 +15,27 @@
 // "verify" expressions are always evaluated, but only cause an error if enabled.
 //
 
+#if !UE_BUILD_SHIPPING
+#define _DebugBreakAndPromptForRemote() \
+	if (!FPlatformMisc::IsDebuggerPresent()) { FPlatformMisc::PromptForRemoteDebugging(false); } FPlatformMisc::DebugBreak();
+#else
+	#define _DebugBreakAndPromptForRemote()
+#endif // !UE_BUILD_SHIPPING
 #if DO_CHECK
 	#define checkCode( Code )		do { Code } while ( false );
-	#define verify(expr)			{ if(!(expr)) { FPlatformMisc::DebugBreak(); FDebug::AssertFailed( #expr, __FILE__, __LINE__ ); CA_ASSUME(expr); } }
-	#define check(expr)				{ if(!(expr)) { FPlatformMisc::DebugBreak(); FDebug::AssertFailed( #expr, __FILE__, __LINE__ ); CA_ASSUME(expr); } }
+	#define verify(expr)			{ if(!(expr)) { FDebug::AssertFailed( #expr, __FILE__, __LINE__ ); _DebugBreakAndPromptForRemote(); FError::FinalLogf( #expr, __FILE__, __LINE__ ); CA_ASSUME(expr); } }
+	#define check(expr)				{ if(!(expr)) { FDebug::AssertFailed( #expr, __FILE__, __LINE__ ); _DebugBreakAndPromptForRemote(); FError::FinalLogf( #expr, __FILE__, __LINE__ ); CA_ASSUME(expr); } }
 	
 	/**
 	 * verifyf, checkf: Same as verify, check but with printf style additional parameters
 	 * Read about __VA_ARGS__ (variadic macros) on http://gcc.gnu.org/onlinedocs/gcc-3.4.4/cpp.pdf.
 	 */
-	#define verifyf(expr, format, ...)		{ if(!(expr)) { FPlatformMisc::DebugBreak(); FDebug::AssertFailed( #expr, __FILE__, __LINE__, format, ##__VA_ARGS__ ); CA_ASSUME(expr); } }
-	#define checkf(expr, format, ...)		{ if(!(expr)) { FPlatformMisc::DebugBreak(); FDebug::AssertFailed( #expr, __FILE__, __LINE__, format, ##__VA_ARGS__ ); CA_ASSUME(expr); } }
+	#define verifyf(expr, format,  ...)		{ if(!(expr)) { FDebug::AssertFailed( #expr, __FILE__, __LINE__, format, ##__VA_ARGS__ ); _DebugBreakAndPromptForRemote(); FError::FinalLogf(#expr, __FILE__, __LINE__, format, ##__VA_ARGS__); CA_ASSUME(expr); } }
+	#define checkf(expr, format,  ...)		{ if(!(expr)) { FDebug::AssertFailed( #expr, __FILE__, __LINE__, format, ##__VA_ARGS__ ); _DebugBreakAndPromptForRemote(); FError::FinalLogf(#expr, __FILE__, __LINE__, format, ##__VA_ARGS__); CA_ASSUME(expr); } }
 	/**
 	 * Denotes code paths that should never be reached.
 	 */
-	#define checkNoEntry()       { FPlatformMisc::DebugBreak(); FDebug::AssertFailed( "Enclosing block should never be called", __FILE__, __LINE__ ); }
+	#define checkNoEntry()       { FDebug::AssertFailed( "Enclosing block should never be called", __FILE__, __LINE__ ); _DebugBreakAndPromptForRemote(); FError::FinalLogf("Enclosing block should never be called", __FILE__, __LINE__ ); }
 
 	/**
 	 * Denotes code paths that should not be executed more than once.
@@ -54,7 +60,7 @@
 	                            checkf( RecursionCounter##__LINE__ == 0, TEXT("Enclosing block was entered recursively") );  \
 	                            const FRecursionScopeMarker ScopeMarker##__LINE__( RecursionCounter##__LINE__ )
 
-	#define unimplemented()       { FPlatformMisc::DebugBreak(); FDebug::AssertFailed( "Unimplemented function called", __FILE__, __LINE__ ); }
+	#define unimplemented()       { FDebug::AssertFailed( "Unimplemented function called", __FILE__, __LINE__ ); _DebugBreakAndPromptForRemote(); FError::FinalLogf("Unimplemented function called", __FILE__, __LINE__); }
 
 #else
 	#define checkCode(...)
@@ -72,9 +78,9 @@
 // Check for development only.
 //
 #if DO_GUARD_SLOW
-	#define checkSlow(expr)					{ if(!(expr)) { FDebug::AssertFailed( #expr, __FILE__, __LINE__ ); CA_ASSUME(expr); } }
-	#define checkfSlow(expr, format, ...)	{ if(!(expr)) { FDebug::AssertFailed( #expr, __FILE__, __LINE__, format, ##__VA_ARGS__ ); CA_ASSUME(expr); } }
-	#define verifySlow(expr)				{ if(!(expr)) { FDebug::AssertFailed( #expr, __FILE__, __LINE__ );} }
+	#define checkSlow(expr)					{ if(!(expr)) { FDebug::AssertFailed( #expr, __FILE__, __LINE__ ); _DebugBreakAndPromptForRemote(); FError::FinalLogf(#expr, __FILE__, __LINE__); CA_ASSUME(expr); } }
+	#define checkfSlow(expr, format, ...)	{ if(!(expr)) { FDebug::AssertFailed( #expr, __FILE__, __LINE__, format, ##__VA_ARGS__ ); _DebugBreakAndPromptForRemote(); FError::FinalLogf( #expr, __FILE__, __LINE__, format, ##__VA_ARGS__ ); CA_ASSUME(expr); } }
+	#define verifySlow(expr)				{ if(!(expr)) { FDebug::AssertFailed( #expr, __FILE__, __LINE__ ); _DebugBreakAndPromptForRemote(); FError::FinalLogf(#expr, __FILE__, __LINE__); } }
 #else
 	#define checkSlow(expr)
 	#define checkfSlow(expr, format, ...)
@@ -113,7 +119,7 @@
 		(((InExpression) == 0) ? FDebug::EnsureNotFalse(false, #InExpression, __FILE__, __LINE__, InMsg ) : true)
 
 	#define ensureMsgf( InExpression, InFormat, ... ) \
-		(((InExpression) == 0) ? FDebug::EnsureNotFalseFormatted(false, #InExpression, __FILE__, __LINE__, InFormat, ##__VA_ARGS__ ) : true)
+		(((InExpression) == 0) ? ( FDebug::EnsureNotFalseFormatted(false, #InExpression, __FILE__, __LINE__, InFormat, ##__VA_ARGS__ ) || FPlatformMisc::DebugBreakAndPromptForRemoteReturningFalse()) : true)
 
 #else	// DO_CHECK
 
@@ -178,8 +184,9 @@ struct FTCharArrayTester
 #define LowLevelFatalError(Format, ...) \
 	{ \
 		static_assert(IS_TCHAR_ARRAY(Format), "Formatting string must be a TCHAR array."); \
-		FPlatformMisc::DebugBreak(); \
 		FError::LowLevelFatal(__FILE__, __LINE__, Format, ##__VA_ARGS__); \
+		_DebugBreakAndPromptForRemote(); \
+		FError::FinalLogf("", __FILE__, __LINE__, Format, ##__VA_ARGS__); \
 	}
 
 
@@ -205,8 +212,9 @@ struct FTCharArrayTester
 		static_assert(IS_TCHAR_ARRAY(Format), "Formatting string must be a TCHAR array."); \
 		if (ELogVerbosity::Verbosity == ELogVerbosity::Fatal) \
 		{ \
-			FPlatformMisc::DebugBreak(); \
 			FError::LowLevelFatal(__FILE__, __LINE__, Format, ##__VA_ARGS__); \
+			_DebugBreakAndPromptForRemote(); \
+			FError::FinalLogf("", __FILE__, __LINE__, Format, ##__VA_ARGS__); \
 		} \
 	}
 	// Conditional logging (fatal errors only).
@@ -217,8 +225,9 @@ struct FTCharArrayTester
 		{ \
 			if (Condition) \
 			{ \
-				FPlatformMisc::DebugBreak(); \
 				FError::LowLevelFatal(__FILE__, __LINE__, Format, ##__VA_ARGS__); \
+				_DebugBreakAndPromptForRemote(); \
+				FError::FinalLogf("", __FILE__, __LINE__, Format, ##__VA_ARGS__); \
 			} \
 		} \
 	}
@@ -277,14 +286,16 @@ struct FTCharArrayTester
 			{ \
 				if (!CategoryName.IsSuppressed(ELogVerbosity::Verbosity)) \
 				{ \
+					FMsg::Logf(__FILE__, __LINE__, CategoryName.GetCategoryName(), ELogVerbosity::Verbosity, Format, ##__VA_ARGS__); \
 					if (ELogVerbosity::Verbosity == ELogVerbosity::Fatal) \
 					{\
-						FPlatformMisc::DebugBreak();\
+						_DebugBreakAndPromptForRemote(); \
+						FError::FinalLogf("", __FILE__, __LINE__, Format, ##__VA_ARGS__); \
 					}\
-					FMsg::Logf(__FILE__, __LINE__, CategoryName.GetCategoryName(), ELogVerbosity::Verbosity, Format, ##__VA_ARGS__); \
 				} \
 			} \
 		}
+
 		// Conditional logging. Will only log if Condition is met.
 		#define UE_CLOG(Condition, CategoryName, Verbosity, Format, ...) \
 		{ \
@@ -296,11 +307,12 @@ struct FTCharArrayTester
 				{ \
 					if (Condition) \
 					{ \
+						FMsg::Logf(__FILE__, __LINE__, CategoryName.GetCategoryName(), ELogVerbosity::Verbosity, Format, ##__VA_ARGS__); \
 						if (ELogVerbosity::Verbosity == ELogVerbosity::Fatal) \
 						{\
-							FPlatformMisc::DebugBreak();\
+							_DebugBreakAndPromptForRemote(); \
+							FError::FinalLogf("", __FILE__, __LINE__, Format, ##__VA_ARGS__); \
 						}\
-						FMsg::Logf(__FILE__, __LINE__, CategoryName.GetCategoryName(), ELogVerbosity::Verbosity, Format, ##__VA_ARGS__); \
 					} \
 				} \
 			} \
diff --git a/Engine/Source/Runtime/Core/Public/Misc/OutputDevice.h b/Engine/Source/Runtime/Core/Public/Misc/OutputDevice.h
--- a/Engine/Source/Runtime/Core/Public/Misc/OutputDevice.h
+++ b/Engine/Source/Runtime/Core/Public/Misc/OutputDevice.h
@@ -413,15 +413,20 @@ struct CORE_API FDebug
 	 * @param	Line	Line number (__LINE__)
 	 * @param	Msg		Optional informative error message text
 	 */
-	static FORCEINLINE bool EnsureNotFalse( bool bExpressionResult, const ANSICHAR* Expr, const ANSICHAR* File, int32 Line, const TCHAR* Msg = TEXT( "" ) )
+    static FORCEINLINE bool EnsureNotFalse( bool bExpressionResult, const ANSICHAR* Expr, const ANSICHAR* File, int32 Line, const TCHAR* Msg = TEXT( "" ) )
 	{
-		if( bExpressionResult == 0 )
-		{
+	    if( bExpressionResult == 0 )
+	    {
+			EnsureFailed(Expr, File, Line, Msg);
+			if (!FPlatformMisc::IsDebuggerPresent())
+			{
+				FPlatformMisc::PromptForRemoteDebugging(false);
+			}
+
 			FPlatformMisc::DebugBreak();
-			EnsureFailed( Expr, File, Line, Msg );
-		}
-	
-		return bExpressionResult;
+	    }
+    
+    	return bExpressionResult;
 	}
 
 	/**
@@ -478,6 +483,9 @@ struct CORE_API FError
 	/** For throwing string-exceptions which safely propagate through guard/unguard. */
 	VARARG_DECL( static void VARARGS, static void, VARARG_NONE, Throwf, VARARG_NONE, const TCHAR*, VARARG_NONE, VARARG_NONE );
 	//@}
+
+	/** Outputs formatted error message and terminates program. */
+	static void VARARGS FinalLogf(const ANSICHAR* Expr, const ANSICHAR* File, int32 Line, const TCHAR* Format = TEXT(""), ...);
 };
 
 
diff --git a/Engine/Source/Runtime/Core/Public/Windows/WindowsPlatformMisc.h b/Engine/Source/Runtime/Core/Public/Windows/WindowsPlatformMisc.h
--- a/Engine/Source/Runtime/Core/Public/Windows/WindowsPlatformMisc.h
+++ b/Engine/Source/Runtime/Core/Public/Windows/WindowsPlatformMisc.h
@@ -61,6 +61,21 @@ struct CORE_API FWindowsPlatformMisc
 		return false;
 	}
 
+	/** Prompts for remote debugging if debugger is not attached. Regardless of result, breaks into debugger afterwards. Returns false for use in conditionals. */
+	static FORCEINLINE bool DebugBreakAndPromptForRemoteReturningFalse()
+	{
+#if !UE_BUILD_SHIPPING
+		if (!IsDebuggerPresent())
+		{
+			PromptForRemoteDebugging(false);
+		}
+
+		DebugBreak();
+#endif
+
+		return false;
+	}
+
 	static void PumpMessages(bool bFromMainLoop);
 	static uint32 GetKeyMap( uint16* KeyCodes, FString* KeyNames, uint32 MaxMappings );
 	static uint32 GetCharKeyMap(uint16* KeyCodes, FString* KeyNames, uint32 MaxMappings);
