Add a specific "compatible with engine version" field to the package file summary, to make it possible to share packages between Rocket and Perforce builds in hotfixes.

diff --git a/Engine/Source/Editor/UnrealEd/Private/Commandlets/PackageUtilities.cpp b/Engine/Source/Editor/UnrealEd/Private/Commandlets/PackageUtilities.cpp
--- a/Engine/Source/Editor/UnrealEd/Private/Commandlets/PackageUtilities.cpp
+++ b/Engine/Source/Editor/UnrealEd/Private/Commandlets/PackageUtilities.cpp
@@ -840,7 +840,8 @@ void FPkgInfoReporter_Log::GeneratePackageReport( ULinkerLoad* InLinker/*=NULL*/
 
 	UE_LOG(LogPackageUtilities, Warning, TEXT("\t         Filename: %s"), *Linker->Filename);
 	UE_LOG(LogPackageUtilities, Warning, TEXT("\t     File Version: %i"), Linker->UE4Ver() );
-	UE_LOG(LogPackageUtilities, Warning, TEXT("\t   Engine Version: %s"), *Linker->Summary.EngineVersion.ToString());
+	UE_LOG(LogPackageUtilities, Warning, TEXT("\t   Engine Version: %s"), *Linker->Summary.SavedByEngineVersion.ToString());
+	UE_LOG(LogPackageUtilities, Warning, TEXT("\t   Compat Version: %s"), *Linker->Summary.CompatibleWithEngineVersion.ToString());
 	UE_LOG(LogPackageUtilities, Warning, TEXT("\t     PackageFlags: %X"), Linker->Summary.PackageFlags );
 	UE_LOG(LogPackageUtilities, Warning, TEXT("\t        NameCount: %d"), Linker->Summary.NameCount );
 	UE_LOG(LogPackageUtilities, Warning, TEXT("\t       NameOffset: %d"), Linker->Summary.NameOffset );
diff --git a/Engine/Source/Editor/UnrealEd/Private/UnrealEdEngine.cpp b/Engine/Source/Editor/UnrealEd/Private/UnrealEdEngine.cpp
--- a/Engine/Source/Editor/UnrealEd/Private/UnrealEdEngine.cpp
+++ b/Engine/Source/Editor/UnrealEd/Private/UnrealEdEngine.cpp
@@ -437,7 +437,7 @@ void UUnrealEdEngine::OnPackageDirtyStateUpdated( UPackage* Pkg)
 				FPackageFileSummary Summary;
 				*PackageReader << Summary;
 
-				if ( Summary.GetFileVersionUE4() > GPackageFileUE4Version || !GEngineVersion.IsCompatibleWith(Summary.EngineVersion) )
+				if ( Summary.GetFileVersionUE4() > GPackageFileUE4Version || !GEngineVersion.IsCompatibleWith(Summary.CompatibleWithEngineVersion) )
 				{
 					WarningStateToSet = WDWS_PendingWarn;
 					bNeedWarningForPkgEngineVer = true;
diff --git a/Engine/Source/Runtime/AssetRegistry/Private/PackageReader.cpp b/Engine/Source/Runtime/AssetRegistry/Private/PackageReader.cpp
--- a/Engine/Source/Runtime/AssetRegistry/Private/PackageReader.cpp
+++ b/Engine/Source/Runtime/AssetRegistry/Private/PackageReader.cpp
@@ -89,10 +89,10 @@ bool FPackageReader::OpenPackageFile(const FString& InPackageFilename)
 	//make sure the filereader gets the correct version number (it defaults to latest version)
 	SetUE4Ver(PackageFileSummary.GetFileVersionUE4());
 	SetLicenseeUE4Ver(PackageFileSummary.GetFileVersionLicenseeUE4());
-	SetEngineVer(PackageFileSummary.EngineVersion);
+	SetEngineVer(PackageFileSummary.SavedByEngineVersion);
 	Loader->SetUE4Ver(PackageFileSummary.GetFileVersionUE4());
 	Loader->SetLicenseeUE4Ver(PackageFileSummary.GetFileVersionLicenseeUE4());
-	Loader->SetEngineVer(PackageFileSummary.EngineVersion);
+	Loader->SetEngineVer(PackageFileSummary.SavedByEngineVersion);
 
 	const FCustomVersionContainer& PackageFileSummaryVersions = PackageFileSummary.GetCustomVersionContainer();
 	SetCustomVersions(PackageFileSummaryVersions);
diff --git a/Engine/Source/Runtime/Core/Private/Misc/EngineVersion.cpp b/Engine/Source/Runtime/Core/Private/Misc/EngineVersion.cpp
--- a/Engine/Source/Runtime/Core/Private/Misc/EngineVersion.cpp
+++ b/Engine/Source/Runtime/Core/Private/Misc/EngineVersion.cpp
@@ -4,6 +4,7 @@
 #include "EngineVersion.h"
 #include "Runtime/Launch/Resources/Version.h"
 #include "EngineBuildSettings.h"
+#include "ModuleVersion.h"
 
 FEngineVersionBase::FEngineVersionBase()
 : Major(0)
@@ -69,8 +70,8 @@ EVersionComparison FEngineVersionBase::GetNewest(const FEngineVersionBase &First
 		return (First.GetPatch() > Second.GetPatch()) ? EVersionComparison::First : EVersionComparison::Second;
 	}
 
-	// Compare changelists (only if they're both from the same vendor)
-	if (First.IsLicenseeVersion() == Second.IsLicenseeVersion() && First.GetChangelist() != Second.GetChangelist())
+	// Compare changelists (only if they're both from the same vendor, and they're both valid)
+	if (First.IsLicenseeVersion() == Second.IsLicenseeVersion() && First.IsPromotedBuild() && Second.IsPromotedBuild() && First.GetChangelist() != Second.GetChangelist())
 	{
 		Component = EVersionComponent::Changelist;
 		return (First.GetChangelist() > Second.GetChangelist()) ? EVersionComparison::First : EVersionComparison::Second;
@@ -112,25 +113,9 @@ bool FEngineVersion::IsCompatibleWith(const FEngineVersionBase &Other) const
 	{
 		return true;
 	}
-
-	// Otherwise compare the versions
-	EVersionComponent Component;
-	EVersionComparison Comparison = FEngineVersion::GetNewest(*this, Other, &Component);
-
-	// If this engine version is the same or newer, it's definitely compatible
-	if(Comparison == EVersionComparison::Neither || Comparison == EVersionComparison::First)
-	{
-		return true;
-	}
-
-	// Otherwise check if we require a strict version match or just a major/minor version match.
-	if (Component != EVersionComponent::Major && Component != EVersionComponent::Minor && FRocketSupport::IsRocket())
-	{
-		return true;
-	}
 	else
 	{
-		return false;
+		return FEngineVersion::GetNewest(*this, Other, nullptr) != EVersionComparison::Second;
 	}
 }
 
@@ -208,6 +193,10 @@ void operator<<(FArchive &Ar, FEngineVersion &Version)
 #endif
 
 // Licensee changelists part of the engine version has the top bit set to 1
-#define ENGINE_VERSION_INTERNAL_OR_LICENSEE (BUILT_FROM_CHANGELIST | (ENGINE_IS_LICENSEE_VERSION << 31))
+#define ENGINE_VERSION_INTERNAL_OR_LICENSEE (2474361 | (ENGINE_IS_LICENSEE_VERSION << 31))
+
 // Global instance of the current engine version
 const FEngineVersion GEngineVersion(ENGINE_MAJOR_VERSION, ENGINE_MINOR_VERSION, ENGINE_PATCH_VERSION, ENGINE_VERSION_INTERNAL_OR_LICENSEE, BRANCH_NAME);
+
+// Version which this engine maintains strict API and package compatibility with
+const FEngineVersion GCompatibleWithEngineVersion(ENGINE_MAJOR_VERSION, ENGINE_MINOR_VERSION, ENGINE_PATCH_VERSION, (2474361 | (ENGINE_IS_LICENSEE_VERSION << 31)), BRANCH_NAME);
diff --git a/Engine/Source/Runtime/Core/Public/Misc/EngineVersion.h b/Engine/Source/Runtime/Core/Public/Misc/EngineVersion.h
--- a/Engine/Source/Runtime/Core/Public/Misc/EngineVersion.h
+++ b/Engine/Source/Runtime/Core/Public/Misc/EngineVersion.h
@@ -49,3 +49,6 @@ private:
 
 /** Global instance of the current engine version. */
 CORE_API extern const FEngineVersion GEngineVersion;
+
+/** Earliest version which this engine maintains strict API and package compatibility with */
+CORE_API extern const FEngineVersion GCompatibleWithEngineVersion;
diff --git a/Engine/Source/Runtime/Core/Public/UObject/ObjectVersion.h b/Engine/Source/Runtime/Core/Public/UObject/ObjectVersion.h
--- a/Engine/Source/Runtime/Core/Public/UObject/ObjectVersion.h
+++ b/Engine/Source/Runtime/Core/Public/UObject/ObjectVersion.h
@@ -482,6 +482,8 @@ enum EUnrealEngineObjectUE4Version
 	VER_UE4_REMOVE_UNUSED_UPOLYS_FROM_UMODEL,
 	// This doesn't do anything except trigger a rebuild on HISMC cluster trees, in this case to get a good "occlusion query" level
 	VER_UE4_REBUILD_HIERARCHICAL_INSTANCE_TREES,
+	// Package summary includes an CompatibleWithEngineVersion field, separately to the version it's saved with
+	VER_UE4_PACKAGE_SUMMARY_HAS_COMPATIBLE_ENGINE_VERSION,
 
 
 	// -----<new versions can be added before this line>-------------------------------------------------
diff --git a/Engine/Source/Runtime/CoreUObject/Private/Misc/WorldCompositionUtility.cpp b/Engine/Source/Runtime/CoreUObject/Private/Misc/WorldCompositionUtility.cpp
--- a/Engine/Source/Runtime/CoreUObject/Private/Misc/WorldCompositionUtility.cpp
+++ b/Engine/Source/Runtime/CoreUObject/Private/Misc/WorldCompositionUtility.cpp
@@ -105,7 +105,7 @@ bool FWorldTileInfo::Read(const FString& InPackageFileName, FWorldTileInfo& OutI
 
 		//make sure the filereader gets the correct version number (it defaults to latest version)
 		FileReader->SetUE4Ver(FileSummary.GetFileVersionUE4());
-		FileReader->SetEngineVer(FileSummary.EngineVersion);
+		FileReader->SetEngineVer(FileSummary.SavedByEngineVersion);
 
 		// Load the structure
 		*FileReader << OutInfo;
diff --git a/Engine/Source/Runtime/CoreUObject/Private/UObject/LinkerLoad.cpp b/Engine/Source/Runtime/CoreUObject/Private/UObject/LinkerLoad.cpp
--- a/Engine/Source/Runtime/CoreUObject/Private/UObject/LinkerLoad.cpp
+++ b/Engine/Source/Runtime/CoreUObject/Private/UObject/LinkerLoad.cpp
@@ -895,13 +895,13 @@ ULinkerLoad::ELinkerStatus ULinkerLoad::SerializePackageFileSummary()
 			return LINKER_Failed;
 		}
 
-		// Don't load packages that were saved engine version newer than the current one.
-		if( !GEngineVersion.IsCompatibleWith(Summary.EngineVersion) )
+		// Don't load packages that are only compatible with an engine version newer than the current one.
+		if( !GEngineVersion.IsCompatibleWith(Summary.CompatibleWithEngineVersion) )
 		{
-			UE_LOG(LogLinker, Warning, TEXT("Asset '%s' has been saved with engine version newer than current and therefore can't be loaded. CurrEngineVersion: %s AssetEngineVersion: %s"), *Filename, *GEngineVersion.ToString(), *Summary.EngineVersion.ToString() );
+			UE_LOG(LogLinker, Warning, TEXT("Asset '%s' has been saved with engine version newer than current and therefore can't be loaded. CurrEngineVersion: %s AssetEngineVersion: %s"), *Filename, *GEngineVersion.ToString(), *Summary.CompatibleWithEngineVersion.ToString() );
 			return LINKER_Failed;
 		}
-		else if( !FPlatformProperties::RequiresCookedData() && !Summary.EngineVersion.IsPromotedBuild() && GEngineVersion.IsPromotedBuild() )
+		else if( !FPlatformProperties::RequiresCookedData() && !Summary.SavedByEngineVersion.IsPromotedBuild() && GEngineVersion.IsPromotedBuild() )
 		{
 			// This warning can be disabled in ini with [Core.System] ZeroEngineVersionWarning=False
 			static struct FInitZeroEngineVersionWarning
@@ -985,11 +985,11 @@ ULinkerLoad::ELinkerStatus ULinkerLoad::SerializePackageFileSummary()
 		// Loader needs to be the same version.
 		Loader->SetUE4Ver(Summary.GetFileVersionUE4());
 		Loader->SetLicenseeUE4Ver(Summary.GetFileVersionLicenseeUE4());
-		Loader->SetEngineVer(Summary.EngineVersion);
+		Loader->SetEngineVer(Summary.SavedByEngineVersion);
 
 		ArUE4Ver = Summary.GetFileVersionUE4();
 		ArLicenseeUE4Ver = Summary.GetFileVersionLicenseeUE4();
-		ArEngineVer = Summary.EngineVersion;
+		ArEngineVer = Summary.SavedByEngineVersion;
 
 		const FCustomVersionContainer& SummaryVersions = Summary.GetCustomVersionContainer();
 		Loader->SetCustomVersions(SummaryVersions);
diff --git a/Engine/Source/Runtime/CoreUObject/Private/UObject/LinkerSave.cpp b/Engine/Source/Runtime/CoreUObject/Private/UObject/LinkerSave.cpp
--- a/Engine/Source/Runtime/CoreUObject/Private/UObject/LinkerSave.cpp
+++ b/Engine/Source/Runtime/CoreUObject/Private/UObject/LinkerSave.cpp
@@ -29,7 +29,8 @@ ULinkerSave::ULinkerSave( const FObjectInitializer& ObjectInitializer, UPackage*
 		// Set main summary info.
 		Summary.Tag           = PACKAGE_FILE_TAG;
 		Summary.SetFileVersions( GPackageFileUE4Version, GPackageFileLicenseeUE4Version, bInSaveUnversioned );
-		Summary.EngineVersion =	GEngineVersion;
+		Summary.SavedByEngineVersion = GEngineVersion;
+		Summary.CompatibleWithEngineVersion = GCompatibleWithEngineVersion;
 		Summary.PackageFlags  = Package ? Package->PackageFlags : 0;
 
 		if (Package)
@@ -62,7 +63,8 @@ ULinkerSave::ULinkerSave(const FObjectInitializer& ObjectInitializer, UPackage*
 		// Set main summary info.
 		Summary.Tag           = PACKAGE_FILE_TAG;
 		Summary.SetFileVersions( GPackageFileUE4Version, GPackageFileLicenseeUE4Version, bInSaveUnversioned );
-		Summary.EngineVersion =	GEngineVersion;
+		Summary.SavedByEngineVersion = GEngineVersion;
+		Summary.CompatibleWithEngineVersion = GCompatibleWithEngineVersion;
 		Summary.PackageFlags  = Package ? Package->PackageFlags : 0;
 
 		if (Package)
diff --git a/Engine/Source/Runtime/CoreUObject/Private/UObject/PackageFileSummary.cpp b/Engine/Source/Runtime/CoreUObject/Private/UObject/PackageFileSummary.cpp
--- a/Engine/Source/Runtime/CoreUObject/Private/UObject/PackageFileSummary.cpp
+++ b/Engine/Source/Runtime/CoreUObject/Private/UObject/PackageFileSummary.cpp
@@ -147,7 +147,7 @@ FArchive& operator<<( FArchive& Ar, FPackageFileSummary& Sum )
 			}
 			else
 			{
-				Ar << Sum.EngineVersion;
+				Ar << Sum.SavedByEngineVersion;
 			}
 		}
 		else
@@ -157,7 +157,27 @@ FArchive& operator<<( FArchive& Ar, FPackageFileSummary& Sum )
 
 			if(Ar.IsLoading() && EngineChangelist != 0)
 			{
-				Sum.EngineVersion.Set(4, 0, 0, EngineChangelist, TEXT(""));
+				Sum.SavedByEngineVersion.Set(4, 0, 0, EngineChangelist, TEXT(""));
+			}
+		}
+
+		if (Sum.GetFileVersionUE4() >= VER_UE4_PACKAGE_SUMMARY_HAS_COMPATIBLE_ENGINE_VERSION )
+		{
+			if(Ar.IsCooking() || (Ar.IsSaving() && !GEngineVersion.IsPromotedBuild()))
+			{
+				FEngineVersion EmptyEngineVersion;
+				Ar << EmptyEngineVersion;
+			}
+			else
+			{
+				Ar << Sum.CompatibleWithEngineVersion;
+			}
+		}
+		else
+		{
+			if (Ar.IsLoading())
+			{
+				Sum.CompatibleWithEngineVersion = Sum.SavedByEngineVersion;
 			}
 		}
 
diff --git a/Engine/Source/Runtime/CoreUObject/Public/UObject/Linker.h b/Engine/Source/Runtime/CoreUObject/Public/UObject/Linker.h
--- a/Engine/Source/Runtime/CoreUObject/Public/UObject/Linker.h
+++ b/Engine/Source/Runtime/CoreUObject/Public/UObject/Linker.h
@@ -626,9 +626,14 @@ public:
 	TArray<FGenerationInfo> Generations;
 
 	/**
-	 * Engine version this package was saved with.
+	 * Engine version this package was saved with. For hotfix releases and engine versions which maintain strict binary compatibility with another version, this may differ from CompatibleWithEngineVersion.
 	 */
-	FEngineVersion EngineVersion;
+	FEngineVersion SavedByEngineVersion;
+
+	/**
+	 * Engine version this package is compatible with. See SavedByEngineVersion.
+	 */
+	FEngineVersion CompatibleWithEngineVersion;
 
 	/**
 	 * Flags used to compress the file on save and uncompress on load.
