Remove test code from EngineVersion.cpp.

diff --git a/Engine/Source/Runtime/Core/Private/Misc/EngineVersion.cpp b/Engine/Source/Runtime/Core/Private/Misc/EngineVersion.cpp
--- a/Engine/Source/Runtime/Core/Private/Misc/EngineVersion.cpp
+++ b/Engine/Source/Runtime/Core/Private/Misc/EngineVersion.cpp
@@ -193,10 +193,10 @@ void operator<<(FArchive &Ar, FEngineVersion &Version)
 #endif
 
 // Licensee changelists part of the engine version has the top bit set to 1
-#define ENGINE_VERSION_INTERNAL_OR_LICENSEE (2474361 | (ENGINE_IS_LICENSEE_VERSION << 31))
+#define ENGINE_VERSION_INTERNAL_OR_LICENSEE (BUILT_FROM_CHANGELIST | (ENGINE_IS_LICENSEE_VERSION << 31))
 
 // Global instance of the current engine version
 const FEngineVersion GEngineVersion(ENGINE_MAJOR_VERSION, ENGINE_MINOR_VERSION, ENGINE_PATCH_VERSION, ENGINE_VERSION_INTERNAL_OR_LICENSEE, BRANCH_NAME);
 
 // Version which this engine maintains strict API and package compatibility with
-const FEngineVersion GCompatibleWithEngineVersion(ENGINE_MAJOR_VERSION, ENGINE_MINOR_VERSION, ENGINE_PATCH_VERSION, (2474361 | (ENGINE_IS_LICENSEE_VERSION << 31)), BRANCH_NAME);
+const FEngineVersion GCompatibleWithEngineVersion(ENGINE_MAJOR_VERSION, ENGINE_MINOR_VERSION, ENGINE_PATCH_VERSION, (MODULE_API_VERSION | (ENGINE_IS_LICENSEE_VERSION << 31)), BRANCH_NAME);
