Allow rendering thread on Tegra devices with ES 3.x support
#ue4
#android

diff --git a/Engine/Source/Runtime/Core/Private/Android/AndroidMisc.cpp b/Engine/Source/Runtime/Core/Private/Android/AndroidMisc.cpp
--- a/Engine/Source/Runtime/Core/Private/Android/AndroidMisc.cpp
+++ b/Engine/Source/Runtime/Core/Private/Android/AndroidMisc.cpp
@@ -258,7 +258,7 @@ bool FAndroidMisc::AllowRenderThread()
 	// tested with lg optimus 2x and motorola xoom 
 	// come back and revisit this later 
 	// https://code.google.com/p/android/issues/detail?id=32636
-	if (FAndroidMisc::GetGPUFamily() == FString(TEXT("NVIDIA Tegra")) && FPlatformMisc::NumberOfCores() <= 2)
+	if (FAndroidMisc::GetGPUFamily() == FString(TEXT("NVIDIA Tegra")) && FPlatformMisc::NumberOfCores() <= 2 && FAndroidMisc::GetGLVersion().StartsWith(TEXT("OpenGL ES 2.")))
 	{
 		return false;
 	}
