Removed unused code

diff --git a/Engine/Source/Runtime/Core/Public/Android/AndroidProperties.h b/Engine/Source/Runtime/Core/Public/Android/AndroidProperties.h
--- a/Engine/Source/Runtime/Core/Public/Android/AndroidProperties.h
+++ b/Engine/Source/Runtime/Core/Public/Android/AndroidProperties.h
@@ -123,12 +123,6 @@ struct FAndroidPlatformProperties
 	{
 		return true;
 	}
-
-	static FORCEINLINE const FString& GetCookFlavor()
-	{
-		extern FString GCookFlavor;
-		return GCookFlavor;
-	}
 };
 
 struct FAndroid_PVRTCPlatformProperties : public FAndroidPlatformProperties
