Refactored FProcHandle API in platform abstraction layer.

- Deprecated FProcHandle::Close() in favor of existing FPlatformProcess::CloseProc() (symmetric to FPlatformProcess::CreateProc()).
- Linux: made FProcHandle instances safe to pass by value.
- Linux: added support for "fire and forget" children, which will not leave zombies (at the expense of extra threads and a leaked thread handle).
- Extended TestPAL with tests for most of the above functionality.

diff --git a/Engine/Source/Programs/TestPAL/Private/Main.cpp b/Engine/Source/Programs/TestPAL/Private/Main.cpp
--- a/Engine/Source/Programs/TestPAL/Private/Main.cpp
+++ b/Engine/Source/Programs/TestPAL/Private/Main.cpp
@@ -23,7 +23,8 @@ int32 ProcRunAsChild(const TCHAR* CommandLine)
 	GEngineLoop.PreInit(CommandLine);
 
 	// set a random delay pretending to do some useful work up to a minute. 
-	double RandomWorkTime = FMath::FRandRange(0.0f, 60.0f);
+	srand(FPlatformProcess::GetCurrentProcessId());
+	double RandomWorkTime = FMath::FRandRange(0.0f, 6.0f);
 
 	UE_LOG(LogTestPAL, Display, TEXT("Running proc test as child (pid %d), will be doing work for %f seconds."), FPlatformProcess::GetCurrentProcessId(), RandomWorkTime);
 
diff --git a/Engine/Source/Programs/TestPAL/Private/Parent.cpp b/Engine/Source/Programs/TestPAL/Private/Parent.cpp
--- a/Engine/Source/Programs/TestPAL/Private/Parent.cpp
+++ b/Engine/Source/Programs/TestPAL/Private/Parent.cpp
@@ -10,14 +10,14 @@ FParent::FParent(int InNumTotalChildren, int InMaxChildrenAtOnce)
 
 }
 
-FProcHandle FParent::Launch()
+FProcHandle FParent::Launch(bool bDetached)
 {
 	// Launch the worker process
 	int32 PriorityModifier = -1; // below normal
 
 	uint32 WorkerId = 0;
 	FString WorkerName = FPlatformProcess::ExecutableName(false);
-	FProcHandle WorkerHandle = FPlatformProcess::CreateProc(*WorkerName, TEXT("proc-child"), true, false, false, &WorkerId, PriorityModifier, NULL, NULL);
+	FProcHandle WorkerHandle = FPlatformProcess::CreateProc(*WorkerName, TEXT("proc-child"), bDetached, false, false, &WorkerId, PriorityModifier, NULL, NULL);
 	if (!WorkerHandle.IsValid())
 	{
 		// If this doesn't error, the app will hang waiting for jobs that can never be completed
@@ -29,6 +29,43 @@ FProcHandle FParent::Launch()
 
 void FParent::Run()
 {
+	// test launching detached
+	{
+		for (int i = 0; i < 100; ++i)
+		{
+			UE_LOG(LogTestPAL, Log, TEXT("Launching a detached child to see if we leak a zombie."));
+			FProcHandle Child = Launch(true);
+			
+			FPlatformProcess::CloseProc(Child);
+		}
+	}
+
+	// test stopping children prematurely
+	{
+		UE_LOG(LogTestPAL, Log, TEXT("Launching a child to wait for it."));
+		FProcHandle Child = Launch();
+
+		UE_LOG(LogTestPAL, Log, TEXT("Closing child's handle (FPlatformProcess::CloseProc)"));
+		FPlatformProcess::CloseProc(Child);
+	}
+
+	{
+		UE_LOG(LogTestPAL, Log, TEXT("Launching a child to terminate it."));
+		FProcHandle Child = Launch();
+
+		UE_LOG(LogTestPAL, Log, TEXT("Sleeping for a bit to let the child ramp up."));
+		FPlatformProcess::Sleep(0.1f);
+
+		UE_LOG(LogTestPAL, Log, TEXT("Terminating the child (FPlatformProcess::TerminateProc())"));
+		FPlatformProcess::TerminateProc(Child);
+
+		UE_LOG(LogTestPAL, Log, TEXT("Closing child's handle (FPlatformProcess::CloseProc)"));
+		FPlatformProcess::CloseProc(Child);
+	}
+
+	UE_LOG(LogTestPAL, Log, TEXT("Proceeding to test multiple children."));
+
+	// test normal working loop
 	while (NumTotalChildren > 0)
 	{
 		// see if there are any new children to spawn
@@ -49,11 +86,13 @@ void FParent::Run()
 		for (int ChildIdx = 0; ChildIdx < Children.Num();)
 		{
 			int32 ReturnCode = -1;
-			if (FPlatformProcess::GetProcReturnCode(Children[ChildIdx], &ReturnCode))
+			FProcHandle ChildCopy = Children[ChildIdx];
+			if (FPlatformProcess::GetProcReturnCode(ChildCopy, &ReturnCode))
 			{
 				// print its return code
 				UE_LOG(LogTestPAL, Log, TEXT("Child finished, return code %d"), ReturnCode);
 
+				FPlatformProcess::CloseProc(ChildCopy);
 				Children.RemoveAt(ChildIdx);
 			}
 			else
diff --git a/Engine/Source/Programs/TestPAL/Private/Parent.h b/Engine/Source/Programs/TestPAL/Private/Parent.h
--- a/Engine/Source/Programs/TestPAL/Private/Parent.h
+++ b/Engine/Source/Programs/TestPAL/Private/Parent.h
@@ -16,7 +16,7 @@ class FParent
 	int MaxChildrenAtOnce;
 
 	/** Launches the child, returns the launched process handle. */
-	FProcHandle Launch();
+	FProcHandle Launch(bool bDetached = false);
 
 public:
 
diff --git a/Engine/Source/Runtime/Core/Private/GenericPlatform/GenericPlatformProcess.cpp b/Engine/Source/Runtime/Core/Private/GenericPlatform/GenericPlatformProcess.cpp
--- a/Engine/Source/Runtime/Core/Private/GenericPlatform/GenericPlatformProcess.cpp
+++ b/Engine/Source/Runtime/Core/Private/GenericPlatform/GenericPlatformProcess.cpp
@@ -180,12 +180,7 @@ void FGenericPlatformProcess::WaitForProc( FProcHandle & ProcessHandle )
 
 void FGenericPlatformProcess::CloseProc(FProcHandle & ProcessHandle)
 {
-	// make sure we don't leave running processes (this is especially important for Linux)
-	check(!FPlatformProcess::IsProcRunning(ProcessHandle));
-
-	// Wait() should return instantly if the process isn't running (which it shouldn't).
-	// However, this is important that we wait for it on Linux/Unix platforms so we don't leave zombies.
-	FPlatformProcess::WaitForProc(ProcessHandle);
+	UE_LOG(LogHAL, Fatal, TEXT("FGenericPlatformProcess::CloseProc not implemented on this platform"));
 }
 
 void FGenericPlatformProcess::TerminateProc( FProcHandle & ProcessHandle, bool KillTree )
diff --git a/Engine/Source/Runtime/Core/Private/Linux/LinuxPlatformProcess.cpp b/Engine/Source/Runtime/Core/Private/Linux/LinuxPlatformProcess.cpp
--- a/Engine/Source/Runtime/Core/Private/Linux/LinuxPlatformProcess.cpp
+++ b/Engine/Source/Runtime/Core/Private/Linux/LinuxPlatformProcess.cpp
@@ -416,6 +416,75 @@ void FLinuxPlatformProcess::LaunchURL(const TCHAR* URL, const TCHAR* Parms, FStr
 	}
 }
 
+/**
+ * This class exists as an imperfect workaround to allow both "fire and forget" children and children about whose return code we actually care.
+ * (maybe we could fork and daemonize ourselves for the first case instead?)
+ */
+struct FChildWaiterThread : public FRunnable
+{
+	/** Global table of all waiter threads */
+	static TArray<FChildWaiterThread *>		ChildWaiterThreadsArray;
+
+	/** Lock guarding the acess to child waiter threads */
+	static FCriticalSection					ChildWaiterThreadsArrayGuard;
+
+	/** Pid of child to wait for */
+	int ChildPid;
+
+	FChildWaiterThread(pid_t InChildPid)
+		:	ChildPid(InChildPid)
+	{
+		// add ourselves to thread array
+		ChildWaiterThreadsArrayGuard.Lock();
+		ChildWaiterThreadsArray.Add(this);
+		ChildWaiterThreadsArrayGuard.Unlock();
+	}
+
+	virtual ~FChildWaiterThread()
+	{
+		// remove
+		ChildWaiterThreadsArrayGuard.Lock();
+		ChildWaiterThreadsArray.RemoveSingle(this);
+		ChildWaiterThreadsArrayGuard.Unlock();
+	}
+
+	virtual uint32 Run()
+	{
+		for(;;)	// infinite loop in case we get EINTR and have to repeat
+		{
+			siginfo_t SignalInfo;
+			if (waitid(P_PID, ChildPid, &SignalInfo, WEXITED))
+			{
+				if (errno != EINTR)
+				{
+					int ErrNo = errno;
+					UE_LOG(LogHAL, Fatal, TEXT("FChildWaiterThread::Run(): waitid for pid %d failed (errno=%d, %s)"), 
+						   static_cast< int32 >(ChildPid), ErrNo, ANSI_TO_TCHAR(strerror(ErrNo)));
+					break;	// exit the loop if for some reason Fatal log (above) returns
+				}
+			}
+			else
+			{
+				check(SignalInfo.si_pid == ChildPid);
+				break;
+			}
+		}
+
+		return 0;
+	}
+
+	virtual void Exit()
+	{
+		// unregister from the array
+		delete this;
+	}
+};
+
+/** See FChildWaiterThread */
+TArray<FChildWaiterThread *> FChildWaiterThread::ChildWaiterThreadsArray;
+/** See FChildWaiterThread */
+FCriticalSection FChildWaiterThread::ChildWaiterThreadsArrayGuard;
+
 FProcHandle FLinuxPlatformProcess::CreateProc(const TCHAR* URL, const TCHAR* Parms, bool bLaunchDetached, bool bLaunchHidden, bool bLaunchReallyHidden, uint32* OutProcessID, int32 PriorityModifier, const TCHAR* OptionalWorkingDirectory, void* PipeWrite)
 {
 	// @TODO bLaunchHidden bLaunchReallyHidden are not handled
@@ -642,17 +711,65 @@ FProcHandle FLinuxPlatformProcess::CreateProc(const TCHAR* URL, const TCHAR* Par
 		*OutProcessID = ChildPid;
 	}
 
-	return FProcHandle( ChildPid );
+	// [RCL] 2015-03-11 @FIXME: is bLaunchDetached usable when determining whether we're in 'fire and forget' mode? This doesn't exactly match what bLaunchDetached is used for.
+	return FProcHandle(new FProcState(ChildPid, bLaunchDetached));
+}
+
+/** Initialization constructor. */
+FProcState::FProcState(pid_t InProcessId, bool bInFireAndForget)
+	:	ProcessId(InProcessId)
+	,	bIsRunning(true)  // assume it is
+	,	bHasBeenWaitedFor(false)
+	,	ReturnCode(-1)
+	,	bFireAndForget(bInFireAndForget)
+{
+}
+
+FProcState::~FProcState()
+{
+	if (!bFireAndForget)
+	{
+		// If not in 'fire and forget' mode, try to catch the common problems that leave zombies:
+		// - We don't want to close the handle of a running process as with our current scheme this will certainly leak a zombie.
+		// - Nor we want to leave the handle unwait()ed for.
+		
+		if (bIsRunning)
+		{
+			// Warn the users before going into what may be a very long block
+			UE_LOG(LogHAL, Warning, TEXT("Closing a process handle while the process (pid=%d) is still running - we will block until it exits to prevent a zombie"),
+				GetProcessId()
+			);
+		}
+		else if (!bHasBeenWaitedFor)	// if child is not running, but has not been waited for, still communicate a problem, but we shouldn't be blocked for long in this case.
+		{
+			UE_LOG(LogHAL, Warning, TEXT("Closing a process handle of a process (pid=%d) that has not been wait()ed for - will wait() now to reap a zombie"),
+				GetProcessId()
+			);
+		}
+
+		Wait();	// will exit immediately if everything is Ok
+	}
+	else if (IsRunning())
+	{
+		// warn about leaking a thread ;/
+		UE_LOG(LogHAL, Warning, TEXT("Process (pid=%d) is still running - we will reap it in a waiter thread, but the thread handle is going to be leaked."),
+			   GetProcessId()
+			);
+
+		FChildWaiterThread * WaiterRunnable = new FChildWaiterThread(GetProcessId());
+		// [RCL] 2015-03-11 @FIXME: do not leak
+		FRunnableThread * WaiterThread = FRunnableThread::Create(WaiterRunnable, *FString::Printf(TEXT("waitpid(%d)"), GetProcessId(), 32768 /* needs just a small stack */, TPri_BelowNormal));
+	}
 }
 
-bool FProcHandle::IsRunning()
+bool FProcState::IsRunning()
 {
 	if (bIsRunning)
 	{
 		check(!bHasBeenWaitedFor);	// check for the sake of internal consistency
 
 		// check if actually running
-		int KillResult = kill(Get(), 0);	// no actual signal is sent
+		int KillResult = kill(GetProcessId(), 0);	// no actual signal is sent
 		check(KillResult != -1 || errno != EINVAL);
 
 		bIsRunning = (KillResult == 0 || (KillResult == -1 && errno == EPERM));
@@ -664,18 +781,19 @@ bool FProcHandle::IsRunning()
 			{
 				siginfo_t SignalInfo;
 				SignalInfo.si_pid = 0;	// if remains 0, treat as child was not waitable (i.e. was running)
-				if (waitid(P_PID, Get(), &SignalInfo, WEXITED | WNOHANG | WNOWAIT))
+				if (waitid(P_PID, GetProcessId(), &SignalInfo, WEXITED | WNOHANG | WNOWAIT))
 				{
 					if (errno != EINTR)
 					{
+						int ErrNo = errno;
 						UE_LOG(LogHAL, Fatal, TEXT("FLinuxPlatformProcess::WaitForProc: waitid for pid %d failed (errno=%d, %s)"), 
-							static_cast< int32 >(Get()), errno, ANSI_TO_TCHAR(strerror(errno)));
+							static_cast< int32 >(GetProcessId()), ErrNo, ANSI_TO_TCHAR(strerror(ErrNo)));
 						break;	// exit the loop if for some reason Fatal log (above) returns
 					}
 				}
 				else
 				{
-					bIsRunning = ( SignalInfo.si_pid != Get() );
+					bIsRunning = ( SignalInfo.si_pid != GetProcessId() );
 					break;
 				}
 			}
@@ -686,7 +804,7 @@ bool FProcHandle::IsRunning()
 		// which is a dubious, but valid behavior. We don't want to keep zombie around though.
 		if (!bIsRunning)
 		{
-			UE_LOG(LogHAL, Log, TEXT("Child %d is no more running (zombie), Wait()ing immediately."), Get() );
+			UE_LOG(LogHAL, Log, TEXT("Child %d is no more running (zombie), Wait()ing immediately."), GetProcessId() );
 			Wait();
 		}
 	}
@@ -694,7 +812,7 @@ bool FProcHandle::IsRunning()
 	return bIsRunning;
 }
 
-bool FProcHandle::GetReturnCode(int32* ReturnCodePtr)
+bool FProcState::GetReturnCode(int32* ReturnCodePtr)
 {
 	check(!bIsRunning || !"You cannot get a return code of a running process");
 	if (!bHasBeenWaitedFor)
@@ -714,7 +832,7 @@ bool FProcHandle::GetReturnCode(int32* ReturnCodePtr)
 	return false;
 }
 
-void FProcHandle::Wait()
+void FProcState::Wait()
 {
 	if (bHasBeenWaitedFor)
 	{
@@ -724,18 +842,19 @@ void FProcHandle::Wait()
 	for(;;)	// infinite loop in case we get EINTR and have to repeat
 	{
 		siginfo_t SignalInfo;
-		if (waitid(P_PID, Get(), &SignalInfo, WEXITED))
+		if (waitid(P_PID, GetProcessId(), &SignalInfo, WEXITED))
 		{
 			if (errno != EINTR)
 			{
+				int ErrNo = errno;
 				UE_LOG(LogHAL, Fatal, TEXT("FLinuxPlatformProcess::WaitForProc: waitid for pid %d failed (errno=%d, %s)"), 
-					static_cast< int32 >(Get()), errno, ANSI_TO_TCHAR(strerror(errno)));
+					static_cast< int32 >(GetProcessId()), ErrNo, ANSI_TO_TCHAR(strerror(ErrNo)));
 				break;	// exit the loop if for some reason Fatal log (above) returns
 			}
 		}
 		else
 		{
-			check(SignalInfo.si_pid == Get());
+			check(SignalInfo.si_pid == GetProcessId());
 
 			ReturnCode = (SignalInfo.si_code == CLD_EXITED) ? SignalInfo.si_status : -1;
 			bHasBeenWaitedFor = true;
@@ -747,12 +866,26 @@ void FProcHandle::Wait()
 
 bool FLinuxPlatformProcess::IsProcRunning( FProcHandle & ProcessHandle )
 {
-	return ProcessHandle.IsRunning();
+	FProcState * ProcInfo = ProcessHandle.GetProcessInfo();
+	return ProcInfo ? ProcInfo->IsRunning() : false;
 }
 
 void FLinuxPlatformProcess::WaitForProc( FProcHandle & ProcessHandle )
 {
-	ProcessHandle.Wait();
+	FProcState * ProcInfo = ProcessHandle.GetProcessInfo();
+	if (ProcInfo)
+	{
+		ProcInfo->Wait();
+	}
+}
+
+void FLinuxPlatformProcess::CloseProc(FProcHandle & ProcessHandle)
+{
+	// dispose of both handle and process info
+	FProcState * ProcInfo = ProcessHandle.GetProcessInfo();
+	ProcessHandle.Reset();
+
+	delete ProcInfo;
 }
 
 void FLinuxPlatformProcess::TerminateProc( FProcHandle & ProcessHandle, bool KillTree )
@@ -763,8 +896,12 @@ void FLinuxPlatformProcess::TerminateProc( FProcHandle & ProcessHandle, bool Kil
 		STUBBED("FLinuxPlatformProcess::TerminateProc() : Killing a subtree is not implemented yet");
 	}
 
-	int KillResult = kill(ProcessHandle.Get(), SIGTERM);	// graceful
-	check(KillResult != -1 || errno != EINVAL);
+	FProcState * ProcInfo = ProcessHandle.GetProcessInfo();
+	if (ProcInfo)
+	{
+		int KillResult = kill(ProcInfo->GetProcessId(), SIGTERM);	// graceful
+		check(KillResult != -1 || errno != EINVAL);
+	}
 }
 
 uint32 FLinuxPlatformProcess::GetCurrentProcessId()
@@ -779,7 +916,8 @@ bool FLinuxPlatformProcess::GetProcReturnCode( FProcHandle& ProcHandle, int32* R
 		return false;
 	}
 
-	return ProcHandle.GetReturnCode(ReturnCode);
+	FProcState * ProcInfo = ProcHandle.GetProcessInfo();
+	return ProcInfo ? ProcInfo->GetReturnCode(ReturnCode) : false;
 }
 
 bool FLinuxPlatformProcess::Daemonize()
diff --git a/Engine/Source/Runtime/Core/Private/Mac/MacPlatformProcess.cpp b/Engine/Source/Runtime/Core/Private/Mac/MacPlatformProcess.cpp
--- a/Engine/Source/Runtime/Core/Private/Mac/MacPlatformProcess.cpp
+++ b/Engine/Source/Runtime/Core/Private/Mac/MacPlatformProcess.cpp
@@ -499,6 +499,11 @@ void FMacPlatformProcess::WaitForProc( FProcHandle& ProcessHandle )
 	[(NSTask*)ProcessHandle.Get() waitUntilExit];
 }
 
+void FMacPlatformProcess::CloseProc( FProcHandle & ProcessHandle )
+{
+	// is any action needed here?
+}
+
 void FMacPlatformProcess::TerminateProc( FProcHandle& ProcessHandle, bool KillTree )
 {
 	SCOPED_AUTORELEASE_POOL;
diff --git a/Engine/Source/Runtime/Core/Private/Windows/WindowsPlatformProcess.cpp b/Engine/Source/Runtime/Core/Private/Windows/WindowsPlatformProcess.cpp
--- a/Engine/Source/Runtime/Core/Private/Windows/WindowsPlatformProcess.cpp
+++ b/Engine/Source/Runtime/Core/Private/Windows/WindowsPlatformProcess.cpp
@@ -384,6 +384,15 @@ void FWindowsPlatformProcess::WaitForProc( FProcHandle & ProcessHandle )
 	::WaitForSingleObject(ProcessHandle.Get(), INFINITE);
 }
 
+void FWindowsPlatformProcess::CloseProc(FProcHandle & ProcessHandle)
+{
+	if (ProcessHandle.IsValid())
+	{
+		::CloseHandle(ProcessHandle.Get());
+		ProcessHandle.Reset();
+	}
+}
+
 void FWindowsPlatformProcess::TerminateProc( FProcHandle & ProcessHandle, bool KillTree )
 {
 	if (KillTree)
diff --git a/Engine/Source/Runtime/Core/Public/GenericPlatform/GenericPlatformProcess.h b/Engine/Source/Runtime/Core/Public/GenericPlatform/GenericPlatformProcess.h
--- a/Engine/Source/Runtime/Core/Public/GenericPlatform/GenericPlatformProcess.h
+++ b/Engine/Source/Runtime/Core/Public/GenericPlatform/GenericPlatformProcess.h
@@ -86,9 +86,11 @@ public:
 	}
 
 	/**
+	 * (Deprecated. Handles created with FPlatformProcess::CreateProc() should be closed with FPlatformProcess::CloseProc())
 	 * Closes handle and frees this resource to the operating system.
 	 * @return true, if this handle was valid before closing it
 	 */
+	DEPRECATED(4.8, "FProcHandle::Close() is deprecated as redundant - handles created with FPlatformProcess::CreateProc() should be closed with FPlatformProcess::CloseProc().")
 	FORCEINLINE bool Close()
 	{
 		return IsValid();
diff --git a/Engine/Source/Runtime/Core/Public/Linux/LinuxPlatformProcess.h b/Engine/Source/Runtime/Core/Public/Linux/LinuxPlatformProcess.h
--- a/Engine/Source/Runtime/Core/Public/Linux/LinuxPlatformProcess.h
+++ b/Engine/Source/Runtime/Core/Public/Linux/LinuxPlatformProcess.h
@@ -10,66 +10,39 @@
 #include "HAL/Platform.h"
 #include "Linux/LinuxSystemIncludes.h"
 
-/** Wrapper around Linux pid_t. */
-struct FProcHandle : public TProcHandle<pid_t, -1>
+/** Wrapper around Linux pid_t. Should not be copyable as changes in the process state won't be properly propagated to all copies. */
+struct FProcState
 {
-	typedef TProcHandle<pid_t, -1> Parent;
-
 	/** Default constructor. */
-	FORCEINLINE FProcHandle()
-		:	TProcHandle()
-		,	bIsRunning( false )
-		,	bHasBeenWaitedFor( false )
-		,	ReturnCode( -1 )
+	FORCEINLINE FProcState()
+		:	ProcessId(0)
+		,	bIsRunning(false)
+		,	bHasBeenWaitedFor(false)
+		,	ReturnCode(-1)
+		,	bFireAndForget(false)
 	{
 	}
 
 	/** Initialization constructor. */
-	FORCEINLINE explicit FProcHandle( HandleType Other )
-		:	TProcHandle( Other )
-		,	bIsRunning( true )	// assume it is
-		,	bHasBeenWaitedFor( false )
-		,	ReturnCode( -1 )
-	{
-	}
+	explicit FProcState(pid_t InProcessId, bool bInFireAndForget);
 
-	/** Copy constructor. */
-	FORCEINLINE FProcHandle( const FProcHandle& Other )
-		:	TProcHandle( Other )
-		,	bIsRunning( Other.bIsRunning )	// assume it is
-		,	bHasBeenWaitedFor( Other.bHasBeenWaitedFor )
-		,	ReturnCode( Other.ReturnCode )
-	{
-	}
+	/** Destructor. */
+	~FProcState();
 
-	/** Assignment operator. */
-	FORCEINLINE FProcHandle& operator=( const FProcHandle& Other )
+	/** Getter for process id */
+	FORCEINLINE pid_t GetProcessId() const
 	{
-		Parent::operator=( Other );
-
-		if( this != &Other )
-		{
-			bIsRunning = Other.bIsRunning;
-			bHasBeenWaitedFor = Other.bHasBeenWaitedFor;
-			ReturnCode = Other.ReturnCode;
-		}
-
-		return *this;
+		return ProcessId;
 	}
 
-
-protected:	// the below is not a public API!
-
-	friend struct FLinuxPlatformProcess;
-
-	/** 
+	/**
 	 * Returns whether this process is running.
 	 *
 	 * @return true if running
 	 */
 	bool	IsRunning();
 
-	/** 
+	/**
 	 * Returns child's return code (only valid to call if not running)
 	 *
 	 * @param ReturnCode set to return code if not NULL (use the value only if method returned true)
@@ -78,14 +51,41 @@ protected:	// the below is not a public API!
 	 */
 	bool	GetReturnCode(int32* ReturnCodePtr);
 
-	/** 
+	/**
 	 * Waits for the process to end.
 	 * Has a side effect (stores child's return code).
 	 */
 	void	Wait();
 
+protected:  // the below is not a public API!
+
+	// FProcState should not be copyable since it breeds problems (e.g. one copy could have wait()ed, but another would not know it)
+
+	/** Copy constructor - should not be publicly accessible */
+	FORCEINLINE FProcState(const FProcState& Other)
+		:	ProcessId(Other.ProcessId)
+		,	bIsRunning(Other.bIsRunning)  // assume it is
+		,	bHasBeenWaitedFor(Other.bHasBeenWaitedFor)
+		,	ReturnCode(Other.ReturnCode)
+		,	bFireAndForget(Other.bFireAndForget)
+	{
+		checkf(false, TEXT("FProcState should not be copied"));
+	}
+
+	/** Assignment operator - should not be publicly accessible */
+	FORCEINLINE FProcState& operator=(const FProcState& Other)
+	{
+		checkf(false, TEXT("FProcState should not be copied"));
+		return *this;
+	}
+
+	friend struct FLinuxPlatformProcess;
+
 	// -------------------------
 
+	/** Process id */
+	pid_t	ProcessId;
+
 	/** Whether the process has finished or not (cached) */
 	bool	bIsRunning;
 
@@ -94,6 +94,57 @@ protected:	// the below is not a public API!
 
 	/** Return code of the process (if negative, means that process did not finish gracefully but was killed/crashed*/
 	int32	ReturnCode;
+
+	/** Whether this child is fire-and-forget */
+	bool	bFireAndForget;
+};
+
+/** FProcHandle can be copied (and thus passed by value). */
+struct FProcHandle
+{
+	FProcState * 		ProcInfo;
+
+	FProcHandle()
+	:	ProcInfo(nullptr)
+	{
+	}
+
+	FProcHandle(FProcState * InHandle)
+	:	ProcInfo(InHandle)
+	{
+	}
+
+	/** Accessors. */
+	FORCEINLINE pid_t Get() const
+	{
+		return ProcInfo ? ProcInfo->GetProcessId() : -1;
+	}
+
+	/** Resets the handle to invalid */
+	FORCEINLINE void Reset()
+	{
+		ProcInfo = nullptr;
+	}
+
+	/** Checks the validity of handle */
+	FORCEINLINE bool IsValid() const
+	{
+		return ProcInfo != nullptr;
+	}
+
+	/**
+	 * (Deprecated. Handles created with FPlatformProcess::CreateProc() should be closed with FPlatformProcess::CloseProc())
+	 * Closes handle and frees this resource to the operating system.
+	 * @return true, if this handle was valid before closing it
+	 */
+	DEPRECATED(4.8, "FProcHandle::Close() is redundant - handles created with FPlatformProcess::CreateProc() should be closed with FPlatformProcess::CloseProc().")
+	bool Close();
+
+	// the below is not part of FProcHandle API and is specific to Linux implementation
+	FORCEINLINE FProcState * GetProcessInfo() const
+	{
+		return ProcInfo;
+	}
 };
 
 /** Wrapper around Linux file descriptors */
@@ -160,6 +211,7 @@ struct CORE_API FLinuxPlatformProcess : public FGenericPlatformProcess
 	static FProcHandle CreateProc(const TCHAR* URL, const TCHAR* Parms, bool bLaunchDetached, bool bLaunchHidden, bool bLaunchReallyHidden, uint32* OutProcessID, int32 PriorityModifier, const TCHAR* OptionalWorkingDirectory, void* PipeWrite);
 	static bool IsProcRunning( FProcHandle & ProcessHandle );
 	static void WaitForProc( FProcHandle & ProcessHandle );
+	static void CloseProc( FProcHandle & ProcessHandle );
 	static void TerminateProc( FProcHandle & ProcessHandle, bool KillTree = false );
 	static uint32 GetCurrentProcessId();
 	static bool GetProcReturnCode( FProcHandle & ProcHandle, int32* ReturnCode );
@@ -173,3 +225,13 @@ struct CORE_API FLinuxPlatformProcess : public FGenericPlatformProcess
 };
 
 typedef FLinuxPlatformProcess FPlatformProcess;
+
+inline bool FProcHandle::Close()
+{
+	if (IsValid())
+	{
+		FPlatformProcess::CloseProc(*this);
+		return true;
+	}
+	return false;
+}
diff --git a/Engine/Source/Runtime/Core/Public/Mac/MacPlatformProcess.h b/Engine/Source/Runtime/Core/Public/Mac/MacPlatformProcess.h
--- a/Engine/Source/Runtime/Core/Public/Mac/MacPlatformProcess.h
+++ b/Engine/Source/Runtime/Core/Public/Mac/MacPlatformProcess.h
@@ -50,6 +50,7 @@ struct CORE_API FMacPlatformProcess : public FGenericPlatformProcess
 	static FProcHandle CreateProc( const TCHAR* URL, const TCHAR* Parms, bool bLaunchDetached, bool bLaunchHidden, bool bLaunchReallyHidden, uint32* OutProcessID, int32 PriorityModifier, const TCHAR* OptionalWorkingDirectory, void* PipeWrite );
 	static bool IsProcRunning( FProcHandle & ProcessHandle );
 	static void WaitForProc( FProcHandle & ProcessHandle );
+	static void CloseProc( FProcHandle & ProcessHandle );
 	static void TerminateProc( FProcHandle & ProcessHandle, bool KillTree = false );
 	static bool GetProcReturnCode( FProcHandle & ProcHandle, int32* ReturnCode );
 	static bool IsApplicationRunning( const TCHAR* ProcName );
diff --git a/Engine/Source/Runtime/Core/Public/Windows/WindowsPlatformProcess.h b/Engine/Source/Runtime/Core/Public/Windows/WindowsPlatformProcess.h
--- a/Engine/Source/Runtime/Core/Public/Windows/WindowsPlatformProcess.h
+++ b/Engine/Source/Runtime/Core/Public/Windows/WindowsPlatformProcess.h
@@ -24,19 +24,8 @@ public:
 		: TProcHandle( Other )
 	{}
 
-	FORCEINLINE bool Close()
-	{
-		if( IsValid() )
-		{
-			::CloseHandle( Handle );
-			Reset();
-			return true;
-		}
-		else
-		{
-			return false;
-		}
-	}
+	DEPRECATED(4.8, "FProcHandle::Close() is redundant - handles created with FPlatformProcess::CreateProc() should be closed with FPlatformProcess::CloseProc().")
+	FORCEINLINE bool Close();
 };
 
 
@@ -100,6 +89,7 @@ public:
 	static FProcHandle CreateProc( const TCHAR* URL, const TCHAR* Parms, bool bLaunchDetached, bool bLaunchHidden, bool bLaunchReallyHidden, uint32* OutProcessID, int32 PriorityModifier, const TCHAR* OptionalWorkingDirectory, void* PipeWrite );
 	static bool IsProcRunning( FProcHandle & ProcessHandle );
 	static void WaitForProc( FProcHandle & ProcessHandle );
+	static void CloseProc( FProcHandle & ProcessHandle );
 	static void TerminateProc( FProcHandle & ProcessHandle, bool KillTree = false );
 	static bool GetProcReturnCode( FProcHandle & ProcHandle, int32* ReturnCode );
 	static bool GetApplicationMemoryUsage(uint32 ProcessId, SIZE_T* OutMemoryUsage);
@@ -148,6 +138,16 @@ private:
 
 typedef FWindowsPlatformProcess FPlatformProcess;
 
+inline bool FProcHandle::Close()
+{
+	if (IsValid())
+	{
+		FPlatformProcess::CloseProc(*this);
+		return true;
+	}
+	return false;
+}
+
 #include "WindowsCriticalSection.h"
 
 typedef FWindowsCriticalSection FCriticalSection;
