Github PR #915 (by PaulEremeeff)

diff --git a/Engine/Source/Runtime/Core/Public/Containers/UnrealString.h b/Engine/Source/Runtime/Core/Public/Containers/UnrealString.h
--- a/Engine/Source/Runtime/Core/Public/Containers/UnrealString.h
+++ b/Engine/Source/Runtime/Core/Public/Containers/UnrealString.h
@@ -1746,7 +1746,7 @@ public:
 			}
 			LineCount++;
 			InData += TermLength;
-		} while (InData);
+		}
 
 		if (bAutoEmitLineTerminator)
 		{
