Implemented FMacPlatformProcess::CloseProc()

diff --git a/Engine/Source/Runtime/Core/Private/Mac/MacPlatformProcess.cpp b/Engine/Source/Runtime/Core/Private/Mac/MacPlatformProcess.cpp
--- a/Engine/Source/Runtime/Core/Private/Mac/MacPlatformProcess.cpp
+++ b/Engine/Source/Runtime/Core/Private/Mac/MacPlatformProcess.cpp
@@ -501,7 +501,8 @@ void FMacPlatformProcess::WaitForProc( FProcHandle& ProcessHandle )
 
 void FMacPlatformProcess::CloseProc( FProcHandle & ProcessHandle )
 {
-	// is any action needed here?
+	SCOPED_AUTORELEASE_POOL;
+	[(NSTask*)ProcessHandle.Get() release];
 }
 
 void FMacPlatformProcess::TerminateProc( FProcHandle& ProcessHandle, bool KillTree )
@@ -534,7 +535,6 @@ void FMacPlatformProcess::TerminateProc( FProcHandle& ProcessHandle, bool KillTr
 	}
 
 	[(NSTask*)ProcessHandle.Get() terminate];
-	[(NSTask*)ProcessHandle.Get() release];
 }
 
 uint32 FMacPlatformProcess::GetCurrentProcessId()
