Moved the call to CloseProc in SwarmInterfaceLocal to the right spot and added missing preoc handle reset to Mac CloseProc

diff --git a/Engine/Source/Runtime/Core/Private/Mac/MacPlatformProcess.cpp b/Engine/Source/Runtime/Core/Private/Mac/MacPlatformProcess.cpp
--- a/Engine/Source/Runtime/Core/Private/Mac/MacPlatformProcess.cpp
+++ b/Engine/Source/Runtime/Core/Private/Mac/MacPlatformProcess.cpp
@@ -503,6 +503,7 @@ void FMacPlatformProcess::CloseProc( FProcHandle & ProcessHandle )
 {
 	SCOPED_AUTORELEASE_POOL;
 	[(NSTask*)ProcessHandle.Get() release];
+	ProcessHandle.Reset();
 }
 
 void FMacPlatformProcess::TerminateProc( FProcHandle& ProcessHandle, bool KillTree )
