Core: GenericPlatform documentation cleanup pass

diff --git a/Engine/Source/Runtime/Core/Private/GenericPlatform/GenericPlatformAtomics.cpp b/Engine/Source/Runtime/Core/Private/GenericPlatform/GenericPlatformAtomics.cpp
--- a/Engine/Source/Runtime/Core/Private/GenericPlatform/GenericPlatformAtomics.cpp
+++ b/Engine/Source/Runtime/Core/Private/GenericPlatform/GenericPlatformAtomics.cpp
@@ -1,10 +1,7 @@
 // Copyright 1998-2015 Epic Games, Inc. All Rights Reserved.
 
-/*=============================================================================
-	GenericPlatformAtomics.cpp: Generic implementations of atomic functions
-=============================================================================*/
-
 #include "CorePrivatePCH.h"
 
+
 static_assert(sizeof(FInt128) == 16, "FInt128 size must be 16 bytes.");
-static_assert(ALIGNOF(FInt128) == 16, "FInt128 alignment must equals 16.");
\ No newline at end of file
+static_assert(ALIGNOF(FInt128) == 16, "FInt128 alignment must equals 16.");
diff --git a/Engine/Source/Runtime/Core/Public/GenericPlatform/GenericPlatform.h b/Engine/Source/Runtime/Core/Public/GenericPlatform/GenericPlatform.h
--- a/Engine/Source/Runtime/Core/Public/GenericPlatform/GenericPlatform.h
+++ b/Engine/Source/Runtime/Core/Public/GenericPlatform/GenericPlatform.h
@@ -1,12 +1,8 @@
 // Copyright 1998-2015 Epic Games, Inc. All Rights Reserved.
 
-
-/*================================================================================
-	GenericPlatform.h: Generic platform classes, mostly implemented with ANSI C++
-==================================================================================*/
-
 #pragma once
 
+
 //---------------------------------------------------------------------
 // Utility for automatically setting up the pointer-sized integer type
 //---------------------------------------------------------------------
diff --git a/Engine/Source/Runtime/Core/Public/GenericPlatform/GenericPlatformAtomics.h b/Engine/Source/Runtime/Core/Public/GenericPlatform/GenericPlatformAtomics.h
--- a/Engine/Source/Runtime/Core/Public/GenericPlatform/GenericPlatformAtomics.h
+++ b/Engine/Source/Runtime/Core/Public/GenericPlatform/GenericPlatformAtomics.h
@@ -1,13 +1,10 @@
 // Copyright 1998-2015 Epic Games, Inc. All Rights Reserved.
 
-
-/*=============================================================================================
-	GenericPlatformAtomics.h: Generic platform Atomics classes
-==============================================================================================*/
-
 #pragma once
+
 #include "HAL/Platform.h"
 
+
 /** Helper struct used to hold 128-bit values, internally represented as two 64-bit integers. */
 struct MS_ALIGN(16) FInt128 
 {
@@ -18,6 +15,7 @@ struct MS_ALIGN(16) FInt128
 } 
 GCC_ALIGN(16);
 
+
 /**
 * Generic implementation...you are required to implement at least FPlatformAtomics::InterlockedCompareExchange but they are listed here 
 * to provide a base implementation in terms of FPlatformAtomics::InterlockedCompareExchange and easy copy-paste
@@ -225,4 +223,3 @@ protected:
 		return !(PTRINT(Ptr) & (Alignment - 1));
 	}
 };
-
diff --git a/Engine/Source/Runtime/Core/Public/GenericPlatform/GenericPlatformMisc.h b/Engine/Source/Runtime/Core/Public/GenericPlatform/GenericPlatformMisc.h
--- a/Engine/Source/Runtime/Core/Public/GenericPlatform/GenericPlatformMisc.h
+++ b/Engine/Source/Runtime/Core/Public/GenericPlatform/GenericPlatformMisc.h
@@ -1,14 +1,11 @@
 // Copyright 1998-2015 Epic Games, Inc. All Rights Reserved.
 
-
-/*=============================================================================================
-	GenericPlatformMisc.h: Generic platform misc classes, mostly implemented with ANSI C++
-==============================================================================================*/
-
 #pragma once
+
 #include "Containers/ContainersFwd.h"
 #include "HAL/Platform.h"
 
+
 struct FGenericCrashContext;
 struct FGenericMemoryWarningContext;
 
@@ -28,42 +25,29 @@ namespace EBuildConfigurations
 	 */
 	enum Type
 	{
-		/**
-		 * Unknown build configuration.
-		 */
+		/** Unknown build configuration. */
 		Unknown,
 
-		/**
-		 * Debug build.
-		 */
+		/** Debug build. */
 		Debug,
 
-		/**
-		 * DebugGame build.
-		 */
+		/** DebugGame build. */
 		DebugGame,
 
-		/**
-		 * Development build.
-		 */
+		/** Development build. */
 		Development,
 
-		/**
-		 * Shipping build.
-		 */
+		/** Shipping build. */
 		Shipping,
 
-		/**
-		 * Test build.
-		 */
+		/** Test build. */
 		Test
 	};
 
 	/**
 	 * Returns the string representation of the specified EBuildConfiguration value.
 	 *
-	 * @param Configuration - The string to get the EBuildConfiguration::Type for.
-	 *
+	 * @param Configuration The string to get the EBuildConfiguration::Type for.
 	 * @return An EBuildConfiguration::Type value.
 	 */
 	CORE_API EBuildConfigurations::Type FromString( const FString& Configuration );
@@ -71,8 +55,7 @@ namespace EBuildConfigurations
 	/**
 	 * Returns the string representation of the specified EBuildConfiguration value.
 	 *
-	 * @param Configuration - The value to get the string for.
-	 *
+	 * @param Configuration The value to get the string for.
 	 * @return The string representation.
 	 */
 	CORE_API const TCHAR* ToString( EBuildConfigurations::Type Configuration );
@@ -80,8 +63,7 @@ namespace EBuildConfigurations
 	/**
 	 * Returns the localized text representation of the specified EBuildConfiguration value.
 	 *
-	 * @param Configuration - The value to get the text for.
-	 *
+	 * @param Configuration The value to get the text for.
 	 * @return The localized Build configuration text
 	 */
 	CORE_API FText ToText( EBuildConfigurations::Type Configuration );
@@ -95,32 +77,23 @@ namespace EBuildTargets
 	 */
 	enum Type
 	{
-		/**
-		 * Unknown build target.
-		 */
+		/** Unknown build target. */
 		Unknown,
 
-		/**
-		 * Editor target.
-		 */
+		/** Editor target. */
 		Editor,
 
-		/**
-		 * Game target.
-		 */
+		/** Game target. */
 		Game,
 
-		/**
-		 * Server target.
-		 */
+		/** Server target. */
 		Server
 	};
 
 	/**
 	 * Returns the string representation of the specified EBuildTarget value.
 	 *
-	 * @param Target - The string to get the EBuildTarget::Type for.
-	 *
+	 * @param Target The string to get the EBuildTarget::Type for.
 	 * @return An EBuildTarget::Type value.
 	 */
 	CORE_API EBuildTargets::Type FromString( const FString& Target );
@@ -128,8 +101,7 @@ namespace EBuildTargets
 	/**
 	 * Returns the string representation of the specified EBuildTarget value.
 	 *
-	 * @param Target - The value to get the string for.
-	 *
+	 * @param Target The value to get the string for.
 	 * @return The string representation.
 	 */
 	CORE_API const TCHAR* ToString( EBuildTargets::Type Target );
@@ -192,6 +164,7 @@ namespace EAppReturnType
 	};
 }
 
+
 /**
 * Generic implementation for most platforms
 **/
@@ -201,20 +174,13 @@ struct CORE_API FGenericPlatformMisc
 	 * Called during appInit() after cmd line setup
 	 */
 	static void PlatformPreInit();
-	static void PlatformInit()
-	{
-	}
-	static void PlatformPostInit(bool ShowSplashScreen = false)
-	{
-	}
+	static void PlatformInit() { }
+	static void PlatformPostInit(bool ShowSplashScreen = false) { }
 
 	/**
-	* Called during AppExit(). Log, Config still exist at this point, but not much else does.
-	*/
-	static void PlatformTearDown()
-	{
-	}
-
+	 * Called during AppExit(). Log, Config still exist at this point, but not much else does.
+	 */
+	static void PlatformTearDown() { }
 	static GenericApplication* CreateApplication();
 
 	static void* GetHardwareWindow()
@@ -222,21 +188,14 @@ struct CORE_API FGenericPlatformMisc
 		return nullptr;
 	}
 
-	/**
-	 * Set/restore the Console Interrupt (Control-C, Control-Break, Close) handler
-	 */
-	static void SetGracefulTerminationHandler()
-	{
-	}
+	/** Set/restore the Console Interrupt (Control-C, Control-Break, Close) handler. */
+	static void SetGracefulTerminationHandler() { }
 
 	/**
 	 * Installs handler for the unexpected (due to error) termination of the program,
 	 * including, but not limited to, crashes.
-	 *
 	 */
-	static void SetCrashHandler(void (* CrashHandler)(const FGenericCrashContext& Context))
-	{
-	}
+	static void SetCrashHandler(void (* CrashHandler)(const FGenericCrashContext& Context)) { }
 
 	/**
 	 * Retrieve a environment variable from the system
@@ -252,11 +211,11 @@ struct CORE_API FGenericPlatformMisc
 
 	
 	/**
-	* Sets an environment variable to the local process's environment
-	*
-	* @param VariableName The name of the variable (ie "Path")
-	* @param Value The string to set the variable to.	
-	*/
+	 * Sets an environment variable to the local process's environment
+	 *
+	 * @param VariableName The name of the variable (ie "Path")
+	 * @param Value The string to set the variable to.	
+	 */
 	static void SetEnvironmentVar(const TCHAR* VariableName, const TCHAR* Value);
 
 	/**
@@ -407,6 +366,7 @@ struct CORE_API FGenericPlatformMisc
 	}
 
 protected:
+
 	/**
 	* Retrieves some standard key code mappings (usually called by a subclass's GetCharKeyMap)
 	 *
@@ -416,9 +376,8 @@ protected:
 	 */
 	static uint32 GetStandardPrintableKeyMap(uint16* KeyCodes, FString* KeyNames, uint32 MaxMappings, bool bMapUppercaseKeys, bool bMapLowercaseKeys);
 
-
-
 public:
+
 	/**
 	 * Platform specific function for adding a named event that can be viewed in PIX
 	 */
@@ -450,7 +409,6 @@ public:
 	 *	@param	InSectionName		The section that this key->value pair is placed within (can contain / separators, eg UserDetails/AccountInfo)
 	 *	@param	InKeyName			The name of the key to set the value for
 	 *	@param	InValue				The value to set
-	 *
 	 *	@return	bool				true if the value was set correctly, false if not
 	 */
 	static bool SetStoredValue(const FString& InStoreId, const FString& InSectionName, const FString& InKeyName, const FString& InValue);
@@ -463,7 +421,6 @@ public:
 	 *	@param	InSectionName		The section that this key->value pair is placed within (can contain / separators, eg UserDetails/AccountInfo)
 	 *	@param	InKeyName			The name of the key to get the value for
 	 *	@param	OutValue			The value found
-	 *
 	 *	@return	bool				true if the entry was found (and OutValue contains the result), false if not
 	 */
 	static bool GetStoredValue(const FString& InStoreId, const FString& InSectionName, const FString& InKeyName, FString& OutValue);
@@ -479,10 +436,7 @@ public:
 	static void LocalPrint( const TCHAR* Str );
 
 
-	/**
-	 * Request application to minimize (goto background)
-	 *
-	 **/
+	/** Request application to minimize (goto background). **/
 	static void RequestMinimize();
 
 	/**
@@ -504,6 +458,7 @@ public:
 
 	/** Copies text to the operating system clipboard. */
 	static void ClipboardCopy(const TCHAR* Str);
+
 	/** Pastes in text from the operating system clipboard. */
 	static void ClipboardPaste(class FString& Dest);
 
@@ -512,12 +467,11 @@ public:
 
 	/** 
 	 * Show a message box if possible, otherwise print a message and return the default
-	 * @param MsgType - what sort of options are provided
-	 * @param Text - specific message
-	 * @param Caption - string indicating the title of the message box
+	 * @param MsgType What sort of options are provided
+	 * @param Text Specific message
+	 * @param Caption String indicating the title of the message box
 	 * @return Very strange convention...not really EAppReturnType, see implementation
-
-	**/
+	 */
 	static EAppReturnType::Type MessageBoxExt( EAppMsgType::Type MsgType, const TCHAR* Text, const TCHAR* Caption );
 
 	/**
@@ -533,6 +487,7 @@ public:
 		Disable,
 		Enable
 	};
+
 	/**
 	 * Disables screensaver (if platform supports such an API)
 	 *
@@ -565,7 +520,8 @@ public:
 
 	/**
 	 * Checks structure of the path against platform formatting requirements
-	 * return - true if path is formatted validly
+	 *
+	 * return true if path is formatted validly
 	 */
 	static bool IsValidAbsolutePathFormat(const FString& Path)
 	{
@@ -581,12 +537,13 @@ public:
 	}
 
 	/**
-	* @return platform specific path separator.
-	*/
+	 * @return platform specific path separator.
+	 */
 	static const TCHAR* GetDefaultPathSeparator();
 
 	/**
 	 * Checks if platform wants to allow a rendering thread on current device (note: does not imply it will, only if okay given other criteria met)
+	 *
 	 * @return true if allowed, false if shouldn't use a separate rendering thread
 	 */
 	static bool AllowRenderThread()
@@ -645,6 +602,7 @@ public:
 	static void LoadPreInitModules()
 	{
 	}
+
 	/**
 	 * Load the platform-specific startup modules
 	 */
@@ -674,7 +632,6 @@ public:
 	 * @param ComandType OS hint as to the type of command 
 	 * @param Command the command to execute
 	 * @param CommandLine the commands to pass to the executable
-	 *
 	 * @return whether the command was successful or not
 	 */
 	static bool OsExecute(const TCHAR* CommandType, const TCHAR* Command, const TCHAR* CommandLine = NULL)
@@ -716,7 +673,6 @@ public:
 	 *  @param	InWorld		World context
 	 *	@param	Cmd			The command to execute
 	 *	@param	Out			The output device to utilize
-	 *
 	 *	@return	bool		true if command was processed, false if not
 	 */
 	static bool Exec(UWorld* InWorld, const TCHAR* Cmd, FOutputDevice& Out)
@@ -732,7 +688,6 @@ public:
 	 *
 	 * @param	InScreenPos		The screen coords to sample for current pixel color
 	 * @param	InGamma			Optional gamma correction to apply to the screen color
-	 *
 	 * @return					The color of the pixel displayed at the chosen location
 	 */
 	static struct FLinearColor GetScreenPixelColor(const struct FVector2D& InScreenPos, float InGamma = 1.0f);
@@ -849,4 +804,3 @@ protected:
 	static bool bPromptForRemoteDebugOnEnsure;
 #endif	//#if !UE_BUILD_SHIPPING
 };
-
diff --git a/Engine/Source/Runtime/Core/Public/GenericPlatform/GenericPlatformNamedPipe.h b/Engine/Source/Runtime/Core/Public/GenericPlatform/GenericPlatformNamedPipe.h
--- a/Engine/Source/Runtime/Core/Public/GenericPlatform/GenericPlatformNamedPipe.h
+++ b/Engine/Source/Runtime/Core/Public/GenericPlatform/GenericPlatformNamedPipe.h
@@ -1,87 +1,82 @@
 // Copyright 1998-2015 Epic Games, Inc. All Rights Reserved.
 
-/*=============================================================================================
-	GenericPlatformNamedPipe.h: Declares the FGenericPlatformNamedPipe class
-==============================================================================================*/
-
 #pragma once
 
 #include "HAL/Platform.h"
 
+
 class FString;
 
+
 #if PLATFORM_SUPPORTS_NAMED_PIPES
 
-// Very initial wrapper for platform named pipe communication
+/**
+ * Wrapper for platform named pipe communication.
+ */
 class CORE_API FGenericPlatformNamedPipe
 {
 public:
 
-	/**
-	 * Default constructor.
-	 */
+	/** Default constructor. */
 	FGenericPlatformNamedPipe();
 
-	/**
-	 * Virtual destructor.
-	 */
+	/** Virtual destructor. */
 	virtual ~FGenericPlatformNamedPipe();
 
-
 public:
 
-	// Create a named pipe as a server or client, using overlapped IO if bAsync=1
+	/** Create a named pipe as a server or client, using overlapped IO if bAsync=1. */
 	virtual bool Create(const FString& PipeName, bool bServer, bool bAsync) { return false; }
 
-	// Get the pipe name used on Create()
+	/** Get the pipe name used on Create(). */
 	virtual const FString& GetName() const
 	{
 		return *NamePtr;
 	}
 
-	// Destroy the pipe
+	/** Destroy the pipe. */
 	virtual bool Destroy() { return false; }
 
-	// Open a connection from a client
+	/** Open a connection from a client. */
 	virtual bool OpenConnection() { return false; }
 
-	// Blocks if there's an IO operation in progress until it's done or errors out
+	/** Block if there's an IO operation in progress until it's done or errors out. */
 	virtual bool BlockForAsyncIO() { return false; }
 
-	// Lets the user know if the pipe is ready to send or receive data
+	/** Let the user know if the pipe is ready to send or receive data. */
 	virtual bool IsReadyForRW() const { return false; }
 
-	// Updates status of async state of the current pipe
+	/** Update status of async state of the current pipe. */
 	virtual bool UpdateAsyncStatus() { return false; }
 
-	// Writes a buffer out
+	/** Write a buffer out. */
 	virtual bool WriteBytes(int32 NumBytes, const void* Data) { return false; }
 
-	// Writes a single int32
+	/** Write a single int32. */
 	inline bool WriteInt32(int32 In)
 	{
 		return WriteBytes(sizeof(int32), &In);
 	}
 
-	// Reads a buffer in
+	/* Read a buffer in. */
 	virtual bool ReadBytes(int32 NumBytes, void* OutData) { return false; }
 
-	// Reads a single int32
+	/** Read a single int32. */
 	inline bool ReadInt32(int32& Out)
 	{
 		return ReadBytes(sizeof(int32), &Out);
 	}
 
-	// Returns true if the pipe has been created and hasn't been destroyed
+	/** Returns true if the pipe has been created and hasn't been destroyed. */
 	virtual bool IsCreated() const { return false; }
 
-	// Return true if the pipe has had any communication error
+	/** Return true if the pipe has had any communication error. */
 	virtual bool HasFailed() const { return true; }
 
-
 protected:
 
 	FString* NamePtr;
 };
 
+
 #endif
diff --git a/Engine/Source/Runtime/Core/Public/GenericPlatform/GenericPlatformOutputDevices.h b/Engine/Source/Runtime/Core/Public/GenericPlatform/GenericPlatformOutputDevices.h
--- a/Engine/Source/Runtime/Core/Public/GenericPlatform/GenericPlatformOutputDevices.h
+++ b/Engine/Source/Runtime/Core/Public/GenericPlatform/GenericPlatformOutputDevices.h
@@ -1,41 +1,37 @@
 // Copyright 1998-2015 Epic Games, Inc. All Rights Reserved.
 
-
-/*=============================================================================================
-	GenericPlatformOutputDevices.h: Generic platform OutputDevices classes, mostly implemented with ANSI C++
-==============================================================================================*/
-
 #pragma once
+
 #include "HAL/Platform.h"
 
+
 class FFeedbackContext;
 class FOutputDevice;
 class FOutputDeviceConsole;
 class FOutputDeviceError;
 class FString;
 
+
 /**
-* Generic implementation for most platforms
-**/
+ * Generic implementation for most platforms
+ */
 struct CORE_API FGenericPlatformOutputDevices
 {
-	/**
-	 * Adds output devices which can vary depending on platform, configuration, command line parameters.
-	 */
+	/** Add output devices which can vary depending on platform, configuration, command line parameters. */
 	static void							SetupOutputDevices();
-
 	static FString						GetAbsoluteLogFilename();
-
 	static FOutputDevice*				GetLog();
+
 	static FOutputDevice*				GetEventLog()
 	{
 		return nullptr; // normally only used for dedicated servers
 	}
+
 	static FOutputDeviceConsole*		GetLogConsole()
 	{
 		return nullptr; // normally only used for PC
 	}
+
 	static FOutputDeviceError*			GetError();
 	static FFeedbackContext*			GetWarn();
 };
-
diff --git a/Engine/Source/Runtime/Core/Public/GenericPlatform/GenericPlatformProperties.h b/Engine/Source/Runtime/Core/Public/GenericPlatform/GenericPlatformProperties.h
--- a/Engine/Source/Runtime/Core/Public/GenericPlatform/GenericPlatformProperties.h
+++ b/Engine/Source/Runtime/Core/Public/GenericPlatform/GenericPlatformProperties.h
@@ -1,19 +1,17 @@
 // Copyright 1998-2015 Epic Games, Inc. All Rights Reserved.
 
-/*=============================================================================================
-	GenericPlatformProperties.h: Generic platform properties, used a base class for platforms.
-		These are shared between:
-		the runtime platform - via FPlatformProperties
-		the target platforms - via ITargetPlatform
-
-==============================================================================================*/
-
 #pragma once
+
 #include "GenericPlatform/GenericPlatformMisc.h"
 #include "HAL/Platform.h"
 
+
 /**
  * Base class for platform properties.
+ *
+ * These are shared between:
+ *     the runtime platform - via FPlatformProperties
+ *     the target platforms - via ITargetPlatform
  */
 struct FGenericPlatformProperties
 {
@@ -144,8 +142,7 @@ struct FGenericPlatformProperties
 	/**
 	 * Checks whether the specified build target is supported.
 	 *
-	 * @param BuildTarget - The build target to check.
-	 *
+	 * @param BuildTarget The build target to check.
 	 * @return true if the build target is supported, false otherwise.
 	 */
 	static FORCEINLINE bool SupportsBuildTarget( EBuildTargets::Type BuildTarget )
diff --git a/Engine/Source/Runtime/Core/Public/GenericPlatform/GenericPlatformSplash.h b/Engine/Source/Runtime/Core/Public/GenericPlatform/GenericPlatformSplash.h
--- a/Engine/Source/Runtime/Core/Public/GenericPlatform/GenericPlatformSplash.h
+++ b/Engine/Source/Runtime/Core/Public/GenericPlatform/GenericPlatformSplash.h
@@ -1,12 +1,10 @@
 // Copyright 1998-2015 Epic Games, Inc. All Rights Reserved.
 
-/*=============================================================================================
-	GenericPlatformSplash.h: Generic platform splash screen...does nothing
-==============================================================================================*/
-
 #pragma once
+
 #include "HAL/Platform.h"
 
+
 /**
  * SplashTextType defines the types of text on the splash screen
  */
@@ -33,23 +31,17 @@ namespace SplashTextType
 	};
 }
 
+
 /**
-* Generic implementation for most platforms
-**/
+ * Generic implementation for most platforms
+ */
 struct CORE_API FGenericPlatformSplash
 {
-	/**
-	* Show the splash screen
-	*/
-	FORCEINLINE static void Show()
-	{
-	}
-	/**
-	* Hide the splash screen
-	*/
-	FORCEINLINE static void Hide()
-	{
-	}
+	/** Show the splash screen. */
+	FORCEINLINE static void Show() { }
+
+	/** Hide the splash screen. */
+	FORCEINLINE static void Hide() { }
 
 	/**
 	 * Sets the text displayed on the splash screen (for startup/loading progress)
@@ -69,5 +61,4 @@ struct CORE_API FGenericPlatformSplash
 	{
 		return true;
 	}
-
 };
diff --git a/Engine/Source/Runtime/Core/Public/GenericPlatform/GenericPlatformStackWalk.h b/Engine/Source/Runtime/Core/Public/GenericPlatform/GenericPlatformStackWalk.h
--- a/Engine/Source/Runtime/Core/Public/GenericPlatform/GenericPlatformStackWalk.h
+++ b/Engine/Source/Runtime/Core/Public/GenericPlatform/GenericPlatformStackWalk.h
@@ -1,20 +1,18 @@
 // Copyright 1998-2015 Epic Games, Inc. All Rights Reserved.
 
-
-/*=============================================================================================
-	GenericPlatformStackWalk.h: Generic platform stack walking functions...do not do anything
-==============================================================================================*/
-
 #pragma once
+
 #include "HAL/Platform.h"
 
+
 class FString;
 struct FGenericCrashContext;
 
+
 /**
-* This is used to capture all of the module information needed to load pdb's.
-* @todo, the memory profiler should not be using this as platform agnostic
-*/
+ * This is used to capture all of the module information needed to load pdb's.
+ * @todo, the memory profiler should not be using this as platform agnostic
+ */
 struct FStackWalkModuleInfo
 {
 	uint64 BaseOfImage;
@@ -72,9 +70,10 @@ struct FProgramCounterSymbolInfo final
 
 struct FProgramCounterSymbolInfoEx;
 
+
 /**
  * Generic implementation for most platforms
- **/
+ */
 struct CORE_API FGenericPlatformStackWalk
 {
 	typedef FGenericPlatformStackWalk Base;
diff --git a/Engine/Source/Runtime/Core/Public/GenericPlatform/GenericPlatformStricmp.h b/Engine/Source/Runtime/Core/Public/GenericPlatform/GenericPlatformStricmp.h
--- a/Engine/Source/Runtime/Core/Public/GenericPlatform/GenericPlatformStricmp.h
+++ b/Engine/Source/Runtime/Core/Public/GenericPlatform/GenericPlatformStricmp.h
@@ -1,11 +1,13 @@
 // Copyright 1998-2015 Epic Games, Inc. All Rights Reserved.
 
 #pragma once
+
 #include "HAL/Platform.h"
 #include "Misc/Char.h"
 #include "Templates/EnableIf.h"
 #include "Templates/UnrealTypeTraits.h"
 
+
 /**
  * This trait tells if given CharTypeA is comparison compatible with CharTypeB, i.e.
  * if CharTypeA contains whole CharTypeB character set on the same positions.
@@ -31,7 +33,6 @@ private:
 	 *
 	 * @param String1 First string to compare.
 	 * @param String2 Second string to compare.
-	 *
 	 * @returns Zero if both strings are equal. Greater than zero if first
 	 *          string is greater than the second one. Less than zero
 	 *          otherwise.
@@ -90,7 +91,6 @@ public:
 	 *
 	 * @param String1 First string to compare.
 	 * @param String2 Second string to compare.
-	 *
 	 * @returns Zero if both strings are equal. Greater than zero if first
 	 *          string is greater than the second one. Less than zero
 	 *          otherwise.
@@ -100,4 +100,4 @@ public:
 	{
 		return FGenericPlatformStricmp::StricmpImpl(String1, String2);
 	}
-};
\ No newline at end of file
+};
diff --git a/Engine/Source/Runtime/Core/Public/GenericPlatform/GenericPlatformString.h b/Engine/Source/Runtime/Core/Public/GenericPlatform/GenericPlatformString.h
--- a/Engine/Source/Runtime/Core/Public/GenericPlatform/GenericPlatformString.h
+++ b/Engine/Source/Runtime/Core/Public/GenericPlatform/GenericPlatformString.h
@@ -1,18 +1,14 @@
 // Copyright 1998-2015 Epic Games, Inc. All Rights Reserved.
 
-
-/*=============================================================================================
-	GenericPlatformString.h: Generic platform string classes, mostly implemented with ANSI C++
-==============================================================================================*/
-
 #pragma once
+
 #include "HAL/Platform.h"
 #include "Templates/EnableIf.h"
 
+
 /**
  * Generic string implementation for most platforms
- *
- **/
+ */
 struct FGenericPlatformString
 {
 	/**
@@ -166,7 +162,6 @@ struct FGenericPlatformString
 		return Dest;
 	}
 
-
 	/**
 	 * Returns the required buffer length for the [Src, Src+SrcSize) string when converted to the DestChar encoding.
 	 * The Src range should contain a null terminator if a null terminator is required in the output.
@@ -204,6 +199,7 @@ private:
 	static CORE_API void LogBogusChars(const SourceEncoding* Src, int32 SrcSize);
 };
 
+
 /**
  * Specialization of IsValidChar for ANSICHARs.
  */
diff --git a/Engine/Source/Runtime/Core/Public/GenericPlatform/GenericPlatformSurvey.h b/Engine/Source/Runtime/Core/Public/GenericPlatform/GenericPlatformSurvey.h
--- a/Engine/Source/Runtime/Core/Public/GenericPlatform/GenericPlatformSurvey.h
+++ b/Engine/Source/Runtime/Core/Public/GenericPlatform/GenericPlatformSurvey.h
@@ -1,13 +1,10 @@
 // Copyright 1998-2015 Epic Games, Inc. All Rights Reserved.
 
-
-/*=============================================================================================
-	GenericPlatformSurvey.h: Generic platform hardware-survey classes
-==============================================================================================*/
-
 #pragma once
+
 #include "HAL/Platform.h"
 
+
 // time and amount of work that was measured
 struct FTimeSample
 {
@@ -23,6 +20,7 @@ struct FTimeSample
 	float NormalizedTime;
 };
 
+
 struct FSynthBenchmarkStat
 {
 	FSynthBenchmarkStat()
@@ -32,8 +30,7 @@ struct FSynthBenchmarkStat
 		, IndexNormalizedTime(-1)
 		, ValueType(0)
 		, Confidence(0)
-	{
-	}
+	{ }
 
 	// @param InDesc descriptions
 	FSynthBenchmarkStat(const TCHAR* InDesc, float InIndexNormalizedTime, const TCHAR* InValueType)
@@ -93,6 +90,7 @@ struct FSynthBenchmarkStat
 	}
 
 private:
+
 	// 0 if not valid
 	const TCHAR *Desc;
 	// -1 if not defined, in seconds, useful to see if a test did run too long (some slower GPUs might timeout)
@@ -107,6 +105,7 @@ private:
 	float Confidence;
 };
 
+
 struct FSynthBenchmarkResults 
 {
 	FSynthBenchmarkStat CPUStats[2];
@@ -157,6 +156,7 @@ struct FSynthBenchmarkResults
 	}
 };
 
+
 struct FHardwareDisplay	
 {
 	static const uint32 MaxStringLength = 260;
@@ -169,6 +169,7 @@ struct FHardwareDisplay
 	TCHAR GPUDriverVersion[MaxStringLength];
 };
 
+
 struct FHardwareSurveyResults
 {
 	static const int32 MaxDisplayCount = 8; 
@@ -211,6 +212,7 @@ struct FHardwareSurveyResults
 	FSynthBenchmarkResults SynthBenchmark;
 };
 
+
 /**
 * Generic implementation for most platforms, these tend to be unused and unimplemented
 **/
@@ -222,11 +224,10 @@ struct CORE_API FGenericPlatformSurvey
 	 *
 	 * @param OutResults	The struct that receives the results if available.
 	 * @param bWait			If true, the function won't return until the results are available or the survey fails. Defaults to false.
-	 *
 	 * @return				True if the results were available, false otherwise.
-	 **/
+	 */
 	static bool GetSurveyResults( FHardwareSurveyResults& OutResults, bool bWait = false )
 	{
 		return false;
 	}
-};
\ No newline at end of file
+};
diff --git a/Engine/Source/Runtime/Core/Public/GenericPlatform/GenericPlatformSymbolication.h b/Engine/Source/Runtime/Core/Public/GenericPlatform/GenericPlatformSymbolication.h
--- a/Engine/Source/Runtime/Core/Public/GenericPlatform/GenericPlatformSymbolication.h
+++ b/Engine/Source/Runtime/Core/Public/GenericPlatform/GenericPlatformSymbolication.h
@@ -1,14 +1,11 @@
 // Copyright 1998-2015 Epic Games, Inc. All Rights Reserved.
 
-
-/*=============================================================================================
-	GenericPlatformSymbolication.h: Generic platform symbolication functions.
-==============================================================================================*/
-
 #pragma once
+
 #include "HAL/Platform.h"
 #include "Serialization/ArchiveBase.h"
 
+
 /**
  * File & line info for a debug symbol region
  */
@@ -31,6 +28,7 @@ struct FGenericPlatformSymbolInfo
 	}
 };
 
+
 /**
  * Debug symbol information
  */
@@ -53,6 +51,7 @@ struct FGenericPlatformSymbolData
 	}
 };
 
+
 /**
  * Container for debug symbols corresponding to a single binary file
  */
@@ -75,6 +74,7 @@ struct FGenericPlatformSymbolDatabase
 	}
 };
 
+
 struct FGenericPlatformSymbolDatabaseKeyFuncs
 {
 	enum { bAllowDuplicateKeys = 0 };
@@ -104,8 +104,10 @@ struct FGenericPlatformSymbolDatabaseKeyFuncs
 	}
 };
 
+
 typedef TSet<FGenericPlatformSymbolDatabase, FGenericPlatformSymbolDatabaseKeyFuncs> FGenericPlatformSymbolDatabaseSet;
 
+
 struct CORE_API FGenericPlatformSymbolication
 {
 	static bool LoadSymbolDatabaseForBinary(FString SourceFolder, FString BinaryPath, FString BinarySignature, FGenericPlatformSymbolDatabase& OutDatabase);
diff --git a/Engine/Source/Runtime/Core/Public/GenericPlatform/GenericPlatformTLS.h b/Engine/Source/Runtime/Core/Public/GenericPlatform/GenericPlatformTLS.h
--- a/Engine/Source/Runtime/Core/Public/GenericPlatform/GenericPlatformTLS.h
+++ b/Engine/Source/Runtime/Core/Public/GenericPlatform/GenericPlatformTLS.h
@@ -1,16 +1,13 @@
 // Copyright 1998-2015 Epic Games, Inc. All Rights Reserved.
 
-
-/*=============================================================================================
-	GenericPlatformTLS.h: Generic platform TLS (Thread local storage and thread ID) classes
-==============================================================================================*/
-
 #pragma once
+
 #include "HAL/Platform.h"
 
+
 /**
  * It should be possible to provide a generic implementation as long as a threadID is provided. We don't do that yet.
-**/
+ */
 struct FGenericPlatformTLS
 {
 #if 0 // provided for reference
@@ -57,4 +54,3 @@ struct FGenericPlatformTLS
 	}
 #endif
 };
-
diff --git a/Engine/Source/Runtime/Core/Public/GenericPlatform/GenericPlatformTime.h b/Engine/Source/Runtime/Core/Public/GenericPlatform/GenericPlatformTime.h
--- a/Engine/Source/Runtime/Core/Public/GenericPlatform/GenericPlatformTime.h
+++ b/Engine/Source/Runtime/Core/Public/GenericPlatform/GenericPlatformTime.h
@@ -1,20 +1,18 @@
 // Copyright 1998-2015 Epic Games, Inc. All Rights Reserved.
 
-
-/*=============================================================================================
-	GenericPlatformTime.h: Generic platform Time classes, mostly implemented with ANSI C++
-==============================================================================================*/
-
 #pragma once
 
 #include "HAL/Platform.h"
 
+
 class FString;
 
+
 #if PLATFORM_HAS_BSD_TIME 
-#include <sys/time.h>
+	#include <sys/time.h>
 #endif
 
+
 /** Contains CPU utilization data. */
 struct FCPUTime
 {
@@ -39,6 +37,7 @@ struct FCPUTime
 	float CPUTimePctRelative;
 };
 
+
 /**
 * Generic implementation for most platforms
 **/
@@ -79,17 +78,17 @@ struct CORE_API FGenericPlatformTime
 	/**
 	 * Get the system date
 	 * 
-	 * @param Dest - destination buffer to copy to
-	 * @param DestSize - size of destination buffer in characters
-	 * @return date string
+	 * @param Dest Destination buffer to copy to
+	 * @param DestSize Size of destination buffer in characters
+	 * @return Date string
 	 */
 	static TCHAR* StrDate( TCHAR* Dest, SIZE_T DestSize );
 	/**
 	 * Get the system time
 	 * 
-	 * @param Dest - destination buffer to copy to
-	 * @param DestSize - size of destination buffer in characters
-	 * @return time string
+	 * @param Dest Destination buffer to copy to
+	 * @param DestSize Size of destination buffer in characters
+	 * @return Time string
 	 */
 	static TCHAR* StrTime( TCHAR* Dest, SIZE_T DestSize );
 
@@ -98,15 +97,15 @@ struct CORE_API FGenericPlatformTime
 	 * NOTE: Only one return value is valid at a time!
 	 *
 	 * @return timestamp string
-	 **/
+	 */
 	static const TCHAR* StrTimestamp();
 
 	/**
 	 * Returns a pretty-string for a time given in seconds. (I.e. "4:31 min", "2:16:30 hours", etc)
 	 *
-	 * @param Seconds	Time in seconds
-	 * @return			Time in a pretty formatted string
-	 **/
+	 * @param Seconds Time in seconds
+	 * @return Time in a pretty formatted string
+	 */
 	static FString PrettyTime( double Seconds );
 
 	/** Updates CPU utilization, called through a delegate from the Core ticker. */
@@ -144,6 +143,6 @@ struct CORE_API FGenericPlatformTime
 	}
 
 protected:
+
 	static double SecondsPerCycle;
 };
-
diff --git a/Engine/Source/Runtime/Core/Public/GenericPlatform/IForceFeedbackSystem.h b/Engine/Source/Runtime/Core/Public/GenericPlatform/IForceFeedbackSystem.h
--- a/Engine/Source/Runtime/Core/Public/GenericPlatform/IForceFeedbackSystem.h
+++ b/Engine/Source/Runtime/Core/Public/GenericPlatform/IForceFeedbackSystem.h
@@ -1,13 +1,19 @@
 // Copyright 1998-2015 Epic Games, Inc. All Rights Reserved.
 
 #pragma once
+
 #include "IInputInterface.h"
 
-// Abstract interface for the force feedback system
-// class is deprecated and will be removed in favor of IInputInterface
-class IForceFeedbackSystem : public IInputInterface
+
+/**
+ * Interface for the force feedback system.
+ *
+ * Note: This class is deprecated and will be removed in favor of IInputInterface
+ */
+class IForceFeedbackSystem
+	: public IInputInterface
 {
 public:
+
 	virtual ~IForceFeedbackSystem() {};
 };
-
diff --git a/Engine/Source/Runtime/Core/Public/GenericPlatform/IInputInterface.h b/Engine/Source/Runtime/Core/Public/GenericPlatform/IInputInterface.h
--- a/Engine/Source/Runtime/Core/Public/GenericPlatform/IInputInterface.h
+++ b/Engine/Source/Runtime/Core/Public/GenericPlatform/IInputInterface.h
@@ -1,8 +1,10 @@
 // Copyright 1998-2015 Epic Games, Inc. All Rights Reserved.
 
 #pragma once
+
 #include "HAL/Platform.h"
 
+
 // General identifiers for potential force feedback channels. These will be mapped according to the
 // platform specific implementation.
 // For example, the PS4 only listens to the XXX_LARGE channels and ignores the rest, while the XBox One could
@@ -16,6 +18,7 @@ enum class FForceFeedbackChannelType
 	RIGHT_SMALL
 };
 
+
 struct FForceFeedbackValues
 {
 	float LeftLarge;
@@ -28,15 +31,19 @@ struct FForceFeedbackValues
 		, LeftSmall(0.f)
 		, RightLarge(0.f)
 		, RightSmall(0.f)
-	{
-	}
+	{ }
 };
 
-// Abstract interface for the input interface.
+
+/**
+ * Interface for the input interface.
+ */
 class IInputInterface
 {
 public:
-	virtual ~IInputInterface() {};
+
+	/** Virtual destructor. */
+	virtual ~IInputInterface() { };
 
 	/**
 	* Sets the strength/speed of the given channel for the given controller id.
@@ -66,4 +73,3 @@ public:
 	 */
 	virtual void SetLightColor(int32 ControllerId, FColor Color) = 0;
 };
-
diff --git a/Engine/Source/Runtime/Core/Public/Windows/WindowsApplication.h b/Engine/Source/Runtime/Core/Public/Windows/WindowsApplication.h
--- a/Engine/Source/Runtime/Core/Public/Windows/WindowsApplication.h
+++ b/Engine/Source/Runtime/Core/Public/Windows/WindowsApplication.h
@@ -54,28 +54,63 @@ namespace ETaskbarProgressState
 }
 
 
+/**
+ * Allows access to task bar lists.
+ *
+ * This class can be used to change the appearance of a window's entry in the windows task bar,
+ * such as setting an overlay icon or showing a progress indicator.
+ */
 class FTaskbarList
 {
 public:
 
+	/**
+	 * Create and initialize a new task bar list.
+	 *
+	 * @return The new task bar list.
+	 */
 	static TSharedRef<FTaskbarList> Create();
 
+	/**
+	 * Sets the overlay icon of a task bar entry.
+	 *
+	 * @param NativeWindow The native window to change the overlay icon for.
+	 * @param Icon The overlay icon to set.
+	 * @param Description The overlay icon's description text.
+	 */
 	void SetOverlayIcon(const TSharedRef<FGenericWindow>& NativeWindow, HICON Icon, FText Description);
 
+	/**
+	 * Sets the progress state of a task bar entry.
+	 *
+	 * @param NativeWindow The native window to change the progress state for.
+	 * @param State The new progress state.
+	 */
 	void SetProgressState(const TSharedRef<FGenericWindow>& NativeWindow, ETaskbarProgressState::Type State);
 
+	/**
+	 * Sets the progress value of a task bar entry.
+	 *
+	 * @param NativeWindow The native window to change the progress value for.
+	 * @param Current The current progress value.
+	 * @param Total The total progress value.
+	 */
 	void SetProgressValue(const TSharedRef<FGenericWindow>& NativeWindow, uint64 Current, uint64 Total);
 
+	/** Destructor. */
 	~FTaskbarList();
 
 private:
 
+	/** Hidden constructor (use FTaskbarList::Create). */
 	FTaskbarList();
 
+	/** Initializes the task bar list instance. */
 	void Initialize();
 
 private:
 
+	/** Holds the internal task bar object. */
 	ITaskbarList3* TaskBarList3;
 };
 
@@ -217,7 +252,8 @@ PRAGMA_DISABLE_DEPRECATION_WARNINGS
  * Windows-specific application implementation.
  */
 class FWindowsApplication
-	: public GenericApplication, IForceFeedbackSystem
+	: public GenericApplication
+	, public IForceFeedbackSystem
 {
 public:
 
@@ -231,10 +267,27 @@ public:
 	 */
 	static FWindowsApplication* CreateWindowsApplication( const HINSTANCE InstanceHandle, const HICON IconHandle );
 
-public:	
-
+	/** Virtual destructor. */
 	virtual ~FWindowsApplication();
-	virtual void DestroyApplication() override;
+
+public:
+
+	/** Called by a window when an OLE Drag and Drop operation occurred on a non-game thread */
+	void DeferDragDropOperation( const FDeferredWindowsDragDropOperation& DeferredDragDropOperation );
+
+	TSharedPtr<FTaskbarList> GetTaskbarList();
+
+	/** Invoked by a window when an OLE Drag and Drop first enters it. */
+	HRESULT OnOLEDragEnter( const HWND HWnd, const FDragDropOLEData& OLEData, ::DWORD KeyState, POINTL CursorPosition, ::DWORD *CursorEffect);
+
+	/** Invoked by a window when an OLE Drag and Drop moves over the window. */
+	HRESULT OnOLEDragOver( const HWND HWnd, ::DWORD KeyState, POINTL CursorPosition, ::DWORD *CursorEffect);
+
+	/** Invoked by a window when an OLE Drag and Drop exits the window. */
+	HRESULT OnOLEDragOut( const HWND HWnd );
+
+	/** Invoked by a window when an OLE Drag and Drop is dropped onto the window. */
+	HRESULT OnOLEDrop( const HWND HWnd, const FDragDropOLEData& OLEData, ::DWORD KeyState, POINTL CursorPosition, ::DWORD *CursorEffect);
 
 public:
 
@@ -259,6 +312,7 @@ public:
 	virtual void GetInitialDisplayMetrics( FDisplayMetrics& OutDisplayMetrics ) const override;
 	virtual EWindowTitleAlignment::Type GetWindowTitleAlignment() const override;
 	virtual EWindowTransparency GetWindowTransparencySupport() const override;
+	virtual void DestroyApplication() override;
 
 	DEPRECATED(4.7, "Please use GetInputInterface()")
 	virtual IForceFeedbackSystem* GetForceFeedbackSystem() override
@@ -276,10 +330,6 @@ public:
 		return this;
 	}
 
-	virtual void SetForceFeedbackChannelValue (int32 ControllerId, FForceFeedbackChannelType ChannelType, float Value) override;
-	virtual void SetForceFeedbackChannelValues(int32 ControllerId, const FForceFeedbackValues &Values) override;
-	virtual void SetLightColor(int32 ControllerId, FColor Color) override {}
-
 	virtual ITextInputMethodSystem *GetTextInputMethodSystem() override
 	{
 		return TextInputMethodSystem.Get();
@@ -287,7 +337,12 @@ public:
 
 	virtual void AddExternalInputDevice(TSharedPtr<class IInputDevice> InputDevice);
 
-	TSharedPtr<FTaskbarList> GetTaskbarList();
+public:
+
+	// IForceFeedbackSystem overrides
+	virtual void SetForceFeedbackChannelValue (int32 ControllerId, FForceFeedbackChannelType ChannelType, float Value) override;
+	virtual void SetForceFeedbackChannelValues(int32 ControllerId, const FForceFeedbackValues &Values) override;
+	virtual void SetLightColor(int32 ControllerId, FColor Color) override { }
 
 protected:
 
@@ -303,23 +358,6 @@ protected:
 	/** Processes deferred drag and drop operations. */
 	void ProcessDeferredDragDropOperation(const FDeferredWindowsDragDropOperation& Op);
 
-public:
-
-	/** Called by a window when an OLE Drag and Drop operation occurred on a non-game thread */
-	void DeferDragDropOperation( const FDeferredWindowsDragDropOperation& DeferredDragDropOperation );
-
-	/** Invoked by a window when an OLE Drag and Drop first enters it. */
-	HRESULT OnOLEDragEnter( const HWND HWnd, const FDragDropOLEData& OLEData, ::DWORD KeyState, POINTL CursorPosition, ::DWORD *CursorEffect);
-
-	/** Invoked by a window when an OLE Drag and Drop moves over the window. */
-	HRESULT OnOLEDragOver( const HWND HWnd, ::DWORD KeyState, POINTL CursorPosition, ::DWORD *CursorEffect);
-
-	/** Invoked by a window when an OLE Drag and Drop exits the window. */
-	HRESULT OnOLEDragOut( const HWND HWnd );
-
-	/** Invoked by a window when an OLE Drag and Drop is dropped onto the window. */
-	HRESULT OnOLEDrop( const HWND HWnd, const FDragDropOLEData& OLEData, ::DWORD KeyState, POINTL CursorPosition, ::DWORD *CursorEffect);
-
 private:
 
 	/** Hidden constructor. */
@@ -386,9 +424,9 @@ private:
 	TSharedPtr<FTaskbarList> TaskbarList;
 
 	// Accessibility shortcut keys
-	STICKYKEYS							StartupStickyKeys;
-	TOGGLEKEYS							StartupToggleKeys;
-	FILTERKEYS							StartupFilterKeys;
+	STICKYKEYS StartupStickyKeys;
+	TOGGLEKEYS StartupToggleKeys;
+	FILTERKEYS StartupFilterKeys;
 };
 
 
