CrashReport - More properties added to unified crash reporting (IsInternalBuild,IsPerforceBuild,IsSourceDistribution)

diff --git a/Engine/Source/Runtime/Core/Private/GenericPlatform/GenericPlatformContext.cpp b/Engine/Source/Runtime/Core/Private/GenericPlatform/GenericPlatformContext.cpp
--- a/Engine/Source/Runtime/Core/Private/GenericPlatform/GenericPlatformContext.cpp
+++ b/Engine/Source/Runtime/Core/Private/GenericPlatform/GenericPlatformContext.cpp
@@ -4,6 +4,7 @@
 #include "GenericPlatform/GenericPlatformContext.h"
 #include "Misc/App.h"
 #include "Runtime/Launch/Resources/Version.h"
+#include "EngineBuildSettings.h"
 
 const ANSICHAR* FGenericCrashContext::CrashContextRuntimeXMLNameA = "CrashContext.runtime-xml";
 const TCHAR* FGenericCrashContext::CrashContextRuntimeXMLNameW = TEXT( "CrashContext.runtime-xml" );
@@ -73,8 +74,13 @@ void FGenericCrashContext::SerializeContentToBuffer()
 	// ProcessId will be used by the crash report client to get more information about the crashed process.
 	// It should be enough to open the process, get the memory stats, a list of loaded modules etc.
 	// It means that the crashed process needs to be alive as long as the crash report client is working on the process.
+	// Proposal, not implemented yet.
 	AddCrashProperty( TEXT( "ProcessId" ), FPlatformProcess::GetCurrentProcessId() );
 
+	AddCrashProperty( TEXT( "IsInternalBuild" ), (int32)FEngineBuildSettings::IsInternalBuild() );
+	AddCrashProperty( TEXT( "IsPerforceBuild" ), (int32)FEngineBuildSettings::IsPerforceBuild() );
+	AddCrashProperty( TEXT( "IsSourceDistribution" ), (int32)FEngineBuildSettings::IsSourceDistribution() );
+
 	// Add common crash properties.
 	AddCrashProperty( TEXT( "GameName" ), FApp::GetGameName() );
 	AddCrashProperty( TEXT( "ExecutableName" ), *ExecutableName );
