Core - GetAllocationSize fix for malloc binned (https://udn.unrealengine.com/questions/236049/incorrect-allocation-size-reported-by-mallocbinned.html)

diff --git a/Engine/Source/Runtime/Core/Public/HAL/MallocBinned.h b/Engine/Source/Runtime/Core/Public/HAL/MallocBinned.h
--- a/Engine/Source/Runtime/Core/Public/HAL/MallocBinned.h
+++ b/Engine/Source/Runtime/Core/Public/HAL/MallocBinned.h
@@ -1116,7 +1116,7 @@ public:
 		}
 		UPTRINT BasePtr;
 		FPoolInfo* Pool = FindPoolInfo((UPTRINT)Original, BasePtr);
-		SizeOut = Pool->TableIndex < BinnedSizeLimit ? MemSizeToPoolTable[Pool->TableIndex]->BlockSize : Pool->GetBytes();
+		SizeOut = Pool->TableIndex < BinnedOSTableIndex ? MemSizeToPoolTable[Pool->TableIndex]->BlockSize : Pool->GetBytes();
 		return true;
 	}
 
