UE4 - fix tests on TlsSlots

diff --git a/Engine/Source/Runtime/Core/Private/Containers/LockFreeList.cpp b/Engine/Source/Runtime/Core/Private/Containers/LockFreeList.cpp
--- a/Engine/Source/Runtime/Core/Private/Containers/LockFreeList.cpp
+++ b/Engine/Source/Runtime/Core/Private/Containers/LockFreeList.cpp
@@ -66,7 +66,7 @@ FLockFreeVoidPointerListBase::FLinkAllocator::FLinkAllocator()
 	static_assert(sizeof(FLink) >= sizeof(void*) && sizeof(FLink) % sizeof(void*) == 0, "Blocks in TLockFreeFixedSizeAllocator must be at least the size of a pointer.");
 	check(IsInGameThread());
 	TlsSlot = FPlatformTLS::AllocTlsSlot();
-	check(TlsSlot);
+	check(FPlatformTLS::IsValidTlsSlot(TlsSlot));
 }
 FLockFreeVoidPointerListBase::FLinkAllocator::~FLinkAllocator()
 {
diff --git a/Engine/Source/Runtime/Core/Private/Stats/Stats2.cpp b/Engine/Source/Runtime/Core/Private/Stats/Stats2.cpp
--- a/Engine/Source/Runtime/Core/Private/Stats/Stats2.cpp
+++ b/Engine/Source/Runtime/Core/Private/Stats/Stats2.cpp
@@ -878,7 +878,7 @@ FThreadStats::FThreadStats():
 {
 	Packet.SetThreadProperties();
 
-	check(TlsSlot);
+	check(FPlatformTLS::IsValidTlsSlot(TlsSlot));
 	FPlatformTLS::SetTlsValue(TlsSlot, this);
 }
 
diff --git a/Engine/Source/Runtime/Core/Public/Containers/LockFreeFixedSizeAllocator.h b/Engine/Source/Runtime/Core/Public/Containers/LockFreeFixedSizeAllocator.h
--- a/Engine/Source/Runtime/Core/Public/Containers/LockFreeFixedSizeAllocator.h
+++ b/Engine/Source/Runtime/Core/Public/Containers/LockFreeFixedSizeAllocator.h
@@ -113,7 +113,7 @@ public:
 		static_assert(SIZE >= sizeof(void*) && SIZE % sizeof(void*) == 0, "Blocks in TLockFreeFixedSizeAllocator must be at least the size of a pointer.");
 		check(IsInGameThread());
 		TlsSlot = FPlatformTLS::AllocTlsSlot();
-		check(TlsSlot);
+		check(FPlatformTLS::IsValidTlsSlot(TlsSlot));
 	}
 	/** Destructor, leaks all of the memory **/
 	~TLockFreeFixedSizeAllocator_TLSCache()
diff --git a/Engine/Source/Runtime/Core/Public/GenericPlatform/GenericPlatformTLS.h b/Engine/Source/Runtime/Core/Public/GenericPlatform/GenericPlatformTLS.h
--- a/Engine/Source/Runtime/Core/Public/GenericPlatform/GenericPlatformTLS.h
+++ b/Engine/Source/Runtime/Core/Public/GenericPlatform/GenericPlatformTLS.h
@@ -10,6 +10,16 @@
  */
 struct FGenericPlatformTLS
 {
+	/**
+	 * Return false if this is an invalid TLS slot
+	 * @param SlotIndex the TLS index to check
+	 * @return true if this looks like a valid slot
+	 */
+	static FORCEINLINE bool IsValidTlsSlot(uint32 SlotIndex)
+	{
+		return SlotIndex != 0xFFFFFFFF;
+	}
+
 #if 0 // provided for reference
 	/**
 	 * Returns the currently executing thread's id
