Do not close the handle in TerminateProc(), there's a separate CloseProc() for this.

diff --git a/Engine/Source/Runtime/Core/Private/Windows/WindowsPlatformProcess.cpp b/Engine/Source/Runtime/Core/Private/Windows/WindowsPlatformProcess.cpp
--- a/Engine/Source/Runtime/Core/Private/Windows/WindowsPlatformProcess.cpp
+++ b/Engine/Source/Runtime/Core/Private/Windows/WindowsPlatformProcess.cpp
@@ -428,8 +428,6 @@ void FWindowsPlatformProcess::TerminateProc( FProcHandle & ProcessHandle, bool K
 	}
 
 	TerminateProcess(ProcessHandle.Get(),0);
-	// Process is terminated, so we can close the process handle.
-	ProcessHandle.Close();
 }
 
 uint32 FWindowsPlatformProcess::GetCurrentProcessId()
