Integrated pull request, Github #908.  Fix to FInterpCurve::EvalDerivative and FInterpCurve::EvalSecondDerivative.

diff --git a/Engine/Source/Runtime/Core/Public/Math/InterpCurve.h b/Engine/Source/Runtime/Core/Public/Math/InterpCurve.h
--- a/Engine/Source/Runtime/Core/Public/Math/InterpCurve.h
+++ b/Engine/Source/Runtime/Core/Public/Math/InterpCurve.h
@@ -281,20 +281,21 @@ T FInterpCurve<T>::EvalDerivative( const float InVal, const T& Default ) const
 
 	if (Diff > 0.0f && PrevPoint.InterpMode != CIM_Constant)
 	{
-		const float Alpha = (InVal - PrevPoint.InVal) / Diff;
-
 		if (PrevPoint.InterpMode == CIM_Linear)
 		{
-			return NextPoint.OutVal - PrevPoint.OutVal;
+			return (NextPoint.OutVal - PrevPoint.OutVal) / Diff;
 		}
 		else
 		{
-			return FMath::CubicInterpDerivative(PrevPoint.OutVal, PrevPoint.LeaveTangent * Diff, NextPoint.OutVal, NextPoint.ArriveTangent * Diff, Alpha);
+			const float Alpha = (InVal - PrevPoint.InVal) / Diff;
+
+			return FMath::CubicInterpDerivative(PrevPoint.OutVal, PrevPoint.LeaveTangent * Diff, NextPoint.OutVal, NextPoint.ArriveTangent * Diff, Alpha) / Diff;
 		}
 	}
 	else
 	{
-		return Points[Index].OutVal;
+		// Derivative of a constant is zero
+		return T();
 	}
 }
 
@@ -333,8 +334,6 @@ T FInterpCurve<T>::EvalSecondDerivative( const float InVal, const T& Default ) c
 
 	if (Diff > 0.0f && PrevPoint.InterpMode != CIM_Constant)
 	{
-		const float Alpha = (InVal - PrevPoint.InVal) / Diff;
-
 		if (PrevPoint.InterpMode == CIM_Linear)
 		{
 			// No change in tangent, return 0.
@@ -342,12 +341,15 @@ T FInterpCurve<T>::EvalSecondDerivative( const float InVal, const T& Default ) c
 		}
 		else
 		{
-			return FMath::CubicInterpSecondDerivative(PrevPoint.OutVal, PrevPoint.LeaveTangent * Diff, NextPoint.OutVal, NextPoint.ArriveTangent * Diff, Alpha);
+			const float Alpha = (InVal - PrevPoint.InVal) / Diff;
+
+			return FMath::CubicInterpSecondDerivative(PrevPoint.OutVal, PrevPoint.LeaveTangent * Diff, NextPoint.OutVal, NextPoint.ArriveTangent * Diff, Alpha) / (Diff * Diff);
 		}
 	}
 	else
 	{
-		return Points[Index].OutVal;
+		// Second derivative of a constant is zero
+		return T();
 	}
 }
 
