Fix one more TlsSlot validity check.

diff --git a/Engine/Source/Runtime/Core/Public/Containers/LockFreeFixedSizeAllocator.h b/Engine/Source/Runtime/Core/Public/Containers/LockFreeFixedSizeAllocator.h
--- a/Engine/Source/Runtime/Core/Public/Containers/LockFreeFixedSizeAllocator.h
+++ b/Engine/Source/Runtime/Core/Public/Containers/LockFreeFixedSizeAllocator.h
@@ -235,7 +235,7 @@ private:
 
 	FThreadLocalCache& GetTLS()
 	{
-		checkSlow(TlsSlot);
+		checkSlow(FPlatformTLS::IsValidTlsSlot(TlsSlot));
 		FThreadLocalCache* TLS = (FThreadLocalCache*)FPlatformTLS::GetTlsValue(TlsSlot);
 		if (!TLS)
 		{
