Fixed Android COTF crash due to file timestamp issue with manifest file (failed to create if opening for write and missing)
#ue4
#android

diff --git a/Engine/Source/Runtime/Core/Private/Android/AndroidFile.cpp b/Engine/Source/Runtime/Core/Private/Android/AndroidFile.cpp
--- a/Engine/Source/Runtime/Core/Private/Android/AndroidFile.cpp
+++ b/Engine/Source/Runtime/Core/Private/Android/AndroidFile.cpp
@@ -336,6 +336,12 @@ public:
 
 	bool SetFileTimeStamp( const FString& FileName, const FDateTime& DateTime )
 	{
+		if (bInitialized == false)
+		{
+			Read();
+			bInitialized = true;
+		}
+
 		FDateTime* Result = ManifestEntries.Find( FileName );
 		if ( Result == NULL )
 		{
@@ -356,12 +362,12 @@ public:
 	{
 		// Local filepaths are directly in the deployment directory.
 		static const FString &BasePath = GetFileBasePath();
-		const FString ManfiestPath = BasePath + ManifestFileName;
+		const FString ManifestPath = BasePath + ManifestFileName;
 
 		ManifestEntries.Empty();
 
 		// int Handle = open( TCHAR_TO_UTF8(ManifestFileName), O_RDWR );
-		int Handle = open( TCHAR_TO_UTF8(*ManfiestPath), O_RDONLY );
+		int Handle = open( TCHAR_TO_UTF8(*ManifestPath), O_RDONLY );
 
 		if ( Handle == -1 )
 		{
@@ -433,10 +439,10 @@ public:
 		
 		// Local filepaths are directly in the deployment directory.
 		static const FString &BasePath = GetFileBasePath();
-		const FString ManfiestPath = BasePath + ManifestFileName;
+		const FString ManifestPath = BasePath + ManifestFileName;
 
 
-		int Handle = open( TCHAR_TO_UTF8(*ManfiestPath), O_WRONLY );
+		int Handle = open(TCHAR_TO_UTF8(*ManifestPath), O_WRONLY | O_CREAT | O_TRUNC);
 
 		if ( Handle == -1 )
 		{
