Core: Allowing for checking major AND minor version of Windows

#UpgradeNotes:
This is an API breaking change. You must rename your usages of VerifyWindowsMajorVersion() to VerifyWindowsVersion() and pass a second parameter for the minor version number (pass '0' to check for Windows Vista, as before).

diff --git a/Engine/Source/Runtime/Core/Private/Windows/WindowsPlatformMisc.cpp b/Engine/Source/Runtime/Core/Private/Windows/WindowsPlatformMisc.cpp
--- a/Engine/Source/Runtime/Core/Private/Windows/WindowsPlatformMisc.cpp
+++ b/Engine/Source/Runtime/Core/Private/Windows/WindowsPlatformMisc.cpp
@@ -1668,11 +1668,12 @@ int32 FWindowsPlatformMisc::GetAppIcon()
 	return IDICON_UE4Game;
 }
 
-bool FWindowsPlatformMisc::VerifyWindowsMajorVersion(uint32 MajorVersion)
+bool FWindowsPlatformMisc::VerifyWindowsVersion(uint32 MajorVersion, uint32 MinorVersion)
 {
 	OSVERSIONINFOEX Version;
 	Version.dwOSVersionInfoSize = sizeof(OSVERSIONINFOEX);
 	Version.dwMajorVersion = MajorVersion;
+	Version.dwMinorVersion = MinorVersion;
 	ULONGLONG ConditionMask = 0;
 	return !!VerifyVersionInfo( &Version, VER_MAJORVERSION, VerSetConditionMask(ConditionMask,VER_MAJORVERSION,VER_GREATER_EQUAL) );
 }
@@ -1773,7 +1774,7 @@ int32 FWindowsPlatformMisc::NumberOfCoresIncludingHyperthreads()
 void FWindowsPlatformMisc::LoadPreInitModules()
 {
 	// D3D11 is not supported on WinXP, so in this case we use the OpenGL RHI
-	if(FWindowsPlatformMisc::VerifyWindowsMajorVersion(6))
+	if(FWindowsPlatformMisc::VerifyWindowsVersion(6, 0))
 	{
 		FModuleManager::Get().LoadModule(TEXT("D3D11RHI"));
 	}
@@ -2348,7 +2349,7 @@ FString FWindowsPlatformMisc::GetOperatingSystemId()
 
 EConvertibleLaptopMode FWindowsPlatformMisc::GetConvertibleLaptopMode()
 {
-	if (!VerifyWindowsMajorVersion(8))
+	if (!VerifyWindowsVersion(6, 2))
 	{
 		return EConvertibleLaptopMode::NotSupported;
 	}
diff --git a/Engine/Source/Runtime/Core/Private/Windows/WindowsPlatformStackWalk.cpp b/Engine/Source/Runtime/Core/Private/Windows/WindowsPlatformStackWalk.cpp
--- a/Engine/Source/Runtime/Core/Private/Windows/WindowsPlatformStackWalk.cpp
+++ b/Engine/Source/Runtime/Core/Private/Windows/WindowsPlatformStackWalk.cpp
@@ -154,7 +154,7 @@ static bool GMaxCallstackDepthInitialized = false;
 void DetermineMaxCallstackDepth()
 {
 	// Check that we're running on Vista or newer (version 6.0+).
-	if ( FWindowsPlatformMisc::VerifyWindowsMajorVersion(6) )
+	if ( FWindowsPlatformMisc::VerifyWindowsVersion(6, 0) )
 	{
 		GMaxCallstackDepth = MAX_CALLSTACK_DEPTH;
 	}
diff --git a/Engine/Source/Runtime/Core/Private/Windows/WindowsPlatformSurvey.cpp b/Engine/Source/Runtime/Core/Private/Windows/WindowsPlatformSurvey.cpp
--- a/Engine/Source/Runtime/Core/Private/Windows/WindowsPlatformSurvey.cpp
+++ b/Engine/Source/Runtime/Core/Private/Windows/WindowsPlatformSurvey.cpp
@@ -171,7 +171,7 @@ void FWindowsPlatformSurvey::TickSurveyHardware( FHardwareSurveyResults& OutResu
 	}
 
 	// Check that we're running on Vista or newer (version 6.0+).
-	bool bIsVistaOrNewer = FWindowsPlatformMisc::VerifyWindowsMajorVersion(6);
+	bool bIsVistaOrNewer = FWindowsPlatformMisc::VerifyWindowsVersion(6, 0);
 
 	FMemory::Memset(&OutResults, 0, sizeof(FHardwareSurveyResults));
 
diff --git a/Engine/Source/Runtime/Core/Public/Windows/WindowsPlatformMisc.h b/Engine/Source/Runtime/Core/Public/Windows/WindowsPlatformMisc.h
--- a/Engine/Source/Runtime/Core/Public/Windows/WindowsPlatformMisc.h
+++ b/Engine/Source/Runtime/Core/Public/Windows/WindowsPlatformMisc.h
@@ -133,9 +133,12 @@ struct CORE_API FWindowsPlatformMisc
 	/** 
 	 * Determines if we are running on the Windows version or newer
 	 *
+	 * See the 'Remarks' section of https://msdn.microsoft.com/en-us/library/windows/desktop/ms724833(v=vs.85).aspx
+	 * for a list of MajorVersion/MinorVersion version combinations for Microsoft Windows.
+	 *
 	 * @return	Returns true if the current Windows version if equal or newer than MajorVersion
 	 */
-	static bool VerifyWindowsMajorVersion(uint32 MajorVersion);
+	static bool VerifyWindowsVersion(uint32 MajorVersion, uint32 MinorVersion);
 
 	/**
 	 * Sample the displayed pixel color from anywhere on the screen using the OS
diff --git a/Engine/Source/Runtime/OpenGLDrv/Private/Windows/OpenGLWindows.cpp b/Engine/Source/Runtime/OpenGLDrv/Private/Windows/OpenGLWindows.cpp
--- a/Engine/Source/Runtime/OpenGLDrv/Private/Windows/OpenGLWindows.cpp
+++ b/Engine/Source/Runtime/OpenGLDrv/Private/Windows/OpenGLWindows.cpp
@@ -182,7 +182,7 @@ static void PlatformCreateDummyGLWindow(FPlatformOpenGLContext* OutContext)
 static bool PlatformOpenGL3()
 {
 	// OpenGL3 is our default platform for Windows XP
-	return FWindowsPlatformMisc::VerifyWindowsMajorVersion(6) == false || FParse::Param(FCommandLine::Get(),TEXT("opengl")) || FParse::Param(FCommandLine::Get(),TEXT("opengl3"));
+	return FWindowsPlatformMisc::VerifyWindowsVersion(6, 0) == false || FParse::Param(FCommandLine::Get(),TEXT("opengl")) || FParse::Param(FCommandLine::Get(),TEXT("opengl3"));
 }
 static bool PlatformOpenGL4()
 {
diff --git a/Engine/Source/Runtime/RHI/Private/Windows/WindowsDynamicRHI.cpp b/Engine/Source/Runtime/RHI/Private/Windows/WindowsDynamicRHI.cpp
--- a/Engine/Source/Runtime/RHI/Private/Windows/WindowsDynamicRHI.cpp
+++ b/Engine/Source/Runtime/RHI/Private/Windows/WindowsDynamicRHI.cpp
@@ -13,7 +13,7 @@ FDynamicRHI* PlatformCreateDynamicRHI()
 	const bool bForceD3D11 = FParse::Param(FCommandLine::Get(),TEXT("d3d11")) || FParse::Param(FCommandLine::Get(),TEXT("sm5")) || FParse::Param(FCommandLine::Get(),TEXT("dx11"));
 	const bool bForceD3D10 = FParse::Param(FCommandLine::Get(),TEXT("d3d10")) || FParse::Param(FCommandLine::Get(),TEXT("sm4")) || FParse::Param(FCommandLine::Get(),TEXT("dx10"));
 
-	const bool bForceOpenGL = FWindowsPlatformMisc::VerifyWindowsMajorVersion(6) == false || FParse::Param(FCommandLine::Get(), TEXT("opengl")) || FParse::Param(FCommandLine::Get(), TEXT("opengl3")) || FParse::Param(FCommandLine::Get(), TEXT("opengl4"));
+	const bool bForceOpenGL = FWindowsPlatformMisc::VerifyWindowsVersion(6, 0) == false || FParse::Param(FCommandLine::Get(), TEXT("opengl")) || FParse::Param(FCommandLine::Get(), TEXT("opengl3")) || FParse::Param(FCommandLine::Get(), TEXT("opengl4"));
 
 	if (((bForceD3D11 ? 1 : 0) + (bForceD3D10 ? 1 : 0) + (bForceOpenGL ? 1 : 0)) > 1)
 	{
