Rolling back an unwanted change to MacApplication from the latest LauncherDev->Main merge

diff --git a/Engine/Source/Runtime/Core/Private/Mac/MacApplication.cpp b/Engine/Source/Runtime/Core/Private/Mac/MacApplication.cpp
--- a/Engine/Source/Runtime/Core/Private/Mac/MacApplication.cpp
+++ b/Engine/Source/Runtime/Core/Private/Mac/MacApplication.cpp
@@ -564,7 +564,7 @@ void FMacApplication::ProcessNSEvent(NSEvent* const Event)
 				// Forward the delta on to Slate
 				MessageHandler->OnRawMouseMove(Delta.X, Delta.Y);
 			}
-			else if (CurrentEventWindow.IsValid())
+			else
 			{
 				NSPoint CursorPos = [NSEvent mouseLocation];
 				CursorPos.y--; // The y coordinate of the point returned by mouseLocation starts from a base of 1
