PR #920: (Git Plugin) Fix FMacPlatformProcess::ExecProcess() to correctly parse double-quoted arguments like CreateProc() does. (Contributed by [REDACTED_PERSON])

https://github.com/EpicGames/UnrealEngine/pull/920

diff --git a/Engine/Source/Runtime/Core/Private/Mac/MacPlatformProcess.cpp b/Engine/Source/Runtime/Core/Private/Mac/MacPlatformProcess.cpp
--- a/Engine/Source/Runtime/Core/Private/Mac/MacPlatformProcess.cpp
+++ b/Engine/Source/Runtime/Core/Private/Mac/MacPlatformProcess.cpp
@@ -246,13 +246,25 @@ bool FMacPlatformProcess::ExecProcess( const TCHAR* URL, const TCHAR* Params, in
 		{
 			if (MultiPartArg.IsEmpty())
 			{
-				if (ArgsArray[Index].StartsWith(TEXT("\"")) && !ArgsArray[Index].EndsWith(TEXT("\"")))
+				if ((ArgsArray[Index].StartsWith(TEXT("\"")) && !ArgsArray[Index].EndsWith(TEXT("\""))) // check for a starting quote but no ending quote, excludes quoted single arguments
+					|| (ArgsArray[Index].Contains(TEXT("=\"")) && !ArgsArray[Index].EndsWith(TEXT("\""))) // check for quote after =, but no ending quote, this gets arguments of the type -blah="string string string"
+					|| ArgsArray[Index].EndsWith(TEXT("=\""))) // check for ending quote after =, this gets arguments of the type -blah=" string string string "
 				{
 					MultiPartArg = ArgsArray[Index];
 				}
 				else
 				{
-					NSString* Arg = (NSString*)FPlatformString::TCHARToCFString(*ArgsArray[Index].TrimQuotes(NULL));
+					NSString* Arg;
+					if (ArgsArray[Index].Contains(TEXT("=\"")))
+					{
+						FString SingleArg = ArgsArray[Index];
+						SingleArg = SingleArg.Replace(TEXT("=\""), TEXT("="));
+						Arg = (NSString*)FPlatformString::TCHARToCFString(*SingleArg.TrimQuotes(NULL));
+					}
+					else
+					{
+						Arg = (NSString*)FPlatformString::TCHARToCFString(*ArgsArray[Index].TrimQuotes(NULL));
+					}
 					[Arguments addObject: Arg];
 					CFRelease((CFStringRef)Arg);
 				}
@@ -263,7 +275,15 @@ bool FMacPlatformProcess::ExecProcess( const TCHAR* URL, const TCHAR* Params, in
 				MultiPartArg += ArgsArray[Index];
 				if (ArgsArray[Index].EndsWith(TEXT("\"")))
 				{
-					NSString* Arg = (NSString*)FPlatformString::TCHARToCFString(*MultiPartArg.TrimQuotes(NULL));
+					NSString* Arg;
+					if (MultiPartArg.StartsWith(TEXT("\"")))
+					{
+						Arg = (NSString*)FPlatformString::TCHARToCFString(*MultiPartArg.TrimQuotes(NULL));
+					}
+					else
+					{
+						Arg = (NSString*)FPlatformString::TCHARToCFString(*MultiPartArg.Replace(TEXT("\""), TEXT("")));
+					}
 					[Arguments addObject: Arg];
 					CFRelease((CFStringRef)Arg);
 					MultiPartArg.Empty();
