UE4 - roll back lockfree allocator; does not work

diff --git a/Engine/Source/Runtime/Core/Private/Containers/LockFreeList.cpp b/Engine/Source/Runtime/Core/Private/Containers/LockFreeList.cpp
--- a/Engine/Source/Runtime/Core/Private/Containers/LockFreeList.cpp
+++ b/Engine/Source/Runtime/Core/Private/Containers/LockFreeList.cpp
@@ -58,36 +58,34 @@ FLockFreeVoidPointerListBase128::FLinkAllocator::~FLinkAllocator()
 	FLockFreeVoidPointerListBase
 -----------------------------------------------------------------------------*/
 
-FLockFreeVoidPointerListBase::FLinkAllocator FLockFreeVoidPointerListBase::FLinkAllocator::TheLinkAllocator;
-
-FLockFreeVoidPointerListBase::FLinkAllocator::FLinkAllocator()
-	: SpecialClosedLink(new FLink())
-{
-	static_assert(sizeof(FLink) >= sizeof(void*) && sizeof(FLink) % sizeof(void*) == 0, "Blocks in TLockFreeFixedSizeAllocator must be at least the size of a pointer.");
-	check(IsInGameThread());
-	TlsSlot = FPlatformTLS::AllocTlsSlot();
-	check(FPlatformTLS::IsValidTlsSlot(TlsSlot));
-}
-FLockFreeVoidPointerListBase::FLinkAllocator::~FLinkAllocator()
+FLockFreeVoidPointerListBase::FLinkAllocator& FLockFreeVoidPointerListBase::FLinkAllocator::Get()
 {
-	FPlatformTLS::FreeTlsSlot(TlsSlot);
-	TlsSlot = 0;
+	static FLinkAllocator TheAllocator;
+	return TheAllocator;
 }
 
-void FLockFreeVoidPointerListBase::FLinkAllocator::PutFreeBundle(void **Bundle)
+FLockFreeVoidPointerListBase::FLinkAllocator::FLinkAllocator()
+	: FreeLinks(NULL)
+	, SpecialClosedLink(new FLink())
 {
-	FScopeLock Lock(&BundleArrayCriticalSection);
-	BundleArray.Push(Bundle);
+	ClosedLink()->LockCount.Increment();
 }
-void** FLockFreeVoidPointerListBase::FLinkAllocator::GetFreeBundle()
+
+FLockFreeVoidPointerListBase::FLinkAllocator::~FLinkAllocator()
 {
-	FScopeLock Lock(&BundleArrayCriticalSection);
-	void** Result = nullptr;
-	if (BundleArray.Num())
+	// - Deliberately leak to allow LockFree to be used during shutdown
+#if	0
+	check(SpecialClosedLink); // double free?
+	check(!NumUsedLinks.GetValue());
+	while (FLink* ToDelete = FLink::Unlink(&FreeLinks))
 	{
-		Result = BundleArray.Pop(false);
+		delete ToDelete;
+		NumFreeLinks.Decrement();
 	}
-	return Result;
+	check(!NumFreeLinks.GetValue());
+	delete SpecialClosedLink;
+	SpecialClosedLink = 0;
+#endif // 0
 }
 
 #endif //USE_LOCKFREELIST_128
diff --git a/Engine/Source/Runtime/Core/Public/Containers/LockFreeVoidPointerListBase.h b/Engine/Source/Runtime/Core/Public/Containers/LockFreeVoidPointerListBase.h
--- a/Engine/Source/Runtime/Core/Public/Containers/LockFreeVoidPointerListBase.h
+++ b/Engine/Source/Runtime/Core/Public/Containers/LockFreeVoidPointerListBase.h
@@ -260,27 +260,13 @@ private:
 		/** Part of the solution to the ABA problem. Links that are free still need locking, but we need to make sure that links are only added to the free list when they are transitioning to the free list. */
 		FThreadSafeCounter	MarkedForDeath;
 
-		/** Constructor, everything is initialized to zero. */
+		/** Construcor, everything is intialized to zero. */
 		FLink()
 			: Next(nullptr)
 			, Item(nullptr)
 		{
 		}
 
-		/** Constructor, initialize the item and lock count to 1. */
-		FLink(void *NewItem)
-			: Next(nullptr)
-			, Item(NewItem)
-			, LockCount(1)
-		{
-#if CHECK_NON_CONCURRENT_ASSUMPTIONS
-			FPlatformMisc::MemoryBarrier();
-			checkLockFreePointerList(NewLink->Item == NewItem);
-			checkLockFreePointerList(!NewLink->Next);
-			checkLockFreePointerList(NewLink->LockCount.GetValue() == 1);
-#endif
-		}
-
 		/**	
 		 *	Link this node into the head of the given list.
 		 *	@param HeadPointer; the head of the list
@@ -529,21 +515,13 @@ private:
 		void Unlock(bool bShouldNeverCauseAFree = false);
 	};
 
-public:
-
-	// new link allocator
-	class FLinkAllocator	
+	/**	
+	 *	Class to allocate and recycle links.
+	 */
+	class FLinkAllocator
 	{
-		enum
-		{
-			NUM_PER_BUNDLE=1024,
-		};
 	public:
 
-		CORE_API FLinkAllocator();
-		/** Destructor, leaks all of the memory **/
-		CORE_API ~FLinkAllocator();
-
 		/**	
 		 *	Return a new link.
 		 *	@param NewItem "item" pointer of the new node.
@@ -551,45 +529,40 @@ public:
 		 */
 		FLink* AllocateLink(void *NewItem)
 		{
-			FThreadLocalCache& TLS = GetTLS();
+			checkLockFreePointerList(NewItem); // we don't allow nullptr items
+			if (NumUsedLinks.Increment() % 8192 == 1)
+			{
+				UE_CLOG(0/*MONITOR_LINK_ALLOCATION*/,LogLockFreeList, Log, TEXT("Number of links %d"),NumUsedLinks.GetValue());
+			}
 
-			if (!TLS.PartialBundle)
+			FLink*  NewLink = nullptr;
+			NewLink = FLink::Unlink(&FreeLinks);
+			if (NewLink)
 			{
-				if (TLS.FullBundle)
-				{
-					TLS.PartialBundle = TLS.FullBundle;
-					TLS.FullBundle = nullptr;
-				}
-				else
+				NumFreeLinks.Decrement();
+			}
+			else
+			{
+				if (NumAllocatedLinks.Increment() % 10 == 1)
 				{
-					TLS.PartialBundle = GetFreeBundle();
-					if (!TLS.PartialBundle)
-					{
-						TLS.PartialBundle = (void**)FMemory::Malloc(sizeof(FLink) * NUM_PER_BUNDLE);
-						void **Next = TLS.PartialBundle;
-						for (int32 Index = 0; Index < NUM_PER_BUNDLE - 1; Index++)
-						{
-							void* NextNext = (void*)(((uint8*)Next) + sizeof(FLink));
-							*Next = NextNext;
-							Next = (void**)NextNext;
-						}
-						*Next = nullptr;
-						NumFree.Add(NUM_PER_BUNDLE);
-					}
+					UE_CLOG(0/*MONITOR_LINK_ALLOCATION*/,LogLockFreeList, Log, TEXT("Number of allocated links %d"),NumAllocatedLinks.GetValue());
 				}
-				TLS.NumPartial = NUM_PER_BUNDLE;
+
+				NewLink = new FLink();
+				NewLink->LockCount.Increment();
 			}
-			NumUsed.Increment();
-			NumFree.Decrement();
-			void* Result = (void*)TLS.PartialBundle;
-			TLS.PartialBundle = (void**)*TLS.PartialBundle;
-			TLS.NumPartial--;
-			checkLockFreePointerList(TLS.NumPartial >= 0 && ((!!TLS.NumPartial) == (!!TLS.PartialBundle)));
-
-			FLink* NewLink = new (Result) FLink(NewItem);
+			checkLockFreePointerList(!NewLink->Item);
+			checkLockFreePointerList(!NewLink->Next);
+			checkLockFreePointerList(NewLink->LockCount.GetValue() >= 1);
+#if CHECK_NON_CONCURRENT_ASSUMPTIONS
+			void* ValueWas = FPlatformAtomics::InterlockedCompareExchangePointer(&NewLink->Item, NewItem, nullptr); // in theory this doesn't need to be interlocked
+			checkLockFreePointerList(ValueWas == nullptr);
+#else
+			NewLink->Item = NewItem;
+#endif
+			FPlatformMisc::MemoryBarrier();
 			return NewLink;
 		}
-
 		/**	
 		 *	Make a link available for recycling. 
 		 *	@param Link link to recycle
@@ -598,53 +571,17 @@ public:
 		void FreeLink(FLink* Link)
 		{
 			checkLockFreePointerList(Link != ClosedLink());  // very bad to recycle the special link
-			Link->~FLink(); // noop, but here for completeness
-			NumUsed.Decrement();
-			NumFree.Increment();
-			FThreadLocalCache& TLS = GetTLS();
-			if (TLS.NumPartial >= NUM_PER_BUNDLE)
-			{
-				if (TLS.FullBundle)
-				{
-					PutFreeBundle(TLS.FullBundle);
-					//TLS.FullBundle = nullptr;
-				}
-				TLS.FullBundle = TLS.PartialBundle;
-				TLS.PartialBundle = nullptr;
-				TLS.NumPartial = 0;
-			}
-			*(void**)Link = (void*)TLS.PartialBundle;
-			TLS.PartialBundle = (void**)Link;
-			TLS.NumPartial++;
-		}
-
-		/**
-		 * Gets the number of allocated memory blocks that are currently in use.
-		 *
-		 * @return Number of used memory blocks.
-		 * @see GetNumFree
-		 */
-		const FLockFreeListCounter& GetNumUsed() const
-		{
-			return NumUsed;
-		}
-
-		/**
-		 * Gets the number of allocated memory blocks that are currently unused.
-		 *
-		 * @return Number of unused memory blocks.
-		 * @see GetNumUsed
-		 */
-		const FLockFreeListCounter& GetNumFree() const
-		{
-			return NumFree;
+			NumUsedLinks.Decrement();
+			Link->LockCount.Increment();
+			FPlatformMisc::MemoryBarrier();
+			Link->Link(&FreeLinks);
+			NumFreeLinks.Increment();
 		}
-
 		/**	
 		 *	Return a pointer to the special closed link
 		 *	@return pointer to special closed link
 		 */
-		FORCEINLINE FLink* ClosedLink()
+		FLink* ClosedLink()
 		{
 			checkLockFreePointerList(nullptr != SpecialClosedLink);
 			return SpecialClosedLink;
@@ -653,69 +590,36 @@ public:
 		 *	Singleton access
 		 *	@return the singleton for the link allocator
 		 */
-		FORCEINLINE static FLinkAllocator& Get()
-		{
-			return TheLinkAllocator;
-		}
+		static CORE_API FLinkAllocator& Get();
 
 	private:
+		/**	
+		 *	Constructor zeros the free link list and creates the special closed link
+		 *	@return the singleton for the link allocator
+		 */
+		FLinkAllocator();
 
-		static CORE_API FLinkAllocator TheLinkAllocator;
-
-		/** struct for the TLS cache. */
-		struct FThreadLocalCache
-		{
-			void **FullBundle;
-			void **PartialBundle;
-			int32 NumPartial;
-
-			FThreadLocalCache()
-				: FullBundle(nullptr)
-				, PartialBundle(nullptr)
-				, NumPartial(0)
-			{
-			}
-		};
-
-		FThreadLocalCache& GetTLS()
-		{
-			checkSlow(TlsSlot);
-			FThreadLocalCache* TLS = (FThreadLocalCache*)FPlatformTLS::GetTlsValue(TlsSlot);
-			if (!TLS)
-			{
-				TLS = new FThreadLocalCache();
-				FPlatformTLS::SetTlsValue(TlsSlot, TLS);
-			}
-			return *TLS;
-		}
-
-		/** put a set of 1024 links to a global store **/
-		CORE_API void PutFreeBundle(void **Bundle);
-		/** retrieve a set of 1024 links from a global store **/
-		CORE_API void** GetFreeBundle();
+		/**	
+		 *	Destructor, should only be called when there are no outstanding links. 
+		 *	Frees the elements of the free list and frees the special link
+		 *	@return the singleton for the link allocator
+		 */
+		~FLinkAllocator();
 
-		/** Slot for TLS struct. */
-		uint32 TlsSlot;
+		/** Head to a list of free links that can be used for new allocations. */
+		FLink*				FreeLinks;
 
 		/** Pointer to the special closed link. It should never be recycled. */
 		FLink*				SpecialClosedLink;
 
-		/** Total number of blocks outstanding and not in the free list. */
-		FLockFreeListCounter NumUsed; 
-
-		/** Total number of blocks in the free list. */
-		FLockFreeListCounter NumFree;
-
-		/** Critical section for the bundle array. */
-		FCriticalSection BundleArrayCriticalSection;
-
-		/** Array of bundles of 1024 links */
-		TArray<void**> BundleArray;
-
+		/** Total number of links outstanding and not in the free list */
+		FLockFreeListCounter	NumUsedLinks; 
+		/** Total number of links in the free list */
+		FLockFreeListCounter	NumFreeLinks;
+		/** Total number of links allocated */
+		FLockFreeListCounter	NumAllocatedLinks; 
 	};
 
-private:
-
 	/** Head of the list */
 	MS_ALIGN(8) FLink*	Head;
 };
