Fix red-blue swap in Linux splash.

diff --git a/Engine/Shaders/StandaloneRenderer/OpenGL/SplashFragmentShader.glsl b/Engine/Shaders/StandaloneRenderer/OpenGL/SplashFragmentShader.glsl
--- a/Engine/Shaders/StandaloneRenderer/OpenGL/SplashFragmentShader.glsl
+++ b/Engine/Shaders/StandaloneRenderer/OpenGL/SplashFragmentShader.glsl
@@ -15,6 +15,9 @@ void main()
 	// OpenGL has 0,0 the "math" way
 	vec2 tc = vec2(textureCoordinate.s, 1.0-textureCoordinate.t);
 
+#if RED_BLUE_SWAP
+	gl_FragColor = texture2D(SplashTexture, tc).bgra;
+#else
 	gl_FragColor = texture2D(SplashTexture, tc);
-
+#endif
 }
\ No newline at end of file
diff --git a/Engine/Source/Runtime/Core/Private/Linux/LinuxPlatformSplash.cpp b/Engine/Source/Runtime/Core/Private/Linux/LinuxPlatformSplash.cpp
--- a/Engine/Source/Runtime/Core/Private/Linux/LinuxPlatformSplash.cpp
+++ b/Engine/Source/Runtime/Core/Private/Linux/LinuxPlatformSplash.cpp
@@ -197,7 +197,7 @@ static int CompileShaderFromString( const FString& Source, GLuint& ShaderID, GLe
 	check( ShaderID );
 	
 	// Allocate a buffer big enough to store the string in ascii format
-	ANSICHAR* Chars[2] = {nullptr};
+	ANSICHAR* Chars[3] = {nullptr};
 	// pass the #define along to the shader
 #if PLATFORM_USES_ES2
 	Chars[0] = (ANSICHAR*)"#define PLATFORM_USES_ES2 1\n\n#define PLATFORM_LINUX 0\n";
@@ -206,12 +206,17 @@ static int CompileShaderFromString( const FString& Source, GLuint& ShaderID, GLe
 #else
 	Chars[0] = (ANSICHAR*)"#version 120\n\n#define PLATFORM_USES_ES2 0\n\n#define PLATFORM_LINUX 0\n";
 #endif
-	Chars[1] = new ANSICHAR[Source.Len()+1];
-	FCStringAnsi::Strcpy(Chars[1], Source.Len() + 1, TCHAR_TO_ANSI(*Source));
+
+	ANSICHAR Buffer[256];
+	sprintf(Buffer, "#define RED_BLUE_SWAP %d\n", (GSplashScreenImage->format->Rmask > GSplashScreenImage->format->Bmask) ? 1 : 0);
+	Chars[1] = Buffer;
+	Chars[2] = new ANSICHAR[Source.Len()+1];
+
+	FCStringAnsi::Strcpy(Chars[2], Source.Len() + 1, TCHAR_TO_ANSI(*Source));
 
 	// give opengl the source code for the shader
-	glShaderSource( ShaderID, 2, (const ANSICHAR**)Chars, NULL );
-	delete[] Chars[1];
+	glShaderSource( ShaderID, 3, (const ANSICHAR**)Chars, NULL );
+	delete[] Chars[2];
 
 	// Compile the shader and check for success
 	glCompileShader( ShaderID );
