PR #911: Linux full-screen only fix.

- Contributed by [REDACTED_PERSON] (https://github.com/EpicGames/UnrealEngine/pull/911).
- We are flip-flopping on this, but currently it doesn't seem to cause issues.

diff --git a/Engine/Source/Runtime/Core/Public/Linux/LinuxPlatformProperties.h b/Engine/Source/Runtime/Core/Public/Linux/LinuxPlatformProperties.h
--- a/Engine/Source/Runtime/Core/Public/Linux/LinuxPlatformProperties.h
+++ b/Engine/Source/Runtime/Core/Public/Linux/LinuxPlatformProperties.h
@@ -91,7 +91,7 @@ struct FLinuxPlatformProperties
 
 	static FORCEINLINE bool HasFixedResolution()
 	{
-		return true; // Fullscreen on Linux is always WindowedFullscreen (to avoid badness happening when switching to various resolutions)
+		return false;
 	}
 
 	static FORCEINLINE bool SupportsWindowedMode()
