Better fix for red-blue swap in Linux splash (UE-6012).

diff --git a/Engine/Shaders/StandaloneRenderer/OpenGL/SplashFragmentShader.glsl b/Engine/Shaders/StandaloneRenderer/OpenGL/SplashFragmentShader.glsl
--- a/Engine/Shaders/StandaloneRenderer/OpenGL/SplashFragmentShader.glsl
+++ b/Engine/Shaders/StandaloneRenderer/OpenGL/SplashFragmentShader.glsl
@@ -15,9 +15,6 @@ void main()
 	// OpenGL has 0,0 the "math" way
 	vec2 tc = vec2(textureCoordinate.s, 1.0-textureCoordinate.t);
 
-#if RED_BLUE_SWAP
-	gl_FragColor = texture2D(SplashTexture, tc).bgra;
-#else
 	gl_FragColor = texture2D(SplashTexture, tc);
-#endif
+
 }
\ No newline at end of file
diff --git a/Engine/Source/Runtime/Core/Private/Linux/LinuxPlatformSplash.cpp b/Engine/Source/Runtime/Core/Private/Linux/LinuxPlatformSplash.cpp
--- a/Engine/Source/Runtime/Core/Private/Linux/LinuxPlatformSplash.cpp
+++ b/Engine/Source/Runtime/Core/Private/Linux/LinuxPlatformSplash.cpp
@@ -197,7 +197,7 @@ static int CompileShaderFromString( const FString& Source, GLuint& ShaderID, GLe
 	check( ShaderID );
 	
 	// Allocate a buffer big enough to store the string in ascii format
-	ANSICHAR* Chars[3] = {nullptr};
+	ANSICHAR* Chars[2] = {nullptr};
 	// pass the #define along to the shader
 #if PLATFORM_USES_ES2
 	Chars[0] = (ANSICHAR*)"#define PLATFORM_USES_ES2 1\n\n#define PLATFORM_LINUX 0\n";
@@ -206,17 +206,12 @@ static int CompileShaderFromString( const FString& Source, GLuint& ShaderID, GLe
 #else
 	Chars[0] = (ANSICHAR*)"#version 120\n\n#define PLATFORM_USES_ES2 0\n\n#define PLATFORM_LINUX 0\n";
 #endif
-
-	ANSICHAR Buffer[256];
-	sprintf(Buffer, "#define RED_BLUE_SWAP %d\n", (GSplashScreenImage->format->Rmask > GSplashScreenImage->format->Bmask) ? 1 : 0);
-	Chars[1] = Buffer;
-	Chars[2] = new ANSICHAR[Source.Len()+1];
-
-	FCStringAnsi::Strcpy(Chars[2], Source.Len() + 1, TCHAR_TO_ANSI(*Source));
+	Chars[1] = new ANSICHAR[Source.Len()+1];
+	FCStringAnsi::Strcpy(Chars[1], Source.Len() + 1, TCHAR_TO_ANSI(*Source));
 
 	// give opengl the source code for the shader
-	glShaderSource( ShaderID, 3, (const ANSICHAR**)Chars, NULL );
-	delete[] Chars[2];
+	glShaderSource( ShaderID, 2, (const ANSICHAR**)Chars, NULL );
+	delete[] Chars[1];
 
 	// Compile the shader and check for success
 	glCompileShader( ShaderID );
@@ -488,17 +483,18 @@ static int RenderString (GLuint tex_idx)
 		// store rendered text as texture
 		glBindTexture(GL_TEXTURE_2D, tex_idx);
 		
+		bool bIsBGR = GSplashScreenImage->format->Rmask > GSplashScreenImage->format->Bmask;
 		if (SplashBPP == 3)
 		{
 			glTexImage2D(GL_TEXTURE_2D, 0, GL_RGB,
 				GSplashScreenImage->w, GSplashScreenImage->h,
-				0, GL_RGB, GL_UNSIGNED_BYTE, ScratchSpace);
+				0, bIsBGR ? GL_BGR : GL_RGB, GL_UNSIGNED_BYTE, ScratchSpace);
 		}
 		else
 		{
 			glTexImage2D(GL_TEXTURE_2D, 0, GL_RGBA,
 				GSplashScreenImage->w, GSplashScreenImage->h,
-				0, GL_RGBA, GL_UNSIGNED_BYTE, ScratchSpace);			
+				0, bIsBGR ? GL_BGRA : GL_RGBA, GL_UNSIGNED_BYTE, ScratchSpace);			
 		}
 	}
 
@@ -659,17 +655,18 @@ static int StartSplashScreenThread(void *ptr)
 	glTexParameteri(GL_TEXTURE_2D, GL_TEXTURE_MAG_FILTER, GL_LINEAR);
 
 	// load splash image as texture in opengl
+	bool bIsBGR = GSplashScreenImage->format->Rmask > GSplashScreenImage->format->Bmask;
 	if (GSplashScreenImage->format->BytesPerPixel == 3)
 	{
 		glTexImage2D(GL_TEXTURE_2D, 0, GL_RGB,
 			GSplashScreenImage->w, GSplashScreenImage->h,
-			0, GL_RGB, GL_UNSIGNED_BYTE, GSplashScreenImage->pixels);
+			0, bIsBGR ? GL_BGR : GL_RGB, GL_UNSIGNED_BYTE, GSplashScreenImage->pixels);
 	}
 	else if (GSplashScreenImage->format->BytesPerPixel == 4)
 	{
 		glTexImage2D(GL_TEXTURE_2D, 0, GL_RGBA,
 			GSplashScreenImage->w, GSplashScreenImage->h,
-			0, GL_RGBA, GL_UNSIGNED_BYTE, GSplashScreenImage->pixels);		
+			0, bIsBGR ? GL_BGRA : GL_RGBA, GL_UNSIGNED_BYTE, GSplashScreenImage->pixels);		
 	}
 	else
 	{
