Disable TLS thread id cache until UE-11541 is investigated.

- We run out of TLS slot limit on Ubuntu 14.04 and earlier.
- Theoretically -ftls-model=local-dynamic (the proper fix) should have fixed this, but for some reason we don't see this happening.
- Issue will need to be investigated, but for now just don't cache thread id for glibc <= 2.19 (this assumes ubuntu-patched glibc by the way).

diff --git a/Engine/Source/Runtime/Core/Public/Linux/LinuxPlatformTLS.h b/Engine/Source/Runtime/Core/Public/Linux/LinuxPlatformTLS.h
--- a/Engine/Source/Runtime/Core/Public/Linux/LinuxPlatformTLS.h
+++ b/Engine/Source/Runtime/Core/Public/Linux/LinuxPlatformTLS.h
@@ -26,6 +26,18 @@ struct CORE_API FLinuxTLS : public FGenericPlatformTLS
 		// note: cannot use pthread_self() without updating the rest of API to opaque (or at least 64-bit) thread handles
 #if defined(_GNU_SOURCE)
 		// syscall() is relatively heavy and shows up in the profiler, given that IsInGameThread() is used quite often. Cache thread id in TLS.
+		
+	// note: this caching currently breaks the editor on Ubuntu 14.04 and earlier. The libraries are using static TLS and glibc 2.19 runs out of limit for TLS slots (see https://bugs.launchpad.net/ubuntu/+source/glibc/+bug/1375555 for a similar problem).
+	// the proper fix is perhaps not to use initial-exec TLS model, but for some reason just passing -ftls-model=local-dynamic doesn't work. Disable this caching until the issue is investigated (tracked as UE-11541)
+	#if defined(__GLIBC__) && (__GLIBC__ == 2) && (__GLIBC_MINOR__ <= 19)
+
+		pid_t ThreadId = static_cast<pid_t>(syscall(SYS_gettid));
+		static_assert(sizeof(pid_t) <= sizeof(uint32), "pid_t is larger than uint32, reconsider implementation of GetCurrentThreadId()");
+		return static_cast<pid_t>(ThreadId);		
+
+	#else
+
+		// cached version - should work fine in future glibc versions (also works fine in Ubuntu 14.10 after glibc 2.19-10ubuntu2, but there's no appropriate #define to discover that)
 		static __thread uint32 ThreadIdTLS = 0;
 		if (ThreadIdTLS == 0)
 		{
@@ -35,6 +47,9 @@ struct CORE_API FLinuxTLS : public FGenericPlatformTLS
 			checkf(ThreadIdTLS != 0, TEXT("ThreadId is 0 - reconsider implementation of GetCurrentThreadId() (syscall changed?)"));
 		}
 		return ThreadIdTLS;
+
+	#endif // glibc version check
+
 #else
 		// better than nothing...
 		return static_cast< uint32 >(pthread_self());
