Core: Added API for intercepting Windows messages

diff --git a/Engine/Source/Runtime/Core/Private/Windows/WindowsApplication.cpp b/Engine/Source/Runtime/Core/Private/Windows/WindowsApplication.cpp
--- a/Engine/Source/Runtime/Core/Private/Windows/WindowsApplication.cpp
+++ b/Engine/Source/Runtime/Core/Private/Windows/WindowsApplication.cpp
@@ -33,15 +33,16 @@ DEFINE_LOG_CATEGORY(LogWindowsDesktop);
 
 const FIntPoint FWindowsApplication::MinimizedWindowPosition(-32000,-32000);
 
-FWindowsApplication* WindowApplication = NULL;
+FWindowsApplication* WindowsApplication = nullptr;
 
 
 FWindowsApplication* FWindowsApplication::CreateWindowsApplication( const HINSTANCE InstanceHandle, const HICON IconHandle )
 {
-	WindowApplication = new FWindowsApplication( InstanceHandle, IconHandle );
-	return WindowApplication;
+	WindowsApplication = new FWindowsApplication( InstanceHandle, IconHandle );
+	return WindowsApplication;
 }
 
+
 FWindowsApplication::FWindowsApplication( const HINSTANCE HInstance, const HICON IconHandle )
 	: GenericApplication( MakeShareable( new FWindowsCursor() ) )
 	, InstanceHandle( HInstance )
@@ -280,7 +281,7 @@ static TSharedPtr< FWindowsWindow > FindWindowByHWND(const TArray< TSharedRef< F
 		}
 	}
 
-	return TSharedPtr< FWindowsWindow >(NULL);
+	return TSharedPtr< FWindowsWindow >(nullptr);
 }
 
 
@@ -505,7 +506,7 @@ void GetMonitorInfo(TArray<FMonitorInfo>& OutMonitorInfo)
 	DisplayDevice.cb = sizeof(DisplayDevice);
 	DWORD DeviceIndex = 0; // device index
 
-	FMonitorInfo* PrimaryDevice = NULL;
+	FMonitorInfo* PrimaryDevice = nullptr;
 	OutMonitorInfo.Empty(2); // Reserve two slots, as that will be the most common maximum
 
 	FString DeviceID;
@@ -534,7 +535,7 @@ void GetMonitorInfo(TArray<FMonitorInfo>& OutMonitorInfo)
 						Info.bIsPrimary = (DisplayDevice.StateFlags & DISPLAY_DEVICE_PRIMARY_DEVICE) > 0;
 						OutMonitorInfo.Add(Info);
 
-						if (PrimaryDevice == NULL && Info.bIsPrimary)
+						if (PrimaryDevice == nullptr && Info.bIsPrimary)
 						{
 							PrimaryDevice = &OutMonitorInfo.Last();
 						}
@@ -619,12 +620,12 @@ EWindowTransparency FWindowsApplication::GetWindowTransparencySupport() const
 #endif
 }
 
-/** All WIN32 messages sent to our app go here; this method simply passes them on */
+
 LRESULT CALLBACK FWindowsApplication::AppWndProc(HWND hwnd, uint32 msg, WPARAM wParam, LPARAM lParam)
 {
 	ensure( IsInGameThread() );
 
-	return WindowApplication->ProcessMessage( hwnd, msg, wParam, lParam );
+	return WindowsApplication->ProcessMessage( hwnd, msg, wParam, lParam );
 }
 
 int32 FWindowsApplication::ProcessMessage( HWND hwnd, uint32 msg, WPARAM wParam, LPARAM lParam )
@@ -1100,6 +1101,20 @@ int32 FWindowsApplication::ProcessMessage( HWND hwnd, uint32 msg, WPARAM wParam,
 				XInput->SetNeedsControllerStateUpdate(); 
 				QueryConnectedMice();
 			}
+
+		default:
+			{
+			    int32 HandlerResult = 0;
+
+				// give others a chance to handle unprocessed messages
+				for (auto Handler : MessageHandlers)
+				{
+					if (Handler->ProcessMessage(hwnd, msg, wParam, lParam, HandlerResult))
+					{
+						return HandlerResult;
+					}
+				}
+			}
 		}
 	}
 
@@ -1884,6 +1899,19 @@ HRESULT FWindowsApplication::OnOLEDrop( const HWND HWnd, const FDragDropOLEData&
 	return 0;
 }
 
+
+void FWindowsApplication::AddMessageHandler(IWindowsMessageHandler& MessageHandler)
+{
+	WindowsApplication->MessageHandlers.AddUnique(&MessageHandler);
+}
+
+
+void FWindowsApplication::RemoveMessageHandler(IWindowsMessageHandler& MessageHandler)
+{
+	WindowsApplication->MessageHandlers.RemoveSwap(&MessageHandler);
+}
+
+
 void FWindowsApplication::QueryConnectedMice()
 {
 	TArray<RAWINPUTDEVICELIST> DeviceList;
diff --git a/Engine/Source/Runtime/Core/Public/Windows/WindowsApplication.h b/Engine/Source/Runtime/Core/Public/Windows/WindowsApplication.h
--- a/Engine/Source/Runtime/Core/Public/Windows/WindowsApplication.h
+++ b/Engine/Source/Runtime/Core/Public/Windows/WindowsApplication.h
@@ -248,6 +248,27 @@ public:
 PRAGMA_DISABLE_DEPRECATION_WARNINGS
 
 
+/**
+ * Interface for classes that handle Windows events.
+ */
+class IWindowsMessageHandler
+{
+public:
+
+	/**
+	 * Processes a Windows message.
+	 *
+	 * @param hwnd Handle to the window that received the message.
+	 * @param msg The message.
+	 * @param wParam Additional message information.
+	 * @param lParam Additional message information.
+	 * @param OutResult Will contain the result if the message was handled.
+	 * @return true if the message was handled, false otherwise.
+	 */
+	virtual bool ProcessMessage(HWND hwnd, uint32 msg, WPARAM wParam, LPARAM lParam, int32& OutResult) = 0;
+};
+
+
 /**
  * Windows-specific application implementation.
  */
@@ -260,10 +281,9 @@ public:
 	/**
 	 * Static: Creates a new Win32 application
 	 *
-	 * @param	InstanceHandle  Win32 instance handle
-	 * @param	IconHandle		Win32 application icon handle
-	 *
-	 * @return  New application object
+	 * @param InstanceHandle Win32 instance handle.
+	 * @param IconHandle Win32 application icon handle.
+	 * @return New application object.
 	 */
 	static FWindowsApplication* CreateWindowsApplication( const HINSTANCE InstanceHandle, const HICON IconHandle );
 
@@ -289,6 +309,22 @@ public:
 	/** Invoked by a window when an OLE Drag and Drop is dropped onto the window. */
 	HRESULT OnOLEDrop( const HWND HWnd, const FDragDropOLEData& OLEData, ::DWORD KeyState, POINTL CursorPosition, ::DWORD *CursorEffect);
 
+	/**
+	 * Adds a Windows message handler with the application instance.
+	 *
+	 * @param MessageHandler The message handler to register.
+	 * @see RemoveMessageHandler
+	 */
+	virtual void AddMessageHandler(IWindowsMessageHandler& MessageHandler);
+
+	/**
+	 * Removes a Windows message handler with the application instance.
+	 *
+	 * @param MessageHandler The message handler to register.
+	 * @see AddMessageHandler
+	 */
+	virtual void RemoveMessageHandler(IWindowsMessageHandler& MessageHandler);
+
 public:
 
 	// GenericApplication overrides
@@ -340,6 +376,7 @@ public:
 public:
 
 	// IForceFeedbackSystem overrides
+
 	virtual void SetForceFeedbackChannelValue (int32 ControllerId, FForceFeedbackChannelType ChannelType, float Value) override;
 	virtual void SetForceFeedbackChannelValues(int32 ControllerId, const FForceFeedbackValues &Values) override;
 	virtual void SetLightColor(int32 ControllerId, FColor Color) override { }
@@ -398,6 +435,9 @@ private:
 
 	TArray<FDeferredWindowsDragDropOperation> DeferredDragDropOperations;
 
+	/** Registered Windows message handlers. */
+	TArray<IWindowsMessageHandler*> MessageHandlers;
+
 	TArray<TSharedRef<FWindowsWindow>> Windows;
 
 	TSharedRef<class XInputInterface> XInput;
diff --git a/Engine/Source/Runtime/Core/Public/Windows/WindowsWindow.h b/Engine/Source/Runtime/Core/Public/Windows/WindowsWindow.h
--- a/Engine/Source/Runtime/Core/Public/Windows/WindowsWindow.h
+++ b/Engine/Source/Runtime/Core/Public/Windows/WindowsWindow.h
@@ -73,6 +73,9 @@ public:
 	/** Called when our parent window is restored (which will in turn cause us to become restored). */
 	void OnParentWindowRestored();
 
+	/** Called by the owning application when the level of transparency support has changed */
+	void OnTransparencySupportChanged(EWindowTransparency NewTransparency);
+
 public:
 
 	// FGenericWindow interface
@@ -120,11 +123,6 @@ public:
 	virtual HRESULT STDCALL DragOver( ::DWORD KeyState, POINTL CursorPosition, __RPC__inout ::DWORD *CursorEffect) override;
 	virtual HRESULT STDCALL DragLeave( void ) override;
 	virtual HRESULT STDCALL Drop( __RPC__in_opt IDataObject *DataObjectPointer, ::DWORD KeyState, POINTL CursorPosition, __RPC__inout ::DWORD *CursorEffect) override;
-	
-public:
-
-	/** Called by the owning application when the level of transparency support has changed */
-	void OnTransparencySupportChanged(EWindowTransparency NewTransparency);
 
 private:
 
