Core: Added API for getting handle of top-level window on Windows

diff --git a/Engine/Source/Runtime/Core/Private/Windows/WindowsPlatformMisc.cpp b/Engine/Source/Runtime/Core/Private/Windows/WindowsPlatformMisc.cpp
--- a/Engine/Source/Runtime/Core/Private/Windows/WindowsPlatformMisc.cpp
+++ b/Engine/Source/Runtime/Core/Private/Windows/WindowsPlatformMisc.cpp
@@ -1810,6 +1810,44 @@ bool FWindowsPlatformMisc::OsExecute(const TCHAR* CommandType, const TCHAR* Comm
 	return bSucceeded;
 }
 
+struct FGetMainWindowHandleData
+{
+	HWND Handle;
+	uint32 ProcessId;
+};
+
+int32 CALLBACK GetMainWindowHandleCallback(HWND Handle, LPARAM lParam)
+{
+	FGetMainWindowHandleData& Data = *(FGetMainWindowHandleData*)lParam;
+	
+	unsigned long ProcessId = 0;
+	{
+		::GetWindowThreadProcessId(Handle, &ProcessId);
+	}
+	
+	if ((Data.ProcessId != ProcessId) || (::GetWindow(Handle, GW_OWNER) != (HWND)0) || !::IsWindowVisible(Handle))
+	{
+		return 1;
+	}
+
+	Data.Handle = Handle;
+
+	return 0;
+}
+
+HWND FWindowsPlatformMisc::GetTopLevelWindowHandle(uint32 ProcessId)
+{
+	FGetMainWindowHandleData Data;
+	{
+		Data.Handle = 0;
+		Data.ProcessId = ProcessId;
+	}
+
+	::EnumWindows(GetMainWindowHandleCallback, (LPARAM)&Data);
+
+	return Data.Handle;
+}
+
 bool FWindowsPlatformMisc::GetWindowTitleMatchingText(const TCHAR* TitleStartsWith, FString& OutTitle)
 {
 	bool bWasFound = false;
diff --git a/Engine/Source/Runtime/Core/Public/Windows/WindowsPlatformMisc.h b/Engine/Source/Runtime/Core/Public/Windows/WindowsPlatformMisc.h
--- a/Engine/Source/Runtime/Core/Public/Windows/WindowsPlatformMisc.h
+++ b/Engine/Source/Runtime/Core/Public/Windows/WindowsPlatformMisc.h
@@ -117,6 +117,17 @@ struct CORE_API FWindowsPlatformMisc
 	 */
 	static bool OsExecute(const TCHAR* CommandType, const TCHAR* Command, const TCHAR* CommandLine = NULL);
 
+	/**
+	 * Attempts to get the handle to a top-level window of the specified process.
+	 *
+	 * If the process has a single main window (root), its handle will be returned.
+	 * If the process has multiple top-level windows, the first one found is returned.
+	 *
+	 * @param ProcessId The identifier of the process to get the window for.
+	 * @return Window handle, or 0 if not found.
+	 */
+	static HWND GetTopLevelWindowHandle(uint32 ProcessId);
+
 	/**
 	 * Searches for a window that matches the window name or the title starts with a particular text. When
 	 * found, it returns the title text for that window
@@ -144,7 +155,7 @@ struct CORE_API FWindowsPlatformMisc
 	/**
 	 * Sample the displayed pixel color from anywhere on the screen using the OS
 	 *
-	 * @param	InScreenPos		The screen coords to sample for current pixel color
+	 * @param	InScreenPos		The screen coordinates to sample for current pixel color
 	 * @param	InGamma			Optional gamma correction to apply to the screen color
 	 *
 	 * @return					The color of the pixel displayed at the chosen location
