Linux: Cache display area to reduce chattiness.

- Still polling, just at a lower rate. What we need is a callback on screen config changes.

diff --git a/Engine/Source/Runtime/Core/Private/Linux/LinuxApplication.cpp b/Engine/Source/Runtime/Core/Private/Linux/LinuxApplication.cpp
--- a/Engine/Source/Runtime/Core/Private/Linux/LinuxApplication.cpp
+++ b/Engine/Source/Runtime/Core/Private/Linux/LinuxApplication.cpp
@@ -69,6 +69,7 @@ FLinuxApplication::FLinuxApplication()
 	,	bIsDragWindowButtonPressed(false)
 	,	bActivateApp(false)
 	,	bLockToCurrentMouseType(false)
+	,	LastTimeCachedWorkArea(-1.0)
 {
 	bUsingHighPrecisionMouseInput = false;
 	bAllowedToDeferMessageProcessing = true;
@@ -1101,40 +1102,49 @@ void FLinuxApplication::SetHighPrecisionMouseMode( const bool Enable, const TSha
 
 FPlatformRect FLinuxApplication::GetWorkArea( const FPlatformRect& CurrentWindow ) const
 {
-	// loop over all monitors to determine which one is the best
-	int NumDisplays = SDL_GetNumVideoDisplays();
-	if (NumDisplays <= 0)
-	{
-		// fake something
-		return CurrentWindow;
-	}
-	
-	SDL_Rect BestDisplayBounds;
-	SDL_GetDisplayBounds(0, &BestDisplayBounds);
-	
-	// see if any other are better (i.e. cover top left)
-	for (int DisplayIdx = 1; DisplayIdx < NumDisplays; ++DisplayIdx)
+	const double kCacheLifetime = 5.0;	// ask once in 5 seconds
+
+	double CurrentTime = FPlatformTime::Seconds();
+	if (CurrentTime - LastTimeCachedWorkArea > kCacheLifetime)
 	{
-		SDL_Rect DisplayBounds;
-		SDL_GetDisplayBounds(DisplayIdx, &DisplayBounds);
+		// loop over all monitors to determine which one is the best
+		int NumDisplays = SDL_GetNumVideoDisplays();
+		if (NumDisplays <= 0)
+		{
+			// fake something
+			return CurrentWindow;
+		}
+		
+		SDL_Rect BestDisplayBounds;
+		SDL_GetDisplayBounds(0, &BestDisplayBounds);
 		
-		// only check top left corner for "bestness"
-		if (DisplayBounds.x <= CurrentWindow.Left && DisplayBounds.x + DisplayBounds.w > CurrentWindow.Left &&
-			DisplayBounds.y <= CurrentWindow.Top && DisplayBounds.y + DisplayBounds.h > CurrentWindow.Bottom)
+		// see if any other are better (i.e. cover top left)
+		for (int DisplayIdx = 1; DisplayIdx < NumDisplays; ++DisplayIdx)
 		{
-			BestDisplayBounds = DisplayBounds;
-			// there can be only one, as we don't expect overlapping displays
-			break;
+			SDL_Rect DisplayBounds;
+			SDL_GetDisplayBounds(DisplayIdx, &DisplayBounds);
+			
+			// only check top left corner for "bestness"
+			if (DisplayBounds.x <= CurrentWindow.Left && DisplayBounds.x + DisplayBounds.w > CurrentWindow.Left &&
+				DisplayBounds.y <= CurrentWindow.Top && DisplayBounds.y + DisplayBounds.h > CurrentWindow.Bottom)
+			{
+				BestDisplayBounds = DisplayBounds;
+				// there can be only one, as we don't expect overlapping displays
+				break;
+			}
 		}
+		
+		FPlatformRect WorkArea;
+		WorkArea.Left	= BestDisplayBounds.x;
+		WorkArea.Top	= BestDisplayBounds.y;
+		WorkArea.Right	= BestDisplayBounds.x + BestDisplayBounds.w;
+		WorkArea.Bottom	= BestDisplayBounds.y + BestDisplayBounds.h;
+
+		CachedWorkArea = WorkArea;
+		LastTimeCachedWorkArea = CurrentTime;
 	}
-	
-	FPlatformRect WorkArea;
-	WorkArea.Left	= BestDisplayBounds.x;
-	WorkArea.Top	= BestDisplayBounds.y;
-	WorkArea.Right	= BestDisplayBounds.x + BestDisplayBounds.w;
-	WorkArea.Bottom	= BestDisplayBounds.y + BestDisplayBounds.h;
 
-	return WorkArea;
+	return CachedWorkArea;
 }
 
 void FLinuxApplication::OnMouseCursorLock( bool bLockEnabled )
diff --git a/Engine/Source/Runtime/Core/Public/Linux/LinuxApplication.h b/Engine/Source/Runtime/Core/Public/Linux/LinuxApplication.h
--- a/Engine/Source/Runtime/Core/Public/Linux/LinuxApplication.h
+++ b/Engine/Source/Runtime/Core/Public/Linux/LinuxApplication.h
@@ -149,7 +149,11 @@ private:
 	/** Used to check with cursor type is current and set to true if left button is pressed.*/
 	bool bLockToCurrentMouseType;
 
-};
+	/** Cached work area - to reduce costly communication with X server (may be better cached in SDL? avoids ugly 'mutable') */
+	mutable FPlatformRect	CachedWorkArea;
 
+	/** Last time we asked about work area (this is a hack. What we need is a callback when screen config changes). */
+	mutable double			LastTimeCachedWorkArea;
+};
 
 extern FLinuxApplication* LinuxApplication;
