Linux: fix encoding of error string.

- Reported by both [REDACTED_PERSON] and [REDACTED_PERSON].

diff --git a/Engine/Source/Runtime/Core/Private/Linux/LinuxPlatformMisc.cpp b/Engine/Source/Runtime/Core/Private/Linux/LinuxPlatformMisc.cpp
--- a/Engine/Source/Runtime/Core/Private/Linux/LinuxPlatformMisc.cpp
+++ b/Engine/Source/Runtime/Core/Private/Linux/LinuxPlatformMisc.cpp
@@ -164,7 +164,7 @@ bool FLinuxPlatformMisc::PlatformInitMultimedia()
 			const char * SDLError = SDL_GetError();
 
 			// do not fail at this point, allow caller handle failure
-			UE_LOG(LogInit, Warning, TEXT("Could not initialize SDL: %s"), SDLError);
+			UE_LOG(LogInit, Warning, TEXT("Could not initialize SDL: %s"), ANSI_TO_TCHAR(SDLError));
 			return false;
 		}
 
