Change assertion function names to play nicely with crash reporter.

diff --git a/Engine/Source/Runtime/Core/Private/Misc/OutputDevice.cpp b/Engine/Source/Runtime/Core/Private/Misc/OutputDevice.cpp
--- a/Engine/Source/Runtime/Core/Private/Misc/OutputDevice.cpp
+++ b/Engine/Source/Runtime/Core/Private/Misc/OutputDevice.cpp
@@ -187,7 +187,7 @@ void StaticFailDebug( const TCHAR* Error, const ANSICHAR* File, int32 Line, cons
 // Failed assertion handler.
 //warning: May be called at library startup time.
 //
-void VARARGS FDebug::AssertFailed( const ANSICHAR* Expr, const ANSICHAR* File, int32 Line, const TCHAR* Format/*=TEXT("")*/, ... )
+void VARARGS FDebug::LogAssertFailedMessage(const ANSICHAR* Expr, const ANSICHAR* File, int32 Line, const TCHAR* Format/*=TEXT("")*/, ...)
 {
 #if !(UE_BUILD_SHIPPING || UE_BUILD_TEST)
 	// Walk the script stack, if any
@@ -224,14 +224,14 @@ void VARARGS FDebug::AssertFailed( const ANSICHAR* Expr, const ANSICHAR* File, i
  * @param	Line	Line number (__LINE__)
  * @param	Msg		Informative error message text
  */
-void FDebug::EnsureFailed( const ANSICHAR* Expr, const ANSICHAR* File, int32 Line, const TCHAR* Msg )
+void FDebug::EnsureFailed(const ANSICHAR* Expr, const ANSICHAR* File, int32 Line, const TCHAR* Msg)
 {
 	// You can set bShouldCrash to true to cause a regular assertion to trigger (stopping program execution) when an ensure() error occurs
 	const bool bShouldCrash = false;		// By default, don't crash on ensure()
 	if( bShouldCrash )
 	{
 		// Just trigger a regular assertion which will crash via GError->Logf()
-		FDebug::AssertFailed( Expr, File, Line, Msg );
+		FDebug::LogAssertFailedMessage( Expr, File, Line, Msg );
 		return;
 	}
 
@@ -339,7 +339,7 @@ void VARARGS FError::LowLevelFatal(const ANSICHAR* File, int32 Line, const TCHAR
 	StaticFailDebug(TEXT("LowLevelFatalError"),File,Line,DescriptionString);
 }
 
-void VARARGS FError::FinalLogf(const ANSICHAR* Expr, const ANSICHAR* File, int32 Line, const TCHAR* Format/*=TEXT("")*/, ...)
+void VARARGS FDebug::AssertFailed(const ANSICHAR* Expr, const ANSICHAR* File, int32 Line, const TCHAR* Format/* = TEXT("")*/, ...)
 {
 	if (GIsCriticalError)
 	{
diff --git a/Engine/Source/Runtime/Core/Public/Misc/AssertionMacros.h b/Engine/Source/Runtime/Core/Public/Misc/AssertionMacros.h
--- a/Engine/Source/Runtime/Core/Public/Misc/AssertionMacros.h
+++ b/Engine/Source/Runtime/Core/Public/Misc/AssertionMacros.h
@@ -23,19 +23,19 @@
 #endif // !UE_BUILD_SHIPPING
 #if DO_CHECK
 	#define checkCode( Code )		do { Code } while ( false );
-	#define verify(expr)			{ if(!(expr)) { FDebug::AssertFailed( #expr, __FILE__, __LINE__ ); _DebugBreakAndPromptForRemote(); FError::FinalLogf( #expr, __FILE__, __LINE__ ); CA_ASSUME(expr); } }
-	#define check(expr)				{ if(!(expr)) { FDebug::AssertFailed( #expr, __FILE__, __LINE__ ); _DebugBreakAndPromptForRemote(); FError::FinalLogf( #expr, __FILE__, __LINE__ ); CA_ASSUME(expr); } }
+	#define verify(expr)			{ if(!(expr)) { FDebug::LogAssertFailedMessage( #expr, __FILE__, __LINE__ ); _DebugBreakAndPromptForRemote(); FDebug::AssertFailed( #expr, __FILE__, __LINE__ ); CA_ASSUME(expr); } }
+	#define check(expr)				{ if(!(expr)) { FDebug::LogAssertFailedMessage( #expr, __FILE__, __LINE__ ); _DebugBreakAndPromptForRemote(); FDebug::AssertFailed( #expr, __FILE__, __LINE__ ); CA_ASSUME(expr); } }
 	
 	/**
 	 * verifyf, checkf: Same as verify, check but with printf style additional parameters
 	 * Read about __VA_ARGS__ (variadic macros) on http://gcc.gnu.org/onlinedocs/gcc-3.4.4/cpp.pdf.
 	 */
-	#define verifyf(expr, format,  ...)		{ if(!(expr)) { FDebug::AssertFailed( #expr, __FILE__, __LINE__, format, ##__VA_ARGS__ ); _DebugBreakAndPromptForRemote(); FError::FinalLogf(#expr, __FILE__, __LINE__, format, ##__VA_ARGS__); CA_ASSUME(expr); } }
-	#define checkf(expr, format,  ...)		{ if(!(expr)) { FDebug::AssertFailed( #expr, __FILE__, __LINE__, format, ##__VA_ARGS__ ); _DebugBreakAndPromptForRemote(); FError::FinalLogf(#expr, __FILE__, __LINE__, format, ##__VA_ARGS__); CA_ASSUME(expr); } }
+	#define verifyf(expr, format,  ...)		{ if(!(expr)) { FDebug::LogAssertFailedMessage( #expr, __FILE__, __LINE__, format, ##__VA_ARGS__ ); _DebugBreakAndPromptForRemote(); FDebug::AssertFailed(#expr, __FILE__, __LINE__, format, ##__VA_ARGS__); CA_ASSUME(expr); } }
+	#define checkf(expr, format,  ...)		{ if(!(expr)) { FDebug::LogAssertFailedMessage( #expr, __FILE__, __LINE__, format, ##__VA_ARGS__ ); _DebugBreakAndPromptForRemote(); FDebug::AssertFailed(#expr, __FILE__, __LINE__, format, ##__VA_ARGS__); CA_ASSUME(expr); } }
 	/**
 	 * Denotes code paths that should never be reached.
 	 */
-	#define checkNoEntry()       { FDebug::AssertFailed( "Enclosing block should never be called", __FILE__, __LINE__ ); _DebugBreakAndPromptForRemote(); FError::FinalLogf("Enclosing block should never be called", __FILE__, __LINE__ ); }
+	#define checkNoEntry()       { FDebug::LogAssertFailedMessage( "Enclosing block should never be called", __FILE__, __LINE__ ); _DebugBreakAndPromptForRemote(); FDebug::AssertFailed("Enclosing block should never be called", __FILE__, __LINE__ ); }
 
 	/**
 	 * Denotes code paths that should not be executed more than once.
@@ -60,7 +60,7 @@
 	                            checkf( RecursionCounter##__LINE__ == 0, TEXT("Enclosing block was entered recursively") );  \
 	                            const FRecursionScopeMarker ScopeMarker##__LINE__( RecursionCounter##__LINE__ )
 
-	#define unimplemented()       { FDebug::AssertFailed( "Unimplemented function called", __FILE__, __LINE__ ); _DebugBreakAndPromptForRemote(); FError::FinalLogf("Unimplemented function called", __FILE__, __LINE__); }
+	#define unimplemented()       { FDebug::LogAssertFailedMessage( "Unimplemented function called", __FILE__, __LINE__ ); _DebugBreakAndPromptForRemote(); FDebug::AssertFailed("Unimplemented function called", __FILE__, __LINE__); }
 
 #else
 	#define checkCode(...)
@@ -78,9 +78,9 @@
 // Check for development only.
 //
 #if DO_GUARD_SLOW
-	#define checkSlow(expr)					{ if(!(expr)) { FDebug::AssertFailed( #expr, __FILE__, __LINE__ ); _DebugBreakAndPromptForRemote(); FError::FinalLogf(#expr, __FILE__, __LINE__); CA_ASSUME(expr); } }
-	#define checkfSlow(expr, format, ...)	{ if(!(expr)) { FDebug::AssertFailed( #expr, __FILE__, __LINE__, format, ##__VA_ARGS__ ); _DebugBreakAndPromptForRemote(); FError::FinalLogf( #expr, __FILE__, __LINE__, format, ##__VA_ARGS__ ); CA_ASSUME(expr); } }
-	#define verifySlow(expr)				{ if(!(expr)) { FDebug::AssertFailed( #expr, __FILE__, __LINE__ ); _DebugBreakAndPromptForRemote(); FError::FinalLogf(#expr, __FILE__, __LINE__); } }
+	#define checkSlow(expr)					{ if(!(expr)) { FDebug::LogAssertFailedMessage( #expr, __FILE__, __LINE__ ); _DebugBreakAndPromptForRemote(); FDebug::AssertFailed(#expr, __FILE__, __LINE__); CA_ASSUME(expr); } }
+	#define checkfSlow(expr, format, ...)	{ if(!(expr)) { FDebug::LogAssertFailedMessage( #expr, __FILE__, __LINE__, format, ##__VA_ARGS__ ); _DebugBreakAndPromptForRemote(); FDebug::AssertFailed( #expr, __FILE__, __LINE__, format, ##__VA_ARGS__ ); CA_ASSUME(expr); } }
+	#define verifySlow(expr)				{ if(!(expr)) { FDebug::LogAssertFailedMessage( #expr, __FILE__, __LINE__ ); _DebugBreakAndPromptForRemote(); FDebug::AssertFailed(#expr, __FILE__, __LINE__); } }
 #else
 	#define checkSlow(expr)
 	#define checkfSlow(expr, format, ...)
@@ -186,7 +186,7 @@ struct FTCharArrayTester
 		static_assert(IS_TCHAR_ARRAY(Format), "Formatting string must be a TCHAR array."); \
 		FError::LowLevelFatal(__FILE__, __LINE__, Format, ##__VA_ARGS__); \
 		_DebugBreakAndPromptForRemote(); \
-		FError::FinalLogf("", __FILE__, __LINE__, Format, ##__VA_ARGS__); \
+		FDebug::AssertFailed("", __FILE__, __LINE__, Format, ##__VA_ARGS__); \
 	}
 
 
@@ -214,7 +214,7 @@ struct FTCharArrayTester
 		{ \
 			FError::LowLevelFatal(__FILE__, __LINE__, Format, ##__VA_ARGS__); \
 			_DebugBreakAndPromptForRemote(); \
-			FError::FinalLogf("", __FILE__, __LINE__, Format, ##__VA_ARGS__); \
+			FDebug::AssertFailed("", __FILE__, __LINE__, Format, ##__VA_ARGS__); \
 		} \
 	}
 	// Conditional logging (fatal errors only).
@@ -227,7 +227,7 @@ struct FTCharArrayTester
 			{ \
 				FError::LowLevelFatal(__FILE__, __LINE__, Format, ##__VA_ARGS__); \
 				_DebugBreakAndPromptForRemote(); \
-				FError::FinalLogf("", __FILE__, __LINE__, Format, ##__VA_ARGS__); \
+				FDebug::AssertFailed("", __FILE__, __LINE__, Format, ##__VA_ARGS__); \
 			} \
 		} \
 	}
@@ -290,7 +290,7 @@ struct FTCharArrayTester
 					if (ELogVerbosity::Verbosity == ELogVerbosity::Fatal) \
 					{\
 						_DebugBreakAndPromptForRemote(); \
-						FError::FinalLogf("", __FILE__, __LINE__, Format, ##__VA_ARGS__); \
+						FDebug::AssertFailed("", __FILE__, __LINE__, Format, ##__VA_ARGS__); \
 					}\
 				} \
 			} \
@@ -311,7 +311,7 @@ struct FTCharArrayTester
 						if (ELogVerbosity::Verbosity == ELogVerbosity::Fatal) \
 						{\
 							_DebugBreakAndPromptForRemote(); \
-							FError::FinalLogf("", __FILE__, __LINE__, Format, ##__VA_ARGS__); \
+							FDebug::AssertFailed("", __FILE__, __LINE__, Format, ##__VA_ARGS__); \
 						}\
 					} \
 				} \
diff --git a/Engine/Source/Runtime/Core/Public/Misc/OutputDevice.h b/Engine/Source/Runtime/Core/Public/Misc/OutputDevice.h
--- a/Engine/Source/Runtime/Core/Public/Misc/OutputDevice.h
+++ b/Engine/Source/Runtime/Core/Public/Misc/OutputDevice.h
@@ -388,12 +388,15 @@ struct CORE_API FMsg
  * These functions offer debugging and diagnostic functionality and its presence 
  * depends on compiler switches.
  **/
-#if DO_CHECK || DO_GUARD_SLOW
 struct CORE_API FDebug
 {
-	/** Failed assertion handler.  Warning: May be called at library startup time. */
-	static void VARARGS AssertFailed( const ANSICHAR* Expr, const ANSICHAR* File, int32 Line, const TCHAR* Format=TEXT(""), ... );
+	/** Logs final assert message and exits the program. */
+	static void VARARGS AssertFailed(const ANSICHAR* Expr, const ANSICHAR* File, int32 Line, const TCHAR* Format = TEXT(""), ...);
 
+#if DO_CHECK || DO_GUARD_SLOW
+	/** Failed assertion handler.  Warning: May be called at library startup time. */
+	static void VARARGS LogAssertFailedMessage( const ANSICHAR* Expr, const ANSICHAR* File, int32 Line, const TCHAR* Format=TEXT(""), ... );
+	
 	/**
 	 * Called when an 'ensure' assertion fails; gathers stack data and generates and error report.
 	 *
@@ -438,9 +441,9 @@ struct CORE_API FDebug
 	 * @param	Line	Line number (__LINE__)
 	 * @param	FormattedMsg	Informative error message text with variable args
 	 */
-	static bool VARARGS EnsureNotFalseFormatted( bool bExpressionResult, const ANSICHAR* Expr, const ANSICHAR* File, int32 Line, const TCHAR* FormattedMsg, ... );
-};
+	static bool VARARGS EnsureNotFalseFormatted(bool bExpressionResult, const ANSICHAR* Expr, const ANSICHAR* File, int32 Line, const TCHAR* FormattedMsg, ...);
 #endif // DO_CHECK || DO_GUARD_SLOW
+};
 
 
 /** 
@@ -483,9 +486,6 @@ struct CORE_API FError
 	/** For throwing string-exceptions which safely propagate through guard/unguard. */
 	VARARG_DECL( static void VARARGS, static void, VARARG_NONE, Throwf, VARARG_NONE, const TCHAR*, VARARG_NONE, VARARG_NONE );
 	//@}
-
-	/** Outputs formatted error message and terminates program. */
-	static void VARARGS FinalLogf(const ANSICHAR* Expr, const ANSICHAR* File, int32 Line, const TCHAR* Format = TEXT(""), ...);
 };
 
 
