Impose consistent Slate<->Cocoa coordinate conversions so that we don't accidently use the wrong coordinate space. This fixes [UE-9848] "Editor Drop Down Menus become out of place when moving across a small screen to a larger display" amongst other problems.

diff --git a/Engine/Source/Runtime/Core/Private/Mac/MacApplication.cpp b/Engine/Source/Runtime/Core/Private/Mac/MacApplication.cpp
--- a/Engine/Source/Runtime/Core/Private/Mac/MacApplication.cpp
+++ b/Engine/Source/Runtime/Core/Private/Mac/MacApplication.cpp
@@ -1258,7 +1258,12 @@ void FMacApplication::OnWindowDidMove( FCocoaWindow* Window )
 	NSRect OpenGLFrame = [Window openGLFrame];
 	
 	const int32 X = (int32)WindowFrame.origin.x;
-	const int32 Y = ([Window windowMode] != EWindowMode::Fullscreen) ? FPlatformMisc::ConvertSlateYPositionToCocoa( (int32)WindowFrame.origin.y ) - OpenGLFrame.size.height + 1 : 0;
+	int32 Y = 0;
+    
+    if ([Window windowMode] != EWindowMode::Fullscreen)
+    {
+        Y = FPlatformMisc::ConvertCocoaYPositionToSlate(WindowFrame.origin.y + WindowFrame.size.height) + (OpenGLFrame.size.height - WindowFrame.size.height);
+    }
 	
 	TSharedPtr< FMacWindow > EventWindow = FindWindowByNSWindow( Windows, &WindowsMutex, Window );
 	if (EventWindow.IsValid())
diff --git a/Engine/Source/Runtime/Core/Private/Mac/MacPlatformMisc.cpp b/Engine/Source/Runtime/Core/Private/Mac/MacPlatformMisc.cpp
--- a/Engine/Source/Runtime/Core/Private/Mac/MacPlatformMisc.cpp
+++ b/Engine/Source/Runtime/Core/Private/Mac/MacPlatformMisc.cpp
@@ -1145,6 +1145,25 @@ int32 FMacPlatformMisc::ConvertSlateYPositionToCocoa(int32 YPosition)
 	return -((YPosition - WholeWorkspaceOrigin) - WholeWorkspaceHeight + 1);
 }
 
+int32 FMacPlatformMisc::ConvertCocoaYPositionToSlate(int32 YPosition)
+{
+    NSArray* AllScreens = [NSScreen screens];
+    NSScreen* PrimaryScreen = (NSScreen*)[AllScreens objectAtIndex: 0];
+    NSRect ScreenFrame = [PrimaryScreen frame];
+    NSRect WholeWorkspace = {{0,0},{0,0}};
+    for(NSScreen* Screen in AllScreens)
+    {
+        if(Screen)
+        {
+            WholeWorkspace = NSUnionRect(WholeWorkspace, [Screen frame]);
+        }
+    }
+    
+    CGFloat const OffsetToPrimary = ((ScreenFrame.origin.y + ScreenFrame.size.height) - (WholeWorkspace.origin.y + WholeWorkspace.size.height));
+    CGFloat const OffsetToWorkspace = (WholeWorkspace.size.height - (YPosition)) + WholeWorkspace.origin.y;
+    return OffsetToWorkspace + OffsetToPrimary;
+}
+
 
 FString FMacPlatformMisc::GetDefaultLocale()
 {
diff --git a/Engine/Source/Runtime/Core/Private/Mac/MacWindow.cpp b/Engine/Source/Runtime/Core/Private/Mac/MacWindow.cpp
--- a/Engine/Source/Runtime/Core/Private/Mac/MacWindow.cpp
+++ b/Engine/Source/Runtime/Core/Private/Mac/MacWindow.cpp
@@ -35,9 +35,8 @@ void FMacWindow::Initialize( FMacApplication* const Application, const TSharedRe
 	int32 Y = FMath::TruncToInt( Definition->YDesiredPositionOnScreen );
 
 	// Make sure it's not under the menu bar on whatever display being targeted
-	const int32 ScreenHeight = FMath::TruncToInt( [TargetScreen frame].size.height );
-	const int32 VisibleHeight = FMath::TruncToInt( [TargetScreen visibleFrame].origin.y + [TargetScreen visibleFrame].size.height );
-	Y = FMath::Max( Y, ScreenHeight - VisibleHeight );
+	const int32 MaxVisibleY = FPlatformMisc::ConvertCocoaYPositionToSlate([TargetScreen visibleFrame].origin.y + [TargetScreen visibleFrame].size.height);
+	Y = (Y - MaxVisibleY) >= 0 ? Y : MaxVisibleY;
 
 	const int32 SizeX = FMath::Max(FMath::TruncToInt( Definition->WidthDesiredOnScreen ), 1);
 	const int32 SizeY = FMath::Max(FMath::TruncToInt( Definition->HeightDesiredOnScreen ), 1);
@@ -480,8 +479,10 @@ bool FMacWindow::IsPointInWindow( int32 X, int32 Y ) const
 	if(![WindowHandle isMiniaturized])
 	{
 		NSRect WindowFrame = [WindowHandle frame];
+		WindowFrame.size = [WindowHandle openGLFrame].size;
 		NSRect VisibleFrame = WindowFrame;
 		VisibleFrame.origin = NSMakePoint(0, 0);
+
 		// Only the editor needs to handle the space-per-display logic introduced in Mavericks.
 	#if WITH_EDITOR
 		// Only fetch the spans-displays once - it requires a log-out to change.
@@ -499,8 +500,8 @@ bool FMacWindow::IsPointInWindow( int32 X, int32 Y ) const
 			NSRect ScreenFrame = [[WindowHandle screen] frame];
 			NSRect Intersection = NSIntersectionRect(ScreenFrame, WindowFrame);
 			VisibleFrame.size = Intersection.size;
-			VisibleFrame.origin.x = WindowFrame.origin.x < ScreenFrame.origin.x ? (ScreenFrame.origin.x - WindowFrame.origin.x) : 0;
-			VisibleFrame.origin.y = (WindowFrame.origin.y + WindowFrame.size.height) > (ScreenFrame.origin.y + ScreenFrame.size.height) ? (WindowFrame.size.height - Intersection.size.height) : 0;
+			VisibleFrame.origin.x = Intersection.origin.x - WindowFrame.origin.x;
+			VisibleFrame.origin.y = Intersection.origin.y - WindowFrame.origin.y;
 		}
 	#endif
 		
@@ -508,7 +509,8 @@ bool FMacWindow::IsPointInWindow( int32 X, int32 Y ) const
 		{
 			FMacCursor* MacCursor = (FMacCursor*)MacApplication->Cursor.Get();
 			FVector2D MouseScale = MacCursor ? MacCursor->GetMouseScaling() : FVector2D(1.0f, 1.0f);
-			PointInWindow = (NSPointInRect(NSMakePoint(X / MouseScale.X, Y / MouseScale.Y), VisibleFrame) == YES);
+			NSPoint CursorPoint = NSMakePoint((X / MouseScale.X), (WindowFrame.size.height - ((Y / MouseScale.Y) + 1)));
+			PointInWindow = (NSPointInRect(CursorPoint, VisibleFrame) == YES);
 		}
 	}
 	return PointInWindow;
diff --git a/Engine/Source/Runtime/Core/Public/Mac/MacPlatformMisc.h b/Engine/Source/Runtime/Core/Public/Mac/MacPlatformMisc.h
--- a/Engine/Source/Runtime/Core/Public/Mac/MacPlatformMisc.h
+++ b/Engine/Source/Runtime/Core/Public/Mac/MacPlatformMisc.h
@@ -152,7 +152,16 @@ struct CORE_API FMacPlatformMisc : public FGenericPlatformMisc
 	 *
 	 * @return					Converted position
 	 */
-	static int32 ConvertSlateYPositionToCocoa(int32 YPosition);
+    static int32 ConvertSlateYPositionToCocoa(int32 YPosition);
+    
+    /**
+     * Converts Y position from Cocoa (left-bottom origin) to Slate (left-top origin) coordinates
+     *
+     * @param	YPosition		Position to convert
+     *
+     * @return					Converted position
+     */
+    static int32 ConvertCocoaYPositionToSlate(int32 YPosition);
 	
 	static void SetGracefulTerminationHandler();
 	
