Mac event handling refactor

diff --git a/Engine/Source/Runtime/Core/Private/Mac/CocoaWindow.cpp b/Engine/Source/Runtime/Core/Private/Mac/CocoaWindow.cpp
--- a/Engine/Source/Runtime/Core/Private/Mac/CocoaWindow.cpp
+++ b/Engine/Source/Runtime/Core/Private/Mac/CocoaWindow.cpp
@@ -163,7 +163,11 @@ NSString* NSPerformDragOperation = @"NSPerformDragOperation";
 	GameThreadCall(^{
 		if (MacApplication)
 		{
-			MacApplication->OnWindowClose(self);
+			TSharedPtr<FMacWindow> Window = MacApplication->FindWindowByNSWindow(self);
+			if (Window.IsValid())
+			{
+				MacApplication->CloseWindow(Window.ToSharedRef());
+			}
 		}
 	}, @[ NSDefaultRunLoopMode ], false);
 }
@@ -210,7 +214,7 @@ NSString* NSPerformDragOperation = @"NSPerformDragOperation";
 {
 	WindowMode = self.TargetWindowMode;
 
-	GameThreadCall(^{ FMacApplication::ProcessEvent([Notification retain]); }, @[ NSDefaultRunLoopMode, UE4FullscreenEventMode ], false);
+	MacApplication->DeferEvent(Notification);
 
 	FMacCursor* MacCursor = (FMacCursor*)MacApplication->Cursor.Get();
 	if (MacCursor)
@@ -236,7 +240,7 @@ NSString* NSPerformDragOperation = @"NSPerformDragOperation";
 	WindowMode = EWindowMode::Windowed;
 	self.TargetWindowMode = EWindowMode::Windowed;
 
-	GameThreadCall(^{ FMacApplication::ProcessEvent([Notification retain]); }, @[ NSDefaultRunLoopMode, UE4FullscreenEventMode ], false);
+	MacApplication->DeferEvent(Notification);
 
 	[self setFrame:self.PreFullScreenRect display:YES];
 	FMacCursor* MacCursor = (FMacCursor*)MacApplication->Cursor.Get();
@@ -254,7 +258,7 @@ NSString* NSPerformDragOperation = @"NSPerformDragOperation";
 		[self orderFrontAndMakeMain:false andKey:false];
 	}
 
-	GameThreadCall(^{ FMacApplication::ProcessEvent([Notification retain]); }, @[ NSDefaultRunLoopMode, UE4ShowEventMode, UE4CloseEventMode, UE4FullscreenEventMode ], false);
+	MacApplication->DeferEvent(Notification);
 }
 
 - (void)windowDidResignMain:(NSNotification*)Notification
@@ -263,12 +267,12 @@ NSString* NSPerformDragOperation = @"NSPerformDragOperation";
 	[self setMovable: YES];
 	[self setMovableByWindowBackground: NO];
 
-	GameThreadCall(^{ FMacApplication::ProcessEvent([Notification retain]); }, @[ NSDefaultRunLoopMode, UE4ShowEventMode, UE4CloseEventMode, UE4FullscreenEventMode ], false);
+	MacApplication->DeferEvent(Notification);
 }
 
 - (void)windowWillMove:(NSNotification*)Notification
 {
-	GameThreadCall(^{ FMacApplication::ProcessEvent([Notification retain]); }, @[ NSDefaultRunLoopMode, UE4ResizeEventMode, UE4ShowEventMode, UE4FullscreenEventMode ], false);
+	MacApplication->DeferEvent(Notification);
 }
 
 - (void)windowDidMove:(NSNotification*)Notification
@@ -279,7 +283,7 @@ NSString* NSPerformDragOperation = @"NSPerformDragOperation";
 	NSView* OpenGLView = [self openGLView];
 	[[NSNotificationCenter defaultCenter] postNotificationName:NSViewGlobalFrameDidChangeNotification object:OpenGLView];
 	
-	GameThreadCall(^{ FMacApplication::ProcessEvent([Notification retain]); }, @[ NSDefaultRunLoopMode, UE4ResizeEventMode, UE4ShowEventMode, UE4FullscreenEventMode ], false);
+	MacApplication->DeferEvent(Notification);
 }
 
 - (void)windowDidChangeScreen:(NSNotification*)Notification
@@ -347,65 +351,30 @@ NSString* NSPerformDragOperation = @"NSPerformDragOperation";
 - (void)windowWillStartLiveResize:(NSNotification*)Notification
 {
 	SCOPED_AUTORELEASE_POOL;
-	GameThreadCall(^{ FMacApplication::ProcessEvent([Notification retain]); }, @[ NSDefaultRunLoopMode, UE4ResizeEventMode, UE4ShowEventMode, UE4FullscreenEventMode ], false);
+	MacApplication->DeferEvent(Notification);
 }
 
 - (void)windowDidEndLiveResize:(NSNotification*)Notification
 {
 	SCOPED_AUTORELEASE_POOL;
-	GameThreadCall(^{ FMacApplication::ProcessEvent([Notification retain]); }, @[ NSDefaultRunLoopMode, UE4ResizeEventMode, UE4ShowEventMode, UE4FullscreenEventMode ], false);
+	MacApplication->DeferEvent(Notification);
 }
 
 - (void)windowDidResize:(NSNotification*)Notification
 {
 	SCOPED_AUTORELEASE_POOL;
 	bZoomed = [self isZoomed];
-	GameThreadCall(^{ FMacApplication::ProcessEvent([Notification retain]); }, @[ NSDefaultRunLoopMode, UE4ResizeEventMode, UE4ShowEventMode, UE4FullscreenEventMode ], true);
+	MacApplication->DeferEvent(Notification);
 }
 
 - (void)windowWillClose:(NSNotification*)Notification
 {
 	SCOPED_AUTORELEASE_POOL;
-	GameThreadCall(^{ FMacApplication::ProcessEvent([Notification retain]); }, @[ NSDefaultRunLoopMode, UE4CloseEventMode ], false);
-	[self setDelegate:nil];
-}
-
-- (void)mouseDown:(NSEvent*)Event
-{
-	GameThreadCall(^{ FMacApplication::ProcessEvent([Event retain]); }, @[ NSDefaultRunLoopMode ], false);
-}
-
-- (void)rightMouseDown:(NSEvent*)Event
-{
-	// Really we shouldn't be doing this - on OS X only left-click changes focus,
-	// but for the moment it is easier than changing Slate.
-	SCOPED_AUTORELEASE_POOL;
-	if ([self canBecomeKeyWindow] && self != [NSApp keyWindow])
+	if (MacApplication)
 	{
-		[self makeKeyWindow];
+		MacApplication->DeferEvent(Notification);
 	}
-	
-	GameThreadCall(^{ FMacApplication::ProcessEvent([Event retain]); }, @[ NSDefaultRunLoopMode ], false);
-}
-
-- (void)otherMouseDown:(NSEvent*)Event
-{
-	GameThreadCall(^{ FMacApplication::ProcessEvent([Event retain]); }, @[ NSDefaultRunLoopMode ], false);
-}
-
-- (void)mouseUp:(NSEvent*)Event
-{
-	GameThreadCall(^{ FMacApplication::ProcessEvent([Event retain]); }, @[ NSDefaultRunLoopMode ], false);
-}
-
-- (void)rightMouseUp:(NSEvent*)Event
-{
-	GameThreadCall(^{ FMacApplication::ProcessEvent([Event retain]); }, @[ NSDefaultRunLoopMode ], false);
-}
-
-- (void)otherMouseUp:(NSEvent*)Event
-{
-	GameThreadCall(^{ FMacApplication::ProcessEvent([Event retain]); }, @[ NSDefaultRunLoopMode ], false);
+	[self setDelegate:nil];
 }
 
 - (NSDragOperation)draggingEntered:(id <NSDraggingInfo>)Sender
@@ -417,14 +386,14 @@ NSString* NSPerformDragOperation = @"NSPerformDragOperation";
 {
 	SCOPED_AUTORELEASE_POOL;
 	NSNotification* Notification = [NSNotification notificationWithName:NSDraggingExited object:Sender];
-	GameThreadCall(^{ FMacApplication::ProcessEvent([Notification retain]); }, @[ NSDefaultRunLoopMode ], false);
+	MacApplication->DeferEvent(Notification);
 }
 
 - (NSDragOperation)draggingUpdated:(id <NSDraggingInfo>)Sender
 {
 	SCOPED_AUTORELEASE_POOL;
 	NSNotification* Notification = [NSNotification notificationWithName:NSDraggingUpdated object:Sender];
-	GameThreadCall(^{ FMacApplication::ProcessEvent([Notification retain]); }, @[ NSDefaultRunLoopMode ], false);
+	MacApplication->DeferEvent(Notification);
 	return NSDragOperationGeneric;
 }
 
@@ -432,7 +401,7 @@ NSString* NSPerformDragOperation = @"NSPerformDragOperation";
 {
 	SCOPED_AUTORELEASE_POOL;
 	NSNotification* Notification = [NSNotification notificationWithName:NSPrepareForDragOperation object:Sender];
-	GameThreadCall(^{ FMacApplication::ProcessEvent([Notification retain]); }, @[ NSDefaultRunLoopMode ], false);
+	MacApplication->DeferEvent(Notification);
 	return YES;
 }
 
@@ -440,7 +409,7 @@ NSString* NSPerformDragOperation = @"NSPerformDragOperation";
 {
 	SCOPED_AUTORELEASE_POOL;
 	NSNotification* Notification = [NSNotification notificationWithName:NSPerformDragOperation object:Sender];
-	GameThreadCall(^{ FMacApplication::ProcessEvent([Notification retain]); }, @[ NSDefaultRunLoopMode ], false);
+	MacApplication->DeferEvent(Notification);
 	return YES;
 }
 
diff --git a/Engine/Source/Runtime/Core/Private/Mac/MacApplication.cpp b/Engine/Source/Runtime/Core/Private/Mac/MacApplication.cpp
--- a/Engine/Source/Runtime/Core/Private/Mac/MacApplication.cpp
+++ b/Engine/Source/Runtime/Core/Private/Mac/MacApplication.cpp
@@ -14,7 +14,9 @@
 #include "ModuleManager.h"
 #include "CocoaTextView.h"
 
-FMacApplication* MacApplication = NULL;
+FMacApplication* MacApplication = nullptr;
+
+const uint32 RESET_EVENT_SUBTYPE = 0x0f00;
 
 #if WITH_EDITOR
 typedef int32 (*MTContactCallbackFunction)(void*, void*, int32, double, int32);
@@ -22,16 +24,6 @@ extern "C" CFMutableArrayRef MTDeviceCreateList(void);
 extern "C" void MTRegisterContactFrameCallback(void*, MTContactCallbackFunction);
 extern "C" void MTDeviceStart(void*, int);
 extern "C" bool MTDeviceIsBuiltIn(void*);
-
-static int MTContactCallback(void* Device, void* Data, int32 NumFingers, double TimeStamp, int32 Frame)
-{
-	if (MacApplication)
-	{
-		const bool bIsTrackpad = MTDeviceIsBuiltIn(Device);
-		MacApplication->SetIsUsingTrackpad(NumFingers > (bIsTrackpad ? 1 : 0));
-	}
-	return 1;
-}
 #endif
 
 FMacApplication* FMacApplication::CreateMacApplication()
@@ -40,113 +32,56 @@ FMacApplication* FMacApplication::CreateMacApplication()
 	return MacApplication;
 }
 
-NSEvent* FMacApplication::HandleNSEvent(NSEvent* Event)
+FMacApplication::FMacApplication()
+:	GenericApplication(MakeShareable(new FMacCursor()))
+,	bUsingHighPrecisionMouseInput(false)
+,	bUsingTrackpad(false)
+,	HIDInput(HIDInputInterface::Create(MessageHandler))
+,	DraggedWindow(nullptr)
+,	bSystemModalMode(false)
+,	ModifierKeysFlags(0)
+,	CurrentModifierFlags(0)
+,	bIsWorkspaceSessionActive(true)
 {
-	NSEvent* ReturnEvent = Event;
-
-	const bool bIsMouseClickOrKeyEvent = [Event type] == NSLeftMouseDown || [Event type] == NSLeftMouseUp
-		|| [Event type] == NSRightMouseDown || [Event type] == NSRightMouseUp
-		|| [Event type] == NSOtherMouseDown || [Event type] == NSOtherMouseUp;
-	const bool bIsResentEvent = [Event type] == NSApplicationDefined && (FMacApplicationEventTypes)[Event subtype] == FMacApplication::ResentEvent;
-
-	if (MacApplication)
+	TextInputMethodSystem = MakeShareable(new FMacTextInputMethodSystem);
+	if (!TextInputMethodSystem->Initialize())
 	{
-		if (bIsResentEvent)
-		{
-			ReturnEvent = (NSEvent*)[Event data1];
-		}
-
-		if (!bIsResentEvent && (!bIsMouseClickOrKeyEvent || [Event window] == NULL))
-		{
-			GameThreadCall(^{ FMacApplication::ProcessEvent([Event retain]); }, @[ NSDefaultRunLoopMode ], false);
-			
-			if ([Event type] == NSKeyDown || [Event type] == NSKeyUp)
-			{
-				ReturnEvent = nil;
-			}
-		}
+		TextInputMethodSystem.Reset();
 	}
 
-	return ReturnEvent;
-}
+	AppActivationObserver = [[NSNotificationCenter defaultCenter] addObserverForName:NSApplicationDidBecomeActiveNotification
+																			  object:[NSApplication sharedApplication]
+																			   queue:[NSOperationQueue mainQueue]
+																		  usingBlock:^(NSNotification* Notification) { OnApplicationDidBecomeActive(); }];
 
-void FMacApplication::OnDisplayReconfiguration(CGDirectDisplayID Display, CGDisplayChangeSummaryFlags Flags, void* UserInfo)
-{
-	FMacApplication* App = (FMacApplication*)UserInfo;
-	if (Flags & kCGDisplayDesktopShapeChangedFlag)
-	{
-		// Slate needs to know when desktop size changes.
-		FDisplayMetrics DisplayMetrics;
-		FDisplayMetrics::GetDisplayMetrics(DisplayMetrics);
-		App->BroadcastDisplayMetricsChanged(DisplayMetrics);
-	}
-	
-	for (int32 WindowIndex=0; WindowIndex < App->Windows.Num(); ++WindowIndex)
-	{
-		TSharedRef< FMacWindow > WindowRef = App->Windows[ WindowIndex ];
-		WindowRef->OnDisplayReconfiguration(Display, Flags);
-	}
-}
+	AppDeactivationObserver = [[NSNotificationCenter defaultCenter] addObserverForName:NSApplicationWillResignActiveNotification
+																				object:[NSApplication sharedApplication]
+																				 queue:[NSOperationQueue mainQueue]
+																			usingBlock:^(NSNotification* Notification) { OnApplicationWillResignActive(); }];
 
-FMacApplication::FMacApplication()
-	: GenericApplication( MakeShareable( new FMacCursor() ) )
-	, bUsingHighPrecisionMouseInput( false )
-	, bUsingTrackpad( false )
-	, HIDInput( HIDInputInterface::Create( MessageHandler ) )
-	, DraggedWindow( NULL )
-	, bIsMouseCursorLocked( false )
-	, bSystemModalMode( false )
-	, bIsProcessingNSEvent( false )
-	, ModifierKeysFlags( 0 )
-	, CurrentModifierFlags( 0 )
-	, bIsWorkspaceSessionActive( true )
-{
-	CGDisplayRegisterReconfigurationCallback(FMacApplication::OnDisplayReconfiguration, this);
+	WorkspaceActivationObserver = [[[NSWorkspace sharedWorkspace] notificationCenter] addObserverForName:NSWorkspaceSessionDidBecomeActiveNotification
+																								  object:[NSWorkspace sharedWorkspace]
+																								   queue:[NSOperationQueue mainQueue]
+																							  usingBlock:^(NSNotification* Notification){ bIsWorkspaceSessionActive = true; }];
 
-	[NSEvent addGlobalMonitorForEventsMatchingMask:NSMouseMovedMask handler:^(NSEvent* Event){
-		GameThreadCall(^{ FMacApplication::ProcessEvent([Event retain]); }, @[ NSDefaultRunLoopMode ], false);
-	}];
+	WorkspaceDeactivationObserver = [[[NSWorkspace sharedWorkspace] notificationCenter] addObserverForName:NSWorkspaceSessionDidResignActiveNotification
+																									object:[NSWorkspace sharedWorkspace]
+																									 queue:[NSOperationQueue mainQueue]
+																								usingBlock:^(NSNotification* Notification){ bIsWorkspaceSessionActive = false; }];
 
-	EventMonitor = [NSEvent addLocalMonitorForEventsMatchingMask:NSAnyEventMask handler:^(NSEvent* IncomingEvent)
-	{
-		NSEvent* ReturnEvent = HandleNSEvent(IncomingEvent);
-		return ReturnEvent;
-	}];
+	MouseMovedEventMonitor = [NSEvent addGlobalMonitorForEventsMatchingMask:NSMouseMovedMask handler:^(NSEvent* Event) { DeferEvent(Event); }];
+	EventMonitor = [NSEvent addLocalMonitorForEventsMatchingMask:NSAnyEventMask handler:^(NSEvent* IncomingEvent) { return HandleNSEvent(IncomingEvent); }];
+
+	CGDisplayRegisterReconfigurationCallback(FMacApplication::OnDisplayReconfiguration, this);
 
 #if WITH_EDITOR
 	NSMutableArray* MultiTouchDevices = (__bridge NSMutableArray*)MTDeviceCreateList();
 	for (id Device in MultiTouchDevices)
 	{
-		MTRegisterContactFrameCallback((void*)Device, MTContactCallback);
+		MTRegisterContactFrameCallback((void*)Device, FMacApplication::MTContactCallback);
 		MTDeviceStart((void*)Device, 0);
 	}
-#endif
-
-	TextInputMethodSystem = MakeShareable( new FMacTextInputMethodSystem );
-	if (!TextInputMethodSystem->Initialize())
-	{
-		TextInputMethodSystem.Reset();
-	}
-
-	AppActivationObserver = [[NSNotificationCenter defaultCenter] addObserverForName:NSApplicationDidBecomeActiveNotification object:[NSApplication sharedApplication] queue:[NSOperationQueue mainQueue] usingBlock:^(NSNotification* Notification)
-							{
-								OnApplicationDidBecomeActive();
-							 }];
-
-	AppDeactivationObserver = [[NSNotificationCenter defaultCenter] addObserverForName:NSApplicationWillResignActiveNotification object:[NSApplication sharedApplication] queue:[NSOperationQueue mainQueue] usingBlock:^(NSNotification* Notification)
-							{
-								OnApplicationWillResignActive();
-							}];
 
-	WorkspaceActivationObserver = [[[NSWorkspace sharedWorkspace] notificationCenter] addObserverForName:NSWorkspaceSessionDidBecomeActiveNotification object:[NSWorkspace sharedWorkspace] queue:[NSOperationQueue mainQueue] usingBlock:^(NSNotification* Notification){
-									   bIsWorkspaceSessionActive = true;
-								   }];
-
-	WorkspaceDeactivationObserver = [[[NSWorkspace sharedWorkspace] notificationCenter] addObserverForName:NSWorkspaceSessionDidResignActiveNotification object:[NSWorkspace sharedWorkspace] queue:[NSOperationQueue mainQueue] usingBlock:^(NSNotification* Notification){
-									   bIsWorkspaceSessionActive = false;
-								   }];
-
-#if WITH_EDITOR
 	FMemory::Memzero(GestureUsage);
 	LastGestureUsed = EGestureEvent::None;
 #endif
@@ -154,76 +89,102 @@ FMacApplication::FMacApplication()
 
 FMacApplication::~FMacApplication()
 {
+	if (MouseMovedEventMonitor)
+	{
+		[NSEvent removeMonitor:MouseMovedEventMonitor];
+	}
 	if (EventMonitor)
 	{
 		[NSEvent removeMonitor:EventMonitor];
-		EventMonitor = nil;
 	}
-	
 	if (AppActivationObserver)
 	{
 		[[NSNotificationCenter defaultCenter] removeObserver:AppActivationObserver];
-		AppActivationObserver = nil;
 	}
-	
 	if (AppDeactivationObserver)
 	{
 		[[NSNotificationCenter defaultCenter] removeObserver:AppDeactivationObserver];
-		AppDeactivationObserver = nil;
 	}
-	
 	if (WorkspaceActivationObserver)
 	{
 		[[NSNotificationCenter defaultCenter] removeObserver:WorkspaceActivationObserver];
-		WorkspaceActivationObserver = nil;
 	}
-
 	if (WorkspaceDeactivationObserver)
 	{
 		[[NSNotificationCenter defaultCenter] removeObserver:WorkspaceDeactivationObserver];
-		WorkspaceDeactivationObserver = nil;
 	}
 
 	CGDisplayRemoveReconfigurationCallback(FMacApplication::OnDisplayReconfiguration, this);
 
 	TextInputMethodSystem->Terminate();
 
-	MacApplication = NULL;
+	MacApplication = nullptr;
 }
 
-TSharedRef< FGenericWindow > FMacApplication::MakeWindow()
+void FMacApplication::SetMessageHandler(const TSharedRef<FGenericApplicationMessageHandler>& InMessageHandler)
 {
-	return FMacWindow::Make();
+	GenericApplication::SetMessageHandler(InMessageHandler);
+	HIDInput->SetMessageHandler(InMessageHandler);
 }
 
-void FMacApplication::InitializeWindow( const TSharedRef< FGenericWindow >& InWindow, const TSharedRef< FGenericWindowDefinition >& InDefinition, const TSharedPtr< FGenericWindow >& InParent, const bool bShowImmediately )
+void FMacApplication::PollGameDeviceState(const float TimeDelta)
 {
-	const TSharedRef< FMacWindow > Window = StaticCastSharedRef< FMacWindow >( InWindow );
-	const TSharedPtr< FMacWindow > ParentWindow = StaticCastSharedPtr< FMacWindow >( InParent );
-
-	Windows.Add( Window );
-	Window->Initialize( this, InDefinition, ParentWindow, bShowImmediately );
+	// Poll game device state and send new events
+	HIDInput->SendControllerEvents();
 }
 
-void FMacApplication::SetMessageHandler( const TSharedRef< FGenericApplicationMessageHandler >& InMessageHandler )
+void FMacApplication::PumpMessages(const float TimeDelta)
 {
-	GenericApplication::SetMessageHandler( InMessageHandler );
-	HIDInput->SetMessageHandler( InMessageHandler );
+	FPlatformMisc::PumpMessages(true);
 }
 
-static TSharedPtr< FMacWindow > FindWindowByNSWindow( const TArray< TSharedRef< FMacWindow > >& WindowsToSearch, FCriticalSection* const Mutex, FCocoaWindow* const WindowHandle )
+void FMacApplication::ProcessDeferredEvents(const float TimeDelta)
 {
-	FScopeLock Lock(Mutex);
-	for (int32 WindowIndex=0; WindowIndex < WindowsToSearch.Num(); ++WindowIndex)
+	TArray<FDeferredMacEvent> EventsToProcess;
+
+	EventsMutex.Lock();
+	EventsToProcess.Append(DeferredEvents);
+	DeferredEvents.Empty();
+	EventsMutex.Unlock();
+
+	for (int32 Index = 0; Index < EventsToProcess.Num(); ++Index)
 	{
-		TSharedRef< FMacWindow > Window = WindowsToSearch[ WindowIndex ];
-		if ( Window->GetWindowHandle() == WindowHandle )
-		{
-			return Window;
-		}
+		ProcessEvent(EventsToProcess[Index]);
 	}
 
-	return TSharedPtr< FMacWindow >( NULL );
+	InvalidateTextLayouts();
+	CloseQueuedWindows();
+}
+
+TSharedRef<FGenericWindow> FMacApplication::MakeWindow()
+{
+	return FMacWindow::Make();
+}
+
+void FMacApplication::InitializeWindow(const TSharedRef<FGenericWindow>& InWindow, const TSharedRef<FGenericWindowDefinition>& InDefinition, const TSharedPtr<FGenericWindow>& InParent, const bool bShowImmediately)
+{
+	const TSharedRef<FMacWindow> Window = StaticCastSharedRef<FMacWindow >(InWindow);
+	const TSharedPtr<FMacWindow> ParentWindow = StaticCastSharedPtr<FMacWindow>(InParent);
+
+	Windows.Add(Window);
+	Window->Initialize(this, InDefinition, ParentWindow, bShowImmediately);
+}
+
+FModifierKeysState FMacApplication::GetModifierKeys() const
+{
+	uint32 CurrentFlags = ModifierKeysFlags;
+
+	const bool bIsLeftShiftDown			= (CurrentFlags & (1 << 0)) != 0;
+	const bool bIsRightShiftDown		= (CurrentFlags & (1 << 1)) != 0;
+	const bool bIsLeftControlDown		= (CurrentFlags & (1 << 6)) != 0; // Mac pretends the Command key is Control
+	const bool bIsRightControlDown		= (CurrentFlags & (1 << 7)) != 0; // Mac pretends the Command key is Control
+	const bool bIsLeftAltDown			= (CurrentFlags & (1 << 4)) != 0;
+	const bool bIsRightAltDown			= (CurrentFlags & (1 << 5)) != 0;
+	const bool bIsLeftCommandDown		= (CurrentFlags & (1 << 2)) != 0; // Mac pretends the Control key is Command
+	const bool bIsRightCommandDown		= (CurrentFlags & (1 << 3)) != 0; // Mac pretends the Control key is Command
+	const bool bAreCapsLocked           = (CurrentFlags & (1 << 8)) != 0;
+
+	return FModifierKeysState(bIsLeftShiftDown, bIsRightShiftDown, bIsLeftControlDown, bIsRightControlDown, bIsLeftAltDown, bIsRightAltDown, bIsLeftCommandDown, bIsRightCommandDown, bAreCapsLocked);
 }
 
 bool FMacApplication::IsCursorDirectlyOverSlateWindow() const
@@ -234,1145 +195,989 @@ bool FMacApplication::IsCursorDirectlyOverSlateWindow() const
 	return Window && [Window isKindOfClass:[FCocoaWindow class]] && Window != DraggedWindow;
 }
 
-void FMacApplication::ProcessEvent(NSObject const* const Event)
+void FMacApplication::SetHighPrecisionMouseMode(const bool Enable, const TSharedPtr<FGenericWindow>& InWindow)
 {
-	// Must have an event
-	check(Event);
-	
-	if ([Event isKindOfClass:[NSEvent class]])
-	{
-		// Process a standard NSEvent as we always did
-		MacApplication->ProcessNSEvent((NSEvent*)Event);
-	}
-	else if ([Event isKindOfClass:[NSNotification class]])
+	bUsingHighPrecisionMouseInput = Enable;
+	((FMacCursor*)Cursor.Get())->SetHighPrecisionMouseMode(Enable);
+}
+
+FPlatformRect FMacApplication::GetWorkArea(const FPlatformRect& CurrentWindow) const
+{
+	SCOPED_AUTORELEASE_POOL;
+
+	NSScreen* Screen = FindScreenByPoint(CurrentWindow.Left, CurrentWindow.Top);
+
+	const int32 ScreenHeight = FMath::TruncToInt([Screen frame].size.height);
+	const NSRect VisibleFrame = [Screen visibleFrame];
+
+	NSArray* AllScreens = [NSScreen screens];
+	NSScreen* PrimaryScreen = (NSScreen*)[AllScreens objectAtIndex: 0];
+	NSRect PrimaryFrame = [PrimaryScreen frame];
+
+	FPlatformRect WorkArea;
+	WorkArea.Left = VisibleFrame.origin.x;
+	WorkArea.Top = (PrimaryFrame.origin.y + PrimaryFrame.size.height) - (VisibleFrame.origin.y + VisibleFrame.size.height);
+	WorkArea.Right = WorkArea.Left + VisibleFrame.size.width;
+	WorkArea.Bottom = WorkArea.Top + VisibleFrame.size.height;
+
+	return WorkArea;
+}
+
+#if WITH_EDITOR
+void FMacApplication::SendAnalytics(IAnalyticsProvider* Provider)
+{
+	static_assert(EGestureEvent::Count == 5, "If the number of gestures changes you need to add more entries below!");
+
+	TArray<FAnalyticsEventAttribute> GestureAttributes;
+	GestureAttributes.Add(FAnalyticsEventAttribute(FString("Scroll"),	GestureUsage[EGestureEvent::Scroll]));
+	GestureAttributes.Add(FAnalyticsEventAttribute(FString("Magnify"),	GestureUsage[EGestureEvent::Magnify]));
+	GestureAttributes.Add(FAnalyticsEventAttribute(FString("Swipe"),	GestureUsage[EGestureEvent::Swipe]));
+	GestureAttributes.Add(FAnalyticsEventAttribute(FString("Rotate"),	GestureUsage[EGestureEvent::Rotate]));
+
+	Provider->RecordEvent(FString("Mac.Gesture.Usage"), GestureAttributes);
+
+	FMemory::Memzero(GestureUsage);
+	LastGestureUsed = EGestureEvent::None;
+}
+#endif
+
+void FMacApplication::CloseWindow(TSharedRef<FMacWindow> Window)
+{
+	MessageHandler->OnWindowClose(Window);
+}
+
+void FMacApplication::DeferEvent(NSObject* Object)
+{
+	FDeferredMacEvent DeferredEvent;
+	FCocoaWindow* WindowHandle = nullptr;
+
+	if (Object && [Object isKindOfClass:[NSEvent class]])
 	{
-		// Notifications need to mapped to the right handler function, that's no longer handled in the
-		// call location.
-		NSNotification* Notification = (NSNotification*)Event;
-		NSString* const NotificationName = [Notification name];
-		FCocoaWindow* CocoaWindow = [[Notification object] isKindOfClass:[FCocoaWindow class]] ? (FCocoaWindow*)[Notification object] : nullptr;
-		if (CocoaWindow)
+		NSEvent* Event = (NSEvent*)Object;
+		WindowHandle = FindEventWindow(Event);
+		DeferredEvent.Type = [Event type];
+		DeferredEvent.LocationInWindow = FVector2D([Event locationInWindow].x, [Event locationInWindow].y);
+		DeferredEvent.ModifierFlags = [Event modifierFlags];
+		DeferredEvent.Timestamp = [Event timestamp];
+		DeferredEvent.WindowNumber = [Event windowNumber];
+		DeferredEvent.Context = [[Event context] retain];
+
+		if (DeferredEvent.Type == NSKeyDown && !bSystemModalMode)
 		{
-			if (NotificationName == NSWindowDidResizeNotification)
-			{
-				MacApplication->OnWindowDidResize(CocoaWindow);
-			}
-			else if (NotificationName == NSWindowWillStartLiveResizeNotification)
-			{
-				MacApplication->OnWindowWillStartLiveResize(CocoaWindow);
-			}
-			else if (NotificationName == NSWindowDidEndLiveResizeNotification)
-			{
-				MacApplication->OnWindowDidEndLiveResize(CocoaWindow);
-			}
-			else if (NotificationName == NSWindowDidEnterFullScreenNotification)
-			{
-				MacApplication->OnWindowDidResize(CocoaWindow);
-			}
-			else if (NotificationName == NSWindowDidExitFullScreenNotification)
-			{
-				MacApplication->OnWindowDidResize(CocoaWindow);
-			}
-			else if (NotificationName == NSWindowDidBecomeMainNotification)
-			{
-				MacApplication->OnWindowDidBecomeMain(CocoaWindow);
-			}
-			else if (NotificationName == NSWindowDidResignMainNotification)
-			{
-				MacApplication->OnWindowDidResignMain(CocoaWindow);
-			}
-			else if (NotificationName == NSWindowWillMoveNotification)
-			{
-				MacApplication->OnWindowWillMove(CocoaWindow);
-			}
-			else if (NotificationName == NSWindowDidMoveNotification)
+			DeferredEvent.Window = FindWindowByNSWindow(WindowHandle);
+			if (DeferredEvent.Window.IsValid() && DeferredEvent.Window->OnIMKKeyDown(Event))
 			{
-				MacApplication->OnWindowDidMove(CocoaWindow);
+				return;
 			}
-			else if (NotificationName == NSWindowWillCloseNotification)
+		}
+
+		switch (DeferredEvent.Type)
+		{
+			case NSMouseMoved:
+			case NSLeftMouseDragged:
+			case NSRightMouseDragged:
+			case NSOtherMouseDragged:
+			case NSEventTypeSwipe:
+				DeferredEvent.Delta = FVector2D([Event deltaX], [Event deltaY]);
+				break;
+
+			case NSLeftMouseDown:
+			case NSRightMouseDown:
+			case NSOtherMouseDown:
+			case NSLeftMouseUp:
+			case NSRightMouseUp:
+			case NSOtherMouseUp:
+				DeferredEvent.ButtonNumber = [Event buttonNumber];
+				DeferredEvent.ClickCount = [Event clickCount];
+				break;
+
+			case NSScrollWheel:
+				DeferredEvent.Delta = FVector2D([Event deltaX], [Event deltaY]);
+				DeferredEvent.ScrollingDelta = FVector2D([Event scrollingDeltaX], [Event scrollingDeltaY]);
+				DeferredEvent.Phase = [Event phase];
+				DeferredEvent.MomentumPhase = [Event momentumPhase];
+				DeferredEvent.IsDirectionInvertedFromDevice = [Event isDirectionInvertedFromDevice];
+				break;
+
+			case NSEventTypeMagnify:
+				DeferredEvent.Delta = FVector2D([Event magnification], [Event magnification]);
+				break;
+
+			case NSEventTypeRotate:
+				DeferredEvent.Delta = FVector2D([Event rotation], [Event rotation]);
+				break;
+
+			case NSKeyDown:
+			case NSKeyUp:
 			{
-				MacApplication->OnWindowDidClose(CocoaWindow);
+				if ([[Event characters] length] > 0)
+				{
+					DeferredEvent.Characters = [[Event characters] retain];
+					DeferredEvent.CharactersIgnoringModifiers = [[Event charactersIgnoringModifiers] retain];
+					DeferredEvent.IsRepeat = [Event isARepeat];
+					DeferredEvent.KeyCode = [Event keyCode];
+				}
+				else
+				{
+					return;
+				}
+				break;
 			}
-			else
+		}
+	}
+	else if (Object && [Object isKindOfClass:[NSNotification class]])
+	{
+		NSNotification* Notification = (NSNotification*)Object;
+		DeferredEvent.NotificationName = [[Notification name] retain];
+		if ([[Notification object] isKindOfClass:[FCocoaWindow class]])
+		{
+			WindowHandle = (FCocoaWindow*)[Notification object];
+
+			if (DeferredEvent.NotificationName == NSWindowDidResizeNotification)
 			{
-				check(false);
+				DeferredEvent.Window = FindWindowByNSWindow(WindowHandle);
+				if (DeferredEvent.Window.IsValid())
+				{
+					GameThreadCall(^{ OnWindowDidResize(DeferredEvent.Window.ToSharedRef()); }, @[ NSDefaultRunLoopMode, UE4ResizeEventMode, UE4ShowEventMode, UE4FullscreenEventMode ], true);
+				}
+				return;
 			}
 		}
 		else if ([[Notification object] conformsToProtocol:@protocol(NSDraggingInfo)])
 		{
-			CocoaWindow = (FCocoaWindow*)[(id<NSDraggingInfo>)[Notification object] draggingDestinationWindow];
-			if (NotificationName == NSDraggingExited)
-			{
-				MacApplication->OnDragOut(CocoaWindow);
-			}
-			else if (NotificationName == NSDraggingUpdated)
-			{
-				MacApplication->OnDragOver(CocoaWindow);
-			}
-			else if (NotificationName == NSPrepareForDragOperation)
-			{
-				MacApplication->OnDragEnter(CocoaWindow, [(id <NSDraggingInfo>)[Notification object] draggingPasteboard]);
-			}
-			else if (NotificationName == NSPerformDragOperation)
-			{
-				MacApplication->OnDragDrop(CocoaWindow);
-			}
-			else
+			WindowHandle = (FCocoaWindow*)[(id<NSDraggingInfo>)[Notification object] draggingDestinationWindow];
+			if (DeferredEvent.NotificationName == NSPrepareForDragOperation)
 			{
-				check(false);
+				DeferredEvent.DraggingPasteboard = [[(id<NSDraggingInfo>)[Notification object] draggingPasteboard] retain];
 			}
 		}
 	}
-	
-	// We are responsible for deallocating the event
-	[Event release];
+
+	DeferredEvent.Window = FindWindowByNSWindow(WindowHandle);
+
+	FScopeLock Lock(&EventsMutex);
+	DeferredEvents.Add(DeferredEvent);
 }
 
-void FMacApplication::RedrawWindow(FCocoaWindow* Window)
+TSharedPtr<FMacWindow> FMacApplication::FindWindowByNSWindow(FCocoaWindow* WindowHandle)
 {
-	TSharedPtr<FMacWindow> EventWindow = FindWindowByNSWindow(Windows, &WindowsMutex, Window);
-	if (EventWindow.IsValid())
+	FScopeLock Lock(&WindowsMutex);
+
+	for (int32 WindowIndex=0; WindowIndex < Windows.Num(); ++WindowIndex)
 	{
-		MessageHandler->OnOSPaint(EventWindow.ToSharedRef());
+		TSharedRef<FMacWindow> Window = Windows[WindowIndex];
+		if (Window->GetWindowHandle() == WindowHandle)
+		{
+			return Window;
+		}
 	}
+
+	return TSharedPtr<FMacWindow>(nullptr);
 }
 
-void FMacApplication::PumpMessages(const float TimeDelta)
+void FMacApplication::InvalidateTextLayout(FCocoaWindow* Window)
 {
-	FPlatformMisc::PumpMessages(true);
+	WindowsRequiringTextInvalidation.AddUnique(Window);
 }
 
-void FMacApplication::CloseQueuedWindows()
+NSEvent* FMacApplication::HandleNSEvent(NSEvent* Event)
 {
-	if (WindowsToClose.Num() > 0)
+	NSEvent* ReturnEvent = Event;
+
+	if (MacApplication)
 	{
-		MainThreadCall(^{
-			SCOPED_AUTORELEASE_POOL;
+		const bool bIsResentEvent = [Event type] == NSApplicationDefined && [Event subtype] == RESET_EVENT_SUBTYPE;
 
-			// Set the new key window, if needed. We cannot trust Cocoa to set the key window to the actual top most window. It will prefer windows with title bars so,
-			// for example, will choose the main window over a context menu window, when closing a submenu.
-			NSArray* AllWindows = [NSApp orderedWindows];
-			for (NSWindow* Window : AllWindows)
-			{
-				if ([Window isKindOfClass:[FCocoaWindow class]] && !WindowsToClose.Contains((FCocoaWindow*)Window) && [Window canBecomeKeyWindow])
-				{
-					if (Window != [NSApp keyWindow])
-					{
-						[Window makeKeyWindow];
-					}
-					break;
-				}
-			}
+		if (bIsResentEvent)
+		{
+			ReturnEvent = (NSEvent*)[Event data1];
+		}
+		else
+		{
+			MacApplication->DeferEvent(Event);
 
-			for (FCocoaWindow* Window : WindowsToClose)
+			if ([Event type] == NSKeyDown || [Event type] == NSKeyUp)
 			{
-				[Window close];
-				[Window release];
+				ReturnEvent = nil;
 			}
-		}, UE4CloseEventMode, true);
-
-		WindowsToClose.Empty();
+		}
 	}
-}
 
-void FMacApplication::InvalidateTextLayout(FCocoaWindow* Window)
-{
-	WindowsRequiringTextInvalidation.AddUnique( Window );
+	return ReturnEvent;
 }
 
-void FMacApplication::InvalidateTextLayouts()
+void FMacApplication::OnDisplayReconfiguration(CGDirectDisplayID Display, CGDisplayChangeSummaryFlags Flags, void* UserInfo)
 {
-	if( WindowsRequiringTextInvalidation.Num() > 0 )
+	FMacApplication* App = (FMacApplication*)UserInfo;
+	if (Flags & kCGDisplayDesktopShapeChangedFlag)
 	{
-	   MainThreadCall(^{
-		   SCOPED_AUTORELEASE_POOL;
-	
-		   for( FCocoaWindow* CocoaWindow : WindowsRequiringTextInvalidation)
-		   {
-			   if(CocoaWindow && [CocoaWindow openGLView])
-			   {
-				   FCocoaTextView* TextView = (FCocoaTextView*)[CocoaWindow openGLView];
-				   [[TextView inputContext] invalidateCharacterCoordinates];
-			   }
-		   }
-
-		}, UE4IMEEventMode, true);
-		
-		WindowsRequiringTextInvalidation.Empty();
+		// Slate needs to know when desktop size changes.
+		FDisplayMetrics DisplayMetrics;
+		FDisplayMetrics::GetDisplayMetrics(DisplayMetrics);
+		App->BroadcastDisplayMetricsChanged(DisplayMetrics);
 	}
 
+	for (int32 WindowIndex=0; WindowIndex < App->Windows.Num(); ++WindowIndex)
+	{
+		TSharedRef<FMacWindow> WindowRef = App->Windows[WindowIndex];
+		WindowRef->OnDisplayReconfiguration(Display, Flags);
+	}
 }
 
-bool FMacApplication::IsWindowMovable(FCocoaWindow* Win, bool* OutMovableByBackground)
+#if WITH_EDITOR
+int32 FMacApplication::MTContactCallback(void* Device, void* Data, int32 NumFingers, double TimeStamp, int32 Frame)
 {
-	if (OutMovableByBackground)
+	if (MacApplication)
 	{
-		*OutMovableByBackground = false;
+		const bool bIsTrackpad = MTDeviceIsBuiltIn(Device);
+		MacApplication->bUsingTrackpad = NumFingers > (bIsTrackpad ? 1 : 0);
 	}
-
-	FCocoaWindow* NativeWindow = (FCocoaWindow*)Win;
-	TSharedPtr< FMacWindow > CurrentEventWindow = FindWindowByNSWindow( Windows, &WindowsMutex, NativeWindow );
-	if ( CurrentEventWindow.IsValid() )
-	{
-		const FVector2D CursorPos = static_cast<FMacCursor*>( Cursor.Get() )->GetPosition();
-		const int32 LocalMouseX = CursorPos.X - CurrentEventWindow->PositionX;
-		const int32 LocalMouseY = CursorPos.Y - CurrentEventWindow->PositionY;
-		
-		const EWindowZone::Type Zone = MessageHandler->GetWindowZoneForPoint( CurrentEventWindow.ToSharedRef(), LocalMouseX, LocalMouseY );
-		bool IsMouseOverTitleBar = Zone == EWindowZone::TitleBar;
-		switch(Zone)
-		{
-			case EWindowZone::NotInWindow:
-			case EWindowZone::TopLeftBorder:
-			case EWindowZone::TopBorder:
-			case EWindowZone::TopRightBorder:
-			case EWindowZone::LeftBorder:
-			case EWindowZone::RightBorder:
-			case EWindowZone::BottomLeftBorder:
-			case EWindowZone::BottomBorder:
-			case EWindowZone::BottomRightBorder:
-				return true;
-			case EWindowZone::TitleBar:
-				if (OutMovableByBackground)
-				{
-					*OutMovableByBackground = true;
-				}
-				return true;
-			case EWindowZone::ClientArea:
-			case EWindowZone::MinimizeButton:
-			case EWindowZone::MaximizeButton:
-			case EWindowZone::CloseButton:
-			case EWindowZone::SysMenu:
-			default:
-				return false;
-		}
-	}
-	return true;
+	return 1;
 }
+#endif
 
-void FMacApplication::HandleModifierChange(NSUInteger NewModifierFlags, NSUInteger FlagsShift, NSUInteger UE4Shift, EMacModifierKeys TranslatedCode)
+void FMacApplication::ProcessEvent(const FDeferredMacEvent& Event)
 {
-	bool CurrentPressed = (CurrentModifierFlags & FlagsShift) != 0;
-	bool NewPressed = (NewModifierFlags & FlagsShift) != 0;
-	if (CurrentPressed != NewPressed)
+	if (Event.Type)
 	{
-		if ( NewPressed )
+		if (CurrentModifierFlags != Event.ModifierFlags)
 		{
-			ModifierKeysFlags |= 1 << UE4Shift;
-			MessageHandler->OnKeyDown( TranslatedCode, 0, false );
+			NSUInteger ModifierFlags = Event.ModifierFlags;
+
+			HandleModifierChange(ModifierFlags, (1<<4), 7, MMK_RightCommand);
+			HandleModifierChange(ModifierFlags, (1<<3), 6, MMK_LeftCommand);
+			HandleModifierChange(ModifierFlags, (1<<1), 0, MMK_LeftShift);
+			HandleModifierChange(ModifierFlags, (1<<16), 8, MMK_CapsLock);
+			HandleModifierChange(ModifierFlags, (1<<5), 4, MMK_LeftAlt);
+			HandleModifierChange(ModifierFlags, (1<<0), 2, MMK_LeftControl);
+			HandleModifierChange(ModifierFlags, (1<<2), 1, MMK_RightShift);
+			HandleModifierChange(ModifierFlags, (1<<6), 5, MMK_RightAlt);
+			HandleModifierChange(ModifierFlags, (1<<13), 3, MMK_RightControl);
+
+			CurrentModifierFlags = ModifierFlags;
 		}
-		else
+
+		switch (Event.Type)
 		{
-			ModifierKeysFlags &= ~(1 << UE4Shift);
-			MessageHandler->OnKeyUp( TranslatedCode, 0, false );
+			case NSMouseMoved:
+			case NSLeftMouseDragged:
+			case NSRightMouseDragged:
+			case NSOtherMouseDragged:
+				ProcessMouseMovedEvent(Event);
+				break;
+
+			case NSLeftMouseDown:
+			case NSRightMouseDown:
+			case NSOtherMouseDown:
+				ProcessMouseDownEvent(Event);
+				break;
+
+			case NSLeftMouseUp:
+			case NSRightMouseUp:
+			case NSOtherMouseUp:
+				ProcessMouseUpEvent(Event);
+				break;
+
+			case NSScrollWheel:
+				ProcessScrollWheelEvent(Event);
+				break;
+
+			case NSEventTypeMagnify:
+			case NSEventTypeSwipe:
+			case NSEventTypeRotate:
+			case NSEventTypeBeginGesture:
+			case NSEventTypeEndGesture:
+				ProcessGestureEvent(Event);
+				break;
+
+			case NSKeyDown:
+				ProcessKeyDownEvent(Event);
+				break;
+
+			case NSKeyUp:
+				ProcessKeyUpEvent(Event);
+				break;
 		}
 	}
-}
-
-void FMacApplication::ProcessNSEvent(NSEvent* const Event)
-{
-	SCOPED_AUTORELEASE_POOL;
-
-	const bool bWasProcessingNSEvent = bIsProcessingNSEvent; // This function can be called recursively
-	bIsProcessingNSEvent = true;
-
-	if (CurrentModifierFlags != [Event modifierFlags])
-	{
-		NSUInteger ModifierFlags = [Event modifierFlags];
-
-		HandleModifierChange(ModifierFlags, (1<<4), 7, MMK_RightCommand);
-		HandleModifierChange(ModifierFlags, (1<<3), 6, MMK_LeftCommand);
-		HandleModifierChange(ModifierFlags, (1<<1), 0, MMK_LeftShift);
-		HandleModifierChange(ModifierFlags, (1<<16), 8, MMK_CapsLock);
-		HandleModifierChange(ModifierFlags, (1<<5), 4, MMK_LeftAlt);
-		HandleModifierChange(ModifierFlags, (1<<0), 2, MMK_LeftControl);
-		HandleModifierChange(ModifierFlags, (1<<2), 1, MMK_RightShift);
-		HandleModifierChange(ModifierFlags, (1<<6), 5, MMK_RightAlt);
-		HandleModifierChange(ModifierFlags, (1<<13), 3, MMK_RightControl);
-
-		CurrentModifierFlags = ModifierFlags;
-	}
-
-	FCocoaWindow* NativeWindow = FindEventWindow(Event);
-	TSharedPtr<FMacWindow> CurrentEventWindow = FindWindowByNSWindow(MacApplication->Windows, &MacApplication->WindowsMutex, NativeWindow);
-
-	const NSEventType EventType = [Event type];
-	switch (EventType)
+	else if (Event.Window.IsValid())
 	{
-		case NSMouseMoved:
-		case NSLeftMouseDragged:
-		case NSRightMouseDragged:
-		case NSOtherMouseDragged:
+		if (Event.NotificationName == NSWindowWillStartLiveResizeNotification)
 		{
-			if (CurrentEventWindow.IsValid() && CurrentEventWindow->IsRegularWindow())
-			{
-				bool IsMouseOverTitleBar = false;
-				const bool IsMovable = IsWindowMovable(NativeWindow, &IsMouseOverTitleBar);
-				[NativeWindow setMovable:IsMovable];
-				[NativeWindow setMovableByWindowBackground:IsMouseOverTitleBar];
-			}
-
-			FMacCursor* MacCursor = (FMacCursor*)Cursor.Get();
-
-			FVector2D const MouseScaling = MacCursor->GetMouseScaling();
-			
-			if (bUsingHighPrecisionMouseInput)
-			{
-				// Under OS X we disassociate the cursor and mouse position during hi-precision mouse input.
-				// The game snaps the mouse cursor back to the starting point when this is disabled, which
-				// accumulates mouse delta that we want to ignore.
-				const FVector2D AccumDelta = MacCursor->GetMouseWarpDelta() * MouseScaling;
-				
-				// Get the mouse position
-				FVector2D HighPrecisionMousePos = MacCursor->GetPosition();
-				
-				// Find the screen the cursor is currently on.
-				NSEnumerator *ScreenEnumerator = [[NSScreen screens] objectEnumerator];
-				NSScreen *Screen;
-				while ((Screen = [ScreenEnumerator nextObject]) && !NSMouseInRect(NSMakePoint(HighPrecisionMousePos.X / MouseScaling.X, HighPrecisionMousePos.Y / MouseScaling.Y), Screen.frame, NO))
-					;
-				
-				// Account for warping delta's
-				FVector2D Delta = FVector2D([Event deltaX], [Event deltaY]) * MouseScaling;
-				const FVector2D WarpDelta(FMath::Abs(AccumDelta.X)<FMath::Abs(Delta.X) ? AccumDelta.X : Delta.X, FMath::Abs(AccumDelta.Y)<FMath::Abs(Delta.Y) ? AccumDelta.Y : Delta.Y);
-				Delta -= WarpDelta;
-				
-				// Update to latest position
-				HighPrecisionMousePos += Delta;
-				
-				// Clip to lock rect
-				MacCursor->UpdateCursorClipping(HighPrecisionMousePos);
-				
-				// Clamp to the current screen and avoid the menu bar and dock to prevent popups and other
-				// assorted potential for mouse abuse.
-				NSRect VisibleFrame = [Screen visibleFrame];
-				// Avoid the menu bar & dock disclosure borders at the top & bottom of fullscreen windows
-				if (CurrentEventWindow.IsValid() && CurrentEventWindow->GetWindowMode() != EWindowMode::Windowed)
-				{
-					VisibleFrame.origin.y += 5;
-					VisibleFrame.size.height -= 10;
-				}
-				NSRect FullFrame = [Screen frame];
-				VisibleFrame.origin.y = (FullFrame.origin.y+FullFrame.size.height) - (VisibleFrame.origin.y + VisibleFrame.size.height);
-				
-				int32 ScaledX = (int32)(HighPrecisionMousePos.X / MouseScaling.X);
-				int32 ScaledY = (int32)(HighPrecisionMousePos.Y / MouseScaling.Y);
-				int32 ClampedPosX = FMath::Clamp(ScaledX, (int32)VisibleFrame.origin.x, (int32)(VisibleFrame.origin.x + VisibleFrame.size.width)-1);
-				int32 ClampedPosY = FMath::Clamp(ScaledY, (int32)VisibleFrame.origin.y, (int32)(VisibleFrame.origin.y + VisibleFrame.size.height)-1);
-				MacCursor->SetPosition(ClampedPosX * MouseScaling.X, ClampedPosY * MouseScaling.Y);
-				
-				// Forward the delta on to Slate
-				MessageHandler->OnRawMouseMove(Delta.X, Delta.Y);
-			}
-			else
-			{
-				NSPoint CursorPos = [NSEvent mouseLocation];
-				CursorPos.y--; // The y coordinate of the point returned by mouseLocation starts from a base of 1
-				const FVector2D MousePosition = FVector2D(CursorPos.x, FPlatformMisc::ConvertSlateYPositionToCocoa(CursorPos.y));
-				FVector2D CurrentPosition = MousePosition * MouseScaling;
-				const FVector2D MouseDelta = CurrentPosition - MacCursor->GetPosition();
-				if (MacCursor->UpdateCursorClipping(CurrentPosition))
-				{
-					MacCursor->SetPosition(CurrentPosition.X, CurrentPosition.Y);
-				}
-				else
-				{
-					MacCursor->UpdateCurrentPosition(CurrentPosition);
-				}
-
-				if (CurrentEventWindow.IsValid())
-				{
-					// Cocoa does not update NSWindow's frame until user stops dragging the window, so while window is being dragged, we calculate
-					// its position based on mouse move delta
-					if (DraggedWindow && DraggedWindow == NativeWindow)
-					{
-						const int32 X = FMath::TruncToInt(CurrentEventWindow->PositionX + MouseDelta.X);
-						const int32 Y = FMath::TruncToInt(CurrentEventWindow->PositionY + MouseDelta.Y);
-						MessageHandler->OnMovedWindow(CurrentEventWindow.ToSharedRef(), X, Y);
-						CurrentEventWindow->PositionX = X;
-						CurrentEventWindow->PositionY = Y;
-					}
-
-					MessageHandler->OnMouseMove();
-				}
-			}
-
-			if (CurrentEventWindow.IsValid() && !DraggedWindow && !GetCapture())
-			{
-				MessageHandler->OnCursorSet();
-			}
-			break;
+			MessageHandler->BeginReshapingWindow(Event.Window.ToSharedRef());
 		}
-
-		case NSLeftMouseDown:
-		case NSRightMouseDown:
-		case NSOtherMouseDown:
+		else if (Event.NotificationName == NSWindowDidEndLiveResizeNotification)
 		{
-			EMouseButtons::Type Button = [Event type] == NSLeftMouseDown ? EMouseButtons::Left : EMouseButtons::Right;
-			if ([Event type] == NSOtherMouseDown)
-			{
-				switch ([Event buttonNumber])
-				{
-					case 2:
-						Button = EMouseButtons::Middle;
-						break;
-
-					case 3:
-						Button = EMouseButtons::Thumb01;
-						break;
-
-					case 4:
-						Button = EMouseButtons::Thumb02;
-						break;
-				}
-			}
-
-			if (CurrentEventWindow.IsValid())
-			{
-				if (Button == LastPressedMouseButton && ([Event clickCount] % 2) == 0)
-				{
-					MessageHandler->OnMouseDoubleClick(CurrentEventWindow, Button);
-
-					bool IsMouseOverTitleBar = false;
-					const bool IsMovable = IsWindowMovable(NativeWindow, &IsMouseOverTitleBar);
-					if (IsMouseOverTitleBar)
-					{
-						const bool bShouldMinimize = [[NSUserDefaults standardUserDefaults] boolForKey:@"AppleMiniaturizeOnDoubleClick"];
-						if (bShouldMinimize)
-						{
-							MainThreadCall(^{ [NativeWindow performMiniaturize:nil]; }, NSDefaultRunLoopMode, true);
-						}
-						else if (!FPlatformMisc::IsRunningOnMavericks())
-						{
-							MainThreadCall(^{ [NativeWindow zoom:nil]; }, NSDefaultRunLoopMode, true);
-						}
-					}
-				}
-				else
-				{
-					MessageHandler->OnMouseDown(CurrentEventWindow, Button);
-				}
-
-				if (CurrentEventWindow->GetWindowHandle() && !DraggedWindow && !GetCapture())
-				{
-					MessageHandler->OnCursorSet();
-				}
-			}
-
-			LastPressedMouseButton = Button;
-
-			break;
+			MessageHandler->FinishedReshapingWindow(Event.Window.ToSharedRef());
 		}
-
-		case NSLeftMouseUp:
-		case NSRightMouseUp:
-		case NSOtherMouseUp:
+		else if (Event.NotificationName == NSWindowDidEnterFullScreenNotification)
 		{
-			EMouseButtons::Type Button = [Event type] == NSLeftMouseUp ? EMouseButtons::Left : EMouseButtons::Right;
-			if ([Event type] == NSOtherMouseUp)
-			{
-				switch ([Event buttonNumber])
-				{
-					case 2:
-						Button = EMouseButtons::Middle;
-						break;
-
-					case 3:
-						Button = EMouseButtons::Thumb01;
-						break;
-
-					case 4:
-						Button = EMouseButtons::Thumb02;
-						break;
-				}
-			}
-
-			MessageHandler->OnMouseUp(Button);
-
-			if (CurrentEventWindow.IsValid() && CurrentEventWindow->GetWindowHandle() && !DraggedWindow && !GetCapture())
-			{
-				MessageHandler->OnCursorSet();
-			}
-
-			FPlatformMisc::bChachedMacMenuStateNeedsUpdate = true;
-			DraggedWindow = nullptr;
-			break;
+			OnWindowDidResize(Event.Window.ToSharedRef());
 		}
-
-		case NSScrollWheel:
+		else if (Event.NotificationName == NSWindowDidExitFullScreenNotification)
 		{
-			const float DeltaX = ([Event modifierFlags] & NSShiftKeyMask) ? [Event deltaY] : [Event deltaX];
-			const float DeltaY = ([Event modifierFlags] & NSShiftKeyMask) ? [Event deltaX] : [Event deltaY];
-
-			NSEventPhase Phase = [Event phase];
-
-			if ([Event momentumPhase] != NSEventPhaseNone || [Event phase] != NSEventPhaseNone)
-			{
-				const bool bInverted = [Event isDirectionInvertedFromDevice];
-
-				const FVector2D ScrollDelta([Event scrollingDeltaX], [Event scrollingDeltaY]);
-
-				// This is actually a scroll gesture from trackpad
-				MessageHandler->OnTouchGesture(EGestureEvent::Scroll, bInverted ? -ScrollDelta : ScrollDelta, DeltaY);
-				RecordUsage(EGestureEvent::Scroll);
-			}
-			else
-			{
-				MessageHandler->OnMouseWheel(DeltaY);
-			}
-
-			if (CurrentEventWindow.IsValid() && !DraggedWindow && !GetCapture())
-			{
-				MessageHandler->OnCursorSet();
-			}
-			break;
+			OnWindowDidResize(Event.Window.ToSharedRef());
 		}
-
-		case NSEventTypeMagnify:
+		else if (Event.NotificationName == NSWindowDidBecomeMainNotification)
 		{
-			MessageHandler->OnTouchGesture(EGestureEvent::Magnify, FVector2D([Event magnification], [Event magnification]), 0);
-			RecordUsage(EGestureEvent::Magnify);
-			break;
+			MessageHandler->OnWindowActivationChanged(Event.Window.ToSharedRef(), EWindowActivation::Activate);
+			OnWindowsReordered(![NSApp isActive]);
 		}
-
-		case NSEventTypeSwipe:
+		else if (Event.NotificationName == NSWindowDidResignMainNotification)
 		{
-			MessageHandler->OnTouchGesture(EGestureEvent::Swipe, FVector2D([Event deltaX], [Event deltaY]), 0);
-			RecordUsage(EGestureEvent::Swipe);
-			break;
+			MessageHandler->OnWindowActivationChanged(Event.Window.ToSharedRef(), EWindowActivation::Deactivate);
+			OnWindowsReordered(![NSApp isActive]);
 		}
-
-		case NSEventTypeRotate:
+		else if (Event.NotificationName == NSWindowWillMoveNotification)
 		{
-			MessageHandler->OnTouchGesture(EGestureEvent::Rotate, FVector2D([Event rotation], [Event rotation]), 0);
-			RecordUsage(EGestureEvent::Rotate);
-			break;
+			DraggedWindow = Event.Window->GetWindowHandle();
 		}
-
-		case NSEventTypeBeginGesture:
+		else if (Event.NotificationName == NSWindowDidMoveNotification)
 		{
-			MessageHandler->OnBeginGesture();
-			break;
+			OnWindowDidMove(Event.Window.ToSharedRef());
 		}
-
-		case NSEventTypeEndGesture:
+		else if (Event.NotificationName == NSWindowWillCloseNotification)
 		{
-			MessageHandler->OnEndGesture();
-#if WITH_EDITOR
-			LastGestureUsed = EGestureEvent::None;
-#endif
-			break;
+			OnWindowDidClose(Event.Window.ToSharedRef());
 		}
-
-		case NSKeyDown:
+		else if (Event.NotificationName == NSDraggingExited)
 		{
-			NSString *Characters = [Event characters];
-			bool bHandled = false;
-			if (!bSystemModalMode && [Characters length] && CurrentEventWindow.IsValid())
-			{
-				bHandled = CurrentEventWindow->OnIMKKeyDown(Event);
-				if (!bHandled)
-				{
-					const bool IsRepeat = [Event isARepeat];
-					const TCHAR Character = ConvertChar([Characters characterAtIndex:0]);
-					const TCHAR CharCode = [[Event charactersIgnoringModifiers] characterAtIndex:0];
-					const uint32 KeyCode = [Event keyCode];
-					const bool IsPrintable = IsPrintableKey(Character);
+			MessageHandler->OnDragLeave(Event.Window.ToSharedRef());
+		}
+		else if (Event.NotificationName == NSDraggingUpdated)
+		{
+			MessageHandler->OnDragOver(Event.Window.ToSharedRef());
+		}
+		else if (Event.NotificationName == NSPrepareForDragOperation)
+		{
+			SCOPED_AUTORELEASE_POOL;
 
-					bHandled = MessageHandler->OnKeyDown(KeyCode, TranslateCharCode(CharCode, KeyCode), IsRepeat);
+			// Decipher the pasteboard data
+			const bool bHaveText = [[Event.DraggingPasteboard types] containsObject:NSPasteboardTypeString];
+			const bool bHaveFiles = [[Event.DraggingPasteboard types] containsObject:NSFilenamesPboardType];
 
-					// First KeyDown, then KeyChar. This is important, as in-game console ignores first character otherwise
-					bool bCmdKeyPressed = [Event modifierFlags] & 0x18;
-					if (!bCmdKeyPressed && IsPrintable)
-					{
-						MessageHandler->OnKeyChar(Character, IsRepeat);
-					}
-				}
-			}
-			if (bHandled)
+			if (bHaveFiles)
 			{
-				FCocoaMenu* MainMenu = [[NSApp mainMenu] isKindOfClass:[FCocoaMenu class]] ? (FCocoaMenu*)[NSApp mainMenu]: nil;
-				if (MainMenu)
+				TArray<FString> FileList;
+
+				NSArray *Files = [Event.DraggingPasteboard propertyListForType:NSFilenamesPboardType];
+				for (int32 Index = 0; Index < [Files count]; Index++)
 				{
-					MainThreadCall(^{ [MainMenu highlightKeyEquivalent:Event]; }, NSDefaultRunLoopMode, true);
+					NSString* FilePath = [Files objectAtIndex: Index];
+					const FString ListElement = FString([FilePath fileSystemRepresentation]);
+					FileList.Add(ListElement);
 				}
+
+				MessageHandler->OnDragEnterFiles(Event.Window.ToSharedRef(), FileList);
 			}
-			else
+			else if (bHaveText)
 			{
-				ResendEvent(Event);
+				NSString* Text = [Event.DraggingPasteboard stringForType:NSPasteboardTypeString];
+				MessageHandler->OnDragEnterText(Event.Window.ToSharedRef(), FString(Text));
 			}
-			break;
 		}
-
-		case NSKeyUp:
+		else if (Event.NotificationName == NSPerformDragOperation)
 		{
-			NSString *Characters = [Event characters];
-			bool bHandled = false;
-			if (!bSystemModalMode && [Characters length])
-			{
-				const bool IsRepeat = [Event isARepeat];
-				const TCHAR Character = ConvertChar([Characters characterAtIndex:0]);
-				const TCHAR CharCode = [[Event charactersIgnoringModifiers] characterAtIndex:0];
-				const uint32 KeyCode = [Event keyCode];
-				const bool IsPrintable = IsPrintableKey(Character);
-
-				bHandled = MessageHandler->OnKeyUp(KeyCode, TranslateCharCode(CharCode, KeyCode), IsRepeat);
-			}
-			if (!bHandled)
-			{
-				ResendEvent(Event);
-			}
-			FPlatformMisc::bChachedMacMenuStateNeedsUpdate = true;
-			break;
+			MessageHandler->OnDragDrop(Event.Window.ToSharedRef());
 		}
 	}
-
-	bIsProcessingNSEvent = bWasProcessingNSEvent;
 }
 
 void FMacApplication::ResendEvent(NSEvent* Event)
 {
 	MainThreadCall(^{
-		NSEvent* Wrapper = [NSEvent otherEventWithType:NSApplicationDefined location:[Event locationInWindow] modifierFlags:[Event modifierFlags] timestamp:[Event timestamp] windowNumber:[Event windowNumber] context:[Event context] subtype:FMacApplication::ResentEvent data1:(NSInteger)Event data2:0];
-		
+		NSEvent* Wrapper = [NSEvent otherEventWithType:NSApplicationDefined location:[Event locationInWindow] modifierFlags:[Event modifierFlags] timestamp:[Event timestamp] windowNumber:[Event windowNumber] context:[Event context] subtype:RESET_EVENT_SUBTYPE data1:(NSInteger)Event data2:0];
 		[NSApp sendEvent:Wrapper];
 	}, NSDefaultRunLoopMode, true);
 }
 
-FCocoaWindow* FMacApplication::FindEventWindow( NSEvent* Event )
+void FMacApplication::ProcessMouseMovedEvent(const FDeferredMacEvent& Event)
 {
-	SCOPED_AUTORELEASE_POOL;
-
-	bool IsMouseEvent = false;
-	switch ([Event type])
+	if (Event.Window.IsValid() && Event.Window->IsRegularWindow())
 	{
-		case NSMouseMoved:
-		case NSLeftMouseDragged:
-		case NSRightMouseDragged:
-		case NSOtherMouseDragged:
-		case NSLeftMouseDown:
-		case NSRightMouseDown:
-		case NSOtherMouseDown:
-		case NSLeftMouseUp:
-		case NSRightMouseUp:
-		case NSOtherMouseUp:
-		case NSScrollWheel:
-			IsMouseEvent = true;
-			break;
+		bool IsMouseOverTitleBar = false;
+		const bool IsMovable = IsWindowMovable(Event.Window.ToSharedRef(), &IsMouseOverTitleBar);
+		[Event.Window->GetWindowHandle() setMovable:IsMovable];
+		[Event.Window->GetWindowHandle() setMovableByWindowBackground:IsMouseOverTitleBar];
 	}
 
-	FCocoaWindow* EventWindow = (FCocoaWindow*)[Event window];
+	FMacCursor* MacCursor = (FMacCursor*)Cursor.Get();
 
-	if ([Event type] == NSMouseMoved)
+	FVector2D const MouseScaling = MacCursor->GetMouseScaling();
+
+	if (bUsingHighPrecisionMouseInput)
 	{
-		// Ignore windows owned by other applications
-		NSInteger WindowNumber = [NSWindow windowNumberAtPoint:[NSEvent mouseLocation] belowWindowWithWindowNumber:0];
-		if ([NSApp windowWithWindowNumber:WindowNumber] == NULL)
+		// Under OS X we disassociate the cursor and mouse position during hi-precision mouse input.
+		// The game snaps the mouse cursor back to the starting point when this is disabled, which
+		// accumulates mouse delta that we want to ignore.
+		const FVector2D AccumDelta = MacCursor->GetMouseWarpDelta() * MouseScaling;
+
+		// Get the mouse position
+		FVector2D HighPrecisionMousePos = MacCursor->GetPosition();
+
+		// Find the screen the cursor is currently on.
+		NSEnumerator *ScreenEnumerator = [[NSScreen screens] objectEnumerator];
+		NSScreen *Screen;
+		while ((Screen = [ScreenEnumerator nextObject]) && !NSMouseInRect(NSMakePoint(HighPrecisionMousePos.X / MouseScaling.X, HighPrecisionMousePos.Y / MouseScaling.Y), Screen.frame, NO))
+			;
+
+		// Account for warping delta's
+		FVector2D Delta = FVector2D(Event.Delta.X, Event.Delta.Y) * MouseScaling;
+		const FVector2D WarpDelta(FMath::Abs(AccumDelta.X)<FMath::Abs(Delta.X) ? AccumDelta.X : Delta.X, FMath::Abs(AccumDelta.Y)<FMath::Abs(Delta.Y) ? AccumDelta.Y : Delta.Y);
+		Delta -= WarpDelta;
+
+		// Update to latest position
+		HighPrecisionMousePos += Delta;
+
+		// Clip to lock rect
+		MacCursor->UpdateCursorClipping(HighPrecisionMousePos);
+
+		// Clamp to the current screen and avoid the menu bar and dock to prevent popups and other
+		// assorted potential for mouse abuse.
+		NSRect VisibleFrame = [Screen visibleFrame];
+		// Avoid the menu bar & dock disclosure borders at the top & bottom of fullscreen windows
+		if (Event.Window.IsValid() && Event.Window->GetWindowMode() != EWindowMode::Windowed)
 		{
-			return NULL;
+			VisibleFrame.origin.y += 5;
+			VisibleFrame.size.height -= 10;
 		}
-	}
+		NSRect FullFrame = [Screen frame];
+		VisibleFrame.origin.y = (FullFrame.origin.y+FullFrame.size.height) - (VisibleFrame.origin.y + VisibleFrame.size.height);
 
-	if ( IsMouseEvent )
+		int32 ScaledX = (int32)(HighPrecisionMousePos.X / MouseScaling.X);
+		int32 ScaledY = (int32)(HighPrecisionMousePos.Y / MouseScaling.Y);
+		int32 ClampedPosX = FMath::Clamp(ScaledX, (int32)VisibleFrame.origin.x, (int32)(VisibleFrame.origin.x + VisibleFrame.size.width)-1);
+		int32 ClampedPosY = FMath::Clamp(ScaledY, (int32)VisibleFrame.origin.y, (int32)(VisibleFrame.origin.y + VisibleFrame.size.height)-1);
+		MacCursor->SetPosition(ClampedPosX * MouseScaling.X, ClampedPosY * MouseScaling.Y);
+
+		// Forward the delta on to Slate
+		MessageHandler->OnRawMouseMove(Delta.X, Delta.Y);
+	}
+	else
 	{
-		if ( DraggedWindow )
+		NSPoint CursorPos = [NSEvent mouseLocation];
+		CursorPos.y--; // The y coordinate of the point returned by mouseLocation starts from a base of 1
+		const FVector2D MousePosition = FVector2D(CursorPos.x, FPlatformMisc::ConvertSlateYPositionToCocoa(CursorPos.y));
+		FVector2D CurrentPosition = MousePosition * MouseScaling;
+		const FVector2D MouseDelta = CurrentPosition - MacCursor->GetPosition();
+		if (MacCursor->UpdateCursorClipping(CurrentPosition))
 		{
-			EventWindow = DraggedWindow;
+			MacCursor->SetPosition(CurrentPosition.X, CurrentPosition.Y);
 		}
 		else
 		{
-			NSPoint CursorPos = [NSEvent mouseLocation];
-			CursorPos.y--; // The y coordinate of the point returned by mouseLocation starts from a base of 1
-			TSharedPtr<FMacWindow> WindowUnderCursor = LocateWindowUnderCursor(CursorPos);
-			if (WindowUnderCursor.IsValid())
+			MacCursor->UpdateCurrentPosition(CurrentPosition);
+		}
+
+		if (Event.Window.IsValid())
+		{
+			// Cocoa does not update NSWindow's frame until user stops dragging the window, so while window is being dragged, we calculate
+			// its position based on mouse move delta
+			if (DraggedWindow && DraggedWindow == Event.Window->GetWindowHandle())
 			{
-				EventWindow = WindowUnderCursor->GetWindowHandle();
+				const int32 X = FMath::TruncToInt(Event.Window->PositionX + MouseDelta.X);
+				const int32 Y = FMath::TruncToInt(Event.Window->PositionY + MouseDelta.Y);
+				MessageHandler->OnMovedWindow(Event.Window.ToSharedRef(), X, Y);
+				Event.Window->PositionX = X;
+				Event.Window->PositionY = Y;
 			}
+
+			MessageHandler->OnMouseMove();
 		}
 	}
 
-	return EventWindow;
-}
-
-TSharedPtr<FMacWindow> FMacApplication::LocateWindowUnderCursor( const NSPoint Position )
-{
-	NSScreen* MouseScreen = nil;
-	if ([NSScreen screensHaveSeparateSpaces])
+	if (Event.Window.IsValid() && !DraggedWindow && !GetCapture())
 	{
-		// New default mode which uses a separate Space per display
-		// Find the screen the cursor is currently on so we can ignore invisible window regions.
-		NSEnumerator* ScreenEnumerator = [[NSScreen screens] objectEnumerator];
-		while ((MouseScreen = [ScreenEnumerator nextObject]) && !NSMouseInRect(Position, MouseScreen.frame, NO))
-			;
+		MessageHandler->OnCursorSet();
 	}
+}
 
-	NSArray* AllWindows = [NSApp orderedWindows];
-	for (int32 Index = 0; Index < [AllWindows count]; Index++)
+void FMacApplication::ProcessMouseDownEvent(const FDeferredMacEvent& Event)
+{
+	EMouseButtons::Type Button = Event.Type == NSLeftMouseDown ? EMouseButtons::Left : EMouseButtons::Right;
+	if (Event.Type == NSOtherMouseDown)
 	{
-		NSWindow* NativeWindow = (NSWindow*)[AllWindows objectAtIndex: Index];
-		if ([NativeWindow isMiniaturized] || ![NativeWindow isVisible] || ![NativeWindow isOnActiveSpace] || ![NativeWindow isKindOfClass: [FCocoaWindow class]])
+		switch (Event.ButtonNumber)
 		{
-			continue;
-		}
+			case 2:
+				Button = EMouseButtons::Middle;
+				break;
 
-        NSRect VisibleFrame = [NativeWindow frame];
-#if WITH_EDITOR
-        if (MouseScreen != nil)
-        {
-            VisibleFrame = NSIntersectionRect([MouseScreen frame], VisibleFrame);
-        }
-#endif
+			case 3:
+				Button = EMouseButtons::Thumb01;
+				break;
 
-		if (NSPointInRect(Position, VisibleFrame))
-		{
-			TSharedPtr<FMacWindow> MacWindow = FindWindowByNSWindow(Windows, &WindowsMutex, (FCocoaWindow*)NativeWindow);
-			return MacWindow;
+			case 4:
+				Button = EMouseButtons::Thumb02;
+				break;
 		}
 	}
 
-	return NULL;
-}
-
-void FMacApplication::OnApplicationDidBecomeActive()
-{
-	for (int32 Index = 0; Index < SavedWindowsOrder.Num(); Index++)
+	if (Event.Window.IsValid())
 	{
-		const FSavedWindowOrderInfo& Info = SavedWindowsOrder[Index];
-		NSWindow* Window = [NSApp windowWithWindowNumber:Info.WindowNumber];
-		if (Window)
+		if (Button == LastPressedMouseButton && (Event.ClickCount % 2) == 0)
 		{
-			[Window setLevel:Info.Level];
-		}
-	}
+			MessageHandler->OnMouseDoubleClick(Event.Window, Button);
 
-	if (SavedWindowsOrder.Num() > 0)
-	{
-		[[NSApp windowWithWindowNumber:SavedWindowsOrder[0].WindowNumber] orderWindow:NSWindowAbove relativeTo:0];
-		for (int32 Index = 1; Index < SavedWindowsOrder.Num(); Index++)
-		{
-			const FSavedWindowOrderInfo& Info = SavedWindowsOrder[Index];
-			NSWindow* Window = [NSApp windowWithWindowNumber:Info.WindowNumber];
-			if (Window)
+			if (Event.Window.IsValid())
 			{
-				[Window orderWindow:NSWindowBelow relativeTo:SavedWindowsOrder[Index - 1].WindowNumber];
+				bool IsMouseOverTitleBar = false;
+				const bool IsMovable = IsWindowMovable(Event.Window.ToSharedRef(), &IsMouseOverTitleBar);
+				if (IsMouseOverTitleBar)
+				{
+					const bool bShouldMinimize = [[NSUserDefaults standardUserDefaults] boolForKey:@"AppleMiniaturizeOnDoubleClick"];
+					if (bShouldMinimize)
+					{
+						MainThreadCall(^{ [Event.Window->GetWindowHandle() performMiniaturize:nil]; }, NSDefaultRunLoopMode, true);
+					}
+					else if (!FPlatformMisc::IsRunningOnMavericks())
+					{
+						MainThreadCall(^{ [Event.Window->GetWindowHandle() zoom:nil]; }, NSDefaultRunLoopMode, true);
+					}
+				}
 			}
 		}
-	}
-	
-	static_cast<FMacCursor*>( Cursor.Get() )->UpdateVisibility();
+		else
+		{
+			MessageHandler->OnMouseDown(Event.Window, Button);
+		}
 
-	// If editor thread doesn't have the focus, don't suck up too much CPU time.
-	if (GIsEditor && !IsRunningCommandlet())
-	{
-		// Boost our priority back to normal.
-		struct sched_param Sched;
-		FMemory::Memzero(&Sched, sizeof(struct sched_param));
-		Sched.sched_priority = 15;
-		pthread_setschedparam(pthread_self(), SCHED_RR, &Sched);
+		if (Event.Window->GetWindowHandle() && !DraggedWindow && !GetCapture())
+		{
+			MessageHandler->OnCursorSet();
+		}
 	}
 
-	// app is active, allow sound
-	FApp::SetVolumeMultiplier(1.0f);
-
-	GameThreadCall(^{
-		MessageHandler->OnApplicationActivationChanged(true);
-	}, @[ NSDefaultRunLoopMode ], false);
+	LastPressedMouseButton = Button;
 }
 
-void FMacApplication::OnApplicationWillResignActive()
+void FMacApplication::ProcessMouseUpEvent(const FDeferredMacEvent& Event)
 {
-	OnWindowsReordered(true);
-
-	if (SavedWindowsOrder.Num() > 0)
+	EMouseButtons::Type Button = Event.Type == NSLeftMouseUp ? EMouseButtons::Left : EMouseButtons::Right;
+	if (Event.Type == NSOtherMouseUp)
 	{
-		[[NSApp windowWithWindowNumber:SavedWindowsOrder[0].WindowNumber] orderWindow:NSWindowAbove relativeTo:0];
-		for (int32 Index = 1; Index < SavedWindowsOrder.Num(); Index++)
+		switch (Event.ButtonNumber)
 		{
-			const FSavedWindowOrderInfo& Info = SavedWindowsOrder[Index];
-			NSWindow* Window = [NSApp windowWithWindowNumber:Info.WindowNumber];
-			if (Window)
-			{
-				if (SavedWindowsOrder[Index - 1].Level == Info.Level)
-				{
-					[Window orderWindow:NSWindowBelow relativeTo:SavedWindowsOrder[Index - 1].WindowNumber];
-				}
-				else
-				{
-					[Window orderWindow:NSWindowAbove relativeTo:0];
-				}
-			}
+			case 2:
+				Button = EMouseButtons::Middle;
+				break;
+
+			case 3:
+				Button = EMouseButtons::Thumb01;
+				break;
+
+			case 4:
+				Button = EMouseButtons::Thumb02;
+				break;
 		}
 	}
-	
-	static_cast<FMacCursor*>( Cursor.Get() )->UpdateVisibility();
 
-	// If editor thread doesn't have the focus, don't suck up too much CPU time.
-	if (GIsEditor && !IsRunningCommandlet())
-	{
-		// Drop our priority to speed up whatever is in the foreground.
-		struct sched_param Sched;
-		FMemory::Memzero(&Sched, sizeof(struct sched_param));
-		Sched.sched_priority = 5;
-		pthread_setschedparam(pthread_self(), SCHED_RR, &Sched);
+	MessageHandler->OnMouseUp(Button);
 
-		// Sleep for a bit to not eat up all CPU time.
-		FPlatformProcess::Sleep(0.005f);
+	if (Event.Window.IsValid() && Event.Window->GetWindowHandle() && !DraggedWindow && !GetCapture())
+	{
+		MessageHandler->OnCursorSet();
 	}
 
-	// app is inactive, apply multiplier
-	FApp::SetVolumeMultiplier(FApp::GetUnfocusedVolumeMultiplier());
-
-	GameThreadCall(^{
-		MessageHandler->OnApplicationActivationChanged(false);
-	}, @[ NSDefaultRunLoopMode ], false);
+	FPlatformMisc::bChachedMacMenuStateNeedsUpdate = true;
+	DraggedWindow = nullptr;
 }
 
-void FMacApplication::OnWindowsReordered(bool bIsAppInBackground)
+void FMacApplication::ProcessScrollWheelEvent(const FDeferredMacEvent& Event)
 {
-	if (bIsAppInBackground)
-	{
-		TMap<int32, int32> Levels;
+	const float DeltaX = (Event.ModifierFlags & NSShiftKeyMask) ? Event.Delta.Y : Event.Delta.X;
+	const float DeltaY = (Event.ModifierFlags & NSShiftKeyMask) ? Event.Delta.X : Event.Delta.Y;
 
-		for (int32 Index = 0; Index < SavedWindowsOrder.Num(); Index++)
-		{
-			const FSavedWindowOrderInfo& Info = SavedWindowsOrder[Index];
-			Levels.Add(Info.WindowNumber, Info.Level);
-		}
+	NSEventPhase Phase = Event.Phase;
 
-		SavedWindowsOrder.Empty();
+	if (Event.MomentumPhase != NSEventPhaseNone || Event.Phase != NSEventPhaseNone)
+	{
+		const FVector2D ScrollDelta(Event.ScrollingDelta.X, Event.ScrollingDelta.Y);
 
-		NSArray* OrderedWindows = [NSApp orderedWindows];
-		for (NSWindow* Window in OrderedWindows)
-		{
-			if ([Window isKindOfClass:[FCocoaWindow class]] && [Window isVisible] && ![Window hidesOnDeactivate])
-			{
-				SavedWindowsOrder.Add(FSavedWindowOrderInfo([Window windowNumber], Levels.Contains([Window windowNumber]) ? Levels[[Window windowNumber]] : [Window level]));
-				[Window setLevel:NSNormalWindowLevel];
-			}
-		}
+		// This is actually a scroll gesture from trackpad
+		MessageHandler->OnTouchGesture(EGestureEvent::Scroll, Event.IsDirectionInvertedFromDevice ? -ScrollDelta : ScrollDelta, DeltaY);
+		RecordUsage(EGestureEvent::Scroll);
+	}
+	else
+	{
+		MessageHandler->OnMouseWheel(DeltaY);
 	}
-}
 
-void FMacApplication::PollGameDeviceState( const float TimeDelta )
-{
-	// Poll game device state and send new events
-	HIDInput->SendControllerEvents();
+	if (Event.Window.IsValid() && !DraggedWindow && !GetCapture())
+	{
+		MessageHandler->OnCursorSet();
+	}
 }
 
-void FMacApplication::SetHighPrecisionMouseMode( const bool Enable, const TSharedPtr< FGenericWindow >& InWindow )
+void FMacApplication::ProcessGestureEvent(const FDeferredMacEvent& Event)
 {
-	bUsingHighPrecisionMouseInput = Enable;
-	static_cast<FMacCursor*>( Cursor.Get() )->SetHighPrecisionMouseMode( Enable );
+	if (Event.Type == NSEventTypeBeginGesture)
+	{
+		MessageHandler->OnBeginGesture();
+	}
+	else if (Event.Type == NSEventTypeEndGesture)
+	{
+		MessageHandler->OnEndGesture();
+#if WITH_EDITOR
+		LastGestureUsed = EGestureEvent::None;
+#endif
+	}
+	else
+	{
+		const EGestureEvent::Type GestureType = Event.Type == NSEventTypeMagnify ? EGestureEvent::Magnify : (Event.Type == NSEventTypeSwipe ? EGestureEvent::Swipe : EGestureEvent::Rotate);
+		MessageHandler->OnTouchGesture(GestureType, Event.Delta, 0);
+		RecordUsage(GestureType);
+	}
 }
 
-FModifierKeysState FMacApplication::GetModifierKeys() const
+void FMacApplication::ProcessKeyDownEvent(const FDeferredMacEvent& Event)
 {
-	uint32 CurrentFlags = ModifierKeysFlags;
+	bool bHandled = false;
+	if (!bSystemModalMode && Event.Window.IsValid())
+	{
+		const TCHAR Character = ConvertChar([Event.Characters characterAtIndex:0]);
+		const TCHAR CharCode = [Event.CharactersIgnoringModifiers characterAtIndex:0];
+		const bool IsPrintable = IsPrintableKey(Character);
 
-	const bool bIsLeftShiftDown			= ( CurrentFlags & ( 1 << 0 ) ) != 0;
-	const bool bIsRightShiftDown		= ( CurrentFlags & ( 1 << 1 ) ) != 0;
-	const bool bIsLeftControlDown		= ( CurrentFlags & ( 1 << 6 ) ) != 0; // Mac pretends the Command key is Control
-	const bool bIsRightControlDown		= ( CurrentFlags & ( 1 << 7 ) ) != 0; // Mac pretends the Command key is Control
-	const bool bIsLeftAltDown			= ( CurrentFlags & ( 1 << 4 ) ) != 0;
-	const bool bIsRightAltDown			= ( CurrentFlags & ( 1 << 5 ) ) != 0;
-	const bool bIsLeftCommandDown		= ( CurrentFlags & ( 1 << 2 ) ) != 0; // Mac pretends the Control key is Command
-	const bool bIsRightCommandDown		= ( CurrentFlags & ( 1 << 3 ) ) != 0; // Mac pretends the Control key is Command
-	const bool bAreCapsLocked           = ( CurrentFlags & ( 1 << 8 ) ) != 0; 
+		bHandled = MessageHandler->OnKeyDown(Event.KeyCode, TranslateCharCode(CharCode, Event.KeyCode), Event.IsRepeat);
 
-	return FModifierKeysState(bIsLeftShiftDown, bIsRightShiftDown, bIsLeftControlDown, bIsRightControlDown, bIsLeftAltDown, bIsRightAltDown, bIsLeftCommandDown, bIsRightCommandDown, bAreCapsLocked);
+		// First KeyDown, then KeyChar. This is important, as in-game console ignores first character otherwise
+		bool bCmdKeyPressed = Event.ModifierFlags & 0x18;
+		if (!bCmdKeyPressed && IsPrintable)
+		{
+			MessageHandler->OnKeyChar(Character, Event.IsRepeat);
+		}
+	}
+	if (bHandled)
+	{
+		FCocoaMenu* MainMenu = [[NSApp mainMenu] isKindOfClass:[FCocoaMenu class]] ? (FCocoaMenu*)[NSApp mainMenu]: nil;
+		if (MainMenu)
+		{
+			MainThreadCall(^{ [MainMenu highlightKeyEquivalent:Event.GetKeyNSEvent()]; }, NSDefaultRunLoopMode, true);
+		}
+	}
+	else
+	{
+		ResendEvent(Event.GetKeyNSEvent());
+	}
 }
 
-FPlatformRect FMacApplication::GetWorkArea( const FPlatformRect& CurrentWindow ) const
+void FMacApplication::ProcessKeyUpEvent(const FDeferredMacEvent& Event)
 {
-	SCOPED_AUTORELEASE_POOL;
-
-	NSScreen* Screen = FindScreenByPoint( CurrentWindow.Left, CurrentWindow.Top );
-
-	const int32 ScreenHeight = FMath::TruncToInt([Screen frame].size.height);
-	const NSRect VisibleFrame = [Screen visibleFrame];
-	
-	NSArray* AllScreens = [NSScreen screens];
-	NSScreen* PrimaryScreen = (NSScreen*)[AllScreens objectAtIndex: 0];
-	NSRect PrimaryFrame = [PrimaryScreen frame];
-
-	FPlatformRect WorkArea;
-	WorkArea.Left = VisibleFrame.origin.x;
-	WorkArea.Top = (PrimaryFrame.origin.y + PrimaryFrame.size.height) - (VisibleFrame.origin.y + VisibleFrame.size.height);
-	WorkArea.Right = WorkArea.Left + VisibleFrame.size.width;
-	WorkArea.Bottom = WorkArea.Top + VisibleFrame.size.height;
-
-	return WorkArea;
-}
+	bool bHandled = false;
+	if (!bSystemModalMode)
+	{
+		const TCHAR Character = ConvertChar([Event.Characters characterAtIndex:0]);
+		const TCHAR CharCode = [Event.CharactersIgnoringModifiers characterAtIndex:0];
 
-NSScreen* FMacApplication::FindScreenByPoint( int32 X, int32 Y ) const
-{
-	NSPoint Point = {0};
-	Point.x = X;
-	Point.y = FPlatformMisc::ConvertSlateYPositionToCocoa(Y);
-	
-	NSArray* AllScreens = [NSScreen screens];
-	NSScreen* TargetScreen = [AllScreens objectAtIndex: 0];
-	for(NSScreen* Screen in AllScreens)
+		bHandled = MessageHandler->OnKeyUp(Event.KeyCode, TranslateCharCode(CharCode, Event.KeyCode), Event.IsRepeat);
+	}
+	if (!bHandled)
 	{
-		if (Screen && NSPointInRect(Point, [Screen frame]))
-		{
-			TargetScreen = Screen;
-			break;
-		}
+		ResendEvent(Event.GetKeyNSEvent());
 	}
-
-	return TargetScreen;
+	FPlatformMisc::bChachedMacMenuStateNeedsUpdate = true;
 }
 
-void FDisplayMetrics::GetDisplayMetrics(FDisplayMetrics& OutDisplayMetrics)
+void FMacApplication::OnWindowDidMove(TSharedRef<FMacWindow> Window)
 {
 	SCOPED_AUTORELEASE_POOL;
 
-	NSArray* AllScreens = [NSScreen screens];
-	NSScreen* PrimaryScreen = (NSScreen*)[AllScreens objectAtIndex: 0];
-
-	NSRect ScreenFrame = [PrimaryScreen frame];
-	NSRect VisibleFrame = [PrimaryScreen visibleFrame];
+	NSRect WindowFrame = [Window->GetWindowHandle() frame];
+	NSRect OpenGLFrame = [Window->GetWindowHandle() openGLFrame];
 
-	// Total screen size of the primary monitor
-	OutDisplayMetrics.PrimaryDisplayWidth = ScreenFrame.size.width;
-	OutDisplayMetrics.PrimaryDisplayHeight = ScreenFrame.size.height;
+	const int32 X = (int32)WindowFrame.origin.x;
+	int32 Y = 0;
 
-	// Virtual desktop area
-	NSRect WholeWorkspace = {{0,0},{0,0}};
-	for (NSScreen* Screen in AllScreens)
+	if ([Window windowMode] != EWindowMode::Fullscreen)
 	{
-		if (Screen)
-		{
-			WholeWorkspace = NSUnionRect(WholeWorkspace, [Screen frame]);
-		}
+		Y = FPlatformMisc::ConvertCocoaYPositionToSlate(WindowFrame.origin.y + WindowFrame.size.height) + (OpenGLFrame.size.height - WindowFrame.size.height);
 	}
-	OutDisplayMetrics.VirtualDisplayRect.Left = WholeWorkspace.origin.x;
-	OutDisplayMetrics.VirtualDisplayRect.Top = FMath::Min((ScreenFrame.size.height - (WholeWorkspace.origin.y + WholeWorkspace.size.height)), 0.0);
-	OutDisplayMetrics.VirtualDisplayRect.Right = WholeWorkspace.origin.x + WholeWorkspace.size.width;
-	OutDisplayMetrics.VirtualDisplayRect.Bottom = WholeWorkspace.size.height + OutDisplayMetrics.VirtualDisplayRect.Top;
-	
-	// Get the screen rect of the primary monitor, excluding taskbar etc.
-	OutDisplayMetrics.PrimaryDisplayWorkAreaRect.Left = VisibleFrame.origin.x;
-	OutDisplayMetrics.PrimaryDisplayWorkAreaRect.Top = ScreenFrame.size.height - (VisibleFrame.origin.y + VisibleFrame.size.height);
-	OutDisplayMetrics.PrimaryDisplayWorkAreaRect.Right = VisibleFrame.origin.x + VisibleFrame.size.width;
-	OutDisplayMetrics.PrimaryDisplayWorkAreaRect.Bottom = OutDisplayMetrics.PrimaryDisplayWorkAreaRect.Top + VisibleFrame.size.height;
+
+	MessageHandler->OnMovedWindow(Window, X, Y);
+	Window->PositionX = X;
+	Window->PositionY = Y;
 }
 
-void FMacApplication::OnDragEnter(FCocoaWindow* Window, NSPasteboard* Pasteboard)
+void FMacApplication::OnWindowDidResize(TSharedRef<FMacWindow> Window)
 {
 	SCOPED_AUTORELEASE_POOL;
 
-	TSharedPtr< FMacWindow > EventWindow = FindWindowByNSWindow(Windows, &WindowsMutex, Window);
-	if (!EventWindow.IsValid())
-	{
-		return;
-	}
-
-	// Decipher the pasteboard data
-	const bool bHaveText = [[Pasteboard types] containsObject:NSPasteboardTypeString];
-	const bool bHaveFiles = [[Pasteboard types] containsObject:NSFilenamesPboardType];
-
-	if (bHaveFiles)
-	{
-		TArray<FString> FileList;
+	OnWindowDidMove(Window);
 
-		NSArray *Files = [Pasteboard propertyListForType:NSFilenamesPboardType];
-		for (int32 Index = 0; Index < [Files count]; Index++)
-		{
-			NSString* FilePath = [Files objectAtIndex: Index];
-			const FString ListElement = FString([FilePath fileSystemRepresentation]);
-			FileList.Add(ListElement);
-		}
+	// default is no override
+	uint32 Width = [Window->GetWindowHandle() openGLFrame].size.width;
+	uint32 Height = [Window->GetWindowHandle() openGLFrame].size.height;
 
-		MessageHandler->OnDragEnterFiles(EventWindow.ToSharedRef(), FileList);
-	}
-	else if (bHaveText)
+	if ([Window->GetWindowHandle() windowMode] == EWindowMode::WindowedFullscreen)
 	{
-		NSString* Text = [Pasteboard stringForType:NSPasteboardTypeString];
-		MessageHandler->OnDragEnterText(EventWindow.ToSharedRef(), FString(Text));
+		// Grab current monitor data for sizing
+		Width = FMath::TruncToInt([[Window->GetWindowHandle() screen] frame].size.width);
+		Height = FMath::TruncToInt([[Window->GetWindowHandle() screen] frame].size.height);
 	}
+
+	// OnResizingWindow flushes the renderer commands which is needed before we start resizing, but also right after that
+	// because window view's drawRect: can be called before Slate has a chance to flush them.
+	MessageHandler->OnResizingWindow(Window);
+	MessageHandler->OnSizeChanged(Window, Width, Height);
+	MessageHandler->OnResizingWindow(Window);
 }
 
-void FMacApplication::OnDragOver( FCocoaWindow* Window )
+void FMacApplication::OnWindowDidClose(TSharedRef<FMacWindow> Window)
 {
-	TSharedPtr< FMacWindow > EventWindow = FindWindowByNSWindow( Windows, &WindowsMutex, Window );
-	if (EventWindow.IsValid())
+	SCOPED_AUTORELEASE_POOL;
+	if ([Window->GetWindowHandle() isKeyWindow])
 	{
-		MessageHandler->OnDragOver( EventWindow.ToSharedRef() );
+		MessageHandler->OnWindowActivationChanged(Window, EWindowActivation::Deactivate);
 	}
+	Windows.Remove(Window);
+	MessageHandler->OnWindowClose(Window);
 }
 
-void FMacApplication::OnDragOut( FCocoaWindow* Window )
+bool FMacApplication::OnWindowDestroyed(TSharedRef<FMacWindow> Window)
 {
-	TSharedPtr< FMacWindow > EventWindow = FindWindowByNSWindow( Windows, &WindowsMutex, Window );
-	if (EventWindow.IsValid())
+	SCOPED_AUTORELEASE_POOL;
+	if ([Window->GetWindowHandle() isKeyWindow])
+	{
+		MessageHandler->OnWindowActivationChanged(Window, EWindowActivation::Deactivate);
+	}
+	Windows.Remove(Window);
+	if (!WindowsToClose.Contains(Window->GetWindowHandle()))
 	{
-		MessageHandler->OnDragLeave( EventWindow.ToSharedRef() );
+		WindowsToClose.Add(Window->GetWindowHandle());
 	}
+	return true;
 }
 
-void FMacApplication::OnDragDrop( FCocoaWindow* Window )
+void FMacApplication::OnApplicationDidBecomeActive()
 {
-	TSharedPtr< FMacWindow > EventWindow = FindWindowByNSWindow( Windows, &WindowsMutex, Window );
-	if (EventWindow.IsValid())
+	for (int32 Index = 0; Index < SavedWindowsOrder.Num(); Index++)
 	{
-		MessageHandler->OnDragDrop( EventWindow.ToSharedRef() );
+		const FSavedWindowOrderInfo& Info = SavedWindowsOrder[Index];
+		NSWindow* Window = [NSApp windowWithWindowNumber:Info.WindowNumber];
+		if (Window)
+		{
+			[Window setLevel:Info.Level];
+		}
 	}
-}
 
-void FMacApplication::OnWindowDidBecomeMain( FCocoaWindow* Window )
-{
-	TSharedPtr< FMacWindow > EventWindow = FindWindowByNSWindow( Windows, &WindowsMutex, Window );
-	if (EventWindow.IsValid())
+	if (SavedWindowsOrder.Num() > 0)
 	{
-		MessageHandler->OnWindowActivationChanged( EventWindow.ToSharedRef(), EWindowActivation::Activate );
+		[[NSApp windowWithWindowNumber:SavedWindowsOrder[0].WindowNumber] orderWindow:NSWindowAbove relativeTo:0];
+		for (int32 Index = 1; Index < SavedWindowsOrder.Num(); Index++)
+		{
+			const FSavedWindowOrderInfo& Info = SavedWindowsOrder[Index];
+			NSWindow* Window = [NSApp windowWithWindowNumber:Info.WindowNumber];
+			if (Window)
+			{
+				[Window orderWindow:NSWindowBelow relativeTo:SavedWindowsOrder[Index - 1].WindowNumber];
+			}
+		}
 	}
 
-	OnWindowsReordered(![NSApp isActive]);
-}
+	((FMacCursor*)Cursor.Get())->UpdateVisibility();
 
-void FMacApplication::OnWindowDidResignMain( FCocoaWindow* Window )
-{
-	TSharedPtr< FMacWindow > EventWindow = FindWindowByNSWindow( Windows, &WindowsMutex, Window );
-	if (EventWindow.IsValid())
+	// If editor thread doesn't have the focus, don't suck up too much CPU time.
+	if (GIsEditor && !IsRunningCommandlet())
 	{
-		MessageHandler->OnWindowActivationChanged( EventWindow.ToSharedRef(), EWindowActivation::Deactivate );
+		// Boost our priority back to normal.
+		struct sched_param Sched;
+		FMemory::Memzero(&Sched, sizeof(struct sched_param));
+		Sched.sched_priority = 15;
+		pthread_setschedparam(pthread_self(), SCHED_RR, &Sched);
 	}
 
-	OnWindowsReordered(![NSApp isActive]);
-}
-
-void FMacApplication::OnWindowWillMove( FCocoaWindow* Window )
-{
-	SCOPED_AUTORELEASE_POOL;
+	// app is active, allow sound
+	FApp::SetVolumeMultiplier(1.0f);
 
-	DraggedWindow = Window;
+	GameThreadCall(^{
+		MessageHandler->OnApplicationActivationChanged(true);
+	}, @[ NSDefaultRunLoopMode ], false);
 }
 
-void FMacApplication::OnWindowDidMove( FCocoaWindow* Window )
+void FMacApplication::OnApplicationWillResignActive()
 {
-	SCOPED_AUTORELEASE_POOL;
+	OnWindowsReordered(true);
 
-	NSRect WindowFrame = [Window frame];
-	NSRect OpenGLFrame = [Window openGLFrame];
-	
-	const int32 X = (int32)WindowFrame.origin.x;
-	int32 Y = 0;
-    
-    if ([Window windowMode] != EWindowMode::Fullscreen)
-    {
-        Y = FPlatformMisc::ConvertCocoaYPositionToSlate(WindowFrame.origin.y + WindowFrame.size.height) + (OpenGLFrame.size.height - WindowFrame.size.height);
-    }
-	
-	TSharedPtr< FMacWindow > EventWindow = FindWindowByNSWindow( Windows, &WindowsMutex, Window );
-	if (EventWindow.IsValid())
+	if (SavedWindowsOrder.Num() > 0)
 	{
-		MessageHandler->OnMovedWindow( EventWindow.ToSharedRef(), X, Y );
-		EventWindow->PositionX = X;
-		EventWindow->PositionY = Y;
+		[[NSApp windowWithWindowNumber:SavedWindowsOrder[0].WindowNumber] orderWindow:NSWindowAbove relativeTo:0];
+		for (int32 Index = 1; Index < SavedWindowsOrder.Num(); Index++)
+		{
+			const FSavedWindowOrderInfo& Info = SavedWindowsOrder[Index];
+			NSWindow* Window = [NSApp windowWithWindowNumber:Info.WindowNumber];
+			if (Window)
+			{
+				if (SavedWindowsOrder[Index - 1].Level == Info.Level)
+				{
+					[Window orderWindow:NSWindowBelow relativeTo:SavedWindowsOrder[Index - 1].WindowNumber];
+				}
+				else
+				{
+					[Window orderWindow:NSWindowAbove relativeTo:0];
+				}
+			}
+		}
 	}
-}
 
-void FMacApplication::OnWindowWillStartLiveResize(FCocoaWindow* Window)
-{
-	SCOPED_AUTORELEASE_POOL;
+	((FMacCursor*)Cursor.Get())->UpdateVisibility();
 
-	TSharedPtr<FMacWindow> EventWindow = FindWindowByNSWindow(Windows, &WindowsMutex, Window);
-	if (EventWindow.IsValid())
+	// If editor thread doesn't have the focus, don't suck up too much CPU time.
+	if (GIsEditor && !IsRunningCommandlet())
 	{
-		MessageHandler->BeginReshapingWindow(EventWindow.ToSharedRef());
+		// Drop our priority to speed up whatever is in the foreground.
+		struct sched_param Sched;
+		FMemory::Memzero(&Sched, sizeof(struct sched_param));
+		Sched.sched_priority = 5;
+		pthread_setschedparam(pthread_self(), SCHED_RR, &Sched);
+
+		// Sleep for a bit to not eat up all CPU time.
+		FPlatformProcess::Sleep(0.005f);
 	}
-}
 
-void FMacApplication::OnWindowDidEndLiveResize(FCocoaWindow* Window)
-{
-	SCOPED_AUTORELEASE_POOL;
+	// app is inactive, apply multiplier
+	FApp::SetVolumeMultiplier(FApp::GetUnfocusedVolumeMultiplier());
 
-	TSharedPtr<FMacWindow> EventWindow = FindWindowByNSWindow(Windows, &WindowsMutex, Window);
-	if (EventWindow.IsValid())
-	{
-		MessageHandler->FinishedReshapingWindow(EventWindow.ToSharedRef());
-	}
+	GameThreadCall(^{
+		MessageHandler->OnApplicationActivationChanged(false);
+	}, @[ NSDefaultRunLoopMode ], false);
 }
 
-void FMacApplication::OnWindowDidResize(FCocoaWindow* Window)
+void FMacApplication::OnWindowsReordered(bool bIsAppInBackground)
 {
-	SCOPED_AUTORELEASE_POOL;
-
-	OnWindowDidMove(Window);
-
-	TSharedPtr<FMacWindow> EventWindow = FindWindowByNSWindow(Windows, &WindowsMutex, Window);
-	if (EventWindow.IsValid())
+	if (bIsAppInBackground)
 	{
-		// default is no override
-		uint32 Width = [Window openGLFrame].size.width;
-		uint32 Height = [Window openGLFrame].size.height;
+		TMap<int32, int32> Levels;
 
-		if ([Window windowMode] == EWindowMode::WindowedFullscreen)
+		for (int32 Index = 0; Index < SavedWindowsOrder.Num(); Index++)
 		{
-			// Grab current monitor data for sizing
-			Width = FMath::TruncToInt([[Window screen] frame].size.width);
-			Height = FMath::TruncToInt([[Window screen] frame].size.height);
+			const FSavedWindowOrderInfo& Info = SavedWindowsOrder[Index];
+			Levels.Add(Info.WindowNumber, Info.Level);
 		}
 
-		// OnResizingWindow flushes the renderer commands which is needed before we start resizing, but also right after that
-		// because window view's drawRect: can be called before Slate has a chance to flush them.
-		MessageHandler->OnResizingWindow(EventWindow.ToSharedRef());
-		MessageHandler->OnSizeChanged(EventWindow.ToSharedRef(), Width, Height);
-		MessageHandler->OnResizingWindow(EventWindow.ToSharedRef());
+		SavedWindowsOrder.Empty();
+
+		NSArray* OrderedWindows = [NSApp orderedWindows];
+		for (NSWindow* Window in OrderedWindows)
+		{
+			if ([Window isKindOfClass:[FCocoaWindow class]] && [Window isVisible] && ![Window hidesOnDeactivate])
+			{
+				SavedWindowsOrder.Add(FSavedWindowOrderInfo([Window windowNumber], Levels.Contains([Window windowNumber]) ? Levels[[Window windowNumber]] : [Window level]));
+				[Window setLevel:NSNormalWindowLevel];
+			}
+		}
 	}
 }
 
-void FMacApplication::OnWindowDidClose( FCocoaWindow* Window )
+void FMacApplication::HandleModifierChange(NSUInteger NewModifierFlags, NSUInteger FlagsShift, NSUInteger UE4Shift, EMacModifierKeys TranslatedCode)
 {
-	TSharedPtr< FMacWindow > EventWindow = FindWindowByNSWindow( Windows, &WindowsMutex, Window );
-	if (EventWindow.IsValid())
+	const bool CurrentPressed = (CurrentModifierFlags & FlagsShift) != 0;
+	const bool NewPressed = (NewModifierFlags & FlagsShift) != 0;
+	if (CurrentPressed != NewPressed)
 	{
-		SCOPED_AUTORELEASE_POOL;
-		if ([Window isKeyWindow])
+		if (NewPressed)
+		{
+			ModifierKeysFlags |= 1 << UE4Shift;
+			MessageHandler->OnKeyDown(TranslatedCode, 0, false);
+		}
+		else
 		{
-			MessageHandler->OnWindowActivationChanged( EventWindow.ToSharedRef(), EWindowActivation::Deactivate );
+			ModifierKeysFlags &= ~(1 << UE4Shift);
+			MessageHandler->OnKeyUp(TranslatedCode, 0, false);
 		}
-		Windows.Remove( EventWindow.ToSharedRef() );
-		MessageHandler->OnWindowClose( EventWindow.ToSharedRef() );
 	}
 }
 
-bool FMacApplication::OnWindowDestroyed( FCocoaWindow* Window )
+FCocoaWindow* FMacApplication::FindEventWindow(NSEvent* Event) const
 {
-	TSharedPtr< FMacWindow > EventWindow = FindWindowByNSWindow( Windows, &WindowsMutex, Window );
-	if (EventWindow.IsValid())
+	SCOPED_AUTORELEASE_POOL;
+
+	FCocoaWindow* EventWindow = [[Event window] isKindOfClass:[FCocoaWindow class]] ? (FCocoaWindow*)[Event window] : nullptr;
+
+	if ([Event type] != NSKeyDown && [Event type] != NSKeyUp)
 	{
-		if ([Window isKeyWindow])
+		NSInteger WindowNumber = [NSWindow windowNumberAtPoint:[NSEvent mouseLocation] belowWindowWithWindowNumber:0];
+		NSWindow* WindowUnderCursor = [NSApp windowWithWindowNumber:WindowNumber];
+
+		if ([Event type] == NSMouseMoved && WindowUnderCursor == nullptr)
+		{
+			// Ignore windows owned by other applications
+			return nullptr;
+		}
+		else if (DraggedWindow)
 		{
-			MessageHandler->OnWindowActivationChanged( EventWindow.ToSharedRef(), EWindowActivation::Deactivate );
+			EventWindow = DraggedWindow;
 		}
-		Windows.Remove( EventWindow.ToSharedRef() );
-		if (!WindowsToClose.Contains(Window))
+		else if (WindowUnderCursor && [WindowUnderCursor isKindOfClass:[FCocoaWindow class]])
 		{
-			WindowsToClose.Add(Window);
+			EventWindow = (FCocoaWindow*)WindowUnderCursor;
 		}
-		return true;
 	}
-	return false;
+
+	return EventWindow;
 }
 
-void FMacApplication::OnWindowClose( FCocoaWindow* Window )
+NSScreen* FMacApplication::FindScreenByPoint(int32 X, int32 Y) const
 {
-	TSharedPtr< FMacWindow > EventWindow = FindWindowByNSWindow( Windows, &WindowsMutex, Window );
-	if (EventWindow.IsValid())
+	NSPoint Point = {0};
+	Point.x = X;
+	Point.y = FPlatformMisc::ConvertSlateYPositionToCocoa(Y);
+
+	NSArray* AllScreens = [NSScreen screens];
+	NSScreen* TargetScreen = [AllScreens objectAtIndex: 0];
+	for(NSScreen* Screen in AllScreens)
 	{
-		MessageHandler->OnWindowClose( EventWindow.ToSharedRef() );
+		if (Screen && NSPointInRect(Point, [Screen frame]))
+		{
+			TargetScreen = Screen;
+			break;
+		}
 	}
+
+	return TargetScreen;
 }
 
-void FMacApplication::OnMouseCursorLock( bool bLockEnabled )
+bool FMacApplication::IsWindowMovable(TSharedRef<FMacWindow> Window, bool* OutMovableByBackground) const
 {
-	bIsMouseCursorLocked = bLockEnabled;
+	if (OutMovableByBackground)
+	{
+		*OutMovableByBackground = false;
+	}
+
+	const FVector2D CursorPos = ((FMacCursor*)Cursor.Get())->GetPosition();
+	const int32 LocalMouseX = CursorPos.X - Window->PositionX;
+	const int32 LocalMouseY = CursorPos.Y - Window->PositionY;
+
+	const EWindowZone::Type Zone = MessageHandler->GetWindowZoneForPoint(Window, LocalMouseX, LocalMouseY);
+	const bool IsMouseOverTitleBar = Zone == EWindowZone::TitleBar;
+	switch (Zone)
+	{
+		case EWindowZone::NotInWindow:
+		case EWindowZone::TopLeftBorder:
+		case EWindowZone::TopBorder:
+		case EWindowZone::TopRightBorder:
+		case EWindowZone::LeftBorder:
+		case EWindowZone::RightBorder:
+		case EWindowZone::BottomLeftBorder:
+		case EWindowZone::BottomBorder:
+		case EWindowZone::BottomRightBorder:
+			return true;
+		case EWindowZone::TitleBar:
+			if (OutMovableByBackground)
+			{
+				*OutMovableByBackground = true;
+			}
+			return true;
+		case EWindowZone::ClientArea:
+		case EWindowZone::MinimizeButton:
+		case EWindowZone::MaximizeButton:
+		case EWindowZone::CloseButton:
+		case EWindowZone::SysMenu:
+		default:
+			return false;
+	}
 }
 
-bool FMacApplication::IsPrintableKey(uint32 Character)
+bool FMacApplication::IsPrintableKey(uint32 Character) const
 {
 	switch (Character)
 	{
@@ -1407,7 +1212,7 @@ bool FMacApplication::IsPrintableKey(uint32 Character)
 	}
 }
 
-TCHAR FMacApplication::ConvertChar(TCHAR Character)
+TCHAR FMacApplication::ConvertChar(TCHAR Character) const
 {
 	switch (Character)
 	{
@@ -1418,31 +1223,31 @@ TCHAR FMacApplication::ConvertChar(TCHAR Character)
 	}
 }
 
-TCHAR FMacApplication::TranslateCharCode(TCHAR CharCode, uint32 KeyCode)
+TCHAR FMacApplication::TranslateCharCode(TCHAR CharCode, uint32 KeyCode) const
 {
 	// Keys like F1-F12 or Enter do not need translation
 	bool bNeedsTranslation = CharCode < NSOpenStepUnicodeReservedBase || CharCode > 0xF8FF;
-	if ( bNeedsTranslation )
+	if (bNeedsTranslation)
 	{
 		// For non-numpad keys, the key code depends on the keyboard layout, so find out what was pressed by converting the key code to a Latin character
 		TISInputSourceRef CurrentKeyboard = TISCopyCurrentKeyboardLayoutInputSource();
-		if ( CurrentKeyboard )
+		if (CurrentKeyboard)
 		{
-			CFDataRef CurrentLayoutData = ( CFDataRef )TISGetInputSourceProperty( CurrentKeyboard, kTISPropertyUnicodeKeyLayoutData );
-			CFRelease( CurrentKeyboard );
+			CFDataRef CurrentLayoutData = (CFDataRef)TISGetInputSourceProperty(CurrentKeyboard, kTISPropertyUnicodeKeyLayoutData);
+			CFRelease(CurrentKeyboard);
 
-			if ( CurrentLayoutData )
+			if (CurrentLayoutData)
 			{
-				const UCKeyboardLayout *KeyboardLayout = ( UCKeyboardLayout *)CFDataGetBytePtr( CurrentLayoutData );
-				if ( KeyboardLayout )
+				const UCKeyboardLayout *KeyboardLayout = (UCKeyboardLayout*)CFDataGetBytePtr(CurrentLayoutData);
+				if (KeyboardLayout)
 				{
 					UniChar Buffer[256] = { 0 };
 					UniCharCount BufferLength = 256;
 					uint32 DeadKeyState = 0;
 
 					// To ensure we get a latin character, we pretend that command modifier key is pressed
-					OSStatus Status = UCKeyTranslate( KeyboardLayout, KeyCode, kUCKeyActionDown, cmdKey >> 8, LMGetKbdType(), kUCKeyTranslateNoDeadKeysMask, &DeadKeyState, BufferLength, &BufferLength, Buffer );
-					if ( Status == noErr )
+					OSStatus Status = UCKeyTranslate(KeyboardLayout, KeyCode, kUCKeyActionDown, cmdKey >> 8, LMGetKbdType(), kUCKeyTranslateNoDeadKeysMask, &DeadKeyState, BufferLength, &BufferLength, Buffer);
+					if (Status == noErr)
 					{
 						CharCode = Buffer[0];
 					}
@@ -1450,45 +1255,113 @@ TCHAR FMacApplication::TranslateCharCode(TCHAR CharCode, uint32 KeyCode)
 			}
 		}
 	}
-	// Private use range should not be returned
 	else
 	{
+		// Private use range should not be returned
 		CharCode = 0;
 	}
 
 	return CharCode;
 }
 
-#if WITH_EDITOR
-
-void FMacApplication::RecordUsage(EGestureEvent::Type Gesture)
+void FMacApplication::CloseQueuedWindows()
 {
-	if ( LastGestureUsed != Gesture )
+	if (WindowsToClose.Num() > 0)
 	{
-		LastGestureUsed = Gesture;
-		GestureUsage[Gesture] += 1;
+		MainThreadCall(^{
+			SCOPED_AUTORELEASE_POOL;
+
+			// Set the new key window, if needed. We cannot trust Cocoa to set the key window to the actual top most window. It will prefer windows with title bars so,
+			// for example, will choose the main window over a context menu window, when closing a submenu.
+			NSArray* AllWindows = [NSApp orderedWindows];
+			for (NSWindow* Window : AllWindows)
+			{
+				if ([Window isKindOfClass:[FCocoaWindow class]] && !WindowsToClose.Contains((FCocoaWindow*)Window) && [Window canBecomeKeyWindow])
+				{
+					if (Window != [NSApp keyWindow])
+					{
+						[Window makeKeyWindow];
+					}
+					break;
+				}
+			}
+
+			for (FCocoaWindow* Window : WindowsToClose)
+			{
+				[Window close];
+				[Window release];
+			}
+		}, UE4CloseEventMode, true);
+
+		WindowsToClose.Empty();
 	}
 }
 
-void FMacApplication::SendAnalytics(IAnalyticsProvider* Provider)
+void FMacApplication::InvalidateTextLayouts()
 {
-	static_assert(EGestureEvent::Count == 5, "If the number of gestures changes you need to add more entries below!");
+	if (WindowsRequiringTextInvalidation.Num() > 0)
+	{
+		MainThreadCall(^{
+			SCOPED_AUTORELEASE_POOL;
 
-	TArray<FAnalyticsEventAttribute> GestureAttributes;
-	GestureAttributes.Add(FAnalyticsEventAttribute(FString("Scroll"),	GestureUsage[EGestureEvent::Scroll]));
-	GestureAttributes.Add(FAnalyticsEventAttribute(FString("Magnify"),	GestureUsage[EGestureEvent::Magnify]));
-	GestureAttributes.Add(FAnalyticsEventAttribute(FString("Swipe"),	GestureUsage[EGestureEvent::Swipe]));
-	GestureAttributes.Add(FAnalyticsEventAttribute(FString("Rotate"),	GestureUsage[EGestureEvent::Rotate]));
+			for (FCocoaWindow* CocoaWindow : WindowsRequiringTextInvalidation)
+			{
+				if (CocoaWindow && [CocoaWindow openGLView])
+				{
+					FCocoaTextView* TextView = (FCocoaTextView*)[CocoaWindow openGLView];
+					[[TextView inputContext] invalidateCharacterCoordinates];
+				}
+			}
 
-    Provider->RecordEvent(FString("Mac.Gesture.Usage"), GestureAttributes);
+		}, UE4IMEEventMode, true);
+
+		WindowsRequiringTextInvalidation.Empty();
+	}
 
-	FMemory::Memzero(GestureUsage);
-	LastGestureUsed = EGestureEvent::None;
 }
 
+#if WITH_EDITOR
+void FMacApplication::RecordUsage(EGestureEvent::Type Gesture)
+{
+	if (LastGestureUsed != Gesture)
+	{
+		LastGestureUsed = Gesture;
+		GestureUsage[Gesture] += 1;
+	}
+}
 #endif
 
-void FMacApplication::SystemModalMode(bool const bInSystemModalMode)
+void FDisplayMetrics::GetDisplayMetrics(FDisplayMetrics& OutDisplayMetrics)
 {
-	bSystemModalMode = bInSystemModalMode;
+	SCOPED_AUTORELEASE_POOL;
+
+	NSArray* AllScreens = [NSScreen screens];
+	NSScreen* PrimaryScreen = (NSScreen*)[AllScreens objectAtIndex: 0];
+
+	NSRect ScreenFrame = [PrimaryScreen frame];
+	NSRect VisibleFrame = [PrimaryScreen visibleFrame];
+
+	// Total screen size of the primary monitor
+	OutDisplayMetrics.PrimaryDisplayWidth = ScreenFrame.size.width;
+	OutDisplayMetrics.PrimaryDisplayHeight = ScreenFrame.size.height;
+
+	// Virtual desktop area
+	NSRect WholeWorkspace = {{0,0},{0,0}};
+	for (NSScreen* Screen in AllScreens)
+	{
+		if (Screen)
+		{
+			WholeWorkspace = NSUnionRect(WholeWorkspace, [Screen frame]);
+		}
+	}
+	OutDisplayMetrics.VirtualDisplayRect.Left = WholeWorkspace.origin.x;
+	OutDisplayMetrics.VirtualDisplayRect.Top = FMath::Min((ScreenFrame.size.height - (WholeWorkspace.origin.y + WholeWorkspace.size.height)), 0.0);
+	OutDisplayMetrics.VirtualDisplayRect.Right = WholeWorkspace.origin.x + WholeWorkspace.size.width;
+	OutDisplayMetrics.VirtualDisplayRect.Bottom = WholeWorkspace.size.height + OutDisplayMetrics.VirtualDisplayRect.Top;
+	
+	// Get the screen rect of the primary monitor, excluding taskbar etc.
+	OutDisplayMetrics.PrimaryDisplayWorkAreaRect.Left = VisibleFrame.origin.x;
+	OutDisplayMetrics.PrimaryDisplayWorkAreaRect.Top = ScreenFrame.size.height - (VisibleFrame.origin.y + VisibleFrame.size.height);
+	OutDisplayMetrics.PrimaryDisplayWorkAreaRect.Right = VisibleFrame.origin.x + VisibleFrame.size.width;
+	OutDisplayMetrics.PrimaryDisplayWorkAreaRect.Bottom = OutDisplayMetrics.PrimaryDisplayWorkAreaRect.Top + VisibleFrame.size.height;
 }
diff --git a/Engine/Source/Runtime/Core/Private/Mac/MacCursor.cpp b/Engine/Source/Runtime/Core/Private/Mac/MacCursor.cpp
--- a/Engine/Source/Runtime/Core/Private/Mac/MacCursor.cpp
+++ b/Engine/Source/Runtime/Core/Private/Mac/MacCursor.cpp
@@ -265,8 +265,6 @@ void FMacCursor::Show( bool bShow )
 void FMacCursor::Lock( const RECT* const Bounds )
 {
 	SCOPED_AUTORELEASE_POOL;
-	
-	MacApplication->OnMouseCursorLock( Bounds != NULL );
 
 	// Lock/Unlock the cursor
 	if ( Bounds == NULL )
diff --git a/Engine/Source/Runtime/Core/Private/Mac/MacPlatformMisc.cpp b/Engine/Source/Runtime/Core/Private/Mac/MacPlatformMisc.cpp
--- a/Engine/Source/Runtime/Core/Private/Mac/MacPlatformMisc.cpp
+++ b/Engine/Source/Runtime/Core/Private/Mac/MacPlatformMisc.cpp
@@ -579,16 +579,13 @@ void FMacPlatformMisc::PumpMessages( bool bFromMainLoop )
 	{
 		ProcessGameThreadEvents();
 
-		if (MacApplication && !MacApplication->IsProcessingNSEvent() && IsInGameThread())
+		if (MacApplication && IsInGameThread())
 		{
 			if (UpdateCachedMacMenuState && bChachedMacMenuStateNeedsUpdate)
 			{
 				UpdateCachedMacMenuState();
 				bChachedMacMenuStateNeedsUpdate = false;
 			}
-
-			MacApplication->InvalidateTextLayouts();
-			MacApplication->CloseQueuedWindows();
 		}
 	}
 }
diff --git a/Engine/Source/Runtime/Core/Private/Mac/MacWindow.cpp b/Engine/Source/Runtime/Core/Private/Mac/MacWindow.cpp
--- a/Engine/Source/Runtime/Core/Private/Mac/MacWindow.cpp
+++ b/Engine/Source/Runtime/Core/Private/Mac/MacWindow.cpp
@@ -7,6 +7,14 @@
 #include "CocoaTextView.h"
 #include "CocoaThread.h"
 
+FMacWindow::FMacWindow()
+:	WindowHandle(NULL)
+,	bIsVisible(false)
+,	bIsClosed(false)
+{
+	PreFullscreenWindowRect.origin.x = PreFullscreenWindowRect.origin.y = PreFullscreenWindowRect.size.width = PreFullscreenWindowRect.size.height = 0.0f;
+}
+
 FMacWindow::~FMacWindow()
 {
 }
@@ -187,14 +195,6 @@ void FMacWindow::Initialize( FMacApplication* const Application, const TSharedRe
 	}
 }
 
-FMacWindow::FMacWindow()
-	: WindowHandle(NULL)
-	, bIsVisible(false)
-	, bIsClosed(false)
-{
-	PreFullscreenWindowRect.origin.x = PreFullscreenWindowRect.origin.y = PreFullscreenWindowRect.size.width = PreFullscreenWindowRect.size.height = 0.0f;
-}
-
 FCocoaWindow* FMacWindow::GetWindowHandle() const
 {
 	return WindowHandle;
@@ -251,7 +251,7 @@ void FMacWindow::ReshapeWindow( int32 X, int32 Y, int32 Width, int32 Height )
 						MacCursor->SetMouseScaling(FVector2D(WidthScale, HeightScale));
 					}
 				}, UE4ResizeEventMode, true);
-				GameThreadCall(^{ FMacApplication::ProcessEvent([[NSNotification notificationWithName:NSWindowDidResizeNotification object:WindowHandle] retain]); }, @[ NSDefaultRunLoopMode, UE4ResizeEventMode, UE4ShowEventMode, UE4FullscreenEventMode ], true);
+				MacApplication->DeferEvent([NSNotification notificationWithName:NSWindowDidResizeNotification object:WindowHandle]);
 			}
 		}
 	}
@@ -297,7 +297,11 @@ void FMacWindow::Destroy()
 		bIsClosed = true;
 		[WindowHandle setAlphaValue:0.0f];
 		[WindowHandle setBackgroundColor:[NSColor clearColor]];
-		MacApplication->OnWindowDestroyed(WindowHandle);
+		TSharedPtr<FMacWindow> Window = MacApplication->FindWindowByNSWindow(WindowHandle);
+		if (Window.IsValid())
+		{
+			MacApplication->OnWindowDestroyed(Window.ToSharedRef());
+		}
 		WindowHandle = nullptr;
 	}
 }
diff --git a/Engine/Source/Runtime/Core/Public/Mac/MacApplication.h b/Engine/Source/Runtime/Core/Public/Mac/MacApplication.h
--- a/Engine/Source/Runtime/Core/Public/Mac/MacApplication.h
+++ b/Engine/Source/Runtime/Core/Public/Mac/MacApplication.h
@@ -1,4 +1,5 @@
 // Copyright 1998-2015 Epic Games, Inc. All Rights Reserved.
+
 #pragma once
 
 #include "Core.h"
@@ -6,147 +7,230 @@
 #include "MacWindow.h"
 #include "MacTextInputMethodSystem.h"
 
+struct FDeferredMacEvent
+{
+	FDeferredMacEvent()
+	:	Type(0)
+	,	LocationInWindow(FVector2D::ZeroVector)
+	,	ModifierFlags(0)
+	,	Timestamp(0.0)
+	,	WindowNumber(0)
+	,	Context(nullptr)
+	,	Delta(FVector2D::ZeroVector)
+	,	ScrollingDelta(FVector2D::ZeroVector)
+	,	ButtonNumber(0)
+	,	ClickCount(0)
+	,	Phase(NSEventPhaseNone)
+	,	MomentumPhase(NSEventPhaseNone)
+	,	IsDirectionInvertedFromDevice(false)
+	,	Characters(nullptr)
+	,	CharactersIgnoringModifiers(nullptr)
+	,	IsRepeat(false)
+	,	KeyCode(0)
+	,	NotificationName(nullptr)
+	,	DraggingPasteboard(nullptr)
+	{
+	}
+
+	FDeferredMacEvent(const FDeferredMacEvent& Other)
+	:	Window(Other.Window)
+	,	Type(Other.Type)
+	,	LocationInWindow(Other.LocationInWindow)
+	,	ModifierFlags(Other.ModifierFlags)
+	,	Timestamp(Other.Timestamp)
+	,	WindowNumber(Other.WindowNumber)
+	,	Context(Other.Context ? [Other.Context retain] : nullptr)
+	,	Delta(Other.Delta)
+	,	ScrollingDelta(Other.ScrollingDelta)
+	,	ButtonNumber(Other.ButtonNumber)
+	,	ClickCount(Other.ClickCount)
+	,	Phase(Other.Phase)
+	,	MomentumPhase(Other.MomentumPhase)
+	,	IsDirectionInvertedFromDevice(Other.IsDirectionInvertedFromDevice)
+	,	Characters(Other.Characters ? [Other.Characters retain] : nullptr)
+	,	CharactersIgnoringModifiers(Other.CharactersIgnoringModifiers ? [Other.CharactersIgnoringModifiers retain] : nullptr)
+	,	IsRepeat(Other.IsRepeat)
+	,	KeyCode(Other.KeyCode)
+	,	NotificationName(Other.NotificationName ? [Other.NotificationName retain] : nullptr)
+	,	DraggingPasteboard(Other.DraggingPasteboard ? [Other.DraggingPasteboard retain] : nullptr)
+	{
+	}
+
+	~FDeferredMacEvent()
+	{
+		SCOPED_AUTORELEASE_POOL;
+		if (Context)
+		{
+			[Context release];
+		}
+		if (Characters)
+		{
+			[Characters release];
+		}
+		if (CharactersIgnoringModifiers)
+		{
+			[CharactersIgnoringModifiers release];
+		}
+		if (NotificationName)
+		{
+			[NotificationName release];
+		}
+		if (DraggingPasteboard)
+		{
+			[DraggingPasteboard release];
+		}
+	}
+
+	NSEvent* GetKeyNSEvent() const
+	{
+		return [NSEvent keyEventWithType:(NSEventType)Type
+								location:NSMakePoint(LocationInWindow.X, LocationInWindow.Y)
+						   modifierFlags:ModifierFlags
+							   timestamp:Timestamp
+							windowNumber:WindowNumber
+								 context:Context
+							  characters:Characters
+			 charactersIgnoringModifiers:CharactersIgnoringModifiers
+							   isARepeat:IsRepeat
+								 keyCode:KeyCode];
+	}
+
+	TSharedPtr<FMacWindow> Window;
+
+	int32 Type;
+	FVector2D LocationInWindow;
+	uint32 ModifierFlags;
+	NSTimeInterval Timestamp;
+	int32 WindowNumber;
+	NSGraphicsContext* Context;
+	FVector2D Delta;
+	FVector2D ScrollingDelta;
+	int32 ButtonNumber;
+	int32 ClickCount;
+	NSEventPhase Phase;
+	NSEventPhase MomentumPhase;
+	bool IsDirectionInvertedFromDevice;
+	NSString* Characters;
+	NSString* CharactersIgnoringModifiers;
+	bool IsRepeat;
+	uint32 KeyCode;
+
+	NSString* NotificationName;
+	NSPasteboard* DraggingPasteboard;
+};
+
 /**
  * Mac-specific application implementation.
  */
 class CORE_API FMacApplication : public GenericApplication
 {
-
 public:
 
 	static FMacApplication* CreateMacApplication();
 
-
 public:	
 
 	~FMacApplication();
 
-
 public:
 
-	virtual void SetMessageHandler( const TSharedRef< class FGenericApplicationMessageHandler >& InMessageHandler ) override;
+	virtual void SetMessageHandler(const TSharedRef<class FGenericApplicationMessageHandler>& InMessageHandler) override;
 
-	virtual void PollGameDeviceState( const float TimeDelta ) override;
+	virtual void PollGameDeviceState(const float TimeDelta) override;
 
-	virtual void PumpMessages( const float TimeDelta ) override;
+	virtual void PumpMessages(const float TimeDelta) override;
 
-	virtual TSharedRef< FGenericWindow > MakeWindow() override;
+	virtual void ProcessDeferredEvents(const float TimeDelta) override;
 
-	virtual void InitializeWindow( const TSharedRef< FGenericWindow >& Window, const TSharedRef< FGenericWindowDefinition >& InDefinition, const TSharedPtr< FGenericWindow >& InParent, const bool bShowImmediately ) override;
+	virtual TSharedRef<FGenericWindow> MakeWindow() override;
 
-	virtual void SetHighPrecisionMouseMode( const bool Enable, const TSharedPtr< FGenericWindow >& InWindow ) override;
-
-	virtual bool IsUsingHighPrecisionMouseMode() const override { return bUsingHighPrecisionMouseInput; }
-
-	virtual bool IsUsingTrackpad() const override { return bUsingTrackpad; }
+	virtual void InitializeWindow(const TSharedRef<FGenericWindow>& Window, const TSharedRef<FGenericWindowDefinition>& InDefinition, const TSharedPtr<FGenericWindow>& InParent, const bool bShowImmediately) override;
 
 	virtual FModifierKeysState GetModifierKeys() const override;
 
 	virtual bool IsCursorDirectlyOverSlateWindow() const override;
 
-	virtual FPlatformRect GetWorkArea( const FPlatformRect& CurrentWindow ) const override;
-
-	virtual EWindowTitleAlignment::Type GetWindowTitleAlignment() const override
-	{
-		return EWindowTitleAlignment::Center;
-	}
-	
-	virtual EWindowTransparency GetWindowTransparencySupport() const override
-	{
-		return EWindowTransparency::PerWindow;
-	}
-
-	virtual ITextInputMethodSystem *GetTextInputMethodSystem() override
-	{
-		return TextInputMethodSystem.Get();
-	}
+	virtual void SetHighPrecisionMouseMode(const bool Enable, const TSharedPtr<FGenericWindow>& InWindow) override;
 
-	void ProcessNSEvent(NSEvent* const Event);
+	virtual bool IsUsingHighPrecisionMouseMode() const override { return bUsingHighPrecisionMouseInput; }
 
-	bool IsWindowMovable(FCocoaWindow* Win, bool* OutMovableByBackground);
+	virtual bool IsUsingTrackpad() const override { return bUsingTrackpad; }
 
-	void ResetModifierKeys() { ModifierKeysFlags = 0; }
+	virtual FPlatformRect GetWorkArea(const FPlatformRect& CurrentWindow) const override;
 
-	uint32 GetModifierKeysFlags() { return ModifierKeysFlags; }
+	virtual EWindowTitleAlignment::Type GetWindowTitleAlignment() const override { return EWindowTitleAlignment::Center; }
 
-	bool IsProcessingNSEvent() const { return bIsProcessingNSEvent; }
+	virtual EWindowTransparency GetWindowTransparencySupport() const override { return EWindowTransparency::PerWindow; }
 
-	bool IsWorkspaceSessionActive() const { return bIsWorkspaceSessionActive; }
+	virtual ITextInputMethodSystem *GetTextInputMethodSystem() override { return TextInputMethodSystem.Get(); }
 
 #if WITH_EDITOR
-    virtual void SendAnalytics(IAnalyticsProvider* Provider) override;
-
-	void SetIsUsingTrackpad(bool bInUsingTrackpad)
-	{
-		bUsingTrackpad = bInUsingTrackpad;
-	}
+	virtual void SendAnalytics(IAnalyticsProvider* Provider) override;
 #endif
 
-	void SystemModalMode(bool const bInSystemModalMode);
-
 public:
 
-	void OnDragEnter(FCocoaWindow* Window, NSPasteboard* InPasteboard);
-	void OnDragOver(FCocoaWindow* Window);
-	void OnDragOut(FCocoaWindow* Window);
-	void OnDragDrop(FCocoaWindow* Window);
-
-	void OnWindowDidBecomeMain(FCocoaWindow* Window);
-	void OnWindowDidResignMain(FCocoaWindow* Window);
-	void OnWindowWillMove(FCocoaWindow* Window);
-	void OnWindowDidMove(FCocoaWindow* Window);
-	void OnWindowWillStartLiveResize(FCocoaWindow* Window);
-	void OnWindowDidEndLiveResize(FCocoaWindow* Window);
-	void OnWindowDidResize(FCocoaWindow* Window);
-	void OnWindowDidClose(FCocoaWindow* Window);
-	bool OnWindowDestroyed(FCocoaWindow* Window);
-	void OnWindowClose(FCocoaWindow* Window);
+	void CloseWindow(TSharedRef<FMacWindow> Window);
 
-	void OnMouseCursorLock(bool bLockEnabled);
+	void DeferEvent(NSObject* Object);
 
-	static void ProcessEvent(NSObject const* const Event);
+	TSharedPtr<FMacWindow> FindWindowByNSWindow(FCocoaWindow* WindowHandle);
 
-	void RedrawWindow(FCocoaWindow* Window);
-
-	const TArray<TSharedRef<FMacWindow>>& GetAllWindows() const { return Windows; }
+	/** Queues a window for text layout invalidation when safe */
+	void InvalidateTextLayout(FCocoaWindow* Window);
 
+	void ResetModifierKeys() { ModifierKeysFlags = 0; }
 
-	void CloseQueuedWindows();
+	bool IsWorkspaceSessionActive() const { return bIsWorkspaceSessionActive; }
 
-	/** Queues a window for text layout invalidation when safe */
-	void InvalidateTextLayout( FCocoaWindow* Window );
+	void SystemModalMode(bool const bInSystemModalMode) { bSystemModalMode = bInSystemModalMode; }
 
-	/** Invalidates all queued windows requiring text layout changes */
-	void InvalidateTextLayouts();
+	const TArray<TSharedRef<FMacWindow>>& GetAllWindows() const { return Windows; }
 
 private:
 
-	enum FMacApplicationEventTypes
-	{
-		ResentEvent = 0x0f00
-	};
-
 	static NSEvent* HandleNSEvent(NSEvent* Event);
 	static void OnDisplayReconfiguration(CGDirectDisplayID Display, CGDisplayChangeSummaryFlags Flags, void* UserInfo);
+#if WITH_EDITOR
+	static int32 MTContactCallback(void* Device, void* Data, int32 NumFingers, double TimeStamp, int32 Frame);
+#endif
 
 	FMacApplication();
 
-	bool IsPrintableKey( uint32 Character );
-	TCHAR ConvertChar( TCHAR Character );
-	TCHAR TranslateCharCode( TCHAR CharCode, uint32 KeyCode );
+	void ProcessEvent(const FDeferredMacEvent& Event);
+	void ResendEvent(NSEvent* Event);
 
-	void ResendEvent( NSEvent* Event );
-	FCocoaWindow* FindEventWindow( NSEvent* CocoaEvent );
-	TSharedPtr<FMacWindow> LocateWindowUnderCursor( const NSPoint Position );
+	void ProcessMouseMovedEvent(const FDeferredMacEvent& Event);
+	void ProcessMouseDownEvent(const FDeferredMacEvent& Event);
+	void ProcessMouseUpEvent(const FDeferredMacEvent& Event);
+	void ProcessScrollWheelEvent(const FDeferredMacEvent& Event);
+	void ProcessGestureEvent(const FDeferredMacEvent& Event);
+	void ProcessKeyDownEvent(const FDeferredMacEvent& Event);
+	void ProcessKeyUpEvent(const FDeferredMacEvent& Event);
+
+	void OnWindowDidMove(TSharedRef<FMacWindow> Window);
+	void OnWindowDidResize(TSharedRef<FMacWindow> Window);
+	void OnWindowDidClose(TSharedRef<FMacWindow> Window);
+	bool OnWindowDestroyed(TSharedRef<FMacWindow> Window);
 
-	void OnApplicationWillResignActive();
 	void OnApplicationDidBecomeActive();
+	void OnApplicationWillResignActive();
 	void OnWindowsReordered(bool bIsAppInBackground);
 
-	NSScreen* FindScreenByPoint( int32 X, int32 Y ) const;
-
 	void HandleModifierChange(NSUInteger NewModifierFlags, NSUInteger FlagsShift, NSUInteger UE4Shift, EMacModifierKeys TranslatedCode);
 
+	FCocoaWindow* FindEventWindow(NSEvent* CocoaEvent) const;
+	NSScreen* FindScreenByPoint(int32 X, int32 Y) const;
+	bool IsWindowMovable(TSharedRef<FMacWindow> Window, bool* OutMovableByBackground) const;
+	bool IsPrintableKey(uint32 Character) const;
+	TCHAR ConvertChar(TCHAR Character) const;
+	TCHAR TranslateCharCode(TCHAR CharCode, uint32 KeyCode) const;
+
+	void CloseQueuedWindows();
+
+	/** Invalidates all queued windows requiring text layout changes */
+	void InvalidateTextLayouts();
+
 #if WITH_EDITOR
 	void RecordUsage(EGestureEvent::Type Gesture);
 #else
@@ -161,6 +245,9 @@ private:
 
 	EMouseButtons::Type LastPressedMouseButton;
 
+	FCriticalSection EventsMutex;
+	TArray<FDeferredMacEvent> DeferredEvents;
+
 	FCriticalSection WindowsMutex;
 	TArray<TSharedRef<FMacWindow>> Windows;
 
@@ -176,10 +263,7 @@ private:
 
 	FCocoaWindow* DraggedWindow;
 
-	bool bIsMouseCursorLocked;
-
 	bool bSystemModalMode;
-	bool bIsProcessingNSEvent;
 
 	/** The current set of modifier keys that are pressed. This is used to detect differences between left and right modifier keys on key up events*/
 	uint32 ModifierKeysFlags;
@@ -195,18 +279,15 @@ private:
 
 	bool bIsWorkspaceSessionActive;
 
-	/** Notification center observer for application activation events */
+	/** Notification center observers */
 	id AppActivationObserver;
-
-	/** Notification center observer for application deactivation events */
 	id AppDeactivationObserver;
-
-	/** Notification center observer for workspace activation events */
 	id WorkspaceActivationObserver;
-
-	/** Notification center observer for workspace deactivation events */
 	id WorkspaceDeactivationObserver;
 
+	id EventMonitor;
+	id MouseMovedEventMonitor;
+
 #if WITH_EDITOR
 	/** Holds the last gesture used to try and capture unique uses for gestures. */
 	EGestureEvent::Type LastGestureUsed;
@@ -215,8 +296,6 @@ private:
 	int32 GestureUsage[EGestureEvent::Count];
 #endif
 
-	id EventMonitor;
-
 	friend class FMacWindow;
 };
 
