Fixed bad merge in MacApplication.cpp

diff --git a/Engine/Source/Runtime/Core/Private/Mac/MacApplication.cpp b/Engine/Source/Runtime/Core/Private/Mac/MacApplication.cpp
--- a/Engine/Source/Runtime/Core/Private/Mac/MacApplication.cpp
+++ b/Engine/Source/Runtime/Core/Private/Mac/MacApplication.cpp
@@ -888,7 +888,7 @@ void FMacApplication::OnWindowDidMove(TSharedRef<FMacWindow> Window)
 	const int32 X = (int32)WindowFrame.origin.x;
 	int32 Y = 0;
 
-	if ([Window windowMode] != EWindowMode::Fullscreen)
+	if ([Window->GetWindowHandle() windowMode] != EWindowMode::Fullscreen)
 	{
 		Y = FPlatformMisc::ConvertCocoaYPositionToSlate(WindowFrame.origin.y + WindowFrame.size.height) + (OpenGLFrame.size.height - WindowFrame.size.height);
 	}
