Make OpenRead() aware about intended sharing mode of the file.

- Currently essential on Windows only; other platforms may make use of this when managing file handles.

diff --git a/Engine/Source/Runtime/Core/Private/Android/AndroidFile.cpp b/Engine/Source/Runtime/Core/Private/Android/AndroidFile.cpp
--- a/Engine/Source/Runtime/Core/Private/Android/AndroidFile.cpp
+++ b/Engine/Source/Runtime/Core/Private/Android/AndroidFile.cpp
@@ -1210,12 +1210,12 @@ public:
 		return Filename;
 	}
 
-	virtual IFileHandle* OpenRead(const TCHAR* Filename) override
+	virtual IFileHandle* OpenRead(const TCHAR* Filename, bool bAllowWrite = false) override
 	{
-		return OpenRead(Filename, false);
+		return OpenRead(Filename, false, bAllowWrite);
 	}
 
-	IFileHandle* OpenRead(const TCHAR* Filename, bool AllowLocal)
+	IFileHandle* OpenRead(const TCHAR* Filename, bool AllowLocal, bool bAllowWrite)
 	{
 #if LOG_ANDROID_FILE
 		FPlatformMisc::LowLevelOutputDebugStringf(TEXT("FAndroidPlatformFile::OpenRead('%s')"), Filename);
diff --git a/Engine/Source/Runtime/Core/Private/Apple/ApplePlatformFile.cpp b/Engine/Source/Runtime/Core/Private/Apple/ApplePlatformFile.cpp
--- a/Engine/Source/Runtime/Core/Private/Apple/ApplePlatformFile.cpp
+++ b/Engine/Source/Runtime/Core/Private/Apple/ApplePlatformFile.cpp
@@ -420,7 +420,7 @@ FString FApplePlatformFile::GetFilenameOnDisk(const TCHAR* Filename)
 	return Filename;
 }
 
-IFileHandle* FApplePlatformFile::OpenRead(const TCHAR* Filename)
+IFileHandle* FApplePlatformFile::OpenRead(const TCHAR* Filename, bool bAllowWrite)
 {
 	int32 Handle = open(TCHAR_TO_UTF8(*NormalizeFilename(Filename)), O_RDONLY);
 	if (Handle != -1)
diff --git a/Engine/Source/Runtime/Core/Private/IOS/IOSPlatformFile.cpp b/Engine/Source/Runtime/Core/Private/IOS/IOSPlatformFile.cpp
--- a/Engine/Source/Runtime/Core/Private/IOS/IOSPlatformFile.cpp
+++ b/Engine/Source/Runtime/Core/Private/IOS/IOSPlatformFile.cpp
@@ -518,7 +518,7 @@ FString FIOSPlatformFile::GetFilenameOnDisk(const TCHAR* Filename)
 	return Filename;
 }
 
-IFileHandle* FIOSPlatformFile::OpenRead(const TCHAR* Filename)
+IFileHandle* FIOSPlatformFile::OpenRead(const TCHAR* Filename, bool bAllowWrite)
 {
 	FString NormalizedFilename = NormalizeFilename(Filename);
 
diff --git a/Engine/Source/Runtime/Core/Private/IOS/IOSPlatformFile.h b/Engine/Source/Runtime/Core/Private/IOS/IOSPlatformFile.h
--- a/Engine/Source/Runtime/Core/Private/IOS/IOSPlatformFile.h
+++ b/Engine/Source/Runtime/Core/Private/IOS/IOSPlatformFile.h
@@ -30,7 +30,7 @@ public:
 	virtual FDateTime GetAccessTimeStamp(const TCHAR* Filename) override;
 	virtual FString GetFilenameOnDisk(const TCHAR* Filename) override;
 
-	virtual IFileHandle* OpenRead(const TCHAR* Filename) override;
+	virtual IFileHandle* OpenRead(const TCHAR* Filename, bool bAllowWrite = false) override;
 	virtual IFileHandle* OpenWrite(const TCHAR* Filename, bool bAppend = false, bool bAllowRead = false) override;
 
 	virtual bool DirectoryExists(const TCHAR* Directory) override;
diff --git a/Engine/Source/Runtime/Core/Private/Linux/LinuxPlatformFile.cpp b/Engine/Source/Runtime/Core/Private/Linux/LinuxPlatformFile.cpp
--- a/Engine/Source/Runtime/Core/Private/Linux/LinuxPlatformFile.cpp
+++ b/Engine/Source/Runtime/Core/Private/Linux/LinuxPlatformFile.cpp
@@ -755,7 +755,7 @@ FString FLinuxPlatformFile::GetFilenameOnDisk(const TCHAR* Filename)
 */
 }
 
-IFileHandle* FLinuxPlatformFile::OpenRead(const TCHAR* Filename)
+IFileHandle* FLinuxPlatformFile::OpenRead(const TCHAR* Filename, bool bAllowWrite)
 {
 	FString MappedToName;
 	int32 Handle = GCaseInsensMapper.OpenCaseInsensitiveRead(TCHAR_TO_UTF8(*NormalizeFilename(Filename)), MappedToName);
diff --git a/Engine/Source/Runtime/Core/Private/Windows/WindowsPlatformFile.cpp b/Engine/Source/Runtime/Core/Private/Windows/WindowsPlatformFile.cpp
--- a/Engine/Source/Runtime/Core/Private/Windows/WindowsPlatformFile.cpp
+++ b/Engine/Source/Runtime/Core/Private/Windows/WindowsPlatformFile.cpp
@@ -255,10 +255,10 @@ public:
 		return Result;
 	}
 
-	virtual IFileHandle* OpenRead(const TCHAR* Filename) override
+	virtual IFileHandle* OpenRead(const TCHAR* Filename, bool bAllowWrite = false) override
 	{
 		uint32  Access    = GENERIC_READ;
-		uint32  WinFlags  = FILE_SHARE_READ;
+		uint32  WinFlags  = FILE_SHARE_READ | (bAllowWrite ? FILE_SHARE_WRITE : 0);
 		uint32  Create    = OPEN_EXISTING;
 		HANDLE Handle    = CreateFileW(*NormalizeFilename(Filename), Access, WinFlags, NULL, Create, FILE_ATTRIBUTE_NORMAL, NULL);
 		if(Handle != INVALID_HANDLE_VALUE)
diff --git a/Engine/Source/Runtime/Core/Public/Apple/ApplePlatformFile.h b/Engine/Source/Runtime/Core/Public/Apple/ApplePlatformFile.h
--- a/Engine/Source/Runtime/Core/Public/Apple/ApplePlatformFile.h
+++ b/Engine/Source/Runtime/Core/Public/Apple/ApplePlatformFile.h
@@ -31,7 +31,7 @@ public:
 	virtual FDateTime GetAccessTimeStamp(const TCHAR* Filename) override;
 	virtual FString GetFilenameOnDisk(const TCHAR* Filename) override;
 
-	virtual IFileHandle* OpenRead(const TCHAR* Filename) override;
+	virtual IFileHandle* OpenRead(const TCHAR* Filename, bool bAllowWrite = false) override;
 	virtual IFileHandle* OpenWrite(const TCHAR* Filename, bool bAppend = false, bool bAllowRead = false) override;
 	virtual bool DirectoryExists(const TCHAR* Directory) override;
 	virtual bool CreateDirectory(const TCHAR* Directory) override;
diff --git a/Engine/Source/Runtime/Core/Public/GenericPlatform/GenericPlatformFile.h b/Engine/Source/Runtime/Core/Public/GenericPlatform/GenericPlatformFile.h
--- a/Engine/Source/Runtime/Core/Public/GenericPlatform/GenericPlatformFile.h
+++ b/Engine/Source/Runtime/Core/Public/GenericPlatform/GenericPlatformFile.h
@@ -147,10 +147,16 @@ public:
 	/** For case insensitive filesystems, returns the full path of the file with the same case as in the filesystem */
 	virtual FString GetFilenameOnDisk(const TCHAR* Filename) = 0;
 
-	/** Attempt to open a file for reading. If successful will return a non-nullptr pointer. Close the file by delete'ing the handle. **/
-	virtual IFileHandle*	OpenRead(const TCHAR* Filename) = 0;
+	/** Attempt to open a file for reading.
+	 *
+	 * @param Filename file to be opened
+	 * @param bAllowWrite (applies to certain platforms only) whether this file is allowed to be written to by other processes. This flag is needed to open files that are currently being written to as well.
+	 *
+	 * @return If successful will return a non-nullptr pointer. Close the file by delete'ing the handle.
+	 */
+	virtual IFileHandle*	OpenRead(const TCHAR* Filename, bool bAllowWrite = false) = 0;
 	/** Attempt to open a file for writing. If successful will return a non-nullptr pointer. Close the file by delete'ing the handle. **/
-	virtual IFileHandle*	OpenWrite(const TCHAR* Filename, bool bAppend = 0, bool bAllowRead = 0) = 0;
+	virtual IFileHandle*	OpenWrite(const TCHAR* Filename, bool bAppend = false, bool bAllowRead = false) = 0;
 
 	/** Return true if the directory exists. **/
 	virtual bool		DirectoryExists(const TCHAR* Directory) = 0;
diff --git a/Engine/Source/Runtime/Core/Public/HAL/IPlatformFileCachedWrapper.h b/Engine/Source/Runtime/Core/Public/HAL/IPlatformFileCachedWrapper.h
--- a/Engine/Source/Runtime/Core/Public/HAL/IPlatformFileCachedWrapper.h
+++ b/Engine/Source/Runtime/Core/Public/HAL/IPlatformFileCachedWrapper.h
@@ -280,9 +280,9 @@ public:
 	{
 		return LowerLevel->GetFilenameOnDisk(Filename);
 	}
-	virtual IFileHandle*	OpenRead(const TCHAR* Filename) override
+	virtual IFileHandle*	OpenRead(const TCHAR* Filename, bool bAllowWrite) override
 	{
-		IFileHandle* InnerHandle=LowerLevel->OpenRead(Filename);
+		IFileHandle* InnerHandle=LowerLevel->OpenRead(Filename, bAllowWrite);
 		if (!InnerHandle)
 		{
 			return nullptr;
diff --git a/Engine/Source/Runtime/Core/Public/HAL/IPlatformFileLogWrapper.h b/Engine/Source/Runtime/Core/Public/HAL/IPlatformFileLogWrapper.h
--- a/Engine/Source/Runtime/Core/Public/HAL/IPlatformFileLogWrapper.h
+++ b/Engine/Source/Runtime/Core/Public/HAL/IPlatformFileLogWrapper.h
@@ -212,11 +212,11 @@ public:
 		FILE_LOG(LogPlatformFile, Log, TEXT("GetFilenameOnDisk return %llx [%s]"), *Result, ThisTime);
 		return Result;
 	}
-	virtual IFileHandle*	OpenRead(const TCHAR* Filename) override
+	virtual IFileHandle*	OpenRead(const TCHAR* Filename, bool bAllowWrite) override
 	{
 		FILE_LOG(LogPlatformFile, Log, TEXT("OpenRead %s"), Filename);
 		double StartTime = FPlatformTime::Seconds();
-		IFileHandle* Result = LowerLevel->OpenRead(Filename);
+		IFileHandle* Result = LowerLevel->OpenRead(Filename, bAllowWrite);
 		float ThisTime = 1000.0f * float(FPlatformTime::Seconds() - StartTime);
 		FILE_LOG(LogPlatformFile, Log, TEXT("OpenRead return %llx [%fms]"), uint64(Result), ThisTime);
 		return Result ? (new FLoggedFileHandle(Result, Filename)) : Result;
diff --git a/Engine/Source/Runtime/Core/Public/HAL/IPlatformFileOpenLogWrapper.h b/Engine/Source/Runtime/Core/Public/HAL/IPlatformFileOpenLogWrapper.h
--- a/Engine/Source/Runtime/Core/Public/HAL/IPlatformFileOpenLogWrapper.h
+++ b/Engine/Source/Runtime/Core/Public/HAL/IPlatformFileOpenLogWrapper.h
@@ -133,9 +133,9 @@ public:
 	{
 		return LowerLevel->GetFilenameOnDisk(Filename);
 	}
-	virtual IFileHandle*	OpenRead(const TCHAR* Filename) override
+	virtual IFileHandle*	OpenRead(const TCHAR* Filename, bool bAllowWrite) override
 	{
-		IFileHandle* Result = LowerLevel->OpenRead(Filename);
+		IFileHandle* Result = LowerLevel->OpenRead(Filename, bAllowWrite);
 		if (Result)
 		{
 			CriticalSection.Lock();
diff --git a/Engine/Source/Runtime/Core/Public/HAL/IPlatformFileProfilerWrapper.h b/Engine/Source/Runtime/Core/Public/HAL/IPlatformFileProfilerWrapper.h
--- a/Engine/Source/Runtime/Core/Public/HAL/IPlatformFileProfilerWrapper.h
+++ b/Engine/Source/Runtime/Core/Public/HAL/IPlatformFileProfilerWrapper.h
@@ -363,11 +363,11 @@ public:
 		OpStat->Duration += FPlatformTime::Seconds() * 1000 - OpStat->LastOpTime;
 		return Result;
 	}
-	virtual IFileHandle*	OpenRead(const TCHAR* Filename) override
+	virtual IFileHandle*	OpenRead(const TCHAR* Filename, bool bAllowWrite = false) override
 	{
 		StatsType* FileStat = CreateStat( Filename );
 		FProfiledFileStatsOp* OpStat = FileStat->CreateOpStat( FProfiledFileStatsOp::OpenRead );
-		IFileHandle* Result = LowerLevel->OpenRead(Filename);
+		IFileHandle* Result = LowerLevel->OpenRead(Filename, bAllowWrite);
 		OpStat->Duration += FPlatformTime::Seconds() * 1000.0 - OpStat->LastOpTime;
 		return Result ? (new TProfiledFileHandle< StatsType >( Result, Filename, FileStat )) : Result;
 	}
@@ -584,9 +584,9 @@ public:
 	{
 		return LowerLevel->GetFilenameOnDisk(Filename);
 	}
-	virtual IFileHandle*	OpenRead(const TCHAR* Filename) override
+	virtual IFileHandle*	OpenRead(const TCHAR* Filename, bool bAllowWrite) override
 	{
-		IFileHandle* Result = LowerLevel->OpenRead(Filename);
+		IFileHandle* Result = LowerLevel->OpenRead(Filename, bAllowWrite);
 		return Result ? (new FPlatformFileReadStatsHandle(Result, Filename, &BytePerSecThisTick, &BytesReadThisTick, &ReadsThisTick)) : Result;
 	}
 	virtual IFileHandle*	OpenWrite(const TCHAR* Filename, bool bAppend = false, bool bAllowRead = false) override
diff --git a/Engine/Source/Runtime/Core/Public/Linux/LinuxPlatformFile.h b/Engine/Source/Runtime/Core/Public/Linux/LinuxPlatformFile.h
--- a/Engine/Source/Runtime/Core/Public/Linux/LinuxPlatformFile.h
+++ b/Engine/Source/Runtime/Core/Public/Linux/LinuxPlatformFile.h
@@ -31,7 +31,7 @@ public:
 	virtual FDateTime GetAccessTimeStamp(const TCHAR* Filename) override;
 	virtual FString GetFilenameOnDisk(const TCHAR* Filename) override;
 
-	virtual IFileHandle* OpenRead(const TCHAR* Filename) override;
+	virtual IFileHandle* OpenRead(const TCHAR* Filename, bool bAllowWrite = false) override;
 	virtual IFileHandle* OpenWrite(const TCHAR* Filename, bool bAppend = false, bool bAllowRead = false) override;
 	virtual bool DirectoryExists(const TCHAR* Directory) override;
 	virtual bool CreateDirectory(const TCHAR* Directory) override;
diff --git a/Engine/Source/Runtime/NetworkFile/Private/NetworkPlatformFile.cpp b/Engine/Source/Runtime/NetworkFile/Private/NetworkPlatformFile.cpp
--- a/Engine/Source/Runtime/NetworkFile/Private/NetworkPlatformFile.cpp
+++ b/Engine/Source/Runtime/NetworkFile/Private/NetworkPlatformFile.cpp
@@ -396,7 +396,7 @@ void FNetworkPlatformFile::SetTimeStamp(const TCHAR* Filename, FDateTime DateTim
 	InnerPlatformFile->SetTimeStamp(Filename, DateTime);
 }
 
-IFileHandle* FNetworkPlatformFile::OpenRead(const TCHAR* Filename)
+IFileHandle* FNetworkPlatformFile::OpenRead(const TCHAR* Filename, bool bAllowWrite)
 {
 //	FScopeLock ScopeLock(&SynchronizationObject);
 
diff --git a/Engine/Source/Runtime/NetworkFile/Public/NetworkPlatformFile.h b/Engine/Source/Runtime/NetworkFile/Public/NetworkPlatformFile.h
--- a/Engine/Source/Runtime/NetworkFile/Public/NetworkPlatformFile.h
+++ b/Engine/Source/Runtime/NetworkFile/Public/NetworkPlatformFile.h
@@ -94,7 +94,7 @@ public:
 	{
 		return Filename;
 	}
-	virtual IFileHandle*	OpenRead(const TCHAR* Filename) override;
+	virtual IFileHandle*	OpenRead(const TCHAR* Filename, bool bAllowWrite = false) override;
 	virtual IFileHandle*	OpenWrite(const TCHAR* Filename, bool bAppend = false, bool bAllowRead = false) override;
 	virtual bool		DirectoryExists(const TCHAR* Directory) override;
 	virtual bool		CreateDirectoryTree(const TCHAR* Directory) override;
diff --git a/Engine/Source/Runtime/PakFile/Private/IPlatformFilePak.cpp b/Engine/Source/Runtime/PakFile/Private/IPlatformFilePak.cpp
--- a/Engine/Source/Runtime/PakFile/Private/IPlatformFilePak.cpp
+++ b/Engine/Source/Runtime/PakFile/Private/IPlatformFilePak.cpp
@@ -787,7 +787,7 @@ bool FPakPlatformFile::HandleUnmountPakDelegate(const FString& PakFilePath)
 	return Unmount(*PakFilePath);
 }
 
-IFileHandle* FPakPlatformFile::OpenRead(const TCHAR* Filename)
+IFileHandle* FPakPlatformFile::OpenRead(const TCHAR* Filename, bool bAllowWrite)
 {
 	IFileHandle* Result = NULL;
 	FPakFile* PakFile = NULL;
diff --git a/Engine/Source/Runtime/PakFile/Public/IPlatformFilePak.h b/Engine/Source/Runtime/PakFile/Public/IPlatformFilePak.h
--- a/Engine/Source/Runtime/PakFile/Public/IPlatformFilePak.h
+++ b/Engine/Source/Runtime/PakFile/Public/IPlatformFilePak.h
@@ -1116,7 +1116,7 @@ public:
 		return LowerLevel->GetFilenameOnDisk(Filename);
 	}
 
-	virtual IFileHandle* OpenRead(const TCHAR* Filename) override;
+	virtual IFileHandle* OpenRead(const TCHAR* Filename, bool bAllowWrite = false) override;
 
 	virtual IFileHandle* OpenWrite(const TCHAR* Filename, bool bAppend = false, bool bAllowRead = false) override
 	{
diff --git a/Engine/Source/Runtime/SandboxFile/Public/IPlatformFileSandboxWrapper.h b/Engine/Source/Runtime/SandboxFile/Public/IPlatformFileSandboxWrapper.h
--- a/Engine/Source/Runtime/SandboxFile/Public/IPlatformFileSandboxWrapper.h
+++ b/Engine/Source/Runtime/SandboxFile/Public/IPlatformFileSandboxWrapper.h
@@ -377,9 +377,9 @@ public:
 		return Result;
 	}
 
-	virtual IFileHandle*	OpenRead(const TCHAR* Filename) override
+	virtual IFileHandle*	OpenRead(const TCHAR* Filename, bool bAllowWrite = false) override
 	{
-		IFileHandle* Result = LowerLevel->OpenRead( *ConvertToSandboxPath( Filename ) );
+		IFileHandle* Result = LowerLevel->OpenRead( *ConvertToSandboxPath(Filename), bAllowWrite );
 		if( !Result  && OkForInnerAccess(Filename) )
 		{
 			Result = LowerLevel->OpenRead( Filename );
diff --git a/Engine/Source/Runtime/StreamingFile/Private/StreamingNetworkPlatformFile.cpp b/Engine/Source/Runtime/StreamingFile/Private/StreamingNetworkPlatformFile.cpp
--- a/Engine/Source/Runtime/StreamingFile/Private/StreamingNetworkPlatformFile.cpp
+++ b/Engine/Source/Runtime/StreamingFile/Private/StreamingNetworkPlatformFile.cpp
@@ -496,7 +496,7 @@ FDateTime FStreamingNetworkPlatformFile::GetAccessTimeStamp(const TCHAR* Filenam
 }
 
 
-IFileHandle* FStreamingNetworkPlatformFile::OpenRead(const TCHAR* Filename)
+IFileHandle* FStreamingNetworkPlatformFile::OpenRead(const TCHAR* Filename, bool bAllowWrite)
 {
 	FString RelativeFilename = Filename;
 	FPaths::MakeStandardFilename(RelativeFilename);
diff --git a/Engine/Source/Runtime/StreamingFile/Public/StreamingNetworkPlatformFile.h b/Engine/Source/Runtime/StreamingFile/Public/StreamingNetworkPlatformFile.h
--- a/Engine/Source/Runtime/StreamingFile/Public/StreamingNetworkPlatformFile.h
+++ b/Engine/Source/Runtime/StreamingFile/Public/StreamingNetworkPlatformFile.h
@@ -182,7 +182,7 @@ public:
 	virtual FDateTime GetTimeStamp(const TCHAR* Filename) override;
 	virtual void SetTimeStamp(const TCHAR* Filename, FDateTime DateTime) override;
 	virtual FDateTime GetAccessTimeStamp(const TCHAR* Filename) override;
-	virtual IFileHandle* OpenRead(const TCHAR* Filename) override;
+	virtual IFileHandle* OpenRead(const TCHAR* Filename, bool bAllowWrite = false) override;
 	virtual IFileHandle* OpenWrite(const TCHAR* Filename, bool bAppend, bool bAllowRead) override;
 	virtual bool DirectoryExists(const TCHAR* Directory) override;
 	virtual bool CreateDirectoryTree(const TCHAR* Directory) override;
