HTML5 PlatformFile CIS fix

diff --git a/Engine/Source/Runtime/Core/Private/HTML5/HTML5PlatformFile.cpp b/Engine/Source/Runtime/Core/Private/HTML5/HTML5PlatformFile.cpp
--- a/Engine/Source/Runtime/Core/Private/HTML5/HTML5PlatformFile.cpp
+++ b/Engine/Source/Runtime/Core/Private/HTML5/HTML5PlatformFile.cpp
@@ -290,7 +290,7 @@ public:
 		return Filename;
 	}
 
-	virtual IFileHandle* OpenRead(const TCHAR* Filename) override
+	virtual IFileHandle* OpenRead(const TCHAR* Filename, bool bAllowWrite /*= false*/) override
 	{
 		FString fn = NormalizeFilename(Filename);
 		int32 Handle = open(TCHAR_TO_UTF8(*fn), O_RDONLY | O_BINARY);
