Replaced some FString copies with moves, and ensured that sizing down the destination FString in ConvertString can't cause a reallocation.

diff --git a/Engine/Source/Runtime/Core/Private/Internationalization/ICUText.cpp b/Engine/Source/Runtime/Core/Private/Internationalization/ICUText.cpp
--- a/Engine/Source/Runtime/Core/Private/Internationalization/ICUText.cpp
+++ b/Engine/Source/Runtime/Core/Private/Internationalization/ICUText.cpp
@@ -38,7 +38,7 @@ FText FText::AsDate(const FDateTime& DateTime, const EDateTimeStyle::Type DateSt
 	FString NativeString;
 	ICUUtilities::ConvertString(FormattedString, NativeString);
 
-	FText ResultText = FText::CreateChronologicalText(NativeString);
+	FText ResultText = FText::CreateChronologicalText(MoveTemp(NativeString));
 	ResultText.History = MakeShareable(new FTextHistory_AsDate(DateTime, DateStyle, TimeZone, TargetCulture));
 
 	return ResultText;
@@ -61,7 +61,7 @@ FText FText::AsTime(const FDateTime& DateTime, const EDateTimeStyle::Type TimeSt
 	FString NativeString;
 	ICUUtilities::ConvertString(FormattedString, NativeString);
 
-	FText ResultText = FText::CreateChronologicalText(NativeString);
+	FText ResultText = FText::CreateChronologicalText(MoveTemp(NativeString));
 	ResultText.History = MakeShareable(new FTextHistory_AsTime(DateTime, TimeStyle, TimeZone, TargetCulture));
 
 	return ResultText;
@@ -108,7 +108,7 @@ FText FText::AsDateTime(const FDateTime& DateTime, const EDateTimeStyle::Type Da
 	FString NativeString;
 	ICUUtilities::ConvertString(FormattedString, NativeString);
 
-	FText ResultText = FText::CreateChronologicalText(NativeString);
+	FText ResultText = FText::CreateChronologicalText(MoveTemp(NativeString));
 	ResultText.History = MakeShareable(new FTextHistory_AsDateTime(DateTime, DateStyle, TimeStyle, TimeZone, TargetCulture));
 
 	return ResultText;
diff --git a/Engine/Source/Runtime/Core/Private/Internationalization/ICUText.h b/Engine/Source/Runtime/Core/Private/Internationalization/ICUText.h
--- a/Engine/Source/Runtime/Core/Private/Internationalization/ICUText.h
+++ b/Engine/Source/Runtime/Core/Private/Internationalization/ICUText.h
@@ -25,7 +25,7 @@ FText FText::AsNumberTemplate(T1 Val, const FNumberFormattingOptions* const Opti
 	FString NativeString;
 	ICUUtilities::ConvertString(FormattedString, NativeString);
 
-	FText ReturnText = FText::CreateNumericalText(NativeString);
+	FText ReturnText = FText::CreateNumericalText(MoveTemp(NativeString));
 	ReturnText.History = MakeShareable(new FTextHistory_AsNumber(Val, Options, TargetCulture));
 	return ReturnText;
 }
@@ -45,7 +45,7 @@ FText FText::AsCurrencyTemplate(T1 Val, const FString& CurrencyCode, const FNumb
 	FString NativeString;
 	ICUUtilities::ConvertString(FormattedString, NativeString);
 
-	FText ReturnText = FText::CreateNumericalText(NativeString);
+	FText ReturnText = FText::CreateNumericalText(MoveTemp(NativeString));
 	ReturnText.History = MakeShareable(new FTextHistory_AsCurrency(Val, CurrencyCode, Options, TargetCulture));
 	return ReturnText;
 }
@@ -65,7 +65,7 @@ FText FText::AsPercentTemplate(T1 Val, const FNumberFormattingOptions* const Opt
 	FString NativeString;
 	ICUUtilities::ConvertString(FormattedString, NativeString);
 
-	FText ReturnText = FText::CreateNumericalText(NativeString);
+	FText ReturnText = FText::CreateNumericalText(MoveTemp(NativeString));
 	ReturnText.History = MakeShareable(new FTextHistory_AsPercent(Val, Options, TargetCulture));
 	return ReturnText;
 }
\ No newline at end of file
diff --git a/Engine/Source/Runtime/Core/Private/Internationalization/ICUUtilities.cpp b/Engine/Source/Runtime/Core/Private/Internationalization/ICUUtilities.cpp
--- a/Engine/Source/Runtime/Core/Private/Internationalization/ICUUtilities.cpp
+++ b/Engine/Source/Runtime/Core/Private/Internationalization/ICUUtilities.cpp
@@ -77,7 +77,7 @@ namespace ICUUtilities
 				const int32_t DestinationSizeBytes = ucnv_fromUChars(ICUConverter, reinterpret_cast<char*>(InternalStringBuffer.GetData()), DestinationCapacityBytes, Source.getBuffer(), Source.length(), &ICUStatus);
 				const int32 DestinationSizeTCHARs = DestinationSizeBytes / sizeof(TCHAR);
 				InternalStringBuffer[DestinationSizeTCHARs] = 0;
-				InternalStringBuffer.SetNumUninitialized(DestinationSizeTCHARs + 1); // the array size includes null
+				InternalStringBuffer.SetNum(DestinationSizeTCHARs + 1, /*bAllowShrinking*/false); // the array size includes null
 
 				ucnv_close(ICUConverter);
 			}
