DegreesToRadians and RadiansToDegrees now use auto-deduced return types, to simplify their definition and allow their use with other types, like vectors.

diff --git a/Engine/Source/Runtime/Core/Public/Math/UnrealMathUtility.h b/Engine/Source/Runtime/Core/Public/Math/UnrealMathUtility.h
--- a/Engine/Source/Runtime/Core/Public/Math/UnrealMathUtility.h
+++ b/Engine/Source/Runtime/Core/Public/Math/UnrealMathUtility.h
@@ -385,13 +385,7 @@ struct FMath : public FPlatformMath
 	 * @return					Value in degrees.
 	 */
 	template<class T>
-	static FORCEINLINE float RadiansToDegrees(T const& RadVal)
-	{
-		return RadVal * (180.f / PI);
-	}
-
-	/** This overload allows us to maintain double precision */
-	static FORCEINLINE double RadiansToDegrees(double RadVal)
+	static FORCEINLINE auto RadiansToDegrees(T const& RadVal) -> decltype(RadVal * (180.f / PI))
 	{
 		return RadVal * (180.f / PI);
 	}
@@ -402,13 +396,7 @@ struct FMath : public FPlatformMath
 	 * @return					Value in radians.
 	 */
 	template<class T>
-	static FORCEINLINE float DegreesToRadians(T const& DegVal)
-	{
-		return DegVal * (PI / 180.f);
-	}
-
-	/** This overload allows us to maintain double precision */
-	static FORCEINLINE double DegreesToRadians(double DegVal)
+	static FORCEINLINE auto DegreesToRadians(T const& DegVal) -> decltype(DegVal * (PI / 180.f))
 	{
 		return DegVal * (PI / 180.f);
 	}
