Linux: fix cache when calculating the work area.

- Previous code was broken as it was caching the result regardless of the actual rectangle being passed.
- Now we cache the display info itself, not the result.
- Maybe we should use maximum overlap instead of top, left (doesn't seem to be important though).

diff --git a/Engine/Source/Runtime/Core/Private/Linux/LinuxApplication.cpp b/Engine/Source/Runtime/Core/Private/Linux/LinuxApplication.cpp
--- a/Engine/Source/Runtime/Core/Private/Linux/LinuxApplication.cpp
+++ b/Engine/Source/Runtime/Core/Private/Linux/LinuxApplication.cpp
@@ -69,7 +69,7 @@ FLinuxApplication::FLinuxApplication()
 	,	bIsDragWindowButtonPressed(false)
 	,	bActivateApp(false)
 	,	bLockToCurrentMouseType(false)
-	,	LastTimeCachedWorkArea(-1.0)
+	,	LastTimeCachedDisplays(-1.0)
 {
 	bUsingHighPrecisionMouseInput = false;
 	bAllowedToDeferMessageProcessing = true;
@@ -1099,52 +1099,65 @@ void FLinuxApplication::SetHighPrecisionMouseMode( const bool Enable, const TSha
 	bUsingHighPrecisionMouseInput = Enable;
 }
 
-
-FPlatformRect FLinuxApplication::GetWorkArea( const FPlatformRect& CurrentWindow ) const
+void FLinuxApplication::RefreshDisplayCache()
 {
 	const double kCacheLifetime = 5.0;	// ask once in 5 seconds
-
+	
 	double CurrentTime = FPlatformTime::Seconds();
-	if (CurrentTime - LastTimeCachedWorkArea > kCacheLifetime)
+	if (CurrentTime - LastTimeCachedDisplays > kCacheLifetime)
 	{
-		// loop over all monitors to determine which one is the best
+		CachedDisplays.Empty();
+
 		int NumDisplays = SDL_GetNumVideoDisplays();
-		if (NumDisplays <= 0)
-		{
-			// fake something
-			return CurrentWindow;
-		}
-		
-		SDL_Rect BestDisplayBounds;
-		SDL_GetDisplayBounds(0, &BestDisplayBounds);
-		
-		// see if any other are better (i.e. cover top left)
-		for (int DisplayIdx = 1; DisplayIdx < NumDisplays; ++DisplayIdx)
+
+		for (int DisplayIdx = 0; DisplayIdx < NumDisplays; ++DisplayIdx)
 		{
 			SDL_Rect DisplayBounds;
 			SDL_GetDisplayBounds(DisplayIdx, &DisplayBounds);
 			
-			// only check top left corner for "bestness"
-			if (DisplayBounds.x <= CurrentWindow.Left && DisplayBounds.x + DisplayBounds.w > CurrentWindow.Left &&
-				DisplayBounds.y <= CurrentWindow.Top && DisplayBounds.y + DisplayBounds.h > CurrentWindow.Bottom)
-			{
-				BestDisplayBounds = DisplayBounds;
-				// there can be only one, as we don't expect overlapping displays
-				break;
-			}
+			CachedDisplays.Add(DisplayBounds);
 		}
-		
-		FPlatformRect WorkArea;
-		WorkArea.Left	= BestDisplayBounds.x;
-		WorkArea.Top	= BestDisplayBounds.y;
-		WorkArea.Right	= BestDisplayBounds.x + BestDisplayBounds.w;
-		WorkArea.Bottom	= BestDisplayBounds.y + BestDisplayBounds.h;
-
-		CachedWorkArea = WorkArea;
-		LastTimeCachedWorkArea = CurrentTime;
+
+		LastTimeCachedDisplays = CurrentTime;
+	}
+}
+
+FPlatformRect FLinuxApplication::GetWorkArea( const FPlatformRect& CurrentWindow ) const
+{
+	(const_cast<FLinuxApplication *>(this))->RefreshDisplayCache();
+
+	// loop over all monitors to determine which one is the best
+	int NumDisplays = CachedDisplays.Num();
+	if (NumDisplays <= 0)
+	{
+		// fake something
+		return CurrentWindow;
 	}
 
-	return CachedWorkArea;
+	SDL_Rect BestDisplayBounds = CachedDisplays[0];
+
+	// see if any other are better (i.e. cover top left)
+	for (int DisplayIdx = 1; DisplayIdx < NumDisplays; ++DisplayIdx)
+	{
+		const SDL_Rect & DisplayBounds = CachedDisplays[DisplayIdx];
+
+		// only check top left corner for "bestness"
+		if (DisplayBounds.x <= CurrentWindow.Left && DisplayBounds.x + DisplayBounds.w > CurrentWindow.Left &&
+			DisplayBounds.y <= CurrentWindow.Top && DisplayBounds.y + DisplayBounds.h > CurrentWindow.Bottom)
+		{
+			BestDisplayBounds = DisplayBounds;
+			// there can be only one, as we don't expect overlapping displays
+			break;
+		}
+	}
+
+	FPlatformRect WorkArea;
+	WorkArea.Left	= BestDisplayBounds.x;
+	WorkArea.Top	= BestDisplayBounds.y;
+	WorkArea.Right	= BestDisplayBounds.x + BestDisplayBounds.w;
+	WorkArea.Bottom	= BestDisplayBounds.y + BestDisplayBounds.h;
+
+	return WorkArea;
 }
 
 void FLinuxApplication::OnMouseCursorLock( bool bLockEnabled )
diff --git a/Engine/Source/Runtime/Core/Public/Linux/LinuxApplication.h b/Engine/Source/Runtime/Core/Public/Linux/LinuxApplication.h
--- a/Engine/Source/Runtime/Core/Public/Linux/LinuxApplication.h
+++ b/Engine/Source/Runtime/Core/Public/Linux/LinuxApplication.h
@@ -97,6 +97,8 @@ private:
 
 private:
 
+	void RefreshDisplayCache();
+
 	struct SDLControllerState
 	{
 		SDL_HController controller;
@@ -149,11 +151,11 @@ private:
 	/** Used to check with cursor type is current and set to true if left button is pressed.*/
 	bool bLockToCurrentMouseType;
 
-	/** Cached work area - to reduce costly communication with X server (may be better cached in SDL? avoids ugly 'mutable') */
-	mutable FPlatformRect	CachedWorkArea;
+	/** Cached displays - to reduce costly communication with X server (may be better cached in SDL? avoids ugly 'mutable') */
+	mutable TArray<SDL_Rect>	CachedDisplays;
 
 	/** Last time we asked about work area (this is a hack. What we need is a callback when screen config changes). */
-	mutable double			LastTimeCachedWorkArea;
+	mutable double			LastTimeCachedDisplays;
 };
 
 extern FLinuxApplication* LinuxApplication;
