Store the original NSEvent in FDeferredMacEvent and use that when resending key events as apparently OS X ignores some events created with [NSEvent keyEventWithType:...]

diff --git a/Engine/Source/Runtime/Core/Private/Mac/MacApplication.cpp b/Engine/Source/Runtime/Core/Private/Mac/MacApplication.cpp
--- a/Engine/Source/Runtime/Core/Private/Mac/MacApplication.cpp
+++ b/Engine/Source/Runtime/Core/Private/Mac/MacApplication.cpp
@@ -255,6 +255,7 @@ void FMacApplication::DeferEvent(NSObject* Object)
 	{
 		NSEvent* Event = (NSEvent*)Object;
 		WindowHandle = FindEventWindow(Event);
+		DeferredEvent.Event = [Event retain];
 		DeferredEvent.Type = [Event type];
 		DeferredEvent.LocationInWindow = FVector2D([Event locationInWindow].x, [Event locationInWindow].y);
 		DeferredEvent.ModifierFlags = [Event modifierFlags];
@@ -852,12 +853,12 @@ void FMacApplication::ProcessKeyDownEvent(const FDeferredMacEvent& Event)
 		FCocoaMenu* MainMenu = [[NSApp mainMenu] isKindOfClass:[FCocoaMenu class]] ? (FCocoaMenu*)[NSApp mainMenu]: nil;
 		if (MainMenu)
 		{
-			MainThreadCall(^{ [MainMenu highlightKeyEquivalent:Event.GetKeyNSEvent()]; }, NSDefaultRunLoopMode, true);
+			MainThreadCall(^{ [MainMenu highlightKeyEquivalent:Event.Event]; }, NSDefaultRunLoopMode, true);
 		}
 	}
 	else
 	{
-		ResendEvent(Event.GetKeyNSEvent());
+		ResendEvent(Event.Event);
 	}
 }
 
@@ -873,7 +874,7 @@ void FMacApplication::ProcessKeyUpEvent(const FDeferredMacEvent& Event)
 	}
 	if (!bHandled)
 	{
-		ResendEvent(Event.GetKeyNSEvent());
+		ResendEvent(Event.Event);
 	}
 	FPlatformMisc::bChachedMacMenuStateNeedsUpdate = true;
 }
diff --git a/Engine/Source/Runtime/Core/Public/Mac/MacApplication.h b/Engine/Source/Runtime/Core/Public/Mac/MacApplication.h
--- a/Engine/Source/Runtime/Core/Public/Mac/MacApplication.h
+++ b/Engine/Source/Runtime/Core/Public/Mac/MacApplication.h
@@ -10,7 +10,8 @@
 struct FDeferredMacEvent
 {
 	FDeferredMacEvent()
-	:	Type(0)
+	:	Event(nullptr)
+	,	Type(0)
 	,	LocationInWindow(FVector2D::ZeroVector)
 	,	ModifierFlags(0)
 	,	Timestamp(0.0)
@@ -33,7 +34,8 @@ struct FDeferredMacEvent
 	}
 
 	FDeferredMacEvent(const FDeferredMacEvent& Other)
-	:	Window(Other.Window)
+	:	Event(Other.Event ? [Other.Event retain] : nullptr)
+	,	Window(Other.Window)
 	,	Type(Other.Type)
 	,	LocationInWindow(Other.LocationInWindow)
 	,	ModifierFlags(Other.ModifierFlags)
@@ -59,6 +61,10 @@ struct FDeferredMacEvent
 	~FDeferredMacEvent()
 	{
 		SCOPED_AUTORELEASE_POOL;
+		if (Event)
+		{
+			[Event release];
+		}
 		if (Context)
 		{
 			[Context release];
@@ -81,19 +87,9 @@ struct FDeferredMacEvent
 		}
 	}
 
-	NSEvent* GetKeyNSEvent() const
-	{
-		return [NSEvent keyEventWithType:(NSEventType)Type
-								location:NSMakePoint(LocationInWindow.X, LocationInWindow.Y)
-						   modifierFlags:ModifierFlags
-							   timestamp:Timestamp
-							windowNumber:WindowNumber
-								 context:Context
-							  characters:Characters
-			 charactersIgnoringModifiers:CharactersIgnoringModifiers
-							   isARepeat:IsRepeat
-								 keyCode:KeyCode];
-	}
+	// Using NSEvent on the game thread is unsafe, so we copy of all its properties and use them when processing the event.
+	// However, in some cases we need the original NSEvent (highlighting menus, resending unhandled key events), so we store it as well.
+	NSEvent* Event;
 
 	TSharedPtr<FMacWindow> Window;
 
