Linux: handle Windows-style paths (e.g. from assets) gracefully.

- PR #933 contributed by [REDACTED_PERSON] ([REDACTED_PERSON]).
- Fixes a crash (assertion failure) on right clicking over some assets in the browser.

diff --git a/Engine/Source/Runtime/Core/Private/Linux/LinuxPlatformFile.cpp b/Engine/Source/Runtime/Core/Private/Linux/LinuxPlatformFile.cpp
--- a/Engine/Source/Runtime/Core/Private/Linux/LinuxPlatformFile.cpp
+++ b/Engine/Source/Runtime/Core/Private/Linux/LinuxPlatformFile.cpp
@@ -457,6 +457,13 @@ public:
 	{
 		// Cannot log anything here, as this may result in infinite recursion when this function is called on log file itself
 
+		// We can get some "absolute" filenames like "D:/Blah/" here (e.g. non-Linux paths to source files embedded in assets).
+		// In that case, fail silently.
+		if (PossiblyWrongFilename[0] != TEXT('/'))
+		{
+			return false;
+		}
+
 		// try the filename as given first
 		struct stat StatInfo;
 		bool bFound = stat(TCHAR_TO_UTF8(*PossiblyWrongFilename), &StatInfo) == 0;
@@ -467,9 +474,7 @@ public:
 		}
 		else
 		{
-			// perform a case-insensitive search
-			// make sure we get the absolute filename
-			checkf(PossiblyWrongFilename[0] == TEXT('/'), TEXT("Filename '%s' given to OpenCaseInsensitiveRead is not absolute!"), *PossiblyWrongFilename);
+			// perform a case-insensitive search from /
 
 			int MaxPathComponents = CountPathComponents(PossiblyWrongFilename);
 			if (MaxPathComponents > 0)
