added comment

https://jira.ol.epicgames.net/browse/UE-10807

diff --git a/Engine/Source/Runtime/Core/Public/Math/Quat.h b/Engine/Source/Runtime/Core/Public/Math/Quat.h
--- a/Engine/Source/Runtime/Core/Public/Math/Quat.h
+++ b/Engine/Source/Runtime/Core/Public/Math/Quat.h
@@ -311,6 +311,9 @@ public:
 	 */
 	CORE_API FQuat Exp() const;
 
+	/**
+	 * @return inverse of this quaternion
+	 */
 	FORCEINLINE FQuat Inverse() const;
 
 	/**
