FTransform::Inverse implementation updated to be the same as the vectorized version.
GetSafeScaleReciprocal functions made static.

diff --git a/Engine/Source/Runtime/Core/Public/Math/Transform.h b/Engine/Source/Runtime/Core/Public/Math/Transform.h
--- a/Engine/Source/Runtime/Core/Public/Math/Transform.h
+++ b/Engine/Source/Runtime/Core/Public/Math/Transform.h
@@ -288,8 +288,11 @@ public:
 	*/
 	FORCEINLINE FTransform Inverse() const
 	{
-		// todo: optimize
-		return FTransform(ToMatrixWithScale().Inverse());
+		FQuat   InvRotation    = Rotation.Inverse();
+		FVector InvScale3D     = Scale3D.Reciprocal();
+		FVector InvTranslation = InvRotation * (InvScale3D * -Translation);
+
+		return FTransform(InvRotation, InvTranslation, InvScale3D);
 	}
 
 	/**
@@ -536,7 +539,7 @@ public:
 	FORCEINLINE FVector		GetScaledAxis(EAxis::Type InAxis) const;
 	FORCEINLINE FVector		GetUnitAxis(EAxis::Type InAxis) const;
 	FORCEINLINE void		Mirror(EAxis::Type MirrorAxis, EAxis::Type FlipAxis);
-	FORCEINLINE FVector		GetSafeScaleReciprocal(const FVector& InScale) const;
+	static FORCEINLINE FVector		GetSafeScaleReciprocal(const FVector& InScale);
 
 	// temp function for easy conversion
 	FORCEINLINE FVector GetLocation() const
@@ -1298,7 +1301,7 @@ inline float FTransform::GetMinimumAxisScale() const
 // anymore because you should be instead of showing gigantic infinite mesh
 // also returning BIG_NUMBER causes sequential NaN issues by multiplying 
 // so we hardcode as 0
-FORCEINLINE FVector FTransform::GetSafeScaleReciprocal(const FVector& InScale) const
+FORCEINLINE FVector FTransform::GetSafeScaleReciprocal(const FVector& InScale)
 {
 	FVector SafeReciprocalScale;
 	// mathematically if you have 0 scale, it should be infinite, 
diff --git a/Engine/Source/Runtime/Core/Public/Math/TransformVectorized.h b/Engine/Source/Runtime/Core/Public/Math/TransformVectorized.h
--- a/Engine/Source/Runtime/Core/Public/Math/TransformVectorized.h
+++ b/Engine/Source/Runtime/Core/Public/Math/TransformVectorized.h
@@ -1226,7 +1226,7 @@ private:
 	 * also returning BIG_NUMBER causes sequential NaN issues by multiplying 
 	 * so we hardcode as 0
 	 */
-	FORCEINLINE VectorRegister		GetSafeScaleReciprocal(const VectorRegister& InScale, const ScalarRegister& Tolerance = ScalarRegister(VectorZero())) const
+	static FORCEINLINE VectorRegister		GetSafeScaleReciprocal(const VectorRegister& InScale, const ScalarRegister& Tolerance = ScalarRegister(VectorZero()))
 	{		
 		// SafeReciprocalScale.X = (InScale.X == 0) ? 0.f : 1/InScale.X; // same for YZW
 		VectorRegister SafeReciprocalScale;
