Linux: reduce communication with display server.

- cache border sizes until the next time the window is shown.
- cache cursor position until the next mouse move or position change.

As of this commit, xtrace shows only OpenGL context changes when we aren't moving the mouse.

diff --git a/Engine/Source/Runtime/Core/Private/Linux/LinuxApplication.cpp b/Engine/Source/Runtime/Core/Private/Linux/LinuxApplication.cpp
--- a/Engine/Source/Runtime/Core/Private/Linux/LinuxApplication.cpp
+++ b/Engine/Source/Runtime/Core/Private/Linux/LinuxApplication.cpp
@@ -246,6 +246,7 @@ void FLinuxApplication::ProcessDeferredMessage( SDL_Event Event )
 		{
 			SDL_MouseMotionEvent motionEvent = Event.motion;
 			FLinuxCursor *LinuxCursor = (FLinuxCursor*)Cursor.Get();
+			LinuxCursor->InvalidateCaches();
 
 			if (LinuxCursor->IsHidden())
 			{
@@ -660,6 +661,9 @@ void FLinuxApplication::ProcessDeferredMessage( SDL_Event Event )
 
 				case SDL_WINDOWEVENT_SHOWN:
 					{
+						// (re)cache native properties
+						CurrentEventWindow->CacheNativeProperties();
+
 						if (CurrentEventWindow->IsRegularWindow())
 						{
 							CurrentEventWindow->SetWindowFocus();
@@ -681,16 +685,12 @@ void FLinuxApplication::ProcessDeferredMessage( SDL_Event Event )
 					{
 						int32 ClientScreenX = windowEvent.data1;
 						int32 ClientScreenY = windowEvent.data2;
-						SDL_Rect Borders;
-						if (SDL_GetWindowBordersSize(NativeWindow, &Borders) == 0)
-						{
-							ClientScreenX += Borders.x;
-							ClientScreenY += Borders.y;
-						}
-						else
-						{
-							UE_LOG(LogLinuxWindow, Verbose, TEXT("Could not get Window border sizes!"));
-						}
+
+						int32 BorderSizeX, BorderSizeY;
+						CurrentEventWindow->GetNativeBordersSize(BorderSizeX, BorderSizeY);
+						ClientScreenX += BorderSizeX;
+						ClientScreenY += BorderSizeY;
+
 						MessageHandler->OnMovedWindow(CurrentEventWindow.ToSharedRef(), ClientScreenX, ClientScreenY);
 					}
 					break;
diff --git a/Engine/Source/Runtime/Core/Private/Linux/LinuxCursor.cpp b/Engine/Source/Runtime/Core/Private/Linux/LinuxCursor.cpp
--- a/Engine/Source/Runtime/Core/Private/Linux/LinuxCursor.cpp
+++ b/Engine/Source/Runtime/Core/Private/Linux/LinuxCursor.cpp
@@ -9,6 +9,9 @@ FLinuxCursor::FLinuxCursor()
 	: 	bHidden(false)
 	,	AccumulatedOffsetX(0)
 	,	AccumulatedOffsetY(0)
+	,	CachedGlobalXPosition(0)
+	,	CachedGlobalYPosition(0)
+	,	bPositionCacheIsValid(false)
 {
 	if (!FPlatformMisc::PlatformInitMultimedia()) //	will not initialize more than once
 	{
@@ -162,15 +165,26 @@ FLinuxCursor::~FLinuxCursor()
 
 FVector2D FLinuxCursor::GetPosition() const
 {
-	int CursorX, CursorY;
+	if (!bPositionCacheIsValid)
+	{
+		SDL_GetGlobalMouseState(&CachedGlobalXPosition, &CachedGlobalYPosition);
+		bPositionCacheIsValid = true;
+	}
 
-	SDL_GetGlobalMouseState(&CursorX, &CursorY);
+	int CursorX = CachedGlobalXPosition;
+	int CursorY = CachedGlobalYPosition;
+	
 	CursorX += AccumulatedOffsetX;
 	CursorY += AccumulatedOffsetY;
 
 	return FVector2D( CursorX, CursorY );
 }
 
+void FLinuxCursor::InvalidateCaches()
+{
+	bPositionCacheIsValid = false;
+}
+
 void FLinuxCursor::SetPosition( const int32 X, const int32 Y )
 {
 	int WndX, WndY;
@@ -179,6 +193,8 @@ void FLinuxCursor::SetPosition( const int32 X, const int32 Y )
 
 	SDL_GetWindowPosition( WndFocus, &WndX, &WndY );	//	get top left
 	SDL_WarpMouseInWindow( NULL, X - WndX, Y - WndY );
+
+	InvalidateCaches();
 }
 
 void FLinuxCursor::AddOffset(const int32 DX, const int32 DY)
diff --git a/Engine/Source/Runtime/Core/Private/Linux/LinuxWindow.cpp b/Engine/Source/Runtime/Core/Private/Linux/LinuxWindow.cpp
--- a/Engine/Source/Runtime/Core/Private/Linux/LinuxWindow.cpp
+++ b/Engine/Source/Runtime/Core/Private/Linux/LinuxWindow.cpp
@@ -32,6 +32,9 @@ FLinuxWindow::FLinuxWindow()
 	, bIsPopupWindow(false)
 	, bIsTooltipWindow(false)
 	, bIsPointerInsideWindow(false)
+	, LeftBorderWidth(0)
+	, TopBorderHeight(0)
+	, bValidNativePropertiesCache(false)
 {
 	PreFullscreenWindowRect.left = PreFullscreenWindowRect.top = PreFullscreenWindowRect.right = PreFullscreenWindowRect.bottom = 0;
 }
@@ -137,6 +140,9 @@ void FLinuxWindow::Initialize( FLinuxApplication* const Application, const TShar
 	VirtualWidth  = ClientWidth;
 	VirtualHeight = ClientHeight;
 
+	// attempt to early cache native properties
+	CacheNativeProperties();
+
 	// We call reshape window here because we didn't take into account the non-client area
 	// in the initial creation of the window. Slate should only pass client area dimensions.
 	// Reshape window may resize the window if the non-client area is encroaching on our
@@ -197,13 +203,9 @@ SDL_HitTestResult FLinuxWindow::HitTest( SDL_Window *SDLwin, const SDL_Point *po
 void FLinuxWindow::MoveWindowTo( int32 X, int32 Y )
 {
 	// we are passed coordinates of a client area, so account for decorations
-	SDL_Rect Borders;
-	if (SDL_GetWindowBordersSize(HWnd, &Borders) == 0)
-	{
-		X -= Borders.x;
-		Y -= Borders.y;
-	}
-    SDL_SetWindowPosition( HWnd, X, Y );
+	checkf(bValidNativePropertiesCache, TEXT("Attempted to use border sizes too early, native properties aren't yet cached. Review the flow"));
+
+	SDL_SetWindowPosition( HWnd, X - LeftBorderWidth, Y - TopBorderHeight );
 }
 
 /** Native windows should implement BringToFront by making this window the top-most window (i.e. focused).
@@ -328,12 +330,9 @@ void FLinuxWindow::ReshapeWindow( int32 NewX, int32 NewY, int32 NewWidth, int32
 			if (Definition->HasOSWindowBorder)
 			{
 				// we are passed coordinates of a client area, so account for decorations
-				SDL_Rect Borders;
-				if (SDL_GetWindowBordersSize(HWnd, &Borders) == 0)
-				{
-					NewX -= Borders.x;
-					NewY -= Borders.y;
-				}
+				checkf(bValidNativePropertiesCache, TEXT("Attempted to use border sizes too early, native properties aren't yet cached. Review the flow"));
+				NewX -= LeftBorderWidth;
+				NewY -= TopBorderHeight;
 			}
 			SDL_SetWindowPosition( HWnd, NewX, NewY );
 			SDL_SetWindowSize( HWnd, NewWidth, NewHeight );
@@ -575,3 +574,23 @@ bool FLinuxWindow::IsPointerInsideWindow() const
 {
 	return bIsPointerInsideWindow;
 }
+
+void FLinuxWindow::GetNativeBordersSize(int32& OutLeftBorderWidth, int32& OutTopBorderHeight) const
+{
+	checkf(bValidNativePropertiesCache, TEXT("Attempted to get border sizes too early, native properties aren't yet cached. Review the flow"));
+	OutLeftBorderWidth = LeftBorderWidth;
+	OutTopBorderHeight = TopBorderHeight;
+}
+
+void FLinuxWindow::CacheNativeProperties()
+{
+	// cache border sizes
+	SDL_Rect Borders;
+	if (SDL_GetWindowBordersSize(HWnd, &Borders) == 0)
+	{
+		LeftBorderWidth = Borders.x;
+		TopBorderHeight = Borders.y;
+	}
+
+	bValidNativePropertiesCache = true;
+}
diff --git a/Engine/Source/Runtime/Core/Public/Linux/LinuxCursor.h b/Engine/Source/Runtime/Core/Public/Linux/LinuxCursor.h
--- a/Engine/Source/Runtime/Core/Public/Linux/LinuxCursor.h
+++ b/Engine/Source/Runtime/Core/Public/Linux/LinuxCursor.h
@@ -47,6 +47,11 @@ public:
 	/** Resets accumulated offsets to (0,0), */
 	void ResetOffset();
 
+	/**
+	 * Invalidates whatever cached data cursor may have
+	 */
+	void InvalidateCaches();
+
 private:
 
 	EMouseCursor::Type CurrentType;
@@ -64,4 +69,13 @@ private:
 
 	/** Accumulated Y offset (from relative mouse movement events) used for reconstructing global cursor position */
 	int32 AccumulatedOffsetY;
+
+	/** Cached global X position */
+	mutable int32 CachedGlobalXPosition;
+
+	/** Cached global Y position */
+	mutable int32 CachedGlobalYPosition;
+
+	/** Whether mouse position cache is valid */
+	mutable bool bPositionCacheIsValid;
 };
diff --git a/Engine/Source/Runtime/Core/Public/Linux/LinuxWindow.h b/Engine/Source/Runtime/Core/Public/Linux/LinuxWindow.h
--- a/Engine/Source/Runtime/Core/Public/Linux/LinuxWindow.h
+++ b/Engine/Source/Runtime/Core/Public/Linux/LinuxWindow.h
@@ -151,7 +151,21 @@ public:
 	 * @return Returns 'true' if the pointer is inside the window else 'false'.
 	 */
 	bool IsPointerInsideWindow() const;
-	
+
+	/**
+	 * Returns window border sizes (frame extents).
+	 * 
+	 * @param BorderWidth - width of a vertical border (left side)
+	 * @param BorderHeight - height of a horizontal border (from the top)
+	 */
+	void GetNativeBordersSize(int32& OutLeftBorderWidth, int32& OutTopBorderHeight) const;
+
+	/**
+	 *  Caches window properties that are too expensive to query all the time.
+	 *  (e.g. border sizes)
+	 */
+	void CacheNativeProperties();
+
 private:
 
 	/**
@@ -210,4 +224,13 @@ private:
 	/** Parent window as given by Slate */
 	TSharedPtr< FLinuxWindow > ParentWindow;
 	static SDL_HitTestResult HitTest( SDL_Window *SDLwin, const SDL_Point *point, void *data );
+
+	/** Cached width of the left border */
+	int32 LeftBorderWidth;
+
+	/** Cached height of the top border */
+	int32 TopBorderHeight;
+
+	/** Whether native properties cache is valid */
+	bool bValidNativePropertiesCache;
 };
