Core: cache results of runtime build settings discovery.

diff --git a/Engine/Source/Runtime/Core/Private/Misc/EngineBuildSettings.cpp b/Engine/Source/Runtime/Core/Private/Misc/EngineBuildSettings.cpp
--- a/Engine/Source/Runtime/Core/Private/Misc/EngineBuildSettings.cpp
+++ b/Engine/Source/Runtime/Core/Private/Misc/EngineBuildSettings.cpp
@@ -6,13 +6,15 @@
 
 bool FEngineBuildSettings::IsInternalBuild()
 {
-	return FPaths::FileExists( FPaths::EngineDir() / TEXT("Build/NotForLicensees/EpicInternal.txt") );
+	static bool bIsInternalBuild = FPaths::FileExists( FPaths::EngineDir() / TEXT("Build/NotForLicensees/EpicInternal.txt") );
+	return bIsInternalBuild;
 }
 
 
 bool FEngineBuildSettings::IsPerforceBuild()
 {
-	return FPaths::FileExists( FPaths::EngineDir() / TEXT("Build/PerforceBuild.txt") );
+	static bool bIsPerforceBuild = FPaths::FileExists(FPaths::EngineDir() / TEXT("Build/PerforceBuild.txt"));
+	return bIsPerforceBuild;
 }
 
 
@@ -24,5 +26,6 @@ bool FEngineBuildSettings::IsSourceDistribution()
 
 bool FEngineBuildSettings::IsSourceDistribution(const FString& RootDir)
 {
-	return FPaths::FileExists( RootDir / TEXT("Engine/Build/SourceDistribution.txt") );
+	static bool bIsSourceDistribution = FPaths::FileExists(RootDir / TEXT("Engine/Build/SourceDistribution.txt"));
+	return bIsSourceDistribution;
 }
