Text Localization Manager:
- Fix for issue where the requested culture was being compared against the prioritized parent cultures of the _current_ culture, instead of the _requested_ culture, which hadn't yet become the current culture.

diff --git a/Engine/Source/Runtime/Core/Private/Internationalization/TextLocalizationManager.cpp b/Engine/Source/Runtime/Core/Private/Internationalization/TextLocalizationManager.cpp
--- a/Engine/Source/Runtime/Core/Private/Internationalization/TextLocalizationManager.cpp
+++ b/Engine/Source/Runtime/Core/Private/Internationalization/TextLocalizationManager.cpp
@@ -90,16 +90,20 @@ void EndInitTextLocalization()
 			TArray< FCultureRef > AvailableCultures;
 			I18N.GetCulturesWithAvailableLocalization(LocalizationPaths, AvailableCultures, false);
 
-			TArray<FString> PrioritizedParentCultureNames = I18N.GetCurrentCulture()->GetPrioritizedParentCultureNames();
-				
 			FString ValidCultureName;
-			for (const FString& CultureName : PrioritizedParentCultureNames)
+			ValidCultureName.Empty();
+			FCulturePtr TargetCulture = I18N.GetCulture(TargetCultureName);
+			if (TargetCulture.IsValid())
 			{
-				FCulturePtr ValidCulture = I18N.GetCulture(CultureName);
-				if (ValidCulture.IsValid() && AvailableCultures.Contains(ValidCulture.ToSharedRef()))
+				TArray<FString> PrioritizedParentCultureNames = TargetCulture->GetPrioritizedParentCultureNames();
+				for (const FString& CultureName : PrioritizedParentCultureNames)
 				{
-					ValidCultureName = CultureName;
-					break;
+					FCulturePtr ValidCulture = I18N.GetCulture(CultureName);
+					if (ValidCulture.IsValid() && AvailableCultures.Contains(ValidCulture.ToSharedRef()))
+					{
+						ValidCultureName = CultureName;
+						break;
+					}
 				}
 			}
 
@@ -109,6 +113,7 @@ void EndInitTextLocalization()
 				{
 					// Make the user aware that the localization data belongs to a parent culture.
 					UE_LOG(LogTextLocalizationManager, Log, TEXT("The requested culture ('%s') has no localization data; parent culture's ('%s') localization data will be used."), *RequestedCultureName, *ValidCultureName);
+					TargetCultureName = ValidCultureName;
 				}
 			}
 			else
