Stop the Mac window moving around when exiting fullscreen again after recent changes - forcing the window back to its original position in windowDidExitFullscreen is no longer correct. Simplified the calculation of window size & location, excluding the system titlebar if there is one.

diff --git a/Engine/Source/Runtime/Core/Private/Mac/CocoaWindow.cpp b/Engine/Source/Runtime/Core/Private/Mac/CocoaWindow.cpp
--- a/Engine/Source/Runtime/Core/Private/Mac/CocoaWindow.cpp
+++ b/Engine/Source/Runtime/Core/Private/Mac/CocoaWindow.cpp
@@ -198,15 +198,20 @@ NSString* NSPerformDragOperation = @"NSPerformDragOperation";
 	// Handle clicking on the titlebar fullscreen item
 	if (self.TargetWindowMode == EWindowMode::Windowed)
 	{
-		self.TargetWindowMode = EWindowMode::Fullscreen;
-#if WITH_EDITORONLY_DATA
+		self.PreFullScreenRect.origin = [self frame].origin;
+		self.PreFullScreenRect.size = [self openGLFrame].size;
+		
+		// Use the current default fullscreen mode when switching via the OS button
+		auto CVar = IConsoleManager::Get().FindTConsoleVariableDataInt(TEXT("r.FullScreenMode"));
+		check(CVar);
+		self.TargetWindowMode = CVar->GetValueOnAnyThread() == 0 ? EWindowMode::Fullscreen : EWindowMode::WindowedFullscreen;
+		
+#if WITH_EDITORONLY_DATA // Always use WindowedFullscreen for the Editor or bad things happen
 		if (GIsEditor)
 		{
 			self.TargetWindowMode = EWindowMode::WindowedFullscreen;
 		}
 #endif
-		self.PreFullScreenRect.origin = [self frame].origin;
-		self.PreFullScreenRect.size = [self openGLFrame].size;
 	}
 }
 
@@ -242,7 +247,6 @@ NSString* NSPerformDragOperation = @"NSPerformDragOperation";
 
 	MacApplication->DeferEvent(Notification);
 
-	[self setFrame:self.PreFullScreenRect display:YES];
 	FMacCursor* MacCursor = (FMacCursor*)MacApplication->Cursor.Get();
 	if (MacCursor)
 	{
diff --git a/Engine/Source/Runtime/Core/Private/Mac/MacApplication.cpp b/Engine/Source/Runtime/Core/Private/Mac/MacApplication.cpp
--- a/Engine/Source/Runtime/Core/Private/Mac/MacApplication.cpp
+++ b/Engine/Source/Runtime/Core/Private/Mac/MacApplication.cpp
@@ -891,7 +891,7 @@ void FMacApplication::OnWindowDidMove(TSharedRef<FMacWindow> Window)
 
 	if ([Window->GetWindowHandle() windowMode] != EWindowMode::Fullscreen)
 	{
-		Y = FPlatformMisc::ConvertCocoaYPositionToSlate(WindowFrame.origin.y + WindowFrame.size.height) + (OpenGLFrame.size.height - WindowFrame.size.height);
+		Y = FPlatformMisc::ConvertCocoaYPositionToSlate(WindowFrame.origin.y + OpenGLFrame.size.height);
 	}
 
 	MessageHandler->OnMovedWindow(Window, X, Y);
