Do not update Mac menu bar while processing deferred events

diff --git a/Engine/Source/Runtime/Core/Private/Mac/MacApplication.cpp b/Engine/Source/Runtime/Core/Private/Mac/MacApplication.cpp
--- a/Engine/Source/Runtime/Core/Private/Mac/MacApplication.cpp
+++ b/Engine/Source/Runtime/Core/Private/Mac/MacApplication.cpp
@@ -36,6 +36,8 @@ FMacApplication::FMacApplication()
 :	GenericApplication(MakeShareable(new FMacCursor()))
 ,	bUsingHighPrecisionMouseInput(false)
 ,	bUsingTrackpad(false)
+,	LastPressedMouseButton(EMouseButtons::Invalid)
+,	bIsProcessingDeferredEvents(false)
 ,	HIDInput(HIDInputInterface::Create(MessageHandler))
 ,	DraggedWindow(nullptr)
 ,	bSystemModalMode(false)
@@ -147,11 +149,15 @@ void FMacApplication::ProcessDeferredEvents(const float TimeDelta)
 	DeferredEvents.Empty();
 	EventsMutex.Unlock();
 
+	bIsProcessingDeferredEvents = true;
+
 	for (int32 Index = 0; Index < EventsToProcess.Num(); ++Index)
 	{
 		ProcessEvent(EventsToProcess[Index]);
 	}
 
+	bIsProcessingDeferredEvents = false;
+
 	InvalidateTextLayouts();
 	CloseQueuedWindows();
 }
diff --git a/Engine/Source/Runtime/Core/Private/Mac/MacPlatformMisc.cpp b/Engine/Source/Runtime/Core/Private/Mac/MacPlatformMisc.cpp
--- a/Engine/Source/Runtime/Core/Private/Mac/MacPlatformMisc.cpp
+++ b/Engine/Source/Runtime/Core/Private/Mac/MacPlatformMisc.cpp
@@ -579,7 +579,7 @@ void FMacPlatformMisc::PumpMessages( bool bFromMainLoop )
 	{
 		ProcessGameThreadEvents();
 
-		if (MacApplication && IsInGameThread())
+		if (MacApplication && !MacApplication->IsProcessingDeferredEvents() && IsInGameThread())
 		{
 			if (UpdateCachedMacMenuState && bChachedMacMenuStateNeedsUpdate)
 			{
diff --git a/Engine/Source/Runtime/Core/Public/Mac/MacApplication.h b/Engine/Source/Runtime/Core/Public/Mac/MacApplication.h
--- a/Engine/Source/Runtime/Core/Public/Mac/MacApplication.h
+++ b/Engine/Source/Runtime/Core/Public/Mac/MacApplication.h
@@ -170,6 +170,8 @@ public:
 
 	void DeferEvent(NSObject* Object);
 
+	bool IsProcessingDeferredEvents() const { return bIsProcessingDeferredEvents; }
+
 	TSharedPtr<FMacWindow> FindWindowByNSWindow(FCocoaWindow* WindowHandle);
 
 	/** Queues a window for text layout invalidation when safe */
@@ -247,6 +249,8 @@ private:
 	FCriticalSection WindowsMutex;
 	TArray<TSharedRef<FMacWindow>> Windows;
 
+	bool bIsProcessingDeferredEvents;
+
 	struct FSavedWindowOrderInfo
 	{
 		int32 WindowNumber;
