Don't allow the use of the native OS X fullscreen on/off titlebar widget as it breaks SceneViewport's mechanism for entering/exiting fullscreen. We'd need to be able to pass a notification into Slate that SceneViewport could register for that could then be used to update its internal WindowMode & PreFullScreenPos variables. This also fixes the editor's green button so that it behaves consistently.

diff --git a/Engine/Source/Runtime/Core/Private/Mac/CocoaWindow.cpp b/Engine/Source/Runtime/Core/Private/Mac/CocoaWindow.cpp
--- a/Engine/Source/Runtime/Core/Private/Mac/CocoaWindow.cpp
+++ b/Engine/Source/Runtime/Core/Private/Mac/CocoaWindow.cpp
@@ -244,6 +244,12 @@ NSString* NSPerformDragOperation = @"NSPerformDragOperation";
 {
 	WindowMode = EWindowMode::Windowed;
 	self.TargetWindowMode = EWindowMode::Windowed;
+	
+	// Remove any primary fullscreen behaviour here, we don't support it properly.
+	NSWindowCollectionBehavior Behaviour = [self collectionBehavior];
+	Behaviour &= ~(NSWindowCollectionBehaviorFullScreenPrimary);
+	Behaviour |= NSWindowCollectionBehaviorFullScreenAuxiliary;
+	[self setCollectionBehavior: Behaviour];
 
 	MacApplication->DeferEvent(Notification);
 
diff --git a/Engine/Source/Runtime/Core/Private/Mac/MacWindow.cpp b/Engine/Source/Runtime/Core/Private/Mac/MacWindow.cpp
--- a/Engine/Source/Runtime/Core/Private/Mac/MacWindow.cpp
+++ b/Engine/Source/Runtime/Core/Private/Mac/MacWindow.cpp
@@ -145,7 +145,7 @@ void FMacWindow::Initialize( FMacApplication* const Application, const TSharedRe
 
 				if( Definition->HasOSWindowBorder )
 				{
-					[WindowHandle setCollectionBehavior: NSWindowCollectionBehaviorFullScreenPrimary|NSWindowCollectionBehaviorDefault|NSWindowCollectionBehaviorManaged|NSWindowCollectionBehaviorParticipatesInCycle];
+					[WindowHandle setCollectionBehavior: NSWindowCollectionBehaviorDefault|NSWindowCollectionBehaviorManaged|NSWindowCollectionBehaviorParticipatesInCycle];
 				}
 				else
 				{
@@ -414,12 +414,11 @@ void FMacWindow::SetWindowMode( EWindowMode::Type NewWindowMode )
 			FPlatformMisc::PumpMessages(true);
 			WindowIsFullScreen = [WindowHandle windowMode] != EWindowMode::Windowed;
 		} while(WindowIsFullScreen != bMakeFullscreen);
-		
-		if(!bMakeFullscreen && !Definition->HasOSWindowBorder)
-		{
-			Behaviour &= ~(NSWindowCollectionBehaviorFullScreenPrimary);
-			Behaviour |= NSWindowCollectionBehaviorFullScreenAuxiliary;
-		}
+	}
+	else // Already in/out fullscreen but a different mode - we should just update the mode rather than forcing the window to change again
+	{
+		WindowHandle.TargetWindowMode = NewWindowMode;
+		[WindowHandle setWindowMode:NewWindowMode];
 	}
 }
 
