Check if MacApplication is valid when deferring window notifications

diff --git a/Engine/Source/Runtime/Core/Private/Mac/CocoaWindow.cpp b/Engine/Source/Runtime/Core/Private/Mac/CocoaWindow.cpp
--- a/Engine/Source/Runtime/Core/Private/Mac/CocoaWindow.cpp
+++ b/Engine/Source/Runtime/Core/Private/Mac/CocoaWindow.cpp
@@ -219,7 +219,10 @@ NSString* NSPerformDragOperation = @"NSPerformDragOperation";
 {
 	WindowMode = self.TargetWindowMode;
 
-	MacApplication->DeferEvent(Notification);
+	if (MacApplication)
+	{
+		MacApplication->DeferEvent(Notification);
+	}
 
 	FMacCursor* MacCursor = (FMacCursor*)MacApplication->Cursor.Get();
 	if (MacCursor)
@@ -251,7 +254,10 @@ NSString* NSPerformDragOperation = @"NSPerformDragOperation";
 	Behaviour |= NSWindowCollectionBehaviorFullScreenAuxiliary;
 	[self setCollectionBehavior: Behaviour];
 
-	MacApplication->DeferEvent(Notification);
+	if (MacApplication)
+	{
+		MacApplication->DeferEvent(Notification);
+	}
 
 	FMacCursor* MacCursor = (FMacCursor*)MacApplication->Cursor.Get();
 	if (MacCursor)
@@ -268,7 +274,10 @@ NSString* NSPerformDragOperation = @"NSPerformDragOperation";
 		[self orderFrontAndMakeMain:false andKey:false];
 	}
 
-	MacApplication->DeferEvent(Notification);
+	if (MacApplication)
+	{
+		MacApplication->DeferEvent(Notification);
+	}
 }
 
 - (void)windowDidResignMain:(NSNotification*)Notification
@@ -277,12 +286,18 @@ NSString* NSPerformDragOperation = @"NSPerformDragOperation";
 	[self setMovable: YES];
 	[self setMovableByWindowBackground: NO];
 
-	MacApplication->DeferEvent(Notification);
+	if (MacApplication)
+	{
+		MacApplication->DeferEvent(Notification);
+	}
 }
 
 - (void)windowWillMove:(NSNotification*)Notification
 {
-	MacApplication->DeferEvent(Notification);
+	if (MacApplication)
+	{
+		MacApplication->DeferEvent(Notification);
+	}
 }
 
 - (void)windowDidMove:(NSNotification*)Notification
@@ -293,7 +308,10 @@ NSString* NSPerformDragOperation = @"NSPerformDragOperation";
 	NSView* OpenGLView = [self openGLView];
 	[[NSNotificationCenter defaultCenter] postNotificationName:NSViewGlobalFrameDidChangeNotification object:OpenGLView];
 	
-	MacApplication->DeferEvent(Notification);
+	if (MacApplication)
+	{
+		MacApplication->DeferEvent(Notification);
+	}
 }
 
 - (void)windowDidChangeScreen:(NSNotification*)Notification
@@ -361,20 +379,29 @@ NSString* NSPerformDragOperation = @"NSPerformDragOperation";
 - (void)windowWillStartLiveResize:(NSNotification*)Notification
 {
 	SCOPED_AUTORELEASE_POOL;
-	MacApplication->DeferEvent(Notification);
+	if (MacApplication)
+	{
+		MacApplication->DeferEvent(Notification);
+	}
 }
 
 - (void)windowDidEndLiveResize:(NSNotification*)Notification
 {
 	SCOPED_AUTORELEASE_POOL;
-	MacApplication->DeferEvent(Notification);
+	if (MacApplication)
+	{
+		MacApplication->DeferEvent(Notification);
+	}
 }
 
 - (void)windowDidResize:(NSNotification*)Notification
 {
 	SCOPED_AUTORELEASE_POOL;
 	bZoomed = [self isZoomed];
-	MacApplication->DeferEvent(Notification);
+	if (MacApplication)
+	{
+		MacApplication->DeferEvent(Notification);
+	}
 }
 
 - (void)windowWillClose:(NSNotification*)Notification
@@ -396,14 +423,20 @@ NSString* NSPerformDragOperation = @"NSPerformDragOperation";
 {
 	SCOPED_AUTORELEASE_POOL;
 	NSNotification* Notification = [NSNotification notificationWithName:NSDraggingExited object:Sender];
-	MacApplication->DeferEvent(Notification);
+	if (MacApplication)
+	{
+		MacApplication->DeferEvent(Notification);
+	}
 }
 
 - (NSDragOperation)draggingUpdated:(id <NSDraggingInfo>)Sender
 {
 	SCOPED_AUTORELEASE_POOL;
 	NSNotification* Notification = [NSNotification notificationWithName:NSDraggingUpdated object:Sender];
-	MacApplication->DeferEvent(Notification);
+	if (MacApplication)
+	{
+		MacApplication->DeferEvent(Notification);
+	}
 	return NSDragOperationGeneric;
 }
 
@@ -411,7 +444,10 @@ NSString* NSPerformDragOperation = @"NSPerformDragOperation";
 {
 	SCOPED_AUTORELEASE_POOL;
 	NSNotification* Notification = [NSNotification notificationWithName:NSPrepareForDragOperation object:Sender];
-	MacApplication->DeferEvent(Notification);
+	if (MacApplication)
+	{
+		MacApplication->DeferEvent(Notification);
+	}
 	return YES;
 }
 
@@ -419,7 +455,10 @@ NSString* NSPerformDragOperation = @"NSPerformDragOperation";
 {
 	SCOPED_AUTORELEASE_POOL;
 	NSNotification* Notification = [NSNotification notificationWithName:NSPerformDragOperation object:Sender];
-	MacApplication->DeferEvent(Notification);
+	if (MacApplication)
+	{
+		MacApplication->DeferEvent(Notification);
+	}
 	return YES;
 }
 
