We now log the IME description when an IME is activated.

UE-11663 - Add additional IME logging to track down problems with external user crashes

diff --git a/Engine/Source/Runtime/Core/Private/Windows/WindowsTextInputMethodSystem.cpp b/Engine/Source/Runtime/Core/Private/Windows/WindowsTextInputMethodSystem.cpp
--- a/Engine/Source/Runtime/Core/Private/Windows/WindowsTextInputMethodSystem.cpp
+++ b/Engine/Source/Runtime/Core/Private/Windows/WindowsTextInputMethodSystem.cpp
@@ -159,18 +159,79 @@ bool FWindowsTextInputMethodSystem::Initialize()
 			check(TSFProfile.hkl == KeyboardLayout);
 
 			CurrentAPI = EAPI::TSF;
-			UE_LOG(LogWindowsTextInputMethodSystem, Verbose, TEXT("IME system now activated using TSF."));
+			LogActiveIMEInfo();
 		}
 		else if(::ImmGetIMEFileName(KeyboardLayout, nullptr, 0) > 0)
 		{
 			CurrentAPI = EAPI::IMM;
-			UE_LOG(LogWindowsTextInputMethodSystem, Verbose, TEXT("IME system now activated using IMM."));
+			LogActiveIMEInfo();
 		}
 	}
 
 	return Result;
 }
 
+void FWindowsTextInputMethodSystem::LogActiveIMEInfo()
+{
+	FString APIString;
+
+	switch(CurrentAPI)
+	{
+	case EAPI::IMM:
+		{
+			APIString = TEXT("IMM");
+
+			// Get the description of the active IME
+			const HKL KeyboardLayout = ::GetKeyboardLayout(0);
+			TArray<TCHAR> DescriptionString;
+			const int32 DescriptionLen = ::ImmGetDescription(KeyboardLayout, nullptr, 0);
+			DescriptionString.SetNumUninitialized(DescriptionLen + 1); // +1 for null
+			::ImmGetDescription(KeyboardLayout, DescriptionString.GetData(), DescriptionLen);
+			DescriptionString[DescriptionLen] = 0;
+
+			if(DescriptionLen > 0)
+			{
+				APIString += TEXT(" (");
+				APIString += DescriptionString.GetData();
+				APIString += TEXT(")");
+			}
+		}
+		break;
+
+	case EAPI::TSF:
+		{
+			APIString = TEXT("TSF");
+
+			TF_INPUTPROCESSORPROFILE TSFProfile;
+			if(SUCCEEDED(TSFInputProcessorProfileManager->GetActiveProfile(GUID_TFCAT_TIP_KEYBOARD, &TSFProfile)) && TSFProfile.dwProfileType == TF_PROFILETYPE_INPUTPROCESSOR)
+			{
+				BSTR TSFDescriptionString;
+				if(SUCCEEDED(TSFInputProcessorProfiles->GetLanguageProfileDescription(TSFProfile.clsid, TSFProfile.langid, TSFProfile.guidProfile, &TSFDescriptionString)))
+				{
+					APIString += TEXT(" (");
+					APIString += TSFDescriptionString;
+					APIString += TEXT(")");
+					::SysFreeString(TSFDescriptionString);
+				}
+			}
+		}
+		break;
+
+	case EAPI::Unknown:
+	default:
+		break;
+	}
+
+	if(APIString.IsEmpty())
+	{
+		UE_LOG(LogWindowsTextInputMethodSystem, Display, TEXT("IME system now deactivated."));
+	}
+	else
+	{
+		UE_LOG(LogWindowsTextInputMethodSystem, Display, TEXT("IME system now activated using %s."), *APIString);
+	}
+}
+
 bool FWindowsTextInputMethodSystem::InitializeIMM()
 {
 	UE_LOG(LogWindowsTextInputMethodSystem, Verbose, TEXT("Initializing IMM..."));
@@ -681,20 +742,22 @@ void FWindowsTextInputMethodSystem::OnIMEActivationStateChanged(const bool bIsEn
 		// It seems that switching away from an IMM based IME doesn't generate a deactivation notification
 		//check(CurrentAPI == EAPI::Unknown);
 
-		FString APIString;
 		const HKL KeyboardLayout = ::GetKeyboardLayout(0);
-		if(::ImmGetIMEFileName(KeyboardLayout, nullptr, 0) > 0)
+
+		TF_INPUTPROCESSORPROFILE TSFProfile;
+		if(SUCCEEDED(TSFInputProcessorProfileManager->GetActiveProfile(GUID_TFCAT_TIP_KEYBOARD, &TSFProfile)) && TSFProfile.dwProfileType == TF_PROFILETYPE_INPUTPROCESSOR)
+		{
+			CurrentAPI = EAPI::TSF;
+		}
+		else if(::ImmGetIMEFileName(KeyboardLayout, nullptr, 0) > 0)
 		{
 			CurrentAPI = EAPI::IMM;
 			UpdateIMMProperty(KeyboardLayout);
-			APIString = TEXT("IMM");
 		}
 		else
 		{
-			CurrentAPI = EAPI::TSF;
-			APIString = TEXT("TSF");
+			CurrentAPI = EAPI::Unknown;
 		}
-		UE_LOG(LogWindowsTextInputMethodSystem, Verbose, TEXT("IME system now activated using %s."), *APIString);
 	}
 	else
 	{
@@ -702,8 +765,9 @@ void FWindowsTextInputMethodSystem::OnIMEActivationStateChanged(const bool bIsEn
 		//check(CurrentAPI != EAPI::Unknown);
 
 		CurrentAPI = EAPI::Unknown;
-		UE_LOG(LogWindowsTextInputMethodSystem, Verbose, TEXT("IME system now deactivated."));
 	}
+
+	LogActiveIMEInfo();
 }
 
 int32 FWindowsTextInputMethodSystem::ProcessMessage(HWND hwnd, uint32 msg, WPARAM wParam, LPARAM lParam)
diff --git a/Engine/Source/Runtime/Core/Public/Windows/WindowsTextInputMethodSystem.h b/Engine/Source/Runtime/Core/Public/Windows/WindowsTextInputMethodSystem.h
--- a/Engine/Source/Runtime/Core/Public/Windows/WindowsTextInputMethodSystem.h
+++ b/Engine/Source/Runtime/Core/Public/Windows/WindowsTextInputMethodSystem.h
@@ -84,6 +84,8 @@ public:
 
 private:
 
+	void LogActiveIMEInfo();
+
 	// IMM Implementation
 
 	bool InitializeIMM();
