Fix issue with failing to open local obb file for read
#ue4
#android

diff --git a/Engine/Source/Runtime/Core/Private/Android/AndroidFile.cpp b/Engine/Source/Runtime/Core/Private/Android/AndroidFile.cpp
--- a/Engine/Source/Runtime/Core/Private/Android/AndroidFile.cpp
+++ b/Engine/Source/Runtime/Core/Private/Android/AndroidFile.cpp
@@ -1686,7 +1686,7 @@ private:
 	void MountOBB(const TCHAR* Filename)
 	{
 		FFileHandleAndroid* File
-			= static_cast<FFileHandleAndroid*>(OpenRead(Filename, true));
+			= static_cast<FFileHandleAndroid*>(OpenRead(Filename, true, false));
 		check(nullptr != File);
 		ZipResource.AddPatchFile(MakeShareable(File));
 		FPlatformMisc::LowLevelOutputDebugStringf(
