TextLocalizationManager: Fixing mistake in CL-2482865. If only parent localization data is found, shouldn't force the parent culture setting in FInternationalization, keep it as the child so things like currency formatting, etc are correct.

diff --git a/Engine/Source/Runtime/Core/Private/Internationalization/TextLocalizationManager.cpp b/Engine/Source/Runtime/Core/Private/Internationalization/TextLocalizationManager.cpp
--- a/Engine/Source/Runtime/Core/Private/Internationalization/TextLocalizationManager.cpp
+++ b/Engine/Source/Runtime/Core/Private/Internationalization/TextLocalizationManager.cpp
@@ -113,7 +113,6 @@ void EndInitTextLocalization()
 				{
 					// Make the user aware that the localization data belongs to a parent culture.
 					UE_LOG(LogTextLocalizationManager, Log, TEXT("The requested culture ('%s') has no localization data; parent culture's ('%s') localization data will be used."), *RequestedCultureName, *ValidCultureName);
-					TargetCultureName = ValidCultureName;
 				}
 			}
 			else
