Heterogeneous comparison for TArray::Contains.
Missing TWeakPtr and TSharedPtr comparison operators added.

diff --git a/Engine/Source/Runtime/Core/Public/Containers/Array.h b/Engine/Source/Runtime/Core/Public/Containers/Array.h
--- a/Engine/Source/Runtime/Core/Public/Containers/Array.h
+++ b/Engine/Source/Runtime/Core/Public/Containers/Array.h
@@ -1081,7 +1081,8 @@ public:
 	 *
 	 * @returns	True if found. False otherwise.
 	 */
-	bool Contains(const ElementType& Item) const
+	template <typename ComparisonType>
+	bool Contains(const ComparisonType& Item) const
 	{
 		for (const ElementType* RESTRICT Data = GetData(), *RESTRICT DataEnd = Data + ArrayNum; Data != DataEnd; ++Data)
 		{
diff --git a/Engine/Source/Runtime/Core/Public/Templates/RefCounting.h b/Engine/Source/Runtime/Core/Public/Templates/RefCounting.h
--- a/Engine/Source/Runtime/Core/Public/Templates/RefCounting.h
+++ b/Engine/Source/Runtime/Core/Public/Templates/RefCounting.h
@@ -104,11 +104,6 @@ public:
 		return *this = InPtr.Reference;
 	}
 
-	bool operator==(const TRefCountPtr& Other) const
-	{
-		return Reference == Other.Reference;
-	}
-
 	ReferencedType* operator->() const
 	{
 		return Reference;
@@ -175,3 +170,21 @@ private:
 
 	ReferencedType* Reference;
 };
+
+template<typename ReferencedType>
+FORCEINLINE bool operator==(const TRefCountPtr<ReferencedType>& A, const TRefCountPtr<ReferencedType>& B)
+{
+	return A.GetReference() == B.GetReference();
+}
+
+template<typename ReferencedType>
+FORCEINLINE bool operator==(const TRefCountPtr<ReferencedType>& A, ReferencedType* B)
+{
+	return A.GetReference() == B;
+}
+
+template<typename ReferencedType>
+FORCEINLINE bool operator==(ReferencedType* A, const TRefCountPtr<ReferencedType>& B)
+{
+	return A == B.GetReference();
+}
diff --git a/Engine/Source/Runtime/Core/Public/Templates/SharedPointer.h b/Engine/Source/Runtime/Core/Public/Templates/SharedPointer.h
--- a/Engine/Source/Runtime/Core/Public/Templates/SharedPointer.h
+++ b/Engine/Source/Runtime/Core/Public/Templates/SharedPointer.h
@@ -1285,6 +1285,54 @@ FORCEINLINE bool operator==( TWeakPtr< ObjectTypeA, Mode > const& InWeakPtrA, TW
 }
 
 
+/**
+ * Global equality operator for TWeakPtr
+ *
+ * @return  True if the weak pointer and the shared ref are equal
+ */
+template< class ObjectTypeA, class ObjectTypeB, ESPMode Mode >
+FORCEINLINE bool operator==( TWeakPtr< ObjectTypeA, Mode > const& InWeakPtrA, TSharedRef< ObjectTypeB, Mode > const& InSharedRefB )
+{
+	return InWeakPtrA.Pin().Get() == &InSharedRefB.Get();
+}
+
+
+/**
+ * Global equality operator for TWeakPtr
+ *
+ * @return  True if the weak pointer and the shared ptr are equal
+ */
+template< class ObjectTypeA, class ObjectTypeB, ESPMode Mode >
+FORCEINLINE bool operator==( TWeakPtr< ObjectTypeA, Mode > const& InWeakPtrA, TSharedPtr< ObjectTypeB, Mode > const& InSharedPtrB )
+{
+	return InWeakPtrA.Pin().Get() == InSharedPtrB.Get();
+}
+
+
+/**
+ * Global equality operator for TWeakPtr
+ *
+ * @return  True if the weak pointer and the shared ref are equal
+ */
+template< class ObjectTypeA, class ObjectTypeB, ESPMode Mode >
+FORCEINLINE bool operator==( TSharedRef< ObjectTypeA, Mode > const& InSharedRefA, TWeakPtr< ObjectTypeB, Mode > const& InWeakPtrB )
+{
+	return &InSharedRefA.Get() == InWeakPtrB.Pin().Get();
+}
+
+
+/**
+ * Global equality operator for TWeakPtr
+ *
+ * @return  True if the weak pointer and the shared ptr are equal
+ */
+template< class ObjectTypeA, class ObjectTypeB, ESPMode Mode >
+FORCEINLINE bool operator==( TSharedPtr< ObjectTypeA, Mode > const& InSharedPtrA, TWeakPtr< ObjectTypeB, Mode > const& InWeakPtrB )
+{
+	return InSharedPtrA.Get() == InWeakPtrB.Pin().Get();
+}
+
+
 /**
  * Global equality operator for TWeakPtr
  *
@@ -1321,6 +1369,54 @@ FORCEINLINE bool operator!=( TWeakPtr< ObjectTypeA, Mode > const& InWeakPtrA, TW
 }
 
 
+/**
+ * Global equality operator for TWeakPtr
+ *
+ * @return  True if the weak pointer and the shared ref are not equal
+ */
+template< class ObjectTypeA, class ObjectTypeB, ESPMode Mode >
+FORCEINLINE bool operator!=( TWeakPtr< ObjectTypeA, Mode > const& InWeakPtrA, TSharedRef< ObjectTypeB, Mode > const& InSharedRefB )
+{
+	return InWeakPtrA.Pin().Get() != &InSharedRefB.Get();
+}
+
+
+/**
+ * Global equality operator for TWeakPtr
+ *
+ * @return  True if the weak pointer and the shared ptr are not equal
+ */
+template< class ObjectTypeA, class ObjectTypeB, ESPMode Mode >
+FORCEINLINE bool operator!=( TWeakPtr< ObjectTypeA, Mode > const& InWeakPtrA, TSharedPtr< ObjectTypeB, Mode > const& InSharedPtrB )
+{
+	return InWeakPtrA.Pin().Get() != InSharedPtrB.Get();
+}
+
+
+/**
+ * Global equality operator for TWeakPtr
+ *
+ * @return  True if the weak pointer and the shared ref are not equal
+ */
+template< class ObjectTypeA, class ObjectTypeB, ESPMode Mode >
+FORCEINLINE bool operator!=( TSharedRef< ObjectTypeA, Mode > const& InSharedRefA, TWeakPtr< ObjectTypeB, Mode > const& InWeakPtrB )
+{
+	return &InSharedRefA.Get() != InWeakPtrB.Pin().Get();
+}
+
+
+/**
+ * Global equality operator for TWeakPtr
+ *
+ * @return  True if the weak pointer and the shared ptr are not equal
+ */
+template< class ObjectTypeA, class ObjectTypeB, ESPMode Mode >
+FORCEINLINE bool operator!=( TSharedPtr< ObjectTypeA, Mode > const& InSharedPtrA, TWeakPtr< ObjectTypeB, Mode > const& InWeakPtrB )
+{
+	return InSharedPtrA.Get() != InWeakPtrB.Pin().Get();
+}
+
+
 /**
  * Global inequality operator for TWeakPtr
  *
