Linux: Fix the crash when trying to map an empty filename.

- Noticed and contributed by [REDACTED_PERSON] (PR #939).

diff --git a/Engine/Source/Runtime/Core/Private/Linux/LinuxPlatformFile.cpp b/Engine/Source/Runtime/Core/Private/Linux/LinuxPlatformFile.cpp
--- a/Engine/Source/Runtime/Core/Private/Linux/LinuxPlatformFile.cpp
+++ b/Engine/Source/Runtime/Core/Private/Linux/LinuxPlatformFile.cpp
@@ -459,7 +459,7 @@ public:
 
 		// We can get some "absolute" filenames like "D:/Blah/" here (e.g. non-Linux paths to source files embedded in assets).
 		// In that case, fail silently.
-		if (PossiblyWrongFilename[0] != TEXT('/'))
+		if (PossiblyWrongFilename.IsEmpty() || PossiblyWrongFilename[0] != TEXT('/'))
 		{
 			return false;
 		}
