Adding NoRedist/[Platform] and NotForLicensees/[Platform] folders to the hierarchy.

diff --git a/Engine/Source/Runtime/Core/Private/Misc/ConfigCacheIni.cpp b/Engine/Source/Runtime/Core/Private/Misc/ConfigCacheIni.cpp
--- a/Engine/Source/Runtime/Core/Private/Misc/ConfigCacheIni.cpp
+++ b/Engine/Source/Runtime/Core/Private/Misc/ConfigCacheIni.cpp
@@ -2780,8 +2780,16 @@ static void GetSourceIniHierarchyFilenames(const TCHAR* InBaseIniName, const TCH
 		// [[[[ PLATFORM DEFAULTS AND PROJECT SETTINGS ]]]]
 		// Engine/Config/Platform/Platform* ini
 		OutHierarchy.Add( FIniFilename(FString::Printf(TEXT("%s%s/%s%s.ini"), EngineConfigDir, *PlatformName, *PlatformName, InBaseIniName), false) );
+		// Engine/Config/NotForLicensees/Platform/Platform* ini
+		OutHierarchy.Add( FIniFilename(FString::Printf(TEXT("%sNotForLicensees/%s/%s%s.ini"), EngineConfigDir, *PlatformName, *PlatformName, InBaseIniName), false) );
+		// Engine/Config/NoRedist/Platform/Platform* ini
+		OutHierarchy.Add( FIniFilename(FString::Printf(TEXT("%sNoRedist/%s/%s%s.ini"), EngineConfigDir, *PlatformName, *PlatformName, InBaseIniName), false) );
 		// Game/Config/Platform/Platform* ini
 		OutHierarchy.Add( FIniFilename(FString::Printf(TEXT("%s%s/%s%s.ini"), SourceConfigDir, *PlatformName, *PlatformName, InBaseIniName), false) );
+		// Game/Config/NotForLicensee/Platform/Platform* ini
+		OutHierarchy.Add( FIniFilename(FString::Printf(TEXT("%sNotForLicensees/%s/%s%s.ini"), SourceConfigDir, *PlatformName, *PlatformName, InBaseIniName), false) );
+		// Game/Config/NoRedist/Platform/Platform* ini
+		OutHierarchy.Add( FIniFilename(FString::Printf(TEXT("%sNoRedist/%s/%s%s.ini"), SourceConfigDir, *PlatformName, *PlatformName, InBaseIniName), false) );
 	}
 
 	// [[[[ GLOBAL USER OVERRIDES ]]]]
