Use TWeakPtr instead of TSharedPtr for a window stored in deferred Mac event struct

diff --git a/Engine/Source/Runtime/Core/Private/Mac/MacApplication.cpp b/Engine/Source/Runtime/Core/Private/Mac/MacApplication.cpp
--- a/Engine/Source/Runtime/Core/Private/Mac/MacApplication.cpp
+++ b/Engine/Source/Runtime/Core/Private/Mac/MacApplication.cpp
@@ -272,7 +272,7 @@ void FMacApplication::DeferEvent(NSObject* Object)
 		if (DeferredEvent.Type == NSKeyDown && !bSystemModalMode)
 		{
 			DeferredEvent.Window = FindWindowByNSWindow(WindowHandle);
-			if (DeferredEvent.Window.IsValid() && DeferredEvent.Window->OnIMKKeyDown(Event))
+			if (DeferredEvent.Window.IsValid() && DeferredEvent.Window.Pin()->OnIMKKeyDown(Event))
 			{
 				return;
 			}
@@ -345,7 +345,7 @@ void FMacApplication::DeferEvent(NSObject* Object)
 				DeferredEvent.Window = FindWindowByNSWindow(WindowHandle);
 				if (DeferredEvent.Window.IsValid())
 				{
-					GameThreadCall(^{ OnWindowDidResize(DeferredEvent.Window.ToSharedRef()); }, @[ NSDefaultRunLoopMode, UE4ResizeEventMode, UE4ShowEventMode, UE4FullscreenEventMode ], true);
+					GameThreadCall(^{ OnWindowDidResize(DeferredEvent.Window.Pin().ToSharedRef()); }, @[ NSDefaultRunLoopMode, UE4ResizeEventMode, UE4ShowEventMode, UE4FullscreenEventMode ], true);
 				}
 				return;
 			}
@@ -508,51 +508,52 @@ void FMacApplication::ProcessEvent(const FDeferredMacEvent& Event)
 	}
 	else if (Event.Window.IsValid())
 	{
+		TSharedRef<FMacWindow> EventWindow = Event.Window.Pin().ToSharedRef();
 		if (Event.NotificationName == NSWindowWillStartLiveResizeNotification)
 		{
-			MessageHandler->BeginReshapingWindow(Event.Window.ToSharedRef());
+			MessageHandler->BeginReshapingWindow(EventWindow);
 		}
 		else if (Event.NotificationName == NSWindowDidEndLiveResizeNotification)
 		{
-			MessageHandler->FinishedReshapingWindow(Event.Window.ToSharedRef());
+			MessageHandler->FinishedReshapingWindow(EventWindow);
 		}
 		else if (Event.NotificationName == NSWindowDidEnterFullScreenNotification)
 		{
-			OnWindowDidResize(Event.Window.ToSharedRef());
+			OnWindowDidResize(EventWindow);
 		}
 		else if (Event.NotificationName == NSWindowDidExitFullScreenNotification)
 		{
-			OnWindowDidResize(Event.Window.ToSharedRef());
+			OnWindowDidResize(EventWindow);
 		}
 		else if (Event.NotificationName == NSWindowDidBecomeMainNotification)
 		{
-			MessageHandler->OnWindowActivationChanged(Event.Window.ToSharedRef(), EWindowActivation::Activate);
+			MessageHandler->OnWindowActivationChanged(EventWindow, EWindowActivation::Activate);
 			OnWindowsReordered(![NSApp isActive]);
 		}
 		else if (Event.NotificationName == NSWindowDidResignMainNotification)
 		{
-			MessageHandler->OnWindowActivationChanged(Event.Window.ToSharedRef(), EWindowActivation::Deactivate);
+			MessageHandler->OnWindowActivationChanged(EventWindow, EWindowActivation::Deactivate);
 			OnWindowsReordered(![NSApp isActive]);
 		}
 		else if (Event.NotificationName == NSWindowWillMoveNotification)
 		{
-			DraggedWindow = Event.Window->GetWindowHandle();
+			DraggedWindow = EventWindow->GetWindowHandle();
 		}
 		else if (Event.NotificationName == NSWindowDidMoveNotification)
 		{
-			OnWindowDidMove(Event.Window.ToSharedRef());
+			OnWindowDidMove(EventWindow);
 		}
 		else if (Event.NotificationName == NSWindowWillCloseNotification)
 		{
-			OnWindowDidClose(Event.Window.ToSharedRef());
+			OnWindowDidClose(EventWindow);
 		}
 		else if (Event.NotificationName == NSDraggingExited)
 		{
-			MessageHandler->OnDragLeave(Event.Window.ToSharedRef());
+			MessageHandler->OnDragLeave(EventWindow);
 		}
 		else if (Event.NotificationName == NSDraggingUpdated)
 		{
-			MessageHandler->OnDragOver(Event.Window.ToSharedRef());
+			MessageHandler->OnDragOver(EventWindow);
 		}
 		else if (Event.NotificationName == NSPrepareForDragOperation)
 		{
@@ -574,17 +575,17 @@ void FMacApplication::ProcessEvent(const FDeferredMacEvent& Event)
 					FileList.Add(ListElement);
 				}
 
-				MessageHandler->OnDragEnterFiles(Event.Window.ToSharedRef(), FileList);
+				MessageHandler->OnDragEnterFiles(EventWindow, FileList);
 			}
 			else if (bHaveText)
 			{
 				NSString* Text = [Event.DraggingPasteboard stringForType:NSPasteboardTypeString];
-				MessageHandler->OnDragEnterText(Event.Window.ToSharedRef(), FString(Text));
+				MessageHandler->OnDragEnterText(EventWindow, FString(Text));
 			}
 		}
 		else if (Event.NotificationName == NSPerformDragOperation)
 		{
-			MessageHandler->OnDragDrop(Event.Window.ToSharedRef());
+			MessageHandler->OnDragDrop(EventWindow);
 		}
 	}
 }
@@ -599,12 +600,13 @@ void FMacApplication::ResendEvent(NSEvent* Event)
 
 void FMacApplication::ProcessMouseMovedEvent(const FDeferredMacEvent& Event)
 {
-	if (Event.Window.IsValid() && Event.Window->IsRegularWindow())
+	TSharedPtr<FMacWindow> EventWindow = Event.Window.Pin();
+	if (EventWindow.IsValid() && EventWindow->IsRegularWindow())
 	{
 		bool IsMouseOverTitleBar = false;
-		const bool IsMovable = IsWindowMovable(Event.Window.ToSharedRef(), &IsMouseOverTitleBar);
-		[Event.Window->GetWindowHandle() setMovable:IsMovable];
-		[Event.Window->GetWindowHandle() setMovableByWindowBackground:IsMouseOverTitleBar];
+		const bool IsMovable = IsWindowMovable(EventWindow.ToSharedRef(), &IsMouseOverTitleBar);
+		[EventWindow->GetWindowHandle() setMovable:IsMovable];
+		[EventWindow->GetWindowHandle() setMovableByWindowBackground:IsMouseOverTitleBar];
 	}
 
 	FMacCursor* MacCursor = (FMacCursor*)Cursor.Get();
@@ -642,7 +644,7 @@ void FMacApplication::ProcessMouseMovedEvent(const FDeferredMacEvent& Event)
 		// assorted potential for mouse abuse.
 		NSRect VisibleFrame = [Screen visibleFrame];
 		// Avoid the menu bar & dock disclosure borders at the top & bottom of fullscreen windows
-		if (Event.Window.IsValid() && Event.Window->GetWindowMode() != EWindowMode::Windowed)
+		if (EventWindow.IsValid() && EventWindow->GetWindowMode() != EWindowMode::Windowed)
 		{
 			VisibleFrame.origin.y += 5;
 			VisibleFrame.size.height -= 10;
@@ -675,17 +677,17 @@ void FMacApplication::ProcessMouseMovedEvent(const FDeferredMacEvent& Event)
 			MacCursor->UpdateCurrentPosition(CurrentPosition);
 		}
 
-		if (Event.Window.IsValid())
+		if (EventWindow.IsValid())
 		{
 			// Cocoa does not update NSWindow's frame until user stops dragging the window, so while window is being dragged, we calculate
 			// its position based on mouse move delta
-			if (DraggedWindow && DraggedWindow == Event.Window->GetWindowHandle())
+			if (DraggedWindow && DraggedWindow == EventWindow->GetWindowHandle())
 			{
-				const int32 X = FMath::TruncToInt(Event.Window->PositionX + MouseDelta.X);
-				const int32 Y = FMath::TruncToInt(Event.Window->PositionY + MouseDelta.Y);
-				MessageHandler->OnMovedWindow(Event.Window.ToSharedRef(), X, Y);
-				Event.Window->PositionX = X;
-				Event.Window->PositionY = Y;
+				const int32 X = FMath::TruncToInt(EventWindow->PositionX + MouseDelta.X);
+				const int32 Y = FMath::TruncToInt(EventWindow->PositionY + MouseDelta.Y);
+				MessageHandler->OnMovedWindow(EventWindow.ToSharedRef(), X, Y);
+				EventWindow->PositionX = X;
+				EventWindow->PositionY = Y;
 			}
 
 			MessageHandler->OnMouseMove();
@@ -700,6 +702,7 @@ void FMacApplication::ProcessMouseMovedEvent(const FDeferredMacEvent& Event)
 
 void FMacApplication::ProcessMouseDownEvent(const FDeferredMacEvent& Event)
 {
+	TSharedPtr<FMacWindow> EventWindow = Event.Window.Pin();
 	EMouseButtons::Type Button = Event.Type == NSLeftMouseDown ? EMouseButtons::Left : EMouseButtons::Right;
 	if (Event.Type == NSOtherMouseDown)
 	{
@@ -723,32 +726,32 @@ void FMacApplication::ProcessMouseDownEvent(const FDeferredMacEvent& Event)
 	{
 		if (Button == LastPressedMouseButton && (Event.ClickCount % 2) == 0)
 		{
-			MessageHandler->OnMouseDoubleClick(Event.Window, Button);
+			MessageHandler->OnMouseDoubleClick(EventWindow, Button);
 
 			if (Event.Window.IsValid())
 			{
 				bool IsMouseOverTitleBar = false;
-				const bool IsMovable = IsWindowMovable(Event.Window.ToSharedRef(), &IsMouseOverTitleBar);
+				const bool IsMovable = IsWindowMovable(EventWindow.ToSharedRef(), &IsMouseOverTitleBar);
 				if (IsMouseOverTitleBar)
 				{
 					const bool bShouldMinimize = [[NSUserDefaults standardUserDefaults] boolForKey:@"AppleMiniaturizeOnDoubleClick"];
 					if (bShouldMinimize)
 					{
-						MainThreadCall(^{ [Event.Window->GetWindowHandle() performMiniaturize:nil]; }, NSDefaultRunLoopMode, true);
+						MainThreadCall(^{ [EventWindow->GetWindowHandle() performMiniaturize:nil]; }, NSDefaultRunLoopMode, true);
 					}
 					else if (!FPlatformMisc::IsRunningOnMavericks())
 					{
-						MainThreadCall(^{ [Event.Window->GetWindowHandle() zoom:nil]; }, NSDefaultRunLoopMode, true);
+						MainThreadCall(^{ [EventWindow->GetWindowHandle() zoom:nil]; }, NSDefaultRunLoopMode, true);
 					}
 				}
 			}
 		}
 		else
 		{
-			MessageHandler->OnMouseDown(Event.Window, Button);
+			MessageHandler->OnMouseDown(EventWindow, Button);
 		}
 
-		if (Event.Window->GetWindowHandle() && !DraggedWindow && !GetCapture())
+		if (EventWindow->GetWindowHandle() && !DraggedWindow && !GetCapture())
 		{
 			MessageHandler->OnCursorSet();
 		}
@@ -780,7 +783,7 @@ void FMacApplication::ProcessMouseUpEvent(const FDeferredMacEvent& Event)
 
 	MessageHandler->OnMouseUp(Button);
 
-	if (Event.Window.IsValid() && Event.Window->GetWindowHandle() && !DraggedWindow && !GetCapture())
+	if (Event.Window.IsValid() && Event.Window.Pin()->GetWindowHandle() && !DraggedWindow && !GetCapture())
 	{
 		MessageHandler->OnCursorSet();
 	}
diff --git a/Engine/Source/Runtime/Core/Public/Mac/MacApplication.h b/Engine/Source/Runtime/Core/Public/Mac/MacApplication.h
--- a/Engine/Source/Runtime/Core/Public/Mac/MacApplication.h
+++ b/Engine/Source/Runtime/Core/Public/Mac/MacApplication.h
@@ -91,7 +91,7 @@ struct FDeferredMacEvent
 	// However, in some cases we need the original NSEvent (highlighting menus, resending unhandled key events), so we store it as well.
 	NSEvent* Event;
 
-	TSharedPtr<FMacWindow> Window;
+	TWeakPtr<FMacWindow> Window;
 
 	int32 Type;
 	FVector2D LocationInWindow;
