- Missed a file with minor changes for iOS in the big Android ES31 changelist (2478853)

diff --git a/Engine/Source/Runtime/Core/Public/IOS/IOSPlatformProperties.h b/Engine/Source/Runtime/Core/Public/IOS/IOSPlatformProperties.h
--- a/Engine/Source/Runtime/Core/Public/IOS/IOSPlatformProperties.h
+++ b/Engine/Source/Runtime/Core/Public/IOS/IOSPlatformProperties.h
@@ -23,11 +23,6 @@ struct FIOSPlatformProperties
 		return false;
 	}
 
-	static FORCEINLINE uint32 MaxGpuSkinBones( )
-	{
-		return 20;
-	}
-
 	static FORCEINLINE const char* PlatformName( )
 	{
 		return "IOS";
@@ -50,7 +45,7 @@ struct FIOSPlatformProperties
 
 	static FORCEINLINE bool SupportsHighQualityLightmaps()
 	{
-		return false;
+		return true;
 	}
 
 	static FORCEINLINE bool SupportsTextureStreaming()
