MacApplication now enumerates and ticks IInputDeviceModule implementatations
Fix bad #ifndef and comment in InputDevice headers
Partially from PR #940: Export input device to allow game input device plugins (Contributed by [REDACTED_PERSON])

diff --git a/Engine/Source/Runtime/Core/Private/Mac/MacApplication.cpp b/Engine/Source/Runtime/Core/Private/Mac/MacApplication.cpp
--- a/Engine/Source/Runtime/Core/Private/Mac/MacApplication.cpp
+++ b/Engine/Source/Runtime/Core/Private/Mac/MacApplication.cpp
@@ -131,8 +131,28 @@ void FMacApplication::SetMessageHandler(const TSharedRef<FGenericApplicationMess
 
 void FMacApplication::PollGameDeviceState(const float TimeDelta)
 {
+	// initialize any externally-implemented input devices (we delay load initialize the array so any plugins have had time to load)
+	if (!bHasLoadedInputPlugins)
+	{
+		TArray<IInputDeviceModule*> PluginImplementations = IModularFeatures::Get().GetModularFeatureImplementations<IInputDeviceModule>( IInputDeviceModule::GetModularFeatureName() );
+		for( auto InputPluginIt = PluginImplementations.CreateIterator(); InputPluginIt; ++InputPluginIt )
+		{
+			TSharedPtr<IInputDevice> Device = (*InputPluginIt)->CreateInputDevice(MessageHandler);
+			AddExternalInputDevice(Device);			
+		}
+
+		bHasLoadedInputPlugins = true;
+	}
+
 	// Poll game device state and send new events
 	HIDInput->SendControllerEvents();
+
+	// Poll externally-implemented devices
+	for( auto DeviceIt = ExternalInputDevices.CreateIterator(); DeviceIt; ++DeviceIt )
+	{
+		(*DeviceIt)->Tick( TimeDelta );
+		(*DeviceIt)->SendControllerEvents();
+	}
 }
 
 void FMacApplication::PumpMessages(const float TimeDelta)
diff --git a/Engine/Source/Runtime/InputDevice/Public/IInputDeviceModule.h b/Engine/Source/Runtime/InputDevice/Public/IInputDeviceModule.h
--- a/Engine/Source/Runtime/InputDevice/Public/IInputDeviceModule.h
+++ b/Engine/Source/Runtime/InputDevice/Public/IInputDeviceModule.h
@@ -8,7 +8,7 @@
 
 
 /**
- * The public interface of the MotionControlsModule
+ * The public interface of the InputDeviceModule
  */
 class IInputDeviceModule : public IModuleInterface, public IModularFeature
 {
diff --git a/Engine/Source/Runtime/InputDevice/Public/InputDevice.h b/Engine/Source/Runtime/InputDevice/Public/InputDevice.h
--- a/Engine/Source/Runtime/InputDevice/Public/InputDevice.h
+++ b/Engine/Source/Runtime/InputDevice/Public/InputDevice.h
@@ -1,10 +1,11 @@
 // Copyright 1998-2015 Epic Games, Inc. All Rights Reserved.
 
-#ifndef __HEADMOUNTEDDISPLAY_H__
-#define __HEADMOUNTEDDISPLAY_H__
+#ifndef __INPUTDEVICE_H__
+#define __INPUTDEVICE_H__
 
 #include "Core.h"
 #include "Engine.h"
 #include "IInputDevice.h"
+#include "IInputDeviceModule.h"
 
 #endif
