missed file for CL# 2484846

diff --git a/Engine/Source/Runtime/Core/Public/Mac/MacApplication.h b/Engine/Source/Runtime/Core/Public/Mac/MacApplication.h
--- a/Engine/Source/Runtime/Core/Public/Mac/MacApplication.h
+++ b/Engine/Source/Runtime/Core/Public/Mac/MacApplication.h
@@ -261,6 +261,10 @@ private:
 
 	TSharedRef<class HIDInputInterface> HIDInput;
 
+	/** List of input devices implemented in external modules. */
+	TArray<TSharedPtr<class IInputDevice>> ExternalInputDevices;
+	bool bHasLoadedInputPlugins;
+
 	FCocoaWindow* DraggedWindow;
 
 	bool bSystemModalMode;
