Fix Mac compile issues

diff --git a/Engine/Source/Runtime/Core/Private/Mac/MacApplication.cpp b/Engine/Source/Runtime/Core/Private/Mac/MacApplication.cpp
--- a/Engine/Source/Runtime/Core/Private/Mac/MacApplication.cpp
+++ b/Engine/Source/Runtime/Core/Private/Mac/MacApplication.cpp
@@ -8,6 +8,8 @@
 #include "CocoaMenu.h"
 #include "GenericApplicationMessageHandler.h"
 #include "HIDInputInterface.h"
+#include "IInputDeviceModule.h"
+#include "IInputDevice.h"
 #include "AnalyticsEventAttribute.h"
 #include "IAnalyticsProvider.h"
 #include "CocoaThread.h"
@@ -39,6 +41,7 @@ FMacApplication::FMacApplication()
 ,	LastPressedMouseButton(EMouseButtons::Invalid)
 ,	bIsProcessingDeferredEvents(false)
 ,	HIDInput(HIDInputInterface::Create(MessageHandler))
+,   bHasLoadedInputPlugins(false)
 ,	DraggedWindow(nullptr)
 ,	bSystemModalMode(false)
 ,	ModifierKeysFlags(0)
@@ -138,8 +141,12 @@ void FMacApplication::PollGameDeviceState(const float TimeDelta)
 		for( auto InputPluginIt = PluginImplementations.CreateIterator(); InputPluginIt; ++InputPluginIt )
 		{
 			TSharedPtr<IInputDevice> Device = (*InputPluginIt)->CreateInputDevice(MessageHandler);
-			AddExternalInputDevice(Device);			
-		}
+            if (Device.IsValid())
+            {
+                UE_LOG(LogInit, Log, TEXT("Adding external input plugin."));
+                ExternalInputDevices.Add(Device);
+            }
+        }
 
 		bHasLoadedInputPlugins = true;
 	}
@@ -1394,4 +1401,4 @@ void FDisplayMetrics::GetDisplayMetrics(FDisplayMetrics& OutDisplayMetrics)
 	OutDisplayMetrics.PrimaryDisplayWorkAreaRect.Top = ScreenFrame.size.height - (VisibleFrame.origin.y + VisibleFrame.size.height);
 	OutDisplayMetrics.PrimaryDisplayWorkAreaRect.Right = VisibleFrame.origin.x + VisibleFrame.size.width;
 	OutDisplayMetrics.PrimaryDisplayWorkAreaRect.Bottom = OutDisplayMetrics.PrimaryDisplayWorkAreaRect.Top + VisibleFrame.size.height;
-}
+}
\ No newline at end of file
diff --git a/Engine/Source/Runtime/Core/Public/Mac/MacApplication.h b/Engine/Source/Runtime/Core/Public/Mac/MacApplication.h
--- a/Engine/Source/Runtime/Core/Public/Mac/MacApplication.h
+++ b/Engine/Source/Runtime/Core/Public/Mac/MacApplication.h
@@ -6,6 +6,7 @@
 #include "GenericApplication.h"
 #include "MacWindow.h"
 #include "MacTextInputMethodSystem.h"
+#include "IInputInterface.h"
 
 struct FDeferredMacEvent
 {
