Added an overload for GetTypeHash that should allow TMap to use a key type that is a C++11 enum class. Code provided by [REDACTED_PERSON]. Preflighted and tested.

diff --git a/Engine/Source/Runtime/Core/Public/Templates/TypeHash.h b/Engine/Source/Runtime/Core/Public/Templates/TypeHash.h
--- a/Engine/Source/Runtime/Core/Public/Templates/TypeHash.h
+++ b/Engine/Source/Runtime/Core/Public/Templates/TypeHash.h
@@ -104,3 +104,15 @@ inline uint32 GetTypeHash( void* A )
 {
 	return PointerHash(A);
 }
+
+template <typename T>
+struct TIsEnum
+{
+	enum { Value = IS_ENUM(T) };
+};
+
+template <typename EnumType>
+FORCEINLINE  typename TEnableIf<TIsEnum<EnumType>::Value, uint32>::Type GetTypeHash(EnumType E)
+{
+	return GetTypeHash((__underlying_type(EnumType))E);
+}
\ No newline at end of file
