TIsEnum moved to UnrealTypeTraits.h.

diff --git a/Engine/Source/Runtime/Core/Public/Templates/TypeHash.h b/Engine/Source/Runtime/Core/Public/Templates/TypeHash.h
--- a/Engine/Source/Runtime/Core/Public/Templates/TypeHash.h
+++ b/Engine/Source/Runtime/Core/Public/Templates/TypeHash.h
@@ -105,12 +105,6 @@ inline uint32 GetTypeHash( void* A )
 	return PointerHash(A);
 }
 
-template <typename T>
-struct TIsEnum
-{
-	enum { Value = IS_ENUM(T) };
-};
-
 template <typename EnumType>
 FORCEINLINE  typename TEnableIf<TIsEnum<EnumType>::Value, uint32>::Type GetTypeHash(EnumType E)
 {
diff --git a/Engine/Source/Runtime/Core/Public/Templates/UnrealTypeTraits.h b/Engine/Source/Runtime/Core/Public/Templates/UnrealTypeTraits.h
--- a/Engine/Source/Runtime/Core/Public/Templates/UnrealTypeTraits.h
+++ b/Engine/Source/Runtime/Core/Public/Templates/UnrealTypeTraits.h
@@ -597,4 +597,10 @@ class THasMemberFunction_##MemberName																			\
 	template <typename U> static int MemberTest(...);															\
 public:																											\
 	enum { Value = sizeof(MemberTest<T>(nullptr)) == sizeof(char) };											\
-};
\ No newline at end of file
+};
+
+template <typename T>
+struct TIsEnum
+{
+	enum { Value = IS_ENUM(T) };
+};
