CrashReport - Clearing all error messages after ensure

diff --git a/Engine/Source/Runtime/Core/Private/Misc/OutputDevice.cpp b/Engine/Source/Runtime/Core/Private/Misc/OutputDevice.cpp
--- a/Engine/Source/Runtime/Core/Private/Misc/OutputDevice.cpp
+++ b/Engine/Source/Runtime/Core/Private/Misc/OutputDevice.cpp
@@ -324,7 +324,13 @@ void FDebug::EnsureFailed(const ANSICHAR* Expr, const ANSICHAR* File, int32 Line
 	if ( bShouldSendNewReport )
 	{
 #if PLATFORM_DESKTOP
+		FScopeLock Lock( &FailDebugCriticalSection );
+
 		NewReportEnsure( GErrorMessage );
+
+		GErrorHist[0] = 0;
+		GErrorMessage[0] = 0;
+		GErrorExceptionDescription[0] = 0;
 #endif
 	}
 }
