Merging a missing .ini crash fix to main

diff --git a/Engine/Source/Runtime/Core/Private/Misc/ConfigCacheIni.cpp b/Engine/Source/Runtime/Core/Private/Misc/ConfigCacheIni.cpp
--- a/Engine/Source/Runtime/Core/Private/Misc/ConfigCacheIni.cpp
+++ b/Engine/Source/Runtime/Core/Private/Misc/ConfigCacheIni.cpp
@@ -3089,7 +3089,7 @@ void FConfigFile::UpdateSections(const TCHAR* DiskFilename, const TCHAR* IniRoot
 	{
 		// walk each line
 		const TCHAR* Ptr = DiskFile.Len() > 0 ? *DiskFile : NULL;
-		bool bDone = false;
+		bool bDone = Ptr ? false : true;
 		bool bIsSkippingSection = true;
 		while (!bDone)
 		{
