Linux: add Set/GetEnvironmentVar in platform layer.

- PR #960 contributed by [REDACTED_PERSON].

diff --git a/Engine/Source/Runtime/Core/Private/Linux/LinuxPlatformMisc.cpp b/Engine/Source/Runtime/Core/Private/Linux/LinuxPlatformMisc.cpp
--- a/Engine/Source/Runtime/Core/Private/Linux/LinuxPlatformMisc.cpp
+++ b/Engine/Source/Runtime/Core/Private/Linux/LinuxPlatformMisc.cpp
@@ -202,6 +202,35 @@ GenericApplication* FLinuxPlatformMisc::CreateApplication()
 	return FLinuxApplication::CreateLinuxApplication();
 }
 
+void FLinuxPlatformMisc::GetEnvironmentVariable(const TCHAR* InVariableName, TCHAR* Result, int32 ResultLength)
+{
+	FString VariableName = InVariableName;
+	VariableName.ReplaceInline(TEXT("-"), TEXT("_"));
+	ANSICHAR *AnsiResult = secure_getenv(TCHAR_TO_ANSI(*VariableName));
+	if (AnsiResult)
+	{
+		wcsncpy(Result, ANSI_TO_TCHAR(AnsiResult), ResultLength);
+	}
+	else
+	{
+		*Result = 0;
+	}
+}
+
+void FLinuxPlatformMisc::SetEnvironmentVar(const TCHAR* InVariableName, const TCHAR* Value)
+{
+	FString VariableName = InVariableName;
+	VariableName.ReplaceInline(TEXT("-"), TEXT("_"));
+	if (Value == NULL || Value[0] == TEXT('\0'))
+	{
+		unsetenv(TCHAR_TO_ANSI(*VariableName));
+	}
+	else
+	{
+		setenv(TCHAR_TO_ANSI(*VariableName), TCHAR_TO_ANSI(Value), 1);
+	}
+}
+
 void FLinuxPlatformMisc::PumpMessages( bool bFromMainLoop )
 {
 	if( bFromMainLoop )
diff --git a/Engine/Source/Runtime/Core/Public/Linux/LinuxPlatformMisc.h b/Engine/Source/Runtime/Core/Public/Linux/LinuxPlatformMisc.h
--- a/Engine/Source/Runtime/Core/Public/Linux/LinuxPlatformMisc.h
+++ b/Engine/Source/Runtime/Core/Public/Linux/LinuxPlatformMisc.h
@@ -23,6 +23,8 @@ struct CORE_API FLinuxPlatformMisc : public FGenericPlatformMisc
 	static void SetGracefulTerminationHandler();
 	static void SetCrashHandler(void (* CrashHandler)(const FGenericCrashContext& Context));
 	static class GenericApplication* CreateApplication();
+	static void GetEnvironmentVariable(const TCHAR* VariableName, TCHAR* Result, int32 ResultLength);
+	static void SetEnvironmentVar(const TCHAR* VariableName, const TCHAR* Value);
 #if !UE_BUILD_SHIPPING
 	static bool IsDebuggerPresent();
 	FORCEINLINE static void DebugBreak()
