Linux: exempt some warnings from -Werror.

- After discussion with Core/Platform teams, it was determined that these warnings [would] affect build failure rate too much.
- We do not want to suppress them altogether though and they are going to be fixed as time and resources permit.

diff --git a/Engine/Source/Programs/UnrealBuildTool/Linux/LinuxToolChain.cs b/Engine/Source/Programs/UnrealBuildTool/Linux/LinuxToolChain.cs
--- a/Engine/Source/Programs/UnrealBuildTool/Linux/LinuxToolChain.cs
+++ b/Engine/Source/Programs/UnrealBuildTool/Linux/LinuxToolChain.cs
@@ -322,7 +322,6 @@ namespace UnrealBuildTool
 
 				if (CompilerVersionGreaterOrEqual(3, 6, 0))
 				{
-					Result += " -Wno-inconsistent-missing-override";	// we may still want to fix those - tracked as UE-12341
 					Result += " -Wno-unused-local-typedef";	// clang is being overly strict here? PhysX headers trigger this.
 				}
             }
diff --git a/Engine/Source/Runtime/Core/Public/Linux/LinuxPlatformCompilerPreSetup.h b/Engine/Source/Runtime/Core/Public/Linux/LinuxPlatformCompilerPreSetup.h
--- a/Engine/Source/Runtime/Core/Public/Linux/LinuxPlatformCompilerPreSetup.h
+++ b/Engine/Source/Runtime/Core/Public/Linux/LinuxPlatformCompilerPreSetup.h
@@ -49,3 +49,9 @@
 #define EMIT_CUSTOM_WARNING(Warning) \
 	_Pragma(PREPROCESSOR_TO_STRING(message(WARNING_LOCATION Warning)))
 #endif // EMIT_CUSTOM_WARNING
+
+// Make certain warnings always be warnings, even despite -Werror.
+// Rationale: we don't want to suppress those as there are plans to address them (e.g. UE-12341), but breaking builds due to these warnings is very expensive
+// since they cannot be caught by all compilers that we support. They are deemed to be relatively safe to be ignored, at least until all SDKs/toolchains start supporting them.
+#pragma clang diagnostic warning "-Winconsistent-missing-override"
+#pragma clang diagnostic warning "-Wreorder"
