Core - Runnable improvements

diff --git a/Engine/Source/Runtime/Core/Private/Android/AndroidApplication.cpp b/Engine/Source/Runtime/Core/Private/Android/AndroidApplication.cpp
--- a/Engine/Source/Runtime/Core/Private/Android/AndroidApplication.cpp
+++ b/Engine/Source/Runtime/Core/Private/Android/AndroidApplication.cpp
@@ -131,10 +131,6 @@ private:
 	}
 };
 
-#if 0
-DECLARE_THREAD_SINGLETON(FJNIHelper);
-#endif
-
 void FAndroidApplication::InitializeJavaEnv( JavaVM* VM, jint Version, jobject GlobalThis )
 {
 	if (CurrentJavaVM == nullptr)
diff --git a/Engine/Source/Runtime/Core/Private/HAL/PThreadRunnableThread.cpp b/Engine/Source/Runtime/Core/Private/HAL/PThreadRunnableThread.cpp
--- a/Engine/Source/Runtime/Core/Private/HAL/PThreadRunnableThread.cpp
+++ b/Engine/Source/Runtime/Core/Private/HAL/PThreadRunnableThread.cpp
@@ -12,6 +12,9 @@ uint32 FRunnableThreadPThread::Run()
 	// Assume we'll fail init
 	uint32 ExitCode = 1;
 	check(Runnable);
+	
+	// Setup TLS for this thread, used by ITlsAutoCleanup objects.
+	SetTls();
 
 	// Initialize the runnable object
 	if (Runnable->Init() == true)
@@ -32,6 +35,7 @@ uint32 FRunnableThreadPThread::Run()
 #if STATS
 	FThreadStats::Shutdown();
 #endif
+	FreeTls();
 
 	// Clean ourselves up without waiting
 	ThreadIsRunning = false;
diff --git a/Engine/Source/Runtime/Core/Private/HAL/PThreadRunnableThread.h b/Engine/Source/Runtime/Core/Private/HAL/PThreadRunnableThread.h
--- a/Engine/Source/Runtime/Core/Private/HAL/PThreadRunnableThread.h
+++ b/Engine/Source/Runtime/Core/Private/HAL/PThreadRunnableThread.h
@@ -24,36 +24,6 @@ protected:
 	 */
 	pthread_t Thread;
 
-	/**
-	 * The runnable object to execute on this thread
-	 */
-	FRunnable* Runnable;
-
-	/** 
-	 * Sync event to make sure that Init() has been completed before allowing the main thread to continue
-	 */
-	FEvent* ThreadInitSyncEvent;
-
-	/**
-	 * The priority to run the thread at
-	 */
-	EThreadPriority ThreadPriority;
-
-	/**
-	 * The CPU Mask to indicate which CPUs the thread will run on.
-	 */
-	uint64 ThreadAffinityMask;
-
-	/**
-	* ID set during thread creation
-	*/
-	uint32 ThreadID;
-
-	/**
-	 * The name of this thread
-	 */
-	FString ThreadName;
-
 	/**
 	 * If true, the thread is ready to be joined.
 	 */
@@ -229,11 +199,6 @@ protected:
 public:
 	FRunnableThreadPThread()
 		: Thread(PTHREAD_NULL)
-		, Runnable(NULL)
-		, ThreadInitSyncEvent(NULL)
-		, ThreadPriority(TPri_Normal)
-		, ThreadAffinityMask(FPlatformAffinity::GetNoAffinityMask())
-		, ThreadID(0)
 		, ThreadIsRunning(false)
 	{
 	}
@@ -245,7 +210,6 @@ public:
 		{
 			Kill(true);
 		}
-		FRunnableThread::GetThreadRegistry().Remove(ThreadID);
 		ThreadID = 0;
 	}
 
@@ -300,16 +264,6 @@ public:
 		}
 	}
 
-	virtual uint32 GetThreadID() override
-	{
-		return ThreadID;
-	}
-
-	virtual FString GetThreadName() override
-	{
-		return ThreadName;
-	}
-
 protected:
 
 	virtual bool CreateInternal(FRunnable* InRunnable, const TCHAR* InThreadName,
@@ -326,7 +280,7 @@ protected:
 		ThreadAffinityMask = InThreadAffinityMask;
 
 		// Create the new thread
-		bool ThreadCreated = SpinPthread(&Thread, GetThreadEntryPoint(), InStackSize, this);
+		const bool ThreadCreated = SpinPthread(&Thread, GetThreadEntryPoint(), InStackSize, this);
 		// If it fails, clear all the vars
 		if (ThreadCreated)
 		{
@@ -336,7 +290,7 @@ protected:
 			ThreadInitSyncEvent->Wait((uint32)-1); // infinite wait
 
 			// set the priority
-			SetThreadPriority(InThreadPri);			
+			SetThreadPriority(InThreadPri);
 		}
 		else // If it fails, clear all the vars
 		{
diff --git a/Engine/Source/Runtime/Core/Private/HAL/ThreadingBase.cpp b/Engine/Source/Runtime/Core/Private/HAL/ThreadingBase.cpp
--- a/Engine/Source/Runtime/Core/Private/HAL/ThreadingBase.cpp
+++ b/Engine/Source/Runtime/Core/Private/HAL/ThreadingBase.cpp
@@ -62,23 +62,17 @@ class FFakeThread : public FRunnableThread
 	/** Thread is suspended. */
 	bool bIsSuspended;
 
-	/** Name of this thread. */
-	FString Name;
-
 	/** Runnable object associated with this thread. */
 	FSingleThreadRunnable* Runnable;
 
-	/** Thread Id. */
-	uint32 ThreadId;
-
 public:
 
 	/** Constructor. */
 	FFakeThread()
 		: bIsSuspended(false)
 		, Runnable(nullptr)
-		, ThreadId(ThreadIdCounter++)
 	{
+		ThreadID = ThreadIdCounter++;
 		// Auto register with single thread manager.
 		FSingleThreadManager::Get().AddThread(this);
 	}
@@ -124,16 +118,6 @@ public:
 		FSingleThreadManager::Get().RemoveThread(this);
 	}
 
-	virtual uint32 GetThreadID() override
-	{
-		return ThreadId;
-	}
-
-	virtual FString GetThreadName() override
-	{
-		return Name;
-	}
-
 	virtual bool CreateInternal(FRunnable* InRunnable, const TCHAR* ThreadName,
 		uint32 InStackSize,
 		EThreadPriority InThreadPri, uint64 InThreadAffinityMask)
@@ -188,9 +172,26 @@ FScopedEvent::~FScopedEvent()
 	Event = nullptr;
 }
 
-FRunnableThread::~FRunnableThread()
+/*-----------------------------------------------------------------------------
+	FRunnableThread
+-----------------------------------------------------------------------------*/
+
+uint32 FRunnableThread::RunnableTlsSlot = 0;
+
+void FRunnableThread::InitializeTls()
+{
+	check( IsInGameThread() );
+	RunnableTlsSlot = FPlatformTLS::AllocTlsSlot();
+	check( FPlatformTLS::IsValidTlsSlot( RunnableTlsSlot ) );
+}
+
+FRunnableThread::FRunnableThread()
+	: Runnable(nullptr)
+	, ThreadInitSyncEvent(nullptr)
+	, ThreadAffinityMask(FPlatformAffinity::GetNoAffinityMask())
+	, ThreadPriority(TPri_Normal)
+	, ThreadID(0)
 {
-	ThreadDestroyedDelegate.Broadcast();
 }
 
 FRunnableThread* FRunnableThread::Create(
@@ -241,17 +242,42 @@ FRunnableThread* FRunnableThread::Create(
 		}
 	}
 
+#if	STATS
 	if( NewThread )
 	{
-		FRunnableThread::GetThreadRegistry().Add( NewThread->GetThreadID(), NewThread );
-#if	STATS
 		FStartupMessages::Get().AddThreadMetadata( FName( *NewThread->GetThreadName() ), NewThread->GetThreadID() );
-#endif // STATS
 	}
+#endif // STATS
 
 	return NewThread;
 }
 
+void FRunnableThread::SetTls()
+{
+	// Make sure it's called from the owning thread.
+	check( ThreadID == FPlatformTLS::GetCurrentThreadId() );
+	check( RunnableTlsSlot );
+	FPlatformTLS::SetTlsValue( RunnableTlsSlot, this );
+}
+
+void FRunnableThread::FreeTls()
+{
+	// Make sure it's called from the owning thread.
+	check( ThreadID == FPlatformTLS::GetCurrentThreadId() );
+	check( RunnableTlsSlot );
+	FPlatformTLS::SetTlsValue( RunnableTlsSlot, nullptr );
+
+	// Delete all ITlsAutoCleanup objects created for this thread.
+	for( auto& Instance : TlsInstances )
+	{
+		delete Instance;
+		Instance = nullptr;
+	}
+}
+
+/*-----------------------------------------------------------------------------
+	FQueuedThread
+-----------------------------------------------------------------------------*/
 
 /**
  * This is the interface used for all poolable threads. The usage pattern for
@@ -594,28 +620,33 @@ FQueuedThreadPool* FQueuedThreadPool::Allocate()
 	FThreadSingletonInitializer
 -----------------------------------------------------------------------------*/
 
-FThreadSingleton* FThreadSingletonInitializer::Get( const FThreadSingleton::TCreateSingletonFuncPtr CreateFunc, const FThreadSingleton::TDestroySingletonFuncPtr DestroyFunc, uint32& TlsSlot )
+ITlsAutoCleanup* FThreadSingletonInitializer::Get( const TFunctionRef<ITlsAutoCleanup*()>& CreateInstance, uint32& TlsSlot )
 {
-	check( CreateFunc != nullptr );
-	check( DestroyFunc != nullptr );
 	if( TlsSlot == 0 )
 	{
-		check( IsInGameThread() );
-		TlsSlot = FPlatformTLS::AllocTlsSlot();
+		const uint32 ThisTlsSlot = FPlatformTLS::AllocTlsSlot();
+		check( FPlatformTLS::IsValidTlsSlot( ThisTlsSlot ) );
+		const uint32 PrevTlsSlot = FPlatformAtomics::InterlockedCompareExchange( (int32*)&TlsSlot, (int32)ThisTlsSlot, 0 );
+		if( PrevTlsSlot != 0 )
+		{
+			FPlatformTLS::FreeTlsSlot( ThisTlsSlot );
+		}
 	}
-	FThreadSingleton* ThreadSingleton = (FThreadSingleton*)FPlatformTLS::GetTlsValue( TlsSlot );
+	ITlsAutoCleanup* ThreadSingleton = (ITlsAutoCleanup*)FPlatformTLS::GetTlsValue( TlsSlot );
 	if( !ThreadSingleton )
 	{
-		const uint32 ThreadId = FPlatformTLS::GetCurrentThreadId();
-		ThreadSingleton = (*CreateFunc)();
-		FRunnableThread::GetThreadRegistry().Lock();
-		FRunnableThread* RunnableThread = FRunnableThread::GetThreadRegistry().GetThread( ThreadId );
-		if( RunnableThread )
-		{
-			RunnableThread->OnThreadDestroyed().AddRaw( ThreadSingleton, DestroyFunc );
-		}
-		FRunnableThread::GetThreadRegistry().Unlock();
+		ThreadSingleton = CreateInstance();
+		ThreadSingleton->Register();
 		FPlatformTLS::SetTlsValue( TlsSlot, ThreadSingleton );
 	}
 	return ThreadSingleton;
 }
+
+void ITlsAutoCleanup::Register()
+{
+	FRunnableThread* RunnableThread = FRunnableThread::GetRunnableThread();
+	if( RunnableThread )
+	{
+		RunnableThread->TlsInstances.Add( this );
+	}
+}
diff --git a/Engine/Source/Runtime/Core/Private/Misc/MemStack.cpp b/Engine/Source/Runtime/Core/Private/Misc/MemStack.cpp
--- a/Engine/Source/Runtime/Core/Private/Misc/MemStack.cpp
+++ b/Engine/Source/Runtime/Core/Private/Misc/MemStack.cpp
@@ -2,9 +2,6 @@
 
 #include "CorePrivatePCH.h"
 
-
-DECLARE_THREAD_SINGLETON( FMemStack );
-
 DECLARE_MEMORY_STAT(TEXT("MemStack Large Block"), STAT_MemStackLargeBLock,STATGROUP_Memory);
 DECLARE_MEMORY_STAT(TEXT("PageAllocator Free"), STAT_PageAllocatorFree, STATGROUP_Memory);
 DECLARE_MEMORY_STAT(TEXT("PageAllocator Used"), STAT_PageAllocatorUsed, STATGROUP_Memory);
diff --git a/Engine/Source/Runtime/Core/Private/Stats/Stats2.cpp b/Engine/Source/Runtime/Core/Private/Stats/Stats2.cpp
--- a/Engine/Source/Runtime/Core/Private/Stats/Stats2.cpp
+++ b/Engine/Source/Runtime/Core/Private/Stats/Stats2.cpp
@@ -6,8 +6,6 @@
 	Global
 -----------------------------------------------------------------------------*/
 
-DECLARE_THREAD_SINGLETON( FThreadIdleStats );
-
 struct FStats2Globals
 {
 	static void Get()
diff --git a/Engine/Source/Runtime/Core/Private/WinRT/WinRTRunnableThread.h b/Engine/Source/Runtime/Core/Private/WinRT/WinRTRunnableThread.h
--- a/Engine/Source/Runtime/Core/Private/WinRT/WinRTRunnableThread.h
+++ b/Engine/Source/Runtime/Core/Private/WinRT/WinRTRunnableThread.h
@@ -13,36 +13,6 @@ class FRunnableThreadWinRT : public FRunnableThread
 	 */
 	HANDLE Thread;
 
-	/**
-	 * The runnable object to execute on this thread
-	 */
-	FRunnable* Runnable;
-
-	/** 
-	 * Sync event to make sure that Init() has been completed before allowing the main thread to continue
-	 */
-	FEvent* ThreadInitSyncEvent;
-
-	/**
-	 * The priority to run the thread at
-	 */
-	EThreadPriority ThreadPriority;
-
-	/**
-	 * The Affinity to run the thread with
-	 */
-	uint64 ThreadAffintyMask;
-
-	/**
-	* ID set during thread creation
-	*/
-	uint32 ThreadID;
-
-	/**
-	 * Holds the name of the thread.
-	 */
-	FString ThreadName;
-
 	/**
 	 * Helper function to set thread names, visible by the debugger.
 	 * @param ThreadID		Thread ID for the thread who's name is going to be set
@@ -126,11 +96,7 @@ class FRunnableThreadWinRT : public FRunnableThread
 
 public:
 	FRunnableThreadWinRT()
-		: Thread(NULL)
-		, Runnable(NULL)
-		, ThreadInitSyncEvent(NULL)
-		, ThreadPriority(TPri_Normal)
-		, ThreadID(NULL)
+		: Thread(nullptr)
 	{
 	}
 
@@ -141,7 +107,6 @@ public:
 		{
 			Kill(true);
 		}
-		FRunnableThread::GetThreadRegistry().Remove(ThreadID);
 	}
 	
 	virtual void SetThreadPriority(EThreadPriority NewPriority) override
@@ -224,7 +189,7 @@ protected:
 	{
 		check(InRunnable);
 		Runnable = InRunnable;
-		ThreadAffintyMask = InThreadAffinityMask;
+		ThreadAffinityMask = InThreadAffinityMask;
 
 		// Create a sync event to guarantee the Init() function is called first
 		ThreadInitSyncEvent	= FPlatformProcess::GetSynchEventFromPool(true);
diff --git a/Engine/Source/Runtime/Core/Private/Windows/WindowsRunnableThread.cpp b/Engine/Source/Runtime/Core/Private/Windows/WindowsRunnableThread.cpp
--- a/Engine/Source/Runtime/Core/Private/Windows/WindowsRunnableThread.cpp
+++ b/Engine/Source/Runtime/Core/Private/Windows/WindowsRunnableThread.cpp
@@ -12,7 +12,7 @@ uint32 FRunnableThreadWin::GuardedRun()
 {
 	uint32 ExitCode = 1;
 
-	FPlatformProcess::SetThreadAffinityMask(ThreadAffintyMask);
+	FPlatformProcess::SetThreadAffinityMask(ThreadAffinityMask);
 
 #if PLATFORM_XBOXONE
 	UE_LOG(LogThreadingWindows, Log, TEXT("Runnable thread %s is on Process %d."), *ThreadName  , static_cast<uint32>(::GetCurrentProcessorNumber()) );
@@ -48,10 +48,6 @@ uint32 FRunnableThreadWin::GuardedRun()
 		ExitCode = Run();
 	}
 
-#if STATS
-	FThreadStats::Shutdown();
-#endif
-
 	return ExitCode;
 }
 
@@ -62,6 +58,9 @@ uint32 FRunnableThreadWin::Run()
 	uint32 ExitCode = 1;
 	check(Runnable);
 
+	// Setup TLS for this thread, used by ITlsAutoCleanup objects.
+	SetTls();
+
 	// Initialize the runnable object
 	if (Runnable->Init() == true)
 	{
@@ -78,5 +77,10 @@ uint32 FRunnableThreadWin::Run()
 		ThreadInitSyncEvent->Trigger();
 	}
 
+#if STATS
+	FThreadStats::Shutdown();
+#endif
+	FreeTls();
+
 	return ExitCode;
 }
diff --git a/Engine/Source/Runtime/Core/Private/Windows/WindowsRunnableThread.h b/Engine/Source/Runtime/Core/Private/Windows/WindowsRunnableThread.h
--- a/Engine/Source/Runtime/Core/Private/Windows/WindowsRunnableThread.h
+++ b/Engine/Source/Runtime/Core/Private/Windows/WindowsRunnableThread.h
@@ -13,24 +13,6 @@ class FRunnableThreadWin
 	/** The thread handle for the thread. */
 	HANDLE Thread;
 
-	/** The runnable object to execute on this thread. */
-	FRunnable* Runnable;
-
-	/** Sync event to make sure that Init() has been completed before allowing the main thread to continue. */
-	FEvent* ThreadInitSyncEvent;
-
-	/** The priority to run the thread at. */
-	EThreadPriority ThreadPriority;
-
-	/** The Affinity to run the thread with. */
-	uint64 ThreadAffintyMask;
-
-	/** ID set during thread creation. */
-	uint32 ThreadID;
-
-	/** Name of the thread. */
-	FString ThreadName;
-
 	/**
 	 * Helper function to set thread names, visible by the debugger.
 	 * @param ThreadID		Thread ID for the thread who's name is going to be set
@@ -97,13 +79,7 @@ public:
 
 	FRunnableThreadWin( )
 		: Thread(NULL)
-		, Runnable(nullptr)
-		, ThreadInitSyncEvent(nullptr)
-		, ThreadPriority(TPri_Normal)
-		, ThreadID(0)
-	{
-
-	}
+	{}
 
 	~FRunnableThreadWin( )
 	{
@@ -112,7 +88,6 @@ public:
 		{
 			Kill(true);
 		}
-		FRunnableThread::GetThreadRegistry().Remove(ThreadID);
 	}
 	
 	virtual void SetThreadPriority( EThreadPriority NewPriority ) override
@@ -175,16 +150,6 @@ public:
 		WaitForSingleObject(Thread,INFINITE);
 	}
 
-	virtual uint32 GetThreadID( ) override
-	{
-		return ThreadID;
-	}
-
-	virtual FString GetThreadName( ) override
-	{
-		return ThreadName;
-	}
-
 protected:
 
 	virtual bool CreateInternal( FRunnable* InRunnable, const TCHAR* InThreadName,
@@ -193,7 +158,7 @@ protected:
 	{
 		check(InRunnable);
 		Runnable = InRunnable;
-		ThreadAffintyMask = InThreadAffinityMask;
+		ThreadAffinityMask = InThreadAffinityMask;
 
 		// Create a sync event to guarantee the Init() function is called first
 		ThreadInitSyncEvent	= FPlatformProcess::GetSynchEventFromPool(true);
diff --git a/Engine/Source/Runtime/Core/Public/HAL/ThreadingBase.h b/Engine/Source/Runtime/Core/Public/HAL/ThreadingBase.h
--- a/Engine/Source/Runtime/Core/Public/HAL/ThreadingBase.h
+++ b/Engine/Source/Runtime/Core/Public/HAL/ThreadingBase.h
@@ -2,6 +2,10 @@
 
 #pragma once
 
+#include "Function.h"
+
+class ITlsAutoCleanup;
+class FRunnableThread;
 
 /**
  * The list of enumerated thread priorities we support
@@ -214,8 +218,6 @@ public:
 		return nullptr;
 	}
 
-public:
-
 	/** Virtual destructor */
 	virtual ~FRunnable() { }
 };
@@ -226,10 +228,19 @@ public:
  *
  * This interface specifies the methods used to manage a thread's life cycle.
  */
-class FRunnableThread
+class CORE_API FRunnableThread
 {
+	friend class FThreadSingletonInitializer;
+	friend class ITlsAutoCleanup;
+
+	/** Index of TLS slot for FRunnableThread pointer. */
+	static uint32 RunnableTlsSlot;
+
 public:
 
+	/** Initializes a Tls slot for storing the runnable thread pointer. */
+	static void InitializeTls();
+
 	/**
 	* Factory method to create a thread with the specified stack size and thread priority.
 	*
@@ -242,7 +253,7 @@ public:
 	* @return The newly created thread or nullptr if it failed
 	*/
 	DEPRECATED(4.3, "Function deprecated. Use FRunnableThread::Create without bAutoDeleteSelf and bAutoDeleteRunnable params and delete thread and runnable manually.")
-	CORE_API static FRunnableThread* Create(
+	static FRunnableThread* Create(
 		class FRunnable* InRunnable,
 		const TCHAR* ThreadName,
 		bool bAutoDeleteSelf,
@@ -260,7 +271,7 @@ public:
 	 * @param InThreadPri Tells the thread whether it needs to adjust its priority or not. Defaults to normal priority
 	 * @return The newly created thread or nullptr if it failed
 	 */
-	CORE_API static FRunnableThread* Create(
+	static FRunnableThread* Create(
 		class FRunnable* InRunnable,
 		const TCHAR* ThreadName,
 		uint32 InStackSize = 0,
@@ -303,7 +314,10 @@ public:
 	 * @return ID that was set by CreateThread
 	 * @see GetThreadName
 	 */
-	virtual uint32 GetThreadID() = 0;
+	const uint32 GetThreadID() const
+	{
+		return ThreadID;
+	}
 
 	/**
 	 * Retrieves the given name of the thread
@@ -311,85 +325,16 @@ public:
 	 * @return Name that was set by CreateThread
 	 * @see GetThreadID
 	 */
-	virtual FString GetThreadName() = 0;
-
-	struct FThreadRegistry
+	const FString& GetThreadName() const
 	{
-		void Add( uint32 ID, FRunnableThread* Thread )
-		{
-			Lock();
-			Registry.FindOrAdd( ID ) = Thread;
-			Updated = true;
-			Unlock();
-		}
-
-		void Remove( uint32 ID )
-		{
-			Lock();
-			Registry.Remove(ID);
-			Updated = true;
-			Unlock();
-		}
-
-		int32 GetThreadCount()
-		{
-			Lock();
-			int32 RetVal = Registry.Num();
-			Unlock();
-			return RetVal;
-		}
-
-		bool IsUpdated()
-		{
-			return Updated;
-		}
-
-		void Lock()
-		{
-			CriticalSection.Lock();
-		}
-
-		void Unlock()
-		{
-			CriticalSection.Unlock();
-		}
-
-		void ClearUpdated()
-		{
-			Updated = false;
-		}
-
-		TMap<uint32, FRunnableThread*>::TConstIterator CreateConstIterator()
-		{
-			return Registry.CreateConstIterator();
-		}
-
-		FRunnableThread* GetThread( uint32 ID )
-		{
-			return Registry.FindRef(ID);
-		}
-
-	private:
-
-		TMap<uint32, FRunnableThread*> Registry;
-		bool Updated;
-		FCriticalSection CriticalSection;
-	};
-
-	static CORE_API FThreadRegistry& GetThreadRegistry()
-	{
-		static FThreadRegistry ThreadRegistry;
-		return ThreadRegistry;
+		return ThreadName;
 	}
 
-	/** @return a delegate that is called when this runnable has been destroyed. */
-	FSimpleMulticastDelegate& OnThreadDestroyed()
-	{
-		return ThreadDestroyedDelegate;
-	}
+	/** Default constructor. */
+	FRunnableThread();
 
 	/** Virtual destructor */
-	virtual ~FRunnableThread();
+	virtual ~FRunnableThread(){}
 
 protected:
 
@@ -406,8 +351,41 @@ protected:
 		uint32 InStackSize = 0,
 		EThreadPriority InThreadPri = TPri_Normal, uint64 InThreadAffinityMask = 0 ) = 0;
 
-	/** Called when the this runnable has been destroyed, so we should clean-up memory allocated by misc classes. */
-	FSimpleMulticastDelegate ThreadDestroyedDelegate;
+	/** Stores this instance in the runnable thread TLS slot. */
+	void SetTls();
+
+	/** Deletes all ITlsAutoCleanup objects created for this thread. */
+	void FreeTls();
+
+	/**
+	 * @return a runnable thread that is executing this runnable, if return value is nullptr, it means the running thread can be game thread or a thread created outside the runnable interface
+	 */
+	static FRunnableThread* GetRunnableThread()
+	{
+		FRunnableThread* RunnableThread = (FRunnableThread*)FPlatformTLS::GetTlsValue( RunnableTlsSlot );
+		return RunnableThread;
+	}
+
+	/** Holds the name of the thread. */
+	FString ThreadName;
+
+	/** The runnable object to execute on this thread. */
+	FRunnable* Runnable;
+
+	/** Sync event to make sure that Init() has been completed before allowing the main thread to continue. */
+	FEvent* ThreadInitSyncEvent;
+
+	/** The Affinity to run the thread with. */
+	uint64 ThreadAffinityMask;
+
+	/** An array of ITlsAutoCleanup based instances that needs to be deleted before the thread will die. */
+	TArray<ITlsAutoCleanup*> TlsInstances;
+
+	/** The priority to run the thread at. */
+	EThreadPriority ThreadPriority;
+
+	/** ID set during thread creation. */
+	uint32 ThreadID;
 };
 
 
@@ -946,21 +924,18 @@ extern CORE_API bool IsInRHIThread();
 /** Thread used for RHI */
 extern CORE_API FRunnableThread* GRHIThread;
 
-
-/**
- * Minimal base class for the thread singleton.
- */
-struct FThreadSingleton
+/** Interface required for auto-cleanup of instances of classes stored in the TLS. */
+class CORE_API ITlsAutoCleanup
 {
-	typedef FThreadSingleton* (*TCreateSingletonFuncPtr)();
-	typedef void (FThreadSingleton::*TDestroySingletonFuncPtr)();
-
-	CORE_API virtual ~FThreadSingleton() { }
+public:
+	/** Virtual destructor. */
+	virtual ~ITlsAutoCleanup()
+	{}
 
-	virtual void DeleteThis() = 0;
+	/** Register this instance to be auto-cleanup. */
+	void Register();
 };
 
-
 /**
  * Thread singleton initializer.
  */
@@ -969,19 +944,18 @@ class FThreadSingletonInitializer
 public:
 
 	/**
-	 * @return an instance of a singleton for the current thread.
-	 */
-	static CORE_API FThreadSingleton* Get( const FThreadSingleton::TCreateSingletonFuncPtr CreateFunc, const FThreadSingleton::TDestroySingletonFuncPtr DestroyFunc, uint32& TlsSlot );
+	* @return an instance of a singleton for the current thread.
+	*/
+	static CORE_API ITlsAutoCleanup* Get( const TFunctionRef<ITlsAutoCleanup*()>& CreateInstance, uint32& TlsSlot );
 };
 
 
 /**
  * This a special version of singleton. It means that there is created only one instance for each thread.
- * Calling Get() method is thread-safe, but first call should be done on the game thread.
- * @see DECLARE_THREAD_SINGLETON usage.
+ * Calling Get() method is thread-safe.
  */
 template < class T >
-class TThreadSingleton : public FThreadSingleton
+class TThreadSingleton : public ITlsAutoCleanup
 {
 	/**
 	 * @return TLS slot that holds a TThreadSingleton.
@@ -997,23 +971,16 @@ protected:
 	/** Default constructor. */
 	TThreadSingleton()
 		: ThreadId(FPlatformTLS::GetCurrentThreadId())
-	{ }
+	{}
 
 	/**
 	 * @return a new instance of the thread singleton.
 	 */
-	static FThreadSingleton* Create()
+	static ITlsAutoCleanup* CreateInstance()
 	{
 		return new T();
 	}
 
-	/**	Deletes this instance of the thread singleton.  */
-	virtual void DeleteThis()
-	{
-		T* This = (T*)this;
-		delete This;
-	}
-
 	/** Thread ID of this thread singleton. */
 	const uint32 ThreadId;
 
@@ -1024,10 +991,9 @@ public:
 	 */
 	FORCEINLINE static T& Get()
 	{
-		return *(T*)FThreadSingletonInitializer::Get(&T::Create, &FThreadSingleton::DeleteThis, T::GetTlsSlot());
+		return *(T*)FThreadSingletonInitializer::Get( [](){ return (ITlsAutoCleanup*)new T(); }, T::GetTlsSlot() );
 	}
 };
 
-
 #define DECLARE_THREAD_SINGLETON(ClassType) \
-	static auto& GForceInitAtBoot_ThreadSingletonInitializer_##ClassType = ClassType::Get();
+	EMIT_DEPRECATED_WARNING_MESSAGE("DECLARE_THREAD_SINGLETON is deprecated in 4.8. It's no longer needed and can be removed.")
\ No newline at end of file
diff --git a/Engine/Source/Runtime/Core/Public/HTML5/HTML5PlatformRunnableThread.h b/Engine/Source/Runtime/Core/Public/HTML5/HTML5PlatformRunnableThread.h
--- a/Engine/Source/Runtime/Core/Public/HTML5/HTML5PlatformRunnableThread.h
+++ b/Engine/Source/Runtime/Core/Public/HTML5/HTML5PlatformRunnableThread.h
@@ -33,16 +33,6 @@ public:
 
 	}
 
-	virtual uint32 GetThreadID ()
-	{
-		return 0;
-	}
-
-	virtual FString GetThreadName ()
-	{
-		return TEXT("");
-	}
-
 public:
 
 
diff --git a/Engine/Source/Runtime/Core/Public/Misc/TypeContainer.h b/Engine/Source/Runtime/Core/Public/Misc/TypeContainer.h
--- a/Engine/Source/Runtime/Core/Public/Misc/TypeContainer.h
+++ b/Engine/Source/Runtime/Core/Public/Misc/TypeContainer.h
@@ -173,6 +173,8 @@ class FTypeContainer
 				FRunnableThread::GetThreadRegistry().Unlock();*/
 
 				FPlatformTLS::SetTlsValue(TlsSlot, Instance);
+				// @see ITlsAutoCleanup
+				//Instance->Register();
 			}
 
 			return *((TSharedPtr<void>*)Instance);
diff --git a/Engine/Source/Runtime/Launch/Private/LaunchEngineLoop.cpp b/Engine/Source/Runtime/Launch/Private/LaunchEngineLoop.cpp
--- a/Engine/Source/Runtime/Launch/Private/LaunchEngineLoop.cpp
+++ b/Engine/Source/Runtime/Launch/Private/LaunchEngineLoop.cpp
@@ -1048,6 +1048,16 @@ int32 FEngineLoop::PreInit( const TCHAR* CmdLine )
 	}
 #endif //WITH_EDITOR
 
+	FRunnableThread::InitializeTls();
+
+	// initialize task graph sub-system with potential multiple threads
+	FTaskGraphInterface::Startup( FPlatformMisc::NumberOfCores() );
+	FTaskGraphInterface::Get().AttachToThread( ENamedThreads::GameThread );
+
+#if STATS
+	FThreadStats::StartThread();
+#endif
+
 	if (FPlatformProcess::SupportsMultithreading())
 	{
 		GThreadPool	= FQueuedThreadPool::Allocate();
@@ -2753,14 +2763,6 @@ bool FEngineLoop::AppInit( )
 	FApp::InitializeSession();
 #endif
 
-	// initialize task graph sub-system with potential multiple threads
-	FTaskGraphInterface::Startup(FPlatformMisc::NumberOfCores());
-	FTaskGraphInterface::Get().AttachToThread(ENamedThreads::GameThread);
-
-#if STATS
-	FThreadStats::StartThread();
-#endif
-
 #if WITH_ENGINE
 	// Earliest place to init the online subsystems
 	// Code needs GConfigFile to be valid
@@ -2835,7 +2837,7 @@ void FEngineLoop::AppPreExit( )
 	FCoreDelegates::OnExit.Broadcast();
 
 	// Clean up the thread pool
-	if (GThreadPool != NULL)
+	if (GThreadPool != nullptr)
 	{
 		GThreadPool->Destroy();
 	}
diff --git a/Engine/Source/Runtime/RenderCore/Private/RenderingThread.cpp b/Engine/Source/Runtime/RenderCore/Private/RenderingThread.cpp
--- a/Engine/Source/Runtime/RenderCore/Private/RenderingThread.cpp
+++ b/Engine/Source/Runtime/RenderCore/Private/RenderingThread.cpp
@@ -337,7 +337,7 @@ public:
 	 */
 	FEvent* TaskGraphBoundSyncEvent;
 
-	FRenderingThread() : FRunnable()
+	FRenderingThread()
 	{
 		TaskGraphBoundSyncEvent	= FPlatformProcess::GetSynchEventFromPool(true);
 		RHIFlushResources();
@@ -350,7 +350,7 @@ public:
 	}
 
 	// FRunnable interface.
-	virtual bool Init(void) 
+	virtual bool Init(void) override
 	{ 
 		GRenderThreadId = FPlatformTLS::GetCurrentThreadId();
 
@@ -360,7 +360,7 @@ public:
 		return true; 
 	}
 
-	virtual void Exit(void) 
+	virtual void Exit(void) override
 	{
 		// Release rendering context ownership on the current thread
 		RHIReleaseThreadOwnership();
@@ -368,11 +368,7 @@ public:
 		GRenderThreadId = 0;
 	}
 
-	virtual void Stop(void)
-	{
-	}
-
-	virtual uint32 Run(void)
+	virtual uint32 Run(void) override
 	{
 		FPlatformProcess::SetupGameOrRenderThread(true);
 
