Fix the Mac cursor not warping to the correct position when exiting high-precision mouse mode.

diff --git a/Engine/Source/Runtime/Core/Private/Mac/MacCursor.cpp b/Engine/Source/Runtime/Core/Private/Mac/MacCursor.cpp
--- a/Engine/Source/Runtime/Core/Private/Mac/MacCursor.cpp
+++ b/Engine/Source/Runtime/Core/Private/Mac/MacCursor.cpp
@@ -435,8 +435,9 @@ void FMacCursor::SetHighPrecisionMouseMode( bool const bEnable )
 		// On disable put the cursor where the user would expect it
 		if ( !bEnable && (!CurrentCursor || !bIsVisible) )
 		{
-			FVector2D Pos = GetPosition();
-			SetPosition(Pos.X, Pos.Y);
+			FVector2D Pos = GetPosition() / MouseScale;
+			UpdateCursorClipping(Pos);
+			WarpCursor(Pos.X, Pos.Y);
 		}
 	}
 }
