PR #964: FBase64::Decode: fix signed/unsigned comparison bug (Contributed by [REDACTED_PERSON])

diff --git a/Engine/Source/Runtime/Core/Private/Misc/Base64.cpp b/Engine/Source/Runtime/Core/Private/Misc/Base64.cpp
--- a/Engine/Source/Runtime/Core/Private/Misc/Base64.cpp
+++ b/Engine/Source/Runtime/Core/Private/Misc/Base64.cpp
@@ -222,7 +222,7 @@ bool FBase64::Decode(const ANSICHAR* Source, uint32 Length, uint8* Dest, uint32&
 			}
 			DecodedValues[Index] = DecodingAlphabet[(int32)(*Source++)];
 			// Abort on values that we don't understand
-			if (DecodedValues[Index] == -1)
+			if (DecodedValues[Index] == 0xFF)
 			{
 				return false;
 			}
