Integrate PR #651 : Input plugin support for iOS and OSX
- The Mac implementation has already been merged, so this adds the iOS version to match.

diff --git a/Engine/Source/Runtime/Core/Private/IOS/IOSApplication.cpp b/Engine/Source/Runtime/Core/Private/IOS/IOSApplication.cpp
--- a/Engine/Source/Runtime/Core/Private/IOS/IOSApplication.cpp
+++ b/Engine/Source/Runtime/Core/Private/IOS/IOSApplication.cpp
@@ -5,6 +5,9 @@
 #include "IOSInputInterface.h"
 #include "IOSWindow.h"
 #include "IOSAppDelegate.h"
+#include "IInputDeviceModule.h"
+#include "IInputInterface.h"
+#include "IInputDevice.h"
 
 FIOSApplication* FIOSApplication::CreateIOSApplication()
 {
@@ -14,7 +17,7 @@ FIOSApplication* FIOSApplication::CreateIOSApplication()
 FIOSApplication::FIOSApplication()
 	: GenericApplication( NULL )
 	, InputInterface( FIOSInputInterface::Create( MessageHandler ) )
-
+	, bHasLoadedInputPlugins(false)
 {
 }
 
@@ -31,13 +34,47 @@ void FIOSApplication::SetMessageHandler( const TSharedRef< FGenericApplicationMe
 {
 	GenericApplication::SetMessageHandler(InMessageHandler);
 	InputInterface->SetMessageHandler(InMessageHandler);
+
+	TArray<IInputDeviceModule*> PluginImplementations = IModularFeatures::Get().GetModularFeatureImplementations<IInputDeviceModule>(IInputDeviceModule::GetModularFeatureName());
+	for (auto DeviceIt = ExternalInputDevices.CreateIterator(); DeviceIt; ++DeviceIt)
+	{
+		(*DeviceIt)->SetMessageHandler(InMessageHandler);
+	}
+}
+
+void FIOSApplication::AddExternalInputDevice(TSharedPtr<IInputDevice> InputDevice)
+{
+	if (InputDevice.IsValid())
+	{
+		ExternalInputDevices.Add(InputDevice);
+	}
 }
 
 void FIOSApplication::PollGameDeviceState( const float TimeDelta )
 {
+	// initialize any externally-implemented input devices (we delay load initialize the array so any plugins have had time to load)
+	if (!bHasLoadedInputPlugins)
+	{
+		TArray<IInputDeviceModule*> PluginImplementations = IModularFeatures::Get().GetModularFeatureImplementations<IInputDeviceModule>(IInputDeviceModule::GetModularFeatureName());
+		for (auto InputPluginIt = PluginImplementations.CreateIterator(); InputPluginIt; ++InputPluginIt)
+		{
+			TSharedPtr<IInputDevice> Device = (*InputPluginIt)->CreateInputDevice(MessageHandler);
+			AddExternalInputDevice(Device);
+		}
+
+		bHasLoadedInputPlugins = true;
+	}
+
 	// Poll game device state and send new events
 	InputInterface->Tick(TimeDelta);
 	InputInterface->SendControllerEvents();
+
+	// Poll externally-implemented devices
+	for (auto DeviceIt = ExternalInputDevices.CreateIterator(); DeviceIt; ++DeviceIt)
+	{
+		(*DeviceIt)->Tick(TimeDelta);
+		(*DeviceIt)->SendControllerEvents();
+	}
 }
 
 FPlatformRect FIOSApplication::GetWorkArea( const FPlatformRect& CurrentWindow ) const
diff --git a/Engine/Source/Runtime/Core/Public/IOS/IOSApplication.h b/Engine/Source/Runtime/Core/Public/IOS/IOSApplication.h
--- a/Engine/Source/Runtime/Core/Public/IOS/IOSApplication.h
+++ b/Engine/Source/Runtime/Core/Public/IOS/IOSApplication.h
@@ -25,6 +25,8 @@ public:
 
 	virtual TSharedRef< FGenericWindow > MakeWindow() override;
 
+	virtual void AddExternalInputDevice(TSharedPtr<class IInputDevice> InputDevice);
+
 protected:
 	virtual void InitializeWindow( const TSharedRef< FGenericWindow >& Window, const TSharedRef< FGenericWindowDefinition >& InDefinition, const TSharedPtr< FGenericWindow >& InParent, const bool bShowImmediately ) override;
 
@@ -37,5 +39,9 @@ private:
 
 	TSharedPtr< class FIOSInputInterface > InputInterface;
 
+	/** List of input devices implemented in external modules. */
+	TArray< TSharedPtr<class IInputDevice> > ExternalInputDevices;
+	bool bHasLoadedInputPlugins;
+
 	TArray< TSharedRef< FIOSWindow > > Windows;
 };
