Adding 'override' qualifiers on fix-as-you-go basis.

diff --git a/Engine/Source/Runtime/Core/Private/Async/TaskGraph.cpp b/Engine/Source/Runtime/Core/Private/Async/TaskGraph.cpp
--- a/Engine/Source/Runtime/Core/Private/Async/TaskGraph.cpp
+++ b/Engine/Source/Runtime/Core/Private/Async/TaskGraph.cpp
@@ -544,7 +544,7 @@ public:
 	 *
 	 * @return True if initialization was successful, false otherwise
 	 */
-	virtual bool Init()
+	virtual bool Init() override
 	{
 		InitializeForCurrentThread();
 		return true;
@@ -556,7 +556,7 @@ public:
 	 *
 	 * @return The exit code of the runnable object
 	 */
-	virtual uint32 Run()
+	virtual uint32 Run() override
 	{
 		ProcessTasksUntilQuit(0);
 		return 0;
@@ -565,7 +565,7 @@ public:
 	/**
 	 * This is called if a thread is requested to terminate early
 	 */
-	virtual void Stop()
+	virtual void Stop() override
 	{
 		RequestQuit(0);
 	}
@@ -573,7 +573,7 @@ public:
 	/**
 	 * Called in the context of the aggregating thread to perform any cleanup.
 	 */
-	virtual void Exit()
+	virtual void Exit() override
 	{
 	}
 
@@ -989,7 +989,7 @@ public:
 		}
 	}
 
-	virtual void TriggerEventWhenTasksComplete(FEvent* InEvent, const FGraphEventArray& Tasks, ENamedThreads::Type CurrentThreadIfKnown = ENamedThreads::AnyThread)
+	virtual void TriggerEventWhenTasksComplete(FEvent* InEvent, const FGraphEventArray& Tasks, ENamedThreads::Type CurrentThreadIfKnown = ENamedThreads::AnyThread) override
 	{
 		check(InEvent);
 		bool bAnyPending = false;
diff --git a/Engine/Source/Runtime/Core/Private/GenericPlatform/GenericPlatformHttp.cpp b/Engine/Source/Runtime/Core/Private/GenericPlatform/GenericPlatformHttp.cpp
--- a/Engine/Source/Runtime/Core/Private/GenericPlatform/GenericPlatformHttp.cpp
+++ b/Engine/Source/Runtime/Core/Private/GenericPlatform/GenericPlatformHttp.cpp
@@ -20,7 +20,7 @@ public:
 	virtual TArray<FString> GetAllHeaders() override { return TArray<FString>(); }
 	virtual FString GetContentType() override { return TEXT(""); }
 	virtual int32 GetContentLength() override { return 0; }
-	virtual const TArray<uint8>& GetContent() { static TArray<uint8> Temp; return Temp; }
+	virtual const TArray<uint8>& GetContent() override { static TArray<uint8> Temp; return Temp; }
 
 	// IHttpRequest
 	virtual FString GetVerb() override { return TEXT(""); }
diff --git a/Engine/Source/Runtime/Core/Private/Linux/LinuxPlatformOutputDevicesPrivate.h b/Engine/Source/Runtime/Core/Private/Linux/LinuxPlatformOutputDevicesPrivate.h
--- a/Engine/Source/Runtime/Core/Private/Linux/LinuxPlatformOutputDevicesPrivate.h
+++ b/Engine/Source/Runtime/Core/Private/Linux/LinuxPlatformOutputDevicesPrivate.h
@@ -19,7 +19,7 @@ public:
 	 * Error handling function that is being called from within the system wide global
 	 * error handler, e.g. using structured exception handling on the PC.
 	 */
-	void HandleError();
+	void HandleError() override;
 
 private:
 
diff --git a/Engine/Source/Runtime/Core/Public/GenericPlatform/GenericPlatformChunkInstall.h b/Engine/Source/Runtime/Core/Public/GenericPlatform/GenericPlatformChunkInstall.h
--- a/Engine/Source/Runtime/Core/Public/GenericPlatform/GenericPlatformChunkInstall.h
+++ b/Engine/Source/Runtime/Core/Public/GenericPlatform/GenericPlatformChunkInstall.h
@@ -224,7 +224,7 @@ public:
 	 * Does nothing in a shipping build.
 	 * @return				true if the opreation succeeds.
 	 **/
-	virtual bool DebugStartNextChunk()
+	virtual bool DebugStartNextChunk() override
 	{
 		return true;
 	}
@@ -235,7 +235,7 @@ public:
 	* @param Delegate		The delegate when the chunk is installed.
 	* @return				False if the delegate was not registered or the chunk is already installed. True on success.
 	*/
-	virtual bool SetChunkInstallDelgate(uint32 ChunkID, FPlatformChunkInstallCompleteDelegate Delegate)
+	virtual bool SetChunkInstallDelgate(uint32 ChunkID, FPlatformChunkInstallCompleteDelegate Delegate) override
 	{
 		return false;
 	}
@@ -246,7 +246,7 @@ public:
 	* @param Delegate		The delegate to remove.
 	* @return				False if the delegate was not registered with a call to SetChunkInstallDelgate. True on success.
 	*/
-	virtual void RemoveChunkInstallDelgate(uint32 ChunkID, FPlatformChunkInstallCompleteDelegate Delegate)
+	virtual void RemoveChunkInstallDelgate(uint32 ChunkID, FPlatformChunkInstallCompleteDelegate Delegate) override
 	{
 		return;
 	}
diff --git a/Engine/Source/Runtime/Core/Public/HAL/ConsoleManager.h b/Engine/Source/Runtime/Core/Public/HAL/ConsoleManager.h
--- a/Engine/Source/Runtime/Core/Public/HAL/ConsoleManager.h
+++ b/Engine/Source/Runtime/Core/Public/HAL/ConsoleManager.h
@@ -59,16 +59,16 @@ public:
 	virtual IConsoleVariable* RegisterConsoleVariableRef(const TCHAR* Name, float& RefValue, const TCHAR* Help, uint32 Flags) override;
 	virtual IConsoleVariable* RegisterConsoleVariableBitRef(const TCHAR* CVarName, const TCHAR* FlagName, uint32 BitNumber, uint8* Force0MaskPtr, uint8* Force1MaskPtr, const TCHAR* Help, uint32 Flags) override;
 	
-	virtual void CallAllConsoleVariableSinks();
+	virtual void CallAllConsoleVariableSinks() override;
 
 	DELEGATE_DEPRECATED("This function is deprecated - please replace any usage with RegisterConsoleVariableSink_Handle.")
-	virtual void RegisterConsoleVariableSink(const FConsoleCommandDelegate& Command);
+	virtual void RegisterConsoleVariableSink(const FConsoleCommandDelegate& Command) override;
 
 	DELEGATE_DEPRECATED("Delegate comparison is deprecated - please replace any usage with UnregisterConsoleVariableSink_Handle, passing the result of RegisterConsoleVariableSink_Handle.")
-	virtual void UnregisterConsoleVariableSink(const FConsoleCommandDelegate& Command);
+	virtual void UnregisterConsoleVariableSink(const FConsoleCommandDelegate& Command) override;
 
-	virtual FConsoleVariableSinkHandle RegisterConsoleVariableSink_Handle(const FConsoleCommandDelegate& Command);
-	virtual void UnregisterConsoleVariableSink_Handle(FConsoleVariableSinkHandle Handle);
+	virtual FConsoleVariableSinkHandle RegisterConsoleVariableSink_Handle(const FConsoleCommandDelegate& Command) override;
+	virtual void UnregisterConsoleVariableSink_Handle(FConsoleVariableSinkHandle Handle) override;
 
 	virtual IConsoleCommand* RegisterConsoleCommand(const TCHAR* Name, const TCHAR* Help, const FConsoleCommandDelegate& Command, uint32 Flags) override;
 	virtual IConsoleCommand* RegisterConsoleCommand(const TCHAR* Name, const TCHAR* Help, const FConsoleCommandWithArgsDelegate& Command, uint32 Flags) override;
@@ -79,8 +79,8 @@ public:
 	virtual IConsoleVariable* FindConsoleVariable(const TCHAR* Name) const override;
 	virtual void ForEachConsoleObject(const FConsoleObjectVisitor& Visitor, const TCHAR* ThatStartsWith) const override;
 	virtual bool ProcessUserConsoleInput(const TCHAR* InInput, FOutputDevice& Ar, UWorld* InWorld) override;
-	virtual void AddConsoleHistoryEntry(const TCHAR* Input);
-	virtual void GetConsoleHistory(TArray<FString>& Out);
+	virtual void AddConsoleHistoryEntry(const TCHAR* Input) override;
+	virtual void GetConsoleHistory(TArray<FString>& Out) override;
 	virtual bool IsNameRegistered(const TCHAR* Name) const override;	
 	virtual void RegisterThreadPropagation(uint32 ThreadId, IConsoleThreadPropagation* InCallback) override;
 	virtual void UnregisterConsoleObject( IConsoleObject* Object, bool bKeepState) override;
diff --git a/Engine/Source/Runtime/Core/Public/HAL/FeedbackContextAnsi.h b/Engine/Source/Runtime/Core/Public/HAL/FeedbackContextAnsi.h
--- a/Engine/Source/Runtime/Core/Public/HAL/FeedbackContextAnsi.h
+++ b/Engine/Source/Runtime/Core/Public/HAL/FeedbackContextAnsi.h
@@ -36,7 +36,7 @@ public:
 	, Context(nullptr)
 	, AuxOut(nullptr)
 	{}
-	void Serialize( const TCHAR* V, ELogVerbosity::Type Verbosity, const class FName& Category )
+	void Serialize( const TCHAR* V, ELogVerbosity::Type Verbosity, const class FName& Category ) override
 	{
 		if( Verbosity==ELogVerbosity::Error || Verbosity==ELogVerbosity::Warning || Verbosity==ELogVerbosity::Display )
 		{
@@ -100,7 +100,7 @@ public:
 		return true;
 	}
 
-	void SetContext( FContextSupplier* InSupplier )
+	void SetContext( FContextSupplier* InSupplier ) override
 	{
 		Context = InSupplier;
 	}
diff --git a/Engine/Source/Runtime/Core/Public/HAL/FileManagerGeneric.h b/Engine/Source/Runtime/Core/Public/HAL/FileManagerGeneric.h
--- a/Engine/Source/Runtime/Core/Public/HAL/FileManagerGeneric.h
+++ b/Engine/Source/Runtime/Core/Public/HAL/FileManagerGeneric.h
@@ -69,9 +69,9 @@ public:
 	bool SetTimeStamp( const TCHAR* Filename, FDateTime Timestamp ) override;
 	virtual FString GetFilenameOnDisk(const TCHAR* Filename) override;
 
-	virtual uint32	Copy( const TCHAR* InDestFile, const TCHAR* InSrcFile, bool ReplaceExisting, bool EvenIfReadOnly, bool Attributes, FCopyProgress* Progress );
-	virtual bool	MakeDirectory( const TCHAR* Path, bool Tree=0 );
-	virtual bool	DeleteDirectory( const TCHAR* Path, bool RequireExists=0, bool Tree=0 );
+	virtual uint32	Copy( const TCHAR* InDestFile, const TCHAR* InSrcFile, bool ReplaceExisting, bool EvenIfReadOnly, bool Attributes, FCopyProgress* Progress ) override;
+	virtual bool	MakeDirectory( const TCHAR* Path, bool Tree=0 ) override;
+	virtual bool	DeleteDirectory( const TCHAR* Path, bool RequireExists=0, bool Tree=0 ) override;
 
 	/** 
 	 * Call the Visit function of the visitor once for each file or directory in a single directory. This function does not explore subdirectories.
diff --git a/Engine/Source/Runtime/Core/Public/HAL/IPlatformFileCachedWrapper.h b/Engine/Source/Runtime/Core/Public/HAL/IPlatformFileCachedWrapper.h
--- a/Engine/Source/Runtime/Core/Public/HAL/IPlatformFileCachedWrapper.h
+++ b/Engine/Source/Runtime/Core/Public/HAL/IPlatformFileCachedWrapper.h
@@ -227,7 +227,7 @@ public:
 		LowerLevel = Inner;
 		return !!LowerLevel;
 	}
-	virtual bool ShouldBeUsed(IPlatformFile* Inner, const TCHAR* CmdLine) const
+	virtual bool ShouldBeUsed(IPlatformFile* Inner, const TCHAR* CmdLine) const override
 	{
 		return FPlatformProperties::RequiresCookedData();
 	}
diff --git a/Engine/Source/Runtime/Core/Public/HAL/OutputDevices.h b/Engine/Source/Runtime/Core/Public/HAL/OutputDevices.h
--- a/Engine/Source/Runtime/Core/Public/HAL/OutputDevices.h
+++ b/Engine/Source/Runtime/Core/Public/HAL/OutputDevices.h
@@ -368,7 +368,7 @@ public:
 	 * as we might have to call "delete" which cannot be done for static/ global
 	 * objects.
 	 */
-	void TearDown();
+	void TearDown() override;
 };
 
 /*-----------------------------------------------------------------------------
@@ -400,13 +400,13 @@ public:
 	 * as we have to call "delete" which cannot be done for static/ global
 	 * objects.
 	 */
-	void TearDown();
+	void TearDown() override;
 
 	/**
 	 * Flush the write cache so the file isn't truncated in case we crash right
 	 * after calling this function.
 	 */
-	void Flush();
+	void Flush() override;
 
 	virtual void Serialize( const TCHAR* Data, ELogVerbosity::Type Verbosity, const class FName& Category ) override;
 
@@ -514,7 +514,7 @@ public:
 	 * Error handling function that is being called from within the system wide global
 	 * error handler, e.g. using structured exception handling on the PC.
 	 */
-	void HandleError();
+	void HandleError() override;
 
 private:
 
@@ -539,4 +539,4 @@ CORE_API DECLARE_LOG_CATEGORY_EXTERN(LogProfilingDebugging, Log, All);
 CORE_API DECLARE_LOG_CATEGORY_EXTERN(LogCore, Log, All);
 
 // Temporary log category, generally you should not check things in that use this
-CORE_API DECLARE_LOG_CATEGORY_EXTERN(LogTemp, Log, All);
\ No newline at end of file
+CORE_API DECLARE_LOG_CATEGORY_EXTERN(LogTemp, Log, All);
diff --git a/Engine/Source/Runtime/Core/Public/Internationalization/InternationalizationMetadata.h b/Engine/Source/Runtime/Core/Public/Internationalization/InternationalizationMetadata.h
--- a/Engine/Source/Runtime/Core/Public/Internationalization/InternationalizationMetadata.h
+++ b/Engine/Source/Runtime/Core/Public/Internationalization/InternationalizationMetadata.h
@@ -172,9 +172,9 @@ class CORE_API FLocMetadataValueString : public FLocMetadataValue
 {
 public:
 	FLocMetadataValueString(const FString& InString) : Value(InString) {Type = ELocMetadataType::String;}
-	virtual FString AsString() {return Value;}
+	virtual FString AsString() override {return Value;}
 
-	virtual TSharedRef<FLocMetadataValue> Clone() const;
+	virtual TSharedRef<FLocMetadataValue> Clone() const override;
 	void SetString( const FString& InString ) { Value = InString; }
 
 protected:
@@ -192,9 +192,9 @@ class CORE_API FLocMetadataValueBoolean : public FLocMetadataValue
 {
 public:
 	FLocMetadataValueBoolean(bool InBool) : Value(InBool) {Type = ELocMetadataType::Boolean;}
-	virtual bool AsBool() {return Value;}
+	virtual bool AsBool() override {return Value;}
 
-	virtual TSharedRef<FLocMetadataValue> Clone() const;
+	virtual TSharedRef<FLocMetadataValue> Clone() const override;
 	
 protected:
 	bool Value;
@@ -210,9 +210,9 @@ class CORE_API FLocMetadataValueArray : public FLocMetadataValue
 {
 public:
 	FLocMetadataValueArray(const TArray< TSharedPtr<FLocMetadataValue> >& InArray) : Value(InArray) {Type = ELocMetadataType::Array;}
-	virtual TArray< TSharedPtr<FLocMetadataValue> > AsArray() {return Value;}
+	virtual TArray< TSharedPtr<FLocMetadataValue> > AsArray() override {return Value;}
 
-	virtual TSharedRef<FLocMetadataValue> Clone() const;
+	virtual TSharedRef<FLocMetadataValue> Clone() const override;
 	
 protected:
 	TArray< TSharedPtr<FLocMetadataValue> > Value;
@@ -228,9 +228,9 @@ class CORE_API FLocMetadataValueObject : public FLocMetadataValue
 {
 public:
 	FLocMetadataValueObject(TSharedPtr<FLocMetadataObject> InObject) : Value(InObject) {Type = ELocMetadataType::Object;}
-	virtual TSharedPtr<FLocMetadataObject> AsObject() {return Value;}
+	virtual TSharedPtr<FLocMetadataObject> AsObject() override {return Value;}
 
-	virtual TSharedRef<FLocMetadataValue> Clone() const;
+	virtual TSharedRef<FLocMetadataValue> Clone() const override;
 	
 protected:
 	TSharedPtr<FLocMetadataObject> Value;
diff --git a/Engine/Source/Runtime/Core/Public/Misc/AutomationTest.h b/Engine/Source/Runtime/Core/Public/Misc/AutomationTest.h
--- a/Engine/Source/Runtime/Core/Public/Misc/AutomationTest.h
+++ b/Engine/Source/Runtime/Core/Public/Misc/AutomationTest.h
@@ -1066,9 +1066,9 @@ public: \
 	public: \
 		TClass( const FString& InName ) \
 		:FAutomationTestBase( InName, false ) {} \
-		virtual uint32 GetTestFlags() const { return TFlags; } \
+		virtual uint32 GetTestFlags() const override { return TFlags; } \
 		virtual bool IsStressTest() const { return false; } \
-		virtual uint32 GetRequiredDeviceNum() const { return 1; } \
+		virtual uint32 GetRequiredDeviceNum() const override { return 1; } \
 	protected: \
 		virtual void GetTests(TArray<FString>& OutBeautifiedNames, TArray <FString>& OutTestCommands) const override \
 		{ \
@@ -1085,9 +1085,9 @@ public: \
 	public: \
 		TClass( const FString& InName ) \
 		:FAutomationTestBase( InName, true ) {} \
-		virtual uint32 GetTestFlags() const { return (TFlags & ~(EAutomationTestFlags::ATF_SmokeTest)); } \
+		virtual uint32 GetTestFlags() const override { return (TFlags & ~(EAutomationTestFlags::ATF_SmokeTest)); } \
 		virtual bool IsStressTest() const { return true; } \
-		virtual uint32 GetRequiredDeviceNum() const { return 1; } \
+		virtual uint32 GetRequiredDeviceNum() const override { return 1; } \
 	protected: \
 		virtual void GetTests(TArray<FString>& OutBeautifiedNames, TArray <FString>& OutTestCommands) const override; \
 		virtual bool RunTest(const FString& Parameters) override; \
@@ -1100,8 +1100,8 @@ public: \
 	public: \
 		TClass( const FString& InName ) \
 		:FAutomationTestBase( InName, false ) {} \
-		virtual uint32 GetTestFlags() const { return (TFlags & ~(EAutomationTestFlags::ATF_Editor | EAutomationTestFlags::ATF_Commandlet | EAutomationTestFlags::ATF_SmokeTest)); } \
-		virtual uint32 GetRequiredDeviceNum() const { return NumParticipants; } \
+		virtual uint32 GetTestFlags() const override { return (TFlags & ~(EAutomationTestFlags::ATF_Editor | EAutomationTestFlags::ATF_Commandlet | EAutomationTestFlags::ATF_SmokeTest)); } \
+		virtual uint32 GetRequiredDeviceNum() const override { return NumParticipants; } \
 	protected: \
 		virtual void GetTests(TArray<FString>& OutBeautifiedNames, TArray <FString>& OutTestCommands) const override \
 		{ \
diff --git a/Engine/Source/Runtime/Core/Public/Serialization/Archive.h b/Engine/Source/Runtime/Core/Public/Serialization/Archive.h
--- a/Engine/Source/Runtime/Core/Public/Serialization/Archive.h
+++ b/Engine/Source/Runtime/Core/Public/Serialization/Archive.h
@@ -113,9 +113,9 @@ public:
 	 *
 	 * This is overridden for the specific Archive Types
 	 **/
-	virtual FString GetArchiveName() const { return TEXT("FMemoryWriter"); }
+	virtual FString GetArchiveName() const override { return TEXT("FMemoryWriter"); }
 
-	int64 TotalSize()
+	int64 TotalSize() override
 	{
 		return Bytes.Num();
 	}
diff --git a/Engine/Source/Runtime/Json/Public/Dom/JsonValue.h b/Engine/Source/Runtime/Json/Public/Dom/JsonValue.h
--- a/Engine/Source/Runtime/Json/Public/Dom/JsonValue.h
+++ b/Engine/Source/Runtime/Json/Public/Dom/JsonValue.h
@@ -92,7 +92,7 @@ public:
 protected:
 	FString Value;
 
-	virtual FString GetType() const {return TEXT("String");}
+	virtual FString GetType() const override {return TEXT("String");}
 };
 
 
@@ -108,7 +108,7 @@ public:
 protected:
 	double Value;
 
-	virtual FString GetType() const {return TEXT("Number");}
+	virtual FString GetType() const override {return TEXT("Number");}
 };
 
 
@@ -124,7 +124,7 @@ public:
 protected:
 	bool Value;
 
-	virtual FString GetType() const {return TEXT("Boolean");}
+	virtual FString GetType() const override {return TEXT("Boolean");}
 };
 
 
@@ -138,7 +138,7 @@ public:
 protected:
 	TArray< TSharedPtr<FJsonValue> > Value;
 
-	virtual FString GetType() const {return TEXT("Array");}
+	virtual FString GetType() const override {return TEXT("Array");}
 };
 
 
@@ -152,7 +152,7 @@ public:
 protected:
 	TSharedPtr<FJsonObject> Value;
 
-	virtual FString GetType() const {return TEXT("Object");}
+	virtual FString GetType() const override {return TEXT("Object");}
 };
 
 
@@ -163,5 +163,5 @@ public:
 	FJsonValueNull() {Type = EJson::Null;}
 
 protected:
-	virtual FString GetType() const {return TEXT("Null");}
+	virtual FString GetType() const override {return TEXT("Null");}
 };
