Massively simplify (and comment!) TAutoPointer

diff --git a/Engine/Source/Runtime/Core/Public/UObject/AutoPointer.h b/Engine/Source/Runtime/Core/Public/UObject/AutoPointer.h
--- a/Engine/Source/Runtime/Core/Public/UObject/AutoPointer.h
+++ b/Engine/Source/Runtime/Core/Public/UObject/AutoPointer.h
@@ -3,155 +3,44 @@
 #pragma once
 
 
+// TAutoPointer wraps a smart-pointer and adds an implicit conversion to raw pointer
+// Its main use is for converting a variable from raw pointer to a smart pointer without breaking existing code
+// Currently only used in TAutoWeakObjectPtr
+// Not the same thing as TAutoPtr :(
 template<class T, class TBASE>
-class TAutoPointer : private TBASE
+class TAutoPointer : public TBASE
 {
 public:
 	/** NULL constructor **/
 	FORCEINLINE TAutoPointer()
 	{
 	}
-	/** Construct from a raw pointer **/
-	FORCEINLINE TAutoPointer(const T *Target) 
-		: TBASE(Target)
+	/** Construct from a single argument (I'd use inheriting constructors if all our compilers supported it) **/
+	template<typename X>
+	FORCEINLINE TAutoPointer(X&& Target)
+		: TBASE(Forward<X>(Target))
 	{
 	}
-	/**  Construct from the base type **/
-	FORCEINLINE explicit TAutoPointer(const TBASE &Other) 
-		: TBASE(Other)
-	{
-	}
-	/**  Construct from another auto pointer **/
-	template<class U, class UBASE>
-	FORCEINLINE explicit TAutoPointer(const TAutoPointer<U,UBASE> &Other) 
-		: TBASE(Other.Get())
-	{
-	}
-	/** Reset the pointer back to the NULL state **/
-	FORCEINLINE void Reset()
-	{
-		TBASE::Reset();
-	}
-	/**  Copy from a raw pointer **/
-	template<class U>
-	FORCEINLINE void operator=(const U *Target)
-	{
-		TBASE::operator=(Target);
-	}
-	/**  Construct from another base pointer type **/
-	FORCEINLINE void operator=(const TBASE &Other)
-	{
-		TBASE::operator=(Other);
-	}
-	/**  Construct from another auto pointer **/
-	template<class U, class UBASE>
-	FORCEINLINE void operator=(const TAutoPointer<U,UBASE> &Other)
-	{
-		(*this) = Other.Get();
-	}
-	/**  Compare for equality with a base pointer **/
-	FORCEINLINE bool operator==(const TBASE &Other) const
-	{
-		return (const TBASE&)*this == Other;
-	}
-	/**  Compare for equality with an auto pointer **/
-	template<class U, class UBASE>
-	FORCEINLINE bool operator==(const TAutoPointer<U,UBASE> &Other) const
-	{
-		return (const TBASE&)*this == Other;
-	}
-	/**  Compare for equality with a raw pointer **/
-	FORCEINLINE bool operator==(T* Other) const
-	{
-		return Get() == Other;
-	}
-	/**  Compare for equality with a raw pointer **/
-	FORCEINLINE bool operator==(const T* Other) const
-	{
-		return Get() == Other;
-	}
-	/**  Compare to NULL **/
-	FORCEINLINE bool operator==(TYPE_OF_NULL Other) const
-	{
-		checkSlow(!Other); // we assume this is NULL
-		return !IsValid();
-	}
-	/**  Compare for inequality with a base pointer **/
-	FORCEINLINE bool operator!=(const TBASE &Other) const
-	{
-		return (const TBASE&)*this != Other;
-	}
-	/**  Compare for equality with an auto pointer **/
-	template<class U, class UBASE>
-	FORCEINLINE bool operator!=(const TAutoPointer<U,UBASE> &Other) const
-	{
-		return (const TBASE&)*this != Other;
-	}
-	/**  Compare for inequality with a raw pointer	**/
-	FORCEINLINE bool operator!=(T* Other) const
-	{
-		return Get() != Other;
-	}
-	/**  Compare for inequality with a raw pointer	**/
-	FORCEINLINE bool operator!=(const T* Other) const
-	{
-		return Get() != Other;
-	}
-	/**  Compare for inequality with NULL **/
-	FORCEINLINE bool operator!=(TYPE_OF_NULL Other) const
-	{
-		checkSlow(!Other); // we assume this is NULL
-		return IsValid();
-	}
-	/** Dereference **/
-	FORCEINLINE T & operator*() const
-	{
-		return *Get();
-	}
-	/** Dereference **/
-	FORCEINLINE T * operator->() const
-	{
-		return Get();
-	}
-	/** Dereference **/
+
 	FORCEINLINE operator T* () const
 	{
 		return Get();
 	}
-	/** Conversion operator **/
 	FORCEINLINE operator const T* () const
 	{
 		return (const T*)Get();
 	}
-	/**  NULL test **/
-	FORCEINLINE operator bool() const
-	{
-		return !!IsValid();
-	}
-	/**  Dereference	**/
-	FORCEINLINE T *Get() const
-	{
-		return TBASE::Get();
-	}
-
-	/**  Test if this pointer is NULL	**/
-	FORCEINLINE bool IsValid() const
-	{
-		return TBASE::IsValid();
-	}
-
-	/** Hash function. */
-	FORCEINLINE friend uint32 GetTypeHash(const TAutoPointer& WeakObjectPtr)
-	{
-		return GetTypeHash(static_cast<const TBASE&>(WeakObjectPtr));
-	}
 
-	friend FArchive& operator<<( FArchive& Ar, TAutoPointer& WeakObjectPtr )
+#if PLATFORM_COMPILER_HAS_EXPLICIT_OPERATORS
+	FORCEINLINE explicit operator bool() const
+#else
+	// Can't use FORCEINLINE_EXPLICIT_OPERATOR_BOOL() due to already having an implicit pointer cast
+	// (which breaks the safe bool idiom it uses when PLATFORM_COMPILER_HAS_EXPLICIT_OPERATORS is 0)
+	FORCEINLINE operator bool() const
+#endif
 	{
-		Ar << static_cast<TBASE&>(WeakObjectPtr);
-		return Ar;
+		return Get() != nullptr;
 	}
-
 };
 
 
