Fix clang compile of TAutoPointer

diff --git a/Engine/Source/Runtime/Core/Public/UObject/AutoPointer.h b/Engine/Source/Runtime/Core/Public/UObject/AutoPointer.h
--- a/Engine/Source/Runtime/Core/Public/UObject/AutoPointer.h
+++ b/Engine/Source/Runtime/Core/Public/UObject/AutoPointer.h
@@ -24,11 +24,11 @@ public:
 
 	FORCEINLINE operator T* () const
 	{
-		return Get();
+		return TBASE::Get();
 	}
 	FORCEINLINE operator const T* () const
 	{
-		return (const T*)Get();
+		return (const T*)TBASE::Get();
 	}
 
 #if PLATFORM_COMPILER_HAS_EXPLICIT_OPERATORS
@@ -39,7 +39,7 @@ public:
 	FORCEINLINE operator bool() const
 #endif
 	{
-		return Get() != nullptr;
+		return TBASE::Get() != nullptr;
 	}
 };
 
