Adding 'override' qualifiers on a fix-as-you-go basis (#2).

diff --git a/Engine/Source/Runtime/Core/Public/HAL/OutputDevices.h b/Engine/Source/Runtime/Core/Public/HAL/OutputDevices.h
--- a/Engine/Source/Runtime/Core/Public/HAL/OutputDevices.h
+++ b/Engine/Source/Runtime/Core/Public/HAL/OutputDevices.h
@@ -361,7 +361,7 @@ public:
 	/**
 	 * Passes on the flush request to all current output devices.
 	 */
-	void Flush();
+	void Flush() override;
 
 	/**
 	 * Closes output device and cleans up. This can't happen in the destructor
diff --git a/Engine/Source/Runtime/Core/Public/Linux/LinuxPlatformFile.h b/Engine/Source/Runtime/Core/Public/Linux/LinuxPlatformFile.h
--- a/Engine/Source/Runtime/Core/Public/Linux/LinuxPlatformFile.h
+++ b/Engine/Source/Runtime/Core/Public/Linux/LinuxPlatformFile.h
@@ -37,5 +37,5 @@ public:
 	virtual bool CreateDirectory(const TCHAR* Directory) override;
 	virtual bool DeleteDirectory(const TCHAR* Directory) override;
 	bool CreateDirectoriesFromPath(const TCHAR* Path);
-	bool IterateDirectory(const TCHAR* Directory, FDirectoryVisitor& Visitor);
+	bool IterateDirectory(const TCHAR* Directory, FDirectoryVisitor& Visitor) override;
 };
diff --git a/Engine/Source/Runtime/Core/Public/Serialization/ArchiveBase.h b/Engine/Source/Runtime/Core/Public/Serialization/ArchiveBase.h
--- a/Engine/Source/Runtime/Core/Public/Serialization/ArchiveBase.h
+++ b/Engine/Source/Runtime/Core/Public/Serialization/ArchiveBase.h
@@ -1069,12 +1069,12 @@ public:
 	 */
 	CORE_API FArchiveProxy(FArchive& InInnerArchive);
 
-	virtual FArchive& operator<<(class FName& Value)
+	virtual FArchive& operator<<(class FName& Value) override
 	{
 		return InnerArchive << Value;
 	}
 
-	virtual FArchive& operator<<(class UObject*& Value)
+	virtual FArchive& operator<<(class UObject*& Value) override
 	{
 		return InnerArchive << Value;
 	}
@@ -1084,104 +1084,104 @@ public:
 		return InnerArchive << Value;
 	}
 
-	virtual void Serialize(void* V, int64 Length)
+	virtual void Serialize(void* V, int64 Length) override
 	{
 		InnerArchive.Serialize(V, Length);
 	}
 
-	virtual void SerializeBits(void* Bits, int64 LengthBits)
+	virtual void SerializeBits(void* Bits, int64 LengthBits) override
 	{
 		InnerArchive.SerializeBits(Bits, LengthBits);
 	}
 
-	virtual void SerializeInt(uint32& Value, uint32 Max)
+	virtual void SerializeInt(uint32& Value, uint32 Max) override
 	{
 		InnerArchive.SerializeInt(Value, Max);
 	}
 
-	virtual void Preload(UObject* Object)
+	virtual void Preload(UObject* Object) override
 	{
 		InnerArchive.Preload(Object);
 	}
 
-	virtual void CountBytes(SIZE_T InNum, SIZE_T InMax)
+	virtual void CountBytes(SIZE_T InNum, SIZE_T InMax) override
 	{
 		InnerArchive.CountBytes(InNum, InMax);
 	}
 
-	CORE_API virtual FString GetArchiveName() const;
+	CORE_API virtual FString GetArchiveName() const override;
 
-	virtual class ULinker* GetLinker()
+	virtual class ULinker* GetLinker() override
 	{
 		return InnerArchive.GetLinker();
 	}
 
-	virtual int64 Tell()
+	virtual int64 Tell() override
 	{
 		return InnerArchive.Tell();
 	}
 
-	virtual int64 TotalSize()
+	virtual int64 TotalSize() override
 	{
 		return InnerArchive.TotalSize();
 	}
 
-	virtual bool AtEnd()
+	virtual bool AtEnd() override
 	{
 		return InnerArchive.AtEnd();
 	}
 
-	virtual void Seek(int64 InPos)
+	virtual void Seek(int64 InPos) override
 	{
 		InnerArchive.Seek(InPos);
 	}
 
-	virtual void AttachBulkData(UObject* Owner, FUntypedBulkData* BulkData)
+	virtual void AttachBulkData(UObject* Owner, FUntypedBulkData* BulkData) override
 	{
 		InnerArchive.AttachBulkData(Owner, BulkData);
 	}
 
-	virtual void DetachBulkData(FUntypedBulkData* BulkData, bool bEnsureBulkDataIsLoaded)
+	virtual void DetachBulkData(FUntypedBulkData* BulkData, bool bEnsureBulkDataIsLoaded) override
 	{
 		InnerArchive.DetachBulkData(BulkData, bEnsureBulkDataIsLoaded);
 	}
 
-	virtual bool Precache(int64 PrecacheOffset, int64 PrecacheSize)
+	virtual bool Precache(int64 PrecacheOffset, int64 PrecacheSize) override
 	{
 		return InnerArchive.Precache(PrecacheOffset, PrecacheSize);
 	}
 
-	virtual bool SetCompressionMap(TArray<struct FCompressedChunk>* CompressedChunks, ECompressionFlags CompressionFlags)
+	virtual bool SetCompressionMap(TArray<struct FCompressedChunk>* CompressedChunks, ECompressionFlags CompressionFlags) override
 	{
 		return InnerArchive.SetCompressionMap(CompressedChunks, CompressionFlags);
 	}
 
-	virtual void Flush()
+	virtual void Flush() override
 	{
 		InnerArchive.Flush();
 	}
 
-	virtual bool Close()
+	virtual bool Close() override
 	{
 		return InnerArchive.Close();
 	}
 
-	virtual bool GetError()
+	virtual bool GetError() override
 	{
 		return InnerArchive.GetError();
 	}
 
-	virtual void MarkScriptSerializationStart(const UObject* Obj)
+	virtual void MarkScriptSerializationStart(const UObject* Obj) override
 	{
 		InnerArchive.MarkScriptSerializationStart(Obj);
 	}
 
-	virtual void MarkScriptSerializationEnd(const UObject* Obj)
+	virtual void MarkScriptSerializationEnd(const UObject* Obj) override
 	{
 		InnerArchive.MarkScriptSerializationEnd(Obj);
 	}
 
-	virtual bool IsCloseComplete(bool& bHasError)
+	virtual bool IsCloseComplete(bool& bHasError) override
 	{
 		return InnerArchive.IsCloseComplete(bHasError);
 	}
diff --git a/Engine/Source/Runtime/CoreUObject/Public/UObject/Linker.h b/Engine/Source/Runtime/CoreUObject/Public/UObject/Linker.h
--- a/Engine/Source/Runtime/CoreUObject/Public/UObject/Linker.h
+++ b/Engine/Source/Runtime/CoreUObject/Public/UObject/Linker.h
@@ -957,7 +957,7 @@ class ULinker : public UObject, public FLinkerTables
 	 * 
 	 * @param	Ar	the archive to read/write into
 	 */
-	void Serialize( FArchive& Ar );
+	void Serialize( FArchive& Ar ) override;
 	static void AddReferencedObjects(UObject* InThis, FReferenceCollector& Collector);
 
 	/**
@@ -1433,7 +1433,7 @@ public:
 	/**
 	 * If this archive is a ULinkerLoad or ULinkerSave, returns a pointer to the ULinker portion.
 	 */
-	virtual ULinker* GetLinker() { return this; }
+	virtual ULinker* GetLinker() override { return this; }
 
 	/**
 	 * Creates and returns a ULinkerLoad object.
@@ -1449,7 +1449,7 @@ public:
 	void Verify();
 
 	COREUOBJECT_API FName GetExportClassPackage( int32 i );
-	virtual FString GetArchiveName() const;
+	virtual FString GetArchiveName() const override;
 
 	/**
 	 * Recursively gathers the dependencies of a given export (the recursive chain of imports
@@ -1530,7 +1530,7 @@ public:
 	 *					If Object is a UClass and the class default object has already been created, calls
 	 *					Preload for the class default object as well.
 	 */
-	void Preload( UObject* Object );
+	void Preload( UObject* Object ) override;
 
 	/**
 	 * Before loading a persistent object from disk, this function can be used to discover
@@ -1571,12 +1571,12 @@ public:
 	/**
 	 * Called when an object begins serializing property data using script serialization.
 	 */
-	virtual void MarkScriptSerializationStart( const UObject* Obj );
+	virtual void MarkScriptSerializationStart( const UObject* Obj ) override;
 
 	/**
 	 * Called when an object stops serializing property data using script serialization.
 	 */
-	virtual void MarkScriptSerializationEnd( const UObject* Obj );
+	virtual void MarkScriptSerializationEnd( const UObject* Obj ) override;
 
 	/**
 	 * Looks for an existing linker for the given package, without trying to make one if it doesn't exist
@@ -1683,7 +1683,7 @@ private:
 	 * @param	PrecacheSize	Number of bytes to precache
 	 * @return	false if precache operation is still pending, true otherwise
 	 */
-	virtual bool Precache( int64 PrecacheOffset, int64 PrecacheSize );
+	virtual bool Precache( int64 PrecacheOffset, int64 PrecacheSize ) override;
 	
 #if WITH_EDITOR
 	/**
@@ -1716,16 +1716,16 @@ public:
 	virtual void Detach( bool bEnsureAllBulkDataIsLoaded );
 private:
 
-	void Seek( int64 InPos );
-	int64 Tell();
-	int64 TotalSize();
+	void Seek( int64 InPos ) override;
+	int64 Tell() override;
+	int64 TotalSize() override;
 	// this fixes the warning : 'ULinkerSave::Serialize' hides overloaded virtual function
 	using ULinker::Serialize;
-	void Serialize( void* V, int64 Length );
-	virtual FArchive& operator<<( UObject*& Object );
-	virtual FArchive& operator<<( FLazyObjectPtr& LazyObjectPtr);
-	virtual FArchive& operator<<( FAssetPtr& AssetPtr);
-	virtual FArchive& operator<<( FName& Name );
+	void Serialize( void* V, int64 Length ) override;
+	virtual FArchive& operator<<( UObject*& Object ) override;
+	virtual FArchive& operator<<( FLazyObjectPtr& LazyObjectPtr) override;
+	virtual FArchive& operator<<( FAssetPtr& AssetPtr) override;
+	virtual FArchive& operator<<( FName& Name ) override;
 
 	/**
 	 * Safely verify that an import in the ImportMap points to a good object. This decides whether or not
diff --git a/Engine/Source/Runtime/Slate/Public/Widgets/Views/SListView.h b/Engine/Source/Runtime/Slate/Public/Widgets/Views/SListView.h
--- a/Engine/Source/Runtime/Slate/Public/Widgets/Views/SListView.h
+++ b/Engine/Source/Runtime/Slate/Public/Widgets/Views/SListView.h
@@ -923,7 +923,7 @@ public:
 	}
 
 	/** @return how many items there are in the TArray being observed */
-	virtual int32 GetNumItemsBeingObserved() const
+	virtual int32 GetNumItemsBeingObserved() const override
 	{
 		return ItemsSource == nullptr ? 0 : ItemsSource->Num();
 	}
