Added EMIT_CUSTOM_WARNING_AT_LINE macro to emit warnings at given line rather than on macro location.

diff --git a/Engine/Source/Runtime/Core/Public/Android/AndroidPlatformCompilerPreSetup.h b/Engine/Source/Runtime/Core/Public/Android/AndroidPlatformCompilerPreSetup.h
--- a/Engine/Source/Runtime/Core/Public/Android/AndroidPlatformCompilerPreSetup.h
+++ b/Engine/Source/Runtime/Core/Public/Android/AndroidPlatformCompilerPreSetup.h
@@ -46,7 +46,7 @@
 
 #endif // DISABLE_DEPRECATION
 
-#ifndef EMIT_CUSTOM_WARNING
-#define EMIT_CUSTOM_WARNING(Warning) \
-	_Pragma(PREPROCESSOR_TO_STRING(message(WARNING_LOCATION Warning)))
-#endif // EMIT_CUSTOM_WARNING
\ No newline at end of file
+#ifndef EMIT_CUSTOM_WARNING_AT_LINE
+#define EMIT_CUSTOM_WARNING_AT_LINE(Line, Warning) \
+	_Pragma(PREPROCESSOR_TO_STRING(message(WARNING_LOCATION(Line) Warning)))
+#endif // EMIT_CUSTOM_WARNING_AT_LINE
\ No newline at end of file
diff --git a/Engine/Source/Runtime/Core/Public/GenericPlatform/GenericPlatformCompilerPreSetup.h b/Engine/Source/Runtime/Core/Public/GenericPlatform/GenericPlatformCompilerPreSetup.h
--- a/Engine/Source/Runtime/Core/Public/GenericPlatform/GenericPlatformCompilerPreSetup.h
+++ b/Engine/Source/Runtime/Core/Public/GenericPlatform/GenericPlatformCompilerPreSetup.h
@@ -79,6 +79,9 @@
 	#define EMIT_DEPRECATED_WARNING_MESSAGE(Msg)
 #endif // EMIT_DEPRECATED_WARNING_MESSAGE
 
-#ifndef EMIT_CUSTOM_WARNING
-	#define EMIT_CUSTOM_WARNING(Warning)
-#endif // EMIT_CUSTOM_WARNING
\ No newline at end of file
+#ifndef EMIT_CUSTOM_WARNING_AT_LINE
+	#define EMIT_CUSTOM_WARNING_AT_LINE(Line, Warning)
+#endif // EMIT_CUSTOM_WARNING_AT_LINE
+
+#define EMIT_CUSTOM_WARNING(Warning) \
+	EMIT_CUSTOM_WARNING_AT_LINE(__LINE__, Warning)
diff --git a/Engine/Source/Runtime/Core/Public/HTML5/HTML5PlatformCompilerPreSetup.h b/Engine/Source/Runtime/Core/Public/HTML5/HTML5PlatformCompilerPreSetup.h
--- a/Engine/Source/Runtime/Core/Public/HTML5/HTML5PlatformCompilerPreSetup.h
+++ b/Engine/Source/Runtime/Core/Public/HTML5/HTML5PlatformCompilerPreSetup.h
@@ -83,7 +83,7 @@
 #endif 
 #endif // DISABLE_DEPRECATION
 
-#ifndef EMIT_CUSTOM_WARNING
-#define EMIT_CUSTOM_WARNING(Warning) \
-	_Pragma(PREPROCESSOR_TO_STRING(message(WARNING_LOCATION Warning)))
-#endif // EMIT_CUSTOM_WARNING
+#ifndef EMIT_CUSTOM_WARNING_AT_LINE
+#define EMIT_CUSTOM_WARNING_AT_LINE(Line, Warning) \
+	_Pragma(PREPROCESSOR_TO_STRING(message(WARNING_LOCATION(Line) Warning)))
+#endif // EMIT_CUSTOM_WARNING_AT_LINE
diff --git a/Engine/Source/Runtime/Core/Public/IOS/IOSPlatformCompilerPreSetup.h b/Engine/Source/Runtime/Core/Public/IOS/IOSPlatformCompilerPreSetup.h
--- a/Engine/Source/Runtime/Core/Public/IOS/IOSPlatformCompilerPreSetup.h
+++ b/Engine/Source/Runtime/Core/Public/IOS/IOSPlatformCompilerPreSetup.h
@@ -80,7 +80,7 @@
 		_Pragma("clang diagnostic pop")
 #endif // DISABLE_DEPRECATION
 
-#ifndef EMIT_CUSTOM_WARNING
-#define EMIT_CUSTOM_WARNING(Warning) \
-	_Pragma(PREPROCESSOR_TO_STRING(message(WARNING_LOCATION Warning)))
-#endif // EMIT_CUSTOM_WARNING
+#ifndef EMIT_CUSTOM_WARNING_AT_LINE
+#define EMIT_CUSTOM_WARNING_AT_LINE(Line, Warning) \
+	_Pragma(PREPROCESSOR_TO_STRING(message(WARNING_LOCATION(Line) Warning)))
+#endif // EMIT_CUSTOM_WARNING_AT_LINE
diff --git a/Engine/Source/Runtime/Core/Public/Linux/LinuxPlatformCompilerPreSetup.h b/Engine/Source/Runtime/Core/Public/Linux/LinuxPlatformCompilerPreSetup.h
--- a/Engine/Source/Runtime/Core/Public/Linux/LinuxPlatformCompilerPreSetup.h
+++ b/Engine/Source/Runtime/Core/Public/Linux/LinuxPlatformCompilerPreSetup.h
@@ -45,10 +45,10 @@
 		_Pragma("clang diagnostic pop")
 #endif // DISABLE_DEPRECATION
 
-#ifndef EMIT_CUSTOM_WARNING
-#define EMIT_CUSTOM_WARNING(Warning) \
-	_Pragma(PREPROCESSOR_TO_STRING(message(WARNING_LOCATION Warning)))
-#endif // EMIT_CUSTOM_WARNING
+#ifndef EMIT_CUSTOM_WARNING_AT_LINE
+#define EMIT_CUSTOM_WARNING_AT_LINE(Line, Warning) \
+	_Pragma(PREPROCESSOR_TO_STRING(message(WARNING_LOCATION(Line) Warning)))
+#endif // EMIT_CUSTOM_WARNING_AT_LINE
 
 // Make certain warnings always be warnings, even despite -Werror.
 // Rationale: we don't want to suppress those as there are plans to address them (e.g. UE-12341), but breaking builds due to these warnings is very expensive
diff --git a/Engine/Source/Runtime/Core/Public/Mac/MacPlatformCompilerPreSetup.h b/Engine/Source/Runtime/Core/Public/Mac/MacPlatformCompilerPreSetup.h
--- a/Engine/Source/Runtime/Core/Public/Mac/MacPlatformCompilerPreSetup.h
+++ b/Engine/Source/Runtime/Core/Public/Mac/MacPlatformCompilerPreSetup.h
@@ -80,7 +80,7 @@
 		_Pragma("clang diagnostic pop")
 #endif // DISABLE_DEPRECATION
 
-#ifndef EMIT_CUSTOM_WARNING
-#define EMIT_CUSTOM_WARNING(Warning) \
-	_Pragma(PREPROCESSOR_TO_STRING(message(WARNING_LOCATION Warning)))
-#endif // EMIT_CUSTOM_WARNING
+#ifndef EMIT_CUSTOM_WARNING_AT_LINE
+#define EMIT_CUSTOM_WARNING_AT_LINE(Line, Warning) \
+	_Pragma(PREPROCESSOR_TO_STRING(message(WARNING_LOCATION(Line) Warning)))
+#endif // EMIT_CUSTOM_WARNING_AT_LINE
diff --git a/Engine/Source/Runtime/Core/Public/Misc/CoreMiscDefines.h b/Engine/Source/Runtime/Core/Public/Misc/CoreMiscDefines.h
--- a/Engine/Source/Runtime/Core/Public/Misc/CoreMiscDefines.h
+++ b/Engine/Source/Runtime/Core/Public/Misc/CoreMiscDefines.h
@@ -110,4 +110,4 @@ const FPlatformUserId PLATFORMUSERID_NONE = INDEX_NONE;
 #define PREPROCESSOR_JOIN(x, y) PREPROCESSOR_JOIN_INNER(x, y)
 
 // When passed to pragma message will result in clickable warning in VS
-#define WARNING_LOCATION __FILE__ "(" PREPROCESSOR_TO_STRING(__LINE__) ")"
+#define WARNING_LOCATION(Line) __FILE__ "(" PREPROCESSOR_TO_STRING(Line) ")"
diff --git a/Engine/Source/Runtime/Core/Public/WinRT/WinRTPlatformCompilerPreSetup.h b/Engine/Source/Runtime/Core/Public/WinRT/WinRTPlatformCompilerPreSetup.h
--- a/Engine/Source/Runtime/Core/Public/WinRT/WinRTPlatformCompilerPreSetup.h
+++ b/Engine/Source/Runtime/Core/Public/WinRT/WinRTPlatformCompilerPreSetup.h
@@ -89,11 +89,11 @@
 #endif // DISABLE_DEPRECATION
 
 #if _MSC_VER
-#define EMIT_CUSTOM_WARNING(Warning) \
-	__pragma(message(WARNING_LOCATION "(\": warning C4996: " Warning))
+#define EMIT_CUSTOM_WARNING_AT_LINE(Line, Warning) \
+	__pragma(message(WARNING_LOCATION(Line) ": warning C4996: " Warning))
 #elif __clang__
-#define EMIT_CUSTOM_WARNING(Warning) \
-	_Pragma(PREPROCESSOR_TO_STRING(message(WARNING_LOCATION Warning)))
+#define EMIT_CUSTOM_WARNING_AT_LINE(Line, Warning) \
+	_Pragma(PREPROCESSOR_TO_STRING(message(WARNING_LOCATION(Line) Warning)))
 #else
 #error Unknown compiler
 #endif
diff --git a/Engine/Source/Runtime/Core/Public/Windows/WindowsPlatformCompilerPreSetup.h b/Engine/Source/Runtime/Core/Public/Windows/WindowsPlatformCompilerPreSetup.h
--- a/Engine/Source/Runtime/Core/Public/Windows/WindowsPlatformCompilerPreSetup.h
+++ b/Engine/Source/Runtime/Core/Public/Windows/WindowsPlatformCompilerPreSetup.h
@@ -54,11 +54,11 @@
 #endif // DISABLE_DEPRECATION
 
 #if _MSC_VER
-#define EMIT_CUSTOM_WARNING(Warning) \
-	__pragma(message(WARNING_LOCATION ": warning C4996: " Warning))
+#define EMIT_CUSTOM_WARNING_AT_LINE(Line, Warning) \
+	__pragma(message(WARNING_LOCATION(Line) ": warning C4996: " Warning))
 #elif __clang__
-#define EMIT_CUSTOM_WARNING(Warning) \
-	_Pragma(PREPROCESSOR_TO_STRING(message(WARNING_LOCATION Warning)))
+#define EMIT_CUSTOM_WARNING_AT_LINE(Line, Warning) \
+	_Pragma(PREPROCESSOR_TO_STRING(message(WARNING_LOCATION(Line) Warning)))
 #else
 #error Unknown compiler
 #endif
