Fix build: not redefining a macro in generic platform headers.

diff --git a/Engine/Source/Runtime/Core/Public/GenericPlatform/GenericPlatformCompilerPreSetup.h b/Engine/Source/Runtime/Core/Public/GenericPlatform/GenericPlatformCompilerPreSetup.h
--- a/Engine/Source/Runtime/Core/Public/GenericPlatform/GenericPlatformCompilerPreSetup.h
+++ b/Engine/Source/Runtime/Core/Public/GenericPlatform/GenericPlatformCompilerPreSetup.h
@@ -83,5 +83,7 @@
 	#define EMIT_CUSTOM_WARNING_AT_LINE(Line, Warning)
 #endif // EMIT_CUSTOM_WARNING_AT_LINE
 
-#define EMIT_CUSTOM_WARNING(Warning) \
-	EMIT_CUSTOM_WARNING_AT_LINE(__LINE__, Warning)
+#ifndef EMIT_CUSTOM_WARNING
+	#define EMIT_CUSTOM_WARNING(Warning) \
+		EMIT_CUSTOM_WARNING_AT_LINE(__LINE__, Warning)
+#endif // EMIT_CUSTOM_WARNING
