fix for immediate crash on startup of iOS game
#ios

diff --git a/Engine/Source/Runtime/Core/Private/IOS/IOSAppDelegate.cpp b/Engine/Source/Runtime/Core/Private/IOS/IOSAppDelegate.cpp
--- a/Engine/Source/Runtime/Core/Private/IOS/IOSAppDelegate.cpp
+++ b/Engine/Source/Runtime/Core/Private/IOS/IOSAppDelegate.cpp
@@ -122,7 +122,7 @@ void InstallSignalHandlers()
 
 	while(!self.bCommandLineReady)
 	{
-		FPlatformProcess::Sleep(.01f);
+		usleep(100);
 	}
 
 
