Core - Fix for uninitialized TLS

diff --git a/Engine/Source/Programs/BuildPatchTool/Private/BuildPatchToolMain.cpp b/Engine/Source/Programs/BuildPatchTool/Private/BuildPatchToolMain.cpp
--- a/Engine/Source/Programs/BuildPatchTool/Private/BuildPatchToolMain.cpp
+++ b/Engine/Source/Programs/BuildPatchTool/Private/BuildPatchToolMain.cpp
@@ -159,9 +159,6 @@ int32 BuildPatchToolMain( const TCHAR* CommandLine )
 
 	FPlatformProcess::SetCurrentWorkingDirectoryToBaseDir();
 
-	// Init TLS for runnable threads.
-	FRunnableThread::InitializeTls();
-
 	bool bSuccess = false;
 
 	FString RootDirectory;
diff --git a/Engine/Source/Runtime/Core/Private/HAL/ThreadingBase.cpp b/Engine/Source/Runtime/Core/Private/HAL/ThreadingBase.cpp
--- a/Engine/Source/Runtime/Core/Private/HAL/ThreadingBase.cpp
+++ b/Engine/Source/Runtime/Core/Private/HAL/ThreadingBase.cpp
@@ -176,13 +176,14 @@ FScopedEvent::~FScopedEvent()
 	FRunnableThread
 -----------------------------------------------------------------------------*/
 
-uint32 FRunnableThread::RunnableTlsSlot = 0;
+uint32 FRunnableThread::RunnableTlsSlot = FRunnableThread::GetTlsSlot();
 
-void FRunnableThread::InitializeTls()
+uint32 FRunnableThread::GetTlsSlot()
 {
 	check( IsInGameThread() );
-	RunnableTlsSlot = FPlatformTLS::AllocTlsSlot();
+	uint32 RunnableTlsSlot = FPlatformTLS::AllocTlsSlot();
 	check( FPlatformTLS::IsValidTlsSlot( RunnableTlsSlot ) );
+	return RunnableTlsSlot;
 }
 
 FRunnableThread::FRunnableThread()
diff --git a/Engine/Source/Runtime/Core/Public/HAL/ThreadingBase.h b/Engine/Source/Runtime/Core/Public/HAL/ThreadingBase.h
--- a/Engine/Source/Runtime/Core/Public/HAL/ThreadingBase.h
+++ b/Engine/Source/Runtime/Core/Public/HAL/ThreadingBase.h
@@ -238,8 +238,8 @@ class CORE_API FRunnableThread
 
 public:
 
-	/** Initializes a Tls slot for storing the runnable thread pointer. */
-	static void InitializeTls();
+	/** Gets a new Tls slot for storing the runnable thread pointer. */
+	static uint32 GetTlsSlot();
 
 	/**
 	* Factory method to create a thread with the specified stack size and thread priority.
diff --git a/Engine/Source/Runtime/Launch/Private/LaunchEngineLoop.cpp b/Engine/Source/Runtime/Launch/Private/LaunchEngineLoop.cpp
--- a/Engine/Source/Runtime/Launch/Private/LaunchEngineLoop.cpp
+++ b/Engine/Source/Runtime/Launch/Private/LaunchEngineLoop.cpp
@@ -1048,8 +1048,6 @@ int32 FEngineLoop::PreInit( const TCHAR* CmdLine )
 	}
 #endif //WITH_EDITOR
 
-	FRunnableThread::InitializeTls();
-
 	// initialize task graph sub-system with potential multiple threads
 	FTaskGraphInterface::Startup( FPlatformMisc::NumberOfCores() );
 	FTaskGraphInterface::Get().AttachToThread( ENamedThreads::GameThread );
