Core: Implemented FTypeContainer.RegisterFactory for Clang

diff --git a/Engine/Source/Runtime/Core/Private/Tests/Misc/TypeContainerTest.cpp b/Engine/Source/Runtime/Core/Private/Tests/Misc/TypeContainerTest.cpp
--- a/Engine/Source/Runtime/Core/Private/Tests/Misc/TypeContainerTest.cpp
+++ b/Engine/Source/Runtime/Core/Private/Tests/Misc/TypeContainerTest.cpp
@@ -175,7 +175,7 @@ bool FTypeContainerTest::RunTest(const FString& Parameters)
 	}
 
 	// factory test
-	{/* @todo gmp: temporarily disabled due to Clang compile errors
+	{
 		struct FLocal
 		{
 			static TSharedRef<IBerry> MakeStrawberry()
@@ -203,7 +203,7 @@ bool FTypeContainerTest::RunTest(const FString& Parameters)
 		auto Berry = Container.GetInstance<IBerry>();
 		auto Fruit = Container.GetInstance<IFruit>();
 		auto Smoothie = Container.GetInstance<ISmoothie>();
-	*/}
+	}
 
 	// delegate test
 	{
diff --git a/Engine/Source/Runtime/Core/Public/Misc/TypeContainer.h b/Engine/Source/Runtime/Core/Public/Misc/TypeContainer.h
--- a/Engine/Source/Runtime/Core/Public/Misc/TypeContainer.h
+++ b/Engine/Source/Runtime/Core/Public/Misc/TypeContainer.h
@@ -5,7 +5,6 @@
 
 #if PLATFORM_COMPILER_HAS_VARIADIC_TEMPLATES
 
-
 /**
  * Enumerates the scopes for instance creation in type containers.
  */
@@ -286,15 +285,13 @@ public:
 	/**
 	 * Register a factory function for the specified class.
 	 *
-	 * Note: VS2013 did not like a purely variadic template function, so we split
-	 * these up into two overloads, one without and one with variadic parameters.
-	 *
 	 * @param R The type of class to register the instance for.
+	 * @param P Additional parameters that the factory function requires.
 	 * @param CreateFunc The factory function that will create instances.
 	 * @see RegisterClass, RegisterInstance, Unregister
 	 */
 	template<class R>
-	void RegisterFactory(TFunction<TSharedRef<R>()> CreateFunc)
+	void RegisterFactory(const TFunction<TSharedRef<R>()>& CreateFunc)
 	{
 		TSharedPtr<IInstanceProvider> Provider = MakeShareable(
 			new TFunctionInstanceProvider<R>(
@@ -310,9 +307,6 @@ public:
 	/**
 	 * Register a factory function for the specified class.
 	 *
-	 * Note: VS2013 did not like a purely variadic template function, so we split
-	 * these up into two overloads, one without and one with variadic parameters.
-	 *
 	 * @param R The type of class to register the instance for.
 	 * @param P0 The first parameter that the factory function requires.
 	 * @param P Additional parameters that the factory function requires.
@@ -320,7 +314,7 @@ public:
 	 * @see RegisterClass, RegisterInstance, Unregister
 	 */
 	template<class R, typename P0, typename... P>
-	void RegisterFactory(TFunction<TSharedRef<R>(TSharedRef<P0>, TSharedRef<P>...)> CreateFunc)
+	void RegisterFactory(const TFunction<TSharedRef<R>(TSharedRef<P0>, TSharedRef<P>...)>& CreateFunc)
 	{
 		TSharedPtr<IInstanceProvider> Provider = MakeShareable(
 			new TFunctionInstanceProvider<R>(
@@ -333,6 +327,32 @@ public:
 		AddProvider(TNameOf<R>::GetName(), Provider);
 	}
 
+#ifdef __clang__
+	/**
+	 * Register a factory function for the specified class.
+	 *
+	 * This is a Clang specific overload for handling raw function pointers.
+	 *
+	 * @param R The type of class to register the instance for.
+	 * @param P Additional parameters that the factory function requires.
+	 * @param CreateFunc The factory function that will create instances.
+	 * @see RegisterClass, RegisterInstance, Unregister
+	 */
+	template<class R, typename... P>
+	void RegisterFactory(TSharedRef<R> (*CreateFunc)(TSharedRef<P>...))
+	{
+		TSharedPtr<IInstanceProvider> Provider = MakeShareable(
+			new TFunctionInstanceProvider<R>(
+				[=]() -> TSharedPtr<void> {
+					return CreateFunc(GetInstance<P>()...);
+				}
+			)
+		);
+
+		AddProvider(TNameOf<R>::GetName(), Provider);
+	}
+#endif
+
 	/**
 	 * Registers an existing instance for the specified class.
 	 *
