Fixes for issues uncovered by a non-unity build.

diff --git a/Engine/Source/Runtime/Core/Private/Linux/LinuxPlatformMisc.cpp b/Engine/Source/Runtime/Core/Private/Linux/LinuxPlatformMisc.cpp
--- a/Engine/Source/Runtime/Core/Private/Linux/LinuxPlatformMisc.cpp
+++ b/Engine/Source/Runtime/Core/Private/Linux/LinuxPlatformMisc.cpp
@@ -7,6 +7,7 @@
 #include "CorePrivatePCH.h"
 #include "LinuxPlatformMisc.h"
 #include "LinuxApplication.h"
+#include "LinuxPlatformCrashContext.h"
 
 #include <sys/sysinfo.h>
 #include <sched.h>
diff --git a/Engine/Source/Runtime/Core/Private/Linux/LinuxPlatformOutputDevicesPrivate.h b/Engine/Source/Runtime/Core/Private/Linux/LinuxPlatformOutputDevicesPrivate.h
--- a/Engine/Source/Runtime/Core/Private/Linux/LinuxPlatformOutputDevicesPrivate.h
+++ b/Engine/Source/Runtime/Core/Private/Linux/LinuxPlatformOutputDevicesPrivate.h
@@ -1,5 +1,7 @@
 // Copyright 1998-2015 Epic Games, Inc. All Rights Reserved.
 
+#include "OutputDeviceConsole.h"
+
 class FOutputDeviceLinuxError : public FOutputDeviceError
 {
 public:
