Allow setting the thread priority on a monitored process. Fixes PS4 device detection eating up all the CPU.

diff --git a/Engine/Source/Runtime/Core/Private/Misc/MonitoredProcess.cpp b/Engine/Source/Runtime/Core/Private/Misc/MonitoredProcess.cpp
--- a/Engine/Source/Runtime/Core/Private/Misc/MonitoredProcess.cpp
+++ b/Engine/Source/Runtime/Core/Private/Misc/MonitoredProcess.cpp
@@ -19,6 +19,7 @@ FMonitoredProcess::FMonitoredProcess( const FString& InURL, const FString& InPar
 	, URL(InURL)
 	, WritePipe(nullptr)
 	, bCreatePipes(InCreatePipes)
+	, SleepInterval(0.0f)
 { }
 
 
@@ -102,7 +103,7 @@ uint32 FMonitoredProcess::Run()
 	{
 		do
 		{
-			FPlatformProcess::Sleep(0.0);
+			FPlatformProcess::Sleep(SleepInterval);
 
 			ProcessOutput(FPlatformProcess::ReadPipe(ReadPipe));
 
diff --git a/Engine/Source/Runtime/Core/Public/Misc/MonitoredProcess.h b/Engine/Source/Runtime/Core/Public/Misc/MonitoredProcess.h
--- a/Engine/Source/Runtime/Core/Public/Misc/MonitoredProcess.h
+++ b/Engine/Source/Runtime/Core/Public/Misc/MonitoredProcess.h
@@ -71,6 +71,16 @@ public:
 	/** Launches the process. */
 	bool Launch();
 
+	/**
+	 * Sets the sleep interval to be used in the main thread loop.
+	 *
+	 * @param InSleepInterval The Sleep interval to use.
+	 */
+	void SetSleepInterval( float InSleepInterval )
+	{
+		SleepInterval = InSleepInterval;
+	}
+
 public:
 
 	/**
@@ -181,6 +191,9 @@ private:
 	// Holds if we should create pipes
 	bool bCreatePipes;
 
+	// Sleep interval to use
+	float SleepInterval;
+
 private:
 
 	// Holds a delegate that is executed when the process has been canceled. */
