Exempt -Wreorder from -Werror so that it doesn't break builds.

diff --git a/Engine/Source/Runtime/Core/Public/Mac/MacPlatformCompilerPreSetup.h b/Engine/Source/Runtime/Core/Public/Mac/MacPlatformCompilerPreSetup.h
--- a/Engine/Source/Runtime/Core/Public/Mac/MacPlatformCompilerPreSetup.h
+++ b/Engine/Source/Runtime/Core/Public/Mac/MacPlatformCompilerPreSetup.h
@@ -84,3 +84,8 @@
 #define EMIT_CUSTOM_WARNING_AT_LINE(Line, Warning) \
 	_Pragma(PREPROCESSOR_TO_STRING(message(WARNING_LOCATION(Line) Warning)))
 #endif // EMIT_CUSTOM_WARNING_AT_LINE
+
+// Make certain warnings always be warnings, even despite -Werror.
+// Rationale: we don't want to suppress those as there are plans to address them (e.g. UE-12341), but breaking builds due to these warnings is very expensive
+// since they cannot be caught by all compilers that we support. They are deemed to be relatively safe to be ignored, at least until all SDKs/toolchains start supporting them.
+#pragma clang diagnostic warning "-Wreorder"
