UE-12184: Lightmass crash after finished building and hitting save

There was a data loss conversion that was causing this crash when files were big enough. Fixed.

diff --git a/Engine/Source/Runtime/Core/Private/Serialization/Archive.cpp b/Engine/Source/Runtime/Core/Private/Serialization/Archive.cpp
--- a/Engine/Source/Runtime/Core/Private/Serialization/Archive.cpp
+++ b/Engine/Source/Runtime/Core/Private/Serialization/Archive.cpp
@@ -675,7 +675,7 @@ void FArchive::SerializeCompressed( void* V, int64 Length, ECompressionFlags Fla
 
 		// Overrwrite chunk infos by seeking to the beginning, serializing the data and then
 		// seeking back to the end.
-		int32 EndPosition = Tell();
+		auto EndPosition = Tell();
 		// Seek to the beginning.
 		Seek( StartPosition );
 		// Serialize chunk infos.
