We no longer write out the clean coalesced ini files on cooked builds, as we have no real need to.

diff --git a/Engine/Source/Runtime/Core/Private/Misc/ConfigCacheIni.cpp b/Engine/Source/Runtime/Core/Private/Misc/ConfigCacheIni.cpp
--- a/Engine/Source/Runtime/Core/Private/Misc/ConfigCacheIni.cpp
+++ b/Engine/Source/Runtime/Core/Private/Misc/ConfigCacheIni.cpp
@@ -1253,26 +1253,29 @@ void FConfigFile::SaveSourceToBackupFile()
 
 void FConfigFile::ProcessSourceAndCheckAgainstBackup()
 {
-	FString BetweenRunsDir = (FPaths::GameIntermediateDir() / TEXT("Config/CoalescedSourceConfigs/"));
-	FString BackupFilename = FString::Printf( TEXT( "%s%s.ini" ), *BetweenRunsDir, *Name.ToString() );
-
-	FConfigFile BackupFile;
-	ProcessIniContents(*BackupFilename, *BackupFilename, &BackupFile, false, false, false);
-
-	for( TMap<FString,FConfigSection>::TIterator SectionIterator(*SourceConfigFile); SectionIterator; ++SectionIterator )
+	if (!FPlatformProperties::RequiresCookedData())
 	{
-		const FString& SectionName = SectionIterator.Key();
-		const FConfigSection& SourceSection = SectionIterator.Value();
-		const FConfigSection* BackupSection = BackupFile.Find( SectionName );
+		FString BetweenRunsDir = (FPaths::GameIntermediateDir() / TEXT("Config/CoalescedSourceConfigs/"));
+		FString BackupFilename = FString::Printf( TEXT( "%s%s.ini" ), *BetweenRunsDir, *Name.ToString() );
 
-		if( BackupSection && SourceSection != *BackupSection )
+		FConfigFile BackupFile;
+		ProcessIniContents(*BackupFilename, *BackupFilename, &BackupFile, false, false, false);
+
+		for( TMap<FString,FConfigSection>::TIterator SectionIterator(*SourceConfigFile); SectionIterator; ++SectionIterator )
 		{
-			this->Remove( SectionName );
-			this->Add( SectionName, SourceSection );
+			const FString& SectionName = SectionIterator.Key();
+			const FConfigSection& SourceSection = SectionIterator.Value();
+			const FConfigSection* BackupSection = BackupFile.Find( SectionName );
+
+			if( BackupSection && SourceSection != *BackupSection )
+			{
+				this->Remove( SectionName );
+				this->Add( SectionName, SourceSection );
+			}
 		}
-	}
 
-	SaveSourceToBackupFile();
+		SaveSourceToBackupFile();
+	}
 }
 
 
