Core - Fixed a crash when capturing stats from the beginning

diff --git a/Engine/Source/Runtime/Core/Private/ProfilingDebugging/ProfilingHelpers.cpp b/Engine/Source/Runtime/Core/Private/ProfilingDebugging/ProfilingHelpers.cpp
--- a/Engine/Source/Runtime/Core/Private/ProfilingDebugging/ProfilingHelpers.cpp
+++ b/Engine/Source/Runtime/Core/Private/ProfilingDebugging/ProfilingHelpers.cpp
@@ -79,8 +79,14 @@ FString CreateProfileFilename( const FString& InFileExtension, bool bIncludeDate
 	FString MapNameStr;
 
 #if WITH_ENGINE
-	check(GGetMapNameDelegate.IsBound());
-	MapNameStr = GGetMapNameDelegate.Execute();
+	if(GGetMapNameDelegate.IsBound())
+	{
+		MapNameStr = GGetMapNameDelegate.Execute();
+	}
+	else
+	{
+		MapNameStr = TEXT( "LoadTimeFile" );
+	}
 #endif		// WITH_ENGINE
 
 	const FString PlatformStr(FPlatformProperties::PlatformName());
