UE4 - ParallelFor

diff --git a/Engine/Source/Runtime/Core/Private/Async/TaskGraph.cpp b/Engine/Source/Runtime/Core/Private/Async/TaskGraph.cpp
--- a/Engine/Source/Runtime/Core/Private/Async/TaskGraph.cpp
+++ b/Engine/Source/Runtime/Core/Private/Async/TaskGraph.cpp
@@ -9,6 +9,8 @@ DEFINE_LOG_CATEGORY_STATIC(LogTaskGraph, Log, All);
 DEFINE_STAT(STAT_FReturnGraphTask);
 DEFINE_STAT(STAT_FTriggerEventGraphTask);
 DEFINE_STAT(STAT_UnknownGraphTask);
+DEFINE_STAT(STAT_ParallelFor);
+DEFINE_STAT(STAT_ParallelForTask);
 
 namespace ENamedThreads
 {
@@ -897,9 +899,14 @@ public:
 		return NumThreads - NumNamedThreads;
 	}
 
-	virtual ENamedThreads::Type GetCurrentThreadIfKnown() override
+	virtual ENamedThreads::Type GetCurrentThreadIfKnown(bool bLocalQueue) override
 	{
-		return GetCurrentThread();
+		ENamedThreads::Type Result = GetCurrentThread();
+		if (Result >= 0 && Result < NumNamedThreads)
+		{
+			Result = ENamedThreads::Type(int32(Result) | int32(ENamedThreads::LocalQueue));
+		}
+		return Result;
 	}
 
 	virtual bool IsThreadProcessingTasks(ENamedThreads::Type ThreadToCheck) override
@@ -1277,5 +1284,84 @@ FGraphEvent::~FGraphEvent()
 }
 
 
+//---
+
+#include "ParallelFor.h"
+
+static void TestParallelFor(const TArray<FString>& Args)
+{
+
+	ParallelFor(10, 
+		[](int32 Index)
+		{
+			UE_LOG(LogConsoleResponse, Display, TEXT("ParallelFor index=%d, thread=%x"), Index, FPlatformTLS::GetCurrentThreadId());
+		}
+	);
+
+	ParallelFor(10, 
+		[](int32 Index)
+		{
+			ParallelFor(10, 
+				[Index](int32 IndexInner)
+				{
+					UE_LOG(LogConsoleResponse, Display, TEXT("ParallelFor index=%d %d, thread=%x"), Index, IndexInner, FPlatformTLS::GetCurrentThreadId());
+				}
+			);
+		}
+	);
+
+	TArray<bool> TestBools;
+
+	TestBools.AddZeroed(10000);
+
+	ParallelFor(TestBools.Num(), [&TestBools](int32 Index){TestBools[Index] = true;});
+
+	for (int32 Index = 0; Index < TestBools.Num(); Index++)
+	{
+		check(TestBools[Index]);
+	}
+
+	TestBools.Empty(10000);
+	TestBools.AddZeroed(10000);
+
+	ParallelFor(TestBools.Num(), [&TestBools](int32 Index){TestBools[Index] = true;}, true);
+
+	for (int32 Index = 0; Index < TestBools.Num(); Index++)
+	{
+		check(TestBools[Index]);
+	}
+
+	uint32 TargetCrc = FCrc::MemCrc32(TestBools.GetData(), 1024);
+
+	{
+		double StartTime = FPlatformTime::Seconds();
+		ParallelFor(1024, 
+			[&TestBools, TargetCrc](int32 Index)
+			{
+				uint32 TestCrc = FCrc::MemCrc32(TestBools.GetData(), 1024);
+				check(TestCrc == TargetCrc);
+			}
+		);
+		UE_LOG(LogConsoleResponse, Display, TEXT("Parallel CRC of 1MB (%d threads) in %6.3fms"), FTaskGraphInterface::Get().GetNumWorkerThreads() + 1, float(FPlatformTime::Seconds() - StartTime) * 1000.0f);
+	}
+	{
+		double StartTime = FPlatformTime::Seconds();
+		ParallelFor(1024, 
+			[&TestBools, TargetCrc](int32 Index)
+			{
+				uint32 TestCrc = FCrc::MemCrc32(TestBools.GetData(), 1024);
+				check(TestCrc == TargetCrc);
+			}, 
+			true
+		);
+		UE_LOG(LogConsoleResponse, Display, TEXT("Serial CRC of 1MB in %6.3fms"), float(FPlatformTime::Seconds() - StartTime) * 1000.0f);
+	}
+}
+
+static FAutoConsoleCommand TestParallelForCmd(
+	TEXT("TestParallelFor"),
+	TEXT("Simple test of ParallelFor."),
+	FConsoleCommandWithArgsDelegate::CreateStatic(&TestParallelFor)
+	);
 
 
diff --git a/Engine/Source/Runtime/Core/Public/Async/ParallelFor.h b/Engine/Source/Runtime/Core/Public/Async/ParallelFor.h
new file mode 100644
--- /dev/null
+++ b/Engine/Source/Runtime/Core/Public/Async/ParallelFor.h
@@ -0,0 +1,110 @@
+// Copyright 1998-2015 Epic Games, Inc. All Rights Reserved.
+
+/*=============================================================================
+	ParllelFor.h: TaskGraph library
+=============================================================================*/
+
+#pragma once
+#include "TaskGraphInterfaces.h"
+#include "Function.h"
+
+/** 
+	*	General purpose parallel for that uses the taskgraph
+	*	@param Num; number of calls of Body; Body(0), Body(1)....Body(Num - 1)
+	*	@param Body; Function to call from multiple threads
+	*	@param bForceSingleThread; Mostly used for testing, if true, run single threaded instead.
+	*	Notes: Please add stats around to calls to parallel for and within your lambda as appropriate. Do not clog the task graph with long running tasks or tasks that block.
+**/
+void ParallelFor(int32 Num, TFunctionRef<void(int32)> Body, bool bForceSingleThread = false)
+{
+	// struct to hold the working data; this outlives the ParallelFor call; lifetime is controlled by a shared pointer
+	struct FParallelForData
+	{
+		int32 Num;
+		TFunctionRef<void(int32)> Body;
+		FScopedEvent& DoneEvent;
+		FThreadSafeCounter IndexToDo;
+		FThreadSafeCounter NumCompleted;
+		FParallelForData(int32 InNum, TFunctionRef<void(int32)> InBody, FScopedEvent& InDoneEvent)
+			: Num(InNum)
+			, Body(InBody)
+			, DoneEvent(InDoneEvent)
+		{
+		}
+		void Process()
+		{
+			while (true)
+			{
+				int32 MyIndex = IndexToDo.Increment() - 1;
+				if (MyIndex < Num)
+				{
+					Body(MyIndex);
+					if (NumCompleted.Increment() == Num)
+					{
+						DoneEvent.Trigger(); // I was the last one; let parallelfor exit
+					}
+				}
+				else
+				{
+					break;
+				}
+			}
+		}
+	};
+	class FParallelForTask
+	{
+		TSharedRef<FParallelForData, ESPMode::ThreadSafe> Data;
+	public:
+		FParallelForTask(TSharedRef<FParallelForData, ESPMode::ThreadSafe>& InData)
+			: Data(InData) 
+		{
+		}
+		static FORCEINLINE TStatId GetStatId()
+		{
+			return GET_STATID(STAT_ParallelForTask);
+		}
+		static FORCEINLINE ENamedThreads::Type GetDesiredThread()
+		{
+			return ENamedThreads::AnyThread;
+		}
+		static FORCEINLINE ESubsequentsMode::Type GetSubsequentsMode() 
+		{ 
+			return ESubsequentsMode::FireAndForget; 
+		}
+		void DoTask(ENamedThreads::Type CurrentThread, const FGraphEventRef& MyCompletionGraphEvent)
+		{
+			Data->Process();
+		}
+	};
+
+	SCOPE_CYCLE_COUNTER(STAT_ParallelFor);
+	check(Num >= 0);
+
+	int32 AnyThreadTasks = 0;
+	if (Num > 1 && !bForceSingleThread && FApp::ShouldUseThreadingForPerformance())
+	{
+		AnyThreadTasks = FMath::Min<int32>(FTaskGraphInterface::Get().GetNumWorkerThreads(), Num - 1);
+	}
+	if (!AnyThreadTasks)
+	{
+		// no threads, just do it and return
+		for (int32 Index = 0; Index < Num; Index++)
+		{
+			Body(Index);
+		}
+		return;
+	}
+
+	FScopedEvent DoneEvent;
+	TSharedRef<FParallelForData, ESPMode::ThreadSafe> Data = MakeShareable(new FParallelForData(Num, Body, DoneEvent));
+	for (int32 Task = 0; Task < AnyThreadTasks; Task++)
+	{
+		TGraphTask<FParallelForTask>::CreateTask().ConstructAndDispatchWhenReady(Data);		
+	}
+	// this thread can help too and this is important to prevent deadlock on recursion 
+	Data->Process();
+	// DoneEvent waits here if some other thread finishes the last item
+	// Data must live on until all of the tasks are cleared which might be long after this function exits
+}
+
+
diff --git a/Engine/Source/Runtime/Core/Public/Async/TaskGraphInterfaces.h b/Engine/Source/Runtime/Core/Public/Async/TaskGraphInterfaces.h
--- a/Engine/Source/Runtime/Core/Public/Async/TaskGraphInterfaces.h
+++ b/Engine/Source/Runtime/Core/Public/Async/TaskGraphInterfaces.h
@@ -23,6 +23,8 @@ DECLARE_STATS_GROUP(TEXT("Task Graph Tasks"), STATGROUP_TaskGraphTasks, STATCAT_
 DECLARE_CYCLE_STAT_EXTERN(TEXT("FReturnGraphTask"), STAT_FReturnGraphTask, STATGROUP_TaskGraphTasks, CORE_API);
 DECLARE_CYCLE_STAT_EXTERN(TEXT("FTriggerEventGraphTask"), STAT_FTriggerEventGraphTask, STATGROUP_TaskGraphTasks, CORE_API);
 DECLARE_CYCLE_STAT_EXTERN(TEXT("Unknown Graph Task"), STAT_UnknownGraphTask, STATGROUP_TaskGraphTasks, CORE_API);
+DECLARE_CYCLE_STAT_EXTERN(TEXT("ParallelFor"), STAT_ParallelFor, STATGROUP_TaskGraphTasks, CORE_API);
+DECLARE_CYCLE_STAT_EXTERN(TEXT("ParallelForTask"), STAT_ParallelForTask, STATGROUP_TaskGraphTasks, CORE_API);
 
 namespace ENamedThreads
 {
@@ -126,7 +128,7 @@ public:
 	static CORE_API FTaskGraphInterface& Get();
 
 	/** Return the current thread type, if known. **/
-	virtual ENamedThreads::Type GetCurrentThreadIfKnown() = 0;
+	virtual ENamedThreads::Type GetCurrentThreadIfKnown(bool bLocalQueue = false) = 0;
 
 	/** Return the number of worker (non-named) threads **/
 	virtual	int32 GetNumWorkerThreads() = 0;
