Made a number of string operations in key files with symbols CaseSensitive
Removed some duplicated checks in FPaths::IsRelative

diff --git a/Engine/Source/Runtime/AssetRegistry/Private/PackageReader.cpp b/Engine/Source/Runtime/AssetRegistry/Private/PackageReader.cpp
--- a/Engine/Source/Runtime/AssetRegistry/Private/PackageReader.cpp
+++ b/Engine/Source/Runtime/AssetRegistry/Private/PackageReader.cpp
@@ -171,7 +171,7 @@ bool FPackageReader::ReadAssetRegistryData (TArray<FBackgroundAssetData*>& Asset
 		FString GroupNames;
 		FString AssetName;
 
-		if ( ObjectPath.Contains(TEXT(".")))
+		if ( ObjectPath.Contains(TEXT("."), ESearchCase::CaseSensitive))
 		{
 			ObjectPath.Split(TEXT("."), &GroupNames, &AssetName, ESearchCase::CaseSensitive, ESearchDir::FromEnd);
 		}
@@ -236,7 +236,7 @@ bool FPackageReader::ReadAssetDataFromThumbnailCache(TArray<FBackgroundAssetData
 		FString GroupNames;
 		FString AssetName;
 
-		if ( ObjectPathWithoutPackageName.Contains(TEXT(".")) )
+		if ( ObjectPathWithoutPackageName.Contains(TEXT("."), ESearchCase::CaseSensitive) )
 		{
 			ObjectPathWithoutPackageName.Split(TEXT("."), &GroupNames, &AssetName, ESearchCase::CaseSensitive, ESearchDir::FromEnd);
 		}
diff --git a/Engine/Source/Runtime/Core/Private/Misc/Paths.cpp b/Engine/Source/Runtime/Core/Private/Misc/Paths.cpp
--- a/Engine/Source/Runtime/Core/Private/Misc/Paths.cpp
+++ b/Engine/Source/Runtime/Core/Private/Misc/Paths.cpp
@@ -214,7 +214,7 @@ FString FPaths::GameUserDeveloperDir()
 		for (int32 CharIdx = 0; CharIdx < InvalidChars.Len(); ++CharIdx)
 		{
 			const FString Char = InvalidChars.Mid(CharIdx, 1);
-			UserFolder = UserFolder.Replace(*Char, TEXT("_"));
+			UserFolder = UserFolder.Replace(*Char, TEXT("_"), ESearchCase::CaseSensitive);
 		}
 	}
 
@@ -571,12 +571,10 @@ bool FPaths::IsRelative(const FString& InPath)
 {
 	// The previous implementation of this function seemed to handle normalized and unnormalized paths, so this one does too for legacy reasons.
 
-	const bool IsRooted =	InPath.StartsWith(TEXT("\\\\"), ESearchCase::CaseSensitive)	||												// "\\" for UNC or "network" paths.
-							InPath.StartsWith(TEXT("//"), ESearchCase::CaseSensitive)	||												// Equivalent to "\\", considering normalization replaces "\\" with "//".
-							InPath.StartsWith(TEXT("\\"), ESearchCase::CaseSensitive)	||												// Root of the current directory on Windows
-							InPath.StartsWith(TEXT("/"), ESearchCase::CaseSensitive)	||												// Root of the current directory on Windows, root on UNIX-likes.
-							InPath.StartsWith(TEXT("root:/")) ||											// Feature packs use this
-							(InPath.Len() >= 2 && FChar::IsAlpha(InPath[0]) && InPath[1] == TEXT(':'));	// Starts with "<DriveLetter>:"
+	const bool IsRooted = InPath.StartsWith(TEXT("\\"), ESearchCase::CaseSensitive)	||					// Root of the current directory on Windows. Also covers "\\" for UNC or "network" paths.
+						  InPath.StartsWith(TEXT("/"), ESearchCase::CaseSensitive)	||					// Root of the current directory on Windows, root on UNIX-likes.  Also covers "\\", considering normalization replaces "\\" with "//".						
+						  InPath.StartsWith(TEXT("root:/")) |											// Feature packs use this
+						  (InPath.Len() >= 2 && FChar::IsAlpha(InPath[0]) && InPath[1] == TEXT(':'));	// Starts with "<DriveLetter>:"
 
 	return !IsRooted;
 }
@@ -847,7 +845,7 @@ bool FPaths::ValidatePath( const FString& InPath, FText* OutReason )
 	RemoveDuplicateSlashes(Standardized);
 
 	// The loop below requires that the path not end with a /
-	if(Standardized.EndsWith(TEXT("/")))
+	if(Standardized.EndsWith(TEXT("/"), ESearchCase::CaseSensitive))
 	{
 		Standardized = Standardized.LeftChop(1);
 	}
@@ -939,7 +937,7 @@ const FString& FPaths::GetRelativePathToRoot()
 			FPaths::MakePathRelativeTo(RelativePathToRoot, *BaseDirectory);
 
 			// Ensure that the path ends w/ '/'
-			if ((RelativePathToRoot.Len() > 0) && (RelativePathToRoot.EndsWith(TEXT("/")) == false) && (RelativePathToRoot.EndsWith(TEXT("\\")) == false))
+			if ((RelativePathToRoot.Len() > 0) && (RelativePathToRoot.EndsWith(TEXT("/"), ESearchCase::CaseSensitive) == false) && (RelativePathToRoot.EndsWith(TEXT("\\"), ESearchCase::CaseSensitive) == false))
 			{
 				RelativePathToRoot += TEXT("/");
 			}
diff --git a/Engine/Source/Runtime/Core/Private/Modules/ModuleManager.cpp b/Engine/Source/Runtime/Core/Private/Modules/ModuleManager.cpp
--- a/Engine/Source/Runtime/Core/Private/Modules/ModuleManager.cpp
+++ b/Engine/Source/Runtime/Core/Private/Modules/ModuleManager.cpp
@@ -187,7 +187,7 @@ void FModuleManager::AddModule( const FName InModuleName )
 									// need to reject some files here that are not numbered...release executables, do have a suffix, so we need to make sure we don't find the debug version
 									check(FoundFilePath.Len() > Prefix.Len() + Suffix.Len());
 									FString Center = FoundFilePath.Mid(Prefix.Len(), FoundFilePath.Len() - Prefix.Len() - Suffix.Len());
-									check(Center.StartsWith(TEXT("-"))); // a minus sign is still considered numeric, so we can leave it.
+									check(Center.StartsWith(TEXT("-"), ESearchCase::CaseSensitive)); // a minus sign is still considered numeric, so we can leave it.
 									if (!Center.IsNumeric())
 									{
 										// this is a debug DLL or something, it is not a numbered hot DLL
diff --git a/Engine/Source/Runtime/Core/Private/Stats/Stats2.cpp b/Engine/Source/Runtime/Core/Private/Stats/Stats2.cpp
--- a/Engine/Source/Runtime/Core/Private/Stats/Stats2.cpp
+++ b/Engine/Source/Runtime/Core/Private/Stats/Stats2.cpp
@@ -580,7 +580,7 @@ FName FStatNameAndInfo::GetShortNameFrom(FName InLongName)
 {
 	FString Input(InLongName.ToString());
 
-	if (Input.StartsWith(TEXT("//")))
+	if (Input.StartsWith(TEXT("//"), ESearchCase::CaseSensitive))
 	{
 		Input = Input.RightChop(2);
 		const int32 IndexEnd = Input.Find(TEXT("//"), ESearchCase::CaseSensitive);
@@ -603,7 +603,7 @@ FName FStatNameAndInfo::GetGroupNameFrom(FName InLongName)
 {
 	FString Input(InLongName.ToString());
 
-	if (Input.StartsWith(TEXT("//")))
+	if (Input.StartsWith(TEXT("//"), ESearchCase::CaseSensitive))
 	{
 		Input = Input.RightChop(2);
 		const int32 IndexEnd = Input.Find(TEXT("//"), ESearchCase::CaseSensitive);
@@ -637,7 +637,7 @@ FName FStatNameAndInfo::GetGroupCategoryFrom(FName InLongName)
 {
 	FString Input(InLongName.ToString());
 
-	const int32 IndexStart = Input.Find(TEXT("///////"), ESearchCase::IgnoreCase, ESearchDir::FromEnd);
+	const int32 IndexStart = Input.Find(TEXT("///////"), ESearchCase::CaseSensitive, ESearchDir::FromEnd);
 	if (IndexStart != INDEX_NONE)
 	{
 		Input = Input.RightChop(IndexStart + 7);
diff --git a/Engine/Source/Runtime/CoreUObject/Private/Misc/PackageName.cpp b/Engine/Source/Runtime/CoreUObject/Private/Misc/PackageName.cpp
--- a/Engine/Source/Runtime/CoreUObject/Private/Misc/PackageName.cpp
+++ b/Engine/Source/Runtime/CoreUObject/Private/Misc/PackageName.cpp
@@ -148,7 +148,7 @@ struct FLongPackagePathsSingleton
 		FString RelativeContentPath = IFileManager::Get().ConvertToRelativePath( *ContentPath );
 
 		// Make sure the path ends in a trailing path separator.  We are expecting that in the InternalFilenameToLongPackageName code.
-		if( !RelativeContentPath.EndsWith( TEXT( "/" ) ) )
+		if( !RelativeContentPath.EndsWith( TEXT( "/" ), ESearchCase::CaseSensitive ) )
 		{
 			RelativeContentPath += TEXT( "/" );
 		}
@@ -169,7 +169,7 @@ struct FLongPackagePathsSingleton
 		FString RelativeContentPath = IFileManager::Get().ConvertToRelativePath(*ContentPath);
 
 		// Make sure the path ends in a trailing path separator.  We are expecting that in the InternalFilenameToLongPackageName code.
-		if (!RelativeContentPath.EndsWith(TEXT("/")))
+		if (!RelativeContentPath.EndsWith(TEXT("/"), ESearchCase::CaseSensitive))
 		{
 			RelativeContentPath += TEXT("/");
 		}
@@ -384,7 +384,7 @@ bool FPackageName::SplitLongPackageName(const FString& InLongPackageName, FStrin
 	OutPackagePath = FPaths::GetPath(RemainingPackageName) / "";
 	OutPackageName = FPaths::GetCleanFilename(RemainingPackageName);
 
-	if(bStripRootLeadingSlash && OutPackageRoot.StartsWith("/"))
+	if(bStripRootLeadingSlash && OutPackageRoot.StartsWith(TEXT("/"), ESearchCase::CaseSensitive))
 	{
 		OutPackageRoot.RemoveAt(0);
 	}
@@ -817,7 +817,7 @@ bool FPackageName::ParseExportTextPath(const FString& InExportTextPath, FString*
 		if ( OutObjectPath )
 		{
 			FString& OutObjectPathRef = *OutObjectPath;
-			if ( OutObjectPathRef.EndsWith(TEXT("'")) )
+			if ( OutObjectPathRef.EndsWith(TEXT("'"), ESearchCase::CaseSensitive) )
 			{
 				OutObjectPathRef = OutObjectPathRef.LeftChop(1);
 			}
diff --git a/Engine/Source/Runtime/CoreUObject/Private/UObject/LinkerLoad.cpp b/Engine/Source/Runtime/CoreUObject/Private/UObject/LinkerLoad.cpp
--- a/Engine/Source/Runtime/CoreUObject/Private/UObject/LinkerLoad.cpp
+++ b/Engine/Source/Runtime/CoreUObject/Private/UObject/LinkerLoad.cpp
@@ -4368,7 +4368,7 @@ ULinkerLoad::ELinkerStatus ULinkerLoad::FixupExportMap()
 				StrRedirectName = RedirectName->ToString();
 
 				// Accepts either "PackageName.ClassName" or just "ClassName"
-				int32 Offset = StrRedirectName.Find(TEXT("."));
+				int32 Offset = StrRedirectName.Find(TEXT("."), ESearchCase::CaseSensitive);
 				if ( Offset >= 0 )
 				{
 					// A package class name redirect
@@ -4429,7 +4429,7 @@ ULinkerLoad::ELinkerStatus ULinkerLoad::FixupExportMap()
 					StrRedirectName = RedirectName->ToString();
 
 					// Accepts either "PackageName.ClassName" or just "ClassName"
-					int32 Offset = StrRedirectName.Find(TEXT("."));
+					int32 Offset = StrRedirectName.Find(TEXT("."), ESearchCase::CaseSensitive);
 					if ( Offset >= 0 )
 					{
 						// A package class name redirect
diff --git a/Engine/Source/Runtime/CoreUObject/Private/UObject/UObjectGlobals.cpp b/Engine/Source/Runtime/CoreUObject/Private/UObject/UObjectGlobals.cpp
--- a/Engine/Source/Runtime/CoreUObject/Private/UObject/UObjectGlobals.cpp
+++ b/Engine/Source/Runtime/CoreUObject/Private/UObject/UObjectGlobals.cpp
@@ -144,9 +144,9 @@ UObject* StaticFindObject( UClass* ObjectClass, UObject* InObjectPackage, const
 		if (s_bCurrentlyLoading == false)
 		{
 			FString NameCheck = OrigInName;
-			if (NameCheck.Contains(TEXT(".")) && 
-				!NameCheck.Contains(TEXT("'")) && 
-				!NameCheck.Contains(TEXT(":")) )
+			if (NameCheck.Contains(TEXT("."), ESearchCase::CaseSensitive) && 
+				!NameCheck.Contains(TEXT("'"), ESearchCase::CaseSensitive) && 
+				!NameCheck.Contains(TEXT(":"), ESearchCase::CaseSensitive) )
 			{
 				s_bCurrentlyLoading = true;
 				MatchingObject = StaticLoadObject(ObjectClass, NULL, OrigInName, NULL,  LOAD_NoWarn, NULL);
@@ -472,7 +472,7 @@ UPackage* CreatePackage( UObject* InOuter, const TCHAR* PackageName )
 		UE_LOG(LogUObjectGlobals, Fatal, TEXT("Attempted to create a package with name containing double slashes. PackageName: %s"), PackageName);
 	}
 
-	if( InName.EndsWith( TEXT( "." ) ) )
+	if( InName.EndsWith( TEXT( "." ), ESearchCase::CaseSensitive ) )
 	{
 		FString InName2 = InName.Left( InName.Len() - 1 );
 		UE_LOG(LogUObjectGlobals, Log,  TEXT( "Invalid Package Name entered - '%s' renamed to '%s'" ), *InName, *InName2 );
@@ -2385,7 +2385,7 @@ void ConstructorHelpers::CheckFoundViaRedirect(UObject *Object, const FString& P
 	if (Redir && Redir->DestinationObject == Object)
 	{
 		FString NewString = Object->GetFullName();
-		NewString.ReplaceInline(TEXT(" "), TEXT("'"));
+		NewString.ReplaceInline(TEXT(" "), TEXT("'"), ESearchCase::CaseSensitive);
 		NewString += TEXT("'");
 
 		auto CurrentInitializer = FTlsObjectInitializers::Top();
