Use more optimized Compare for FName< operator

diff --git a/Engine/Source/Runtime/Core/Public/UObject/NameTypes.h b/Engine/Source/Runtime/Core/Public/UObject/NameTypes.h
--- a/Engine/Source/Runtime/Core/Public/UObject/NameTypes.h
+++ b/Engine/Source/Runtime/Core/Public/UObject/NameTypes.h
@@ -551,7 +551,7 @@ public:
 	 */
 	FORCEINLINE bool operator<( const FName& Other ) const
 	{
-		return ToString() < Other.ToString();
+		return Compare(Other) < 0;
 	}
 
 	FORCEINLINE bool IsNone() const
