Removing a stray / character from the final user ini, in the config cache.

diff --git a/Engine/Source/Runtime/Core/Private/Misc/ConfigCacheIni.cpp b/Engine/Source/Runtime/Core/Private/Misc/ConfigCacheIni.cpp
--- a/Engine/Source/Runtime/Core/Private/Misc/ConfigCacheIni.cpp
+++ b/Engine/Source/Runtime/Core/Private/Misc/ConfigCacheIni.cpp
@@ -2803,7 +2803,7 @@ static void GetSourceIniHierarchyFilenames(const TCHAR* InBaseIniName, const TCH
 
 	// [[[[ PROJECT USER OVERRIDES ]]]]
 	// Game/Config/User* ini (Checkpointed here at the end)
-	HierarchyCheckpointPath = FString::Printf(TEXT("%s/User%s.ini"), SourceConfigDir, InBaseIniName);
+	HierarchyCheckpointPath = FString::Printf(TEXT("%sUser%s.ini"), SourceConfigDir, InBaseIniName);
 	OutHierarchy.Add(FIniFilename(HierarchyCheckpointPath, false, GenerateHierarchyCacheKey(OutHierarchy, HierarchyCheckpointPath, InBaseIniName)));
 	
 }
