PR #967: Added support for Android menu and volume keys (Contributed by [REDACTED_PERSON])
#ue4
#android

diff --git a/Engine/Source/Runtime/Core/Private/Android/AndroidInputInterface.cpp b/Engine/Source/Runtime/Core/Private/Android/AndroidInputInterface.cpp
--- a/Engine/Source/Runtime/Core/Private/Android/AndroidInputInterface.cpp
+++ b/Engine/Source/Runtime/Core/Private/Android/AndroidInputInterface.cpp
@@ -53,6 +53,9 @@ FAndroidInputInterface::FAndroidInputInterface( const TSharedRef< FGenericApplic
 	ButtonMapping[14] = EControllerButtons::DPadLeft;
 	ButtonMapping[15] = EControllerButtons::DPadRight;
 	ButtonMapping[16] = EControllerButtons::AndroidBack;  // Technically just an alias for SpecialLeft
+	ButtonMapping[17] = EControllerButtons::AndroidVolumeUp;
+	ButtonMapping[18] = EControllerButtons::AndroidVolumeDown;
+	ButtonMapping[19] = EControllerButtons::AndroidMenu;  // Technically just an alias for SpecialRightNewControllerData[deviceId].ButtonStates[6] = buttonDown; 
 
 	InitialButtonRepeatDelay = 0.2f;
 	ButtonRepeatDelay = 0.1f;
@@ -803,7 +806,7 @@ void FAndroidInputInterface::JoystickButtonEvent(int32 deviceId, int32 buttonId,
 		case AKEYCODE_BUTTON_L1:     NewControllerData[deviceId].ButtonStates[4] = buttonDown; break;
 		case AKEYCODE_BUTTON_R1:     NewControllerData[deviceId].ButtonStates[5] = buttonDown; break;
 		case AKEYCODE_BUTTON_START:
-		case AKEYCODE_MENU:          NewControllerData[deviceId].ButtonStates[6] = buttonDown; break;
+		case AKEYCODE_MENU:          NewControllerData[deviceId].ButtonStates[6] = buttonDown; NewControllerData[deviceId].ButtonStates[19] = buttonDown;  break;
 		case AKEYCODE_BUTTON_SELECT: 
 		case AKEYCODE_BACK:          NewControllerData[deviceId].ButtonStates[7] = buttonDown; NewControllerData[deviceId].ButtonStates[16] = buttonDown;  break;
 		case AKEYCODE_BUTTON_THUMBL: NewControllerData[deviceId].ButtonStates[8] = buttonDown; break;
diff --git a/Engine/Source/Runtime/Core/Private/Android/AndroidMisc.cpp b/Engine/Source/Runtime/Core/Private/Android/AndroidMisc.cpp
--- a/Engine/Source/Runtime/Core/Private/Android/AndroidMisc.cpp
+++ b/Engine/Source/Runtime/Core/Private/Android/AndroidMisc.cpp
@@ -706,6 +706,9 @@ uint32 FAndroidMisc::GetKeyMap( uint16* KeyCodes, FString* KeyNames, uint32 MaxM
 		//ADDKEYMAP( VK_RWIN, TEXT("RightCommand") );
 
 		ADDKEYMAP(AKEYCODE_BACK, TEXT("Android_Back"));
+		ADDKEYMAP(AKEYCODE_VOLUME_UP, TEXT("Android_Volume_Up"));
+		ADDKEYMAP(AKEYCODE_VOLUME_DOWN, TEXT("Android_Volume_Down"));
+		ADDKEYMAP(AKEYCODE_MENU, TEXT("Android_Menu"));
 
 	}
 
diff --git a/Engine/Source/Runtime/Core/Public/Android/AndroidInputInterface.h b/Engine/Source/Runtime/Core/Public/Android/AndroidInputInterface.h
--- a/Engine/Source/Runtime/Core/Public/Android/AndroidInputInterface.h
+++ b/Engine/Source/Runtime/Core/Public/Android/AndroidInputInterface.h
@@ -7,9 +7,9 @@
 #include "IInputInterface.h"
 #include "IForceFeedbackSystem.h"
 
-#if __ANDROID_API__ <= 18
+#if __ANDROID_API__ < 13
 
-// Joystick functions and constants only available API level 19 and above
+// Joystick functions and constants only available API level 13 and above
 // Definitions are provided to allow compiling against lower API levels, but
 // still using the features when available.
 
@@ -52,7 +52,7 @@ enum
 	AMOTION_EVENT_AXIS_GENERIC_9 = 40,
 	AMOTION_EVENT_AXIS_GENERIC_10 = 41,
 	AMOTION_EVENT_AXIS_GENERIC_11 = 42,
-	AMOTION_EVENT_AXIS_GENERIC_12 = 43,
+	AMOTION_EVENT_AXIS_GENERIC_12 = 43, 
 	AMOTION_EVENT_AXIS_GENERIC_13 = 44,
 	AMOTION_EVENT_AXIS_GENERIC_14 = 45,
 	AMOTION_EVENT_AXIS_GENERIC_15 = 46,
@@ -67,7 +67,7 @@ enum
 	AINPUT_SOURCE_GAMEPAD = 0x00000400 | AINPUT_SOURCE_CLASS_BUTTON,
 	AINPUT_SOURCE_JOYSTICK = 0x01000000 | AINPUT_SOURCE_CLASS_JOYSTICK,
 };
-#endif // __ANDROID_API__ <= 9
+#endif // __ANDROID_API__ < 13
 
 
 enum TouchType
@@ -86,7 +86,7 @@ struct TouchInput
 };
 
 #define MAX_NUM_CONTROLLERS					1  // Only one for now
-#define MAX_NUM_CONTROLLER_BUTTONS			17
+#define MAX_NUM_CONTROLLER_BUTTONS			20
 #define MAX_DEFERRED_MESSAGE_QUEUE_SIZE		128
 
 struct FAndroidControllerData
diff --git a/Engine/Source/Runtime/Core/Public/GenericPlatform/GenericApplicationMessageHandler.h b/Engine/Source/Runtime/Core/Public/GenericPlatform/GenericApplicationMessageHandler.h
--- a/Engine/Source/Runtime/Core/Public/GenericPlatform/GenericApplicationMessageHandler.h
+++ b/Engine/Source/Runtime/Core/Public/GenericPlatform/GenericApplicationMessageHandler.h
@@ -81,6 +81,9 @@ namespace EControllerButtons
 		GlobalBack,
 
 		AndroidBack,
+		AndroidVolumeUp,
+		AndroidVolumeDown,
+		AndroidMenu,
 
 		Invalid,
 	};
diff --git a/Engine/Source/Runtime/InputCore/Classes/InputCoreTypes.h b/Engine/Source/Runtime/InputCore/Classes/InputCoreTypes.h
--- a/Engine/Source/Runtime/InputCore/Classes/InputCoreTypes.h
+++ b/Engine/Source/Runtime/InputCore/Classes/InputCoreTypes.h
@@ -386,7 +386,11 @@ struct INPUTCORE_API EKeys
 	static const FKey Global_Play;
 	static const FKey Global_Back;
 
+	// Android-specific
 	static const FKey Android_Back;
+	static const FKey Android_Volume_Up;
+	static const FKey Android_Volume_Down;
+	static const FKey Android_Menu;
 
 	static const FKey Invalid;
 
diff --git a/Engine/Source/Runtime/InputCore/Private/InputCoreTypes.cpp b/Engine/Source/Runtime/InputCore/Private/InputCoreTypes.cpp
--- a/Engine/Source/Runtime/InputCore/Private/InputCoreTypes.cpp
+++ b/Engine/Source/Runtime/InputCore/Private/InputCoreTypes.cpp
@@ -229,8 +229,12 @@ const FKey EKeys::Global_Pause("Global_Pause");
 const FKey EKeys::Global_Play("Global_Play");
 const FKey EKeys::Global_Back("Global_Back");
 
+// Android-specific
 const FKey EKeys::Android_Back("Android_Back");
- 
+const FKey EKeys::Android_Volume_Up("Android_Volume_Up");
+const FKey EKeys::Android_Volume_Down("Android_Volume_Down");
+const FKey EKeys::Android_Menu("Android_Menu");
+
 const FKey EKeys::Invalid(NAME_None);
 
 bool EKeys::bInitialized = false;
@@ -464,7 +468,11 @@ void EKeys::Initialize()
 	AddKey(FKeyDetails(EKeys::Global_Play, LOCTEXT("Global_Play", "Global Play"), FKeyDetails::GamepadKey));
 	AddKey(FKeyDetails(EKeys::Global_Back, LOCTEXT("Global_Back", "Global Back"), FKeyDetails::GamepadKey));
 
+	// Android-specific
 	AddKey(FKeyDetails(EKeys::Android_Back, LOCTEXT("Android_Back", "Android Back"), FKeyDetails::GamepadKey));
+	AddKey(FKeyDetails(EKeys::Android_Volume_Up, LOCTEXT("Android_Volume_Up", "Android Volume Up"), FKeyDetails::GamepadKey));
+	AddKey(FKeyDetails(EKeys::Android_Volume_Down, LOCTEXT("Android_Volume_Down", "Android Volume Down"), FKeyDetails::GamepadKey));
+	AddKey(FKeyDetails(EKeys::Android_Menu, LOCTEXT("Android_Menu", "Android Menu"), FKeyDetails::GamepadKey));
 
 	// Initialize the input key manager.  This will cause any additional OEM keys to get added
 	FInputKeyManager::Get();
diff --git a/Engine/Source/Runtime/Slate/Private/Framework/Application/SlateApplication.cpp b/Engine/Source/Runtime/Slate/Private/Framework/Application/SlateApplication.cpp
--- a/Engine/Source/Runtime/Slate/Private/Framework/Application/SlateApplication.cpp
+++ b/Engine/Source/Runtime/Slate/Private/Framework/Application/SlateApplication.cpp
@@ -4648,6 +4648,9 @@ FKey TranslateControllerButtonToKey( EControllerButtons::Type Button )
 	case EControllerButtons::GlobalBack: Key = EKeys::Global_Back; break;
 
 	case EControllerButtons::AndroidBack: Key = EKeys::Android_Back; break;
+	case EControllerButtons::AndroidVolumeUp: Key = EKeys::Android_Volume_Up; break;
+	case EControllerButtons::AndroidVolumeDown: Key = EKeys::Android_Volume_Down; break;
+	case EControllerButtons::AndroidMenu: Key = EKeys::Android_Menu; break;
 
 	case EControllerButtons::Invalid: Key = EKeys::Invalid; break;
 	}
