Rename awkwardly named "-engineinstalled" command line option to "-installedengine".

diff --git a/Engine/Source/Programs/AutomationTool/Scripts/RocketBuild.Automation.cs b/Engine/Source/Programs/AutomationTool/Scripts/RocketBuild.Automation.cs
--- a/Engine/Source/Programs/AutomationTool/Scripts/RocketBuild.Automation.cs
+++ b/Engine/Source/Programs/AutomationTool/Scripts/RocketBuild.Automation.cs
@@ -462,7 +462,7 @@ namespace Rocket
 
 			// Write the default command line to the root of the output directory
 			string CommandLinePath = Path.Combine(TargetDir, "UE4CommandLine.txt");
-			File.WriteAllText(CommandLinePath, "-engineinstalled -rocket");
+			File.WriteAllText(CommandLinePath, "-installedengine -rocket");
 			AddBuildProduct(CommandLinePath);
 		}
 	}
diff --git a/Engine/Source/Runtime/Core/Private/Misc/App.cpp b/Engine/Source/Runtime/Core/Private/Misc/App.cpp
--- a/Engine/Source/Runtime/Core/Private/Misc/App.cpp
+++ b/Engine/Source/Runtime/Core/Private/Misc/App.cpp
@@ -118,7 +118,7 @@ bool FApp::IsInstalled()
 
 bool FApp::IsEngineInstalled()
 {
-	static bool bIsInstalledEngine = IsInstalled() || (FRocketSupport::IsRocket() ? !FParse::Param(FCommandLine::Get(), TEXT("EngineNotInstalled")) : FParse::Param(FCommandLine::Get(), TEXT("EngineInstalled")));
+	static bool bIsInstalledEngine = IsInstalled() || (FRocketSupport::IsRocket() ? !FParse::Param(FCommandLine::Get(), TEXT("NotInstalledEngine")) : FParse::Param(FCommandLine::Get(), TEXT("InstalledEngine")));
 	return bIsInstalledEngine;
 }
 
