Move virtual destructor to the interface level instead of the implementor

diff --git a/Engine/Source/Runtime/Core/Public/Misc/ITransaction.h b/Engine/Source/Runtime/Core/Public/Misc/ITransaction.h
--- a/Engine/Source/Runtime/Core/Public/Misc/ITransaction.h
+++ b/Engine/Source/Runtime/Core/Public/Misc/ITransaction.h
@@ -20,6 +20,7 @@ typedef void(*STRUCT_DTOR)( void* TPtr );					// destruct
 class ITransactionObjectAnnotation 
 { 
 public:
+	virtual ~ITransactionObjectAnnotation() {}
 	virtual void AddReferencedObjects(class FReferenceCollector& Collector) = 0;
 };
 
diff --git a/Engine/Source/Runtime/Engine/Classes/GameFramework/Actor.h b/Engine/Source/Runtime/Engine/Classes/GameFramework/Actor.h
--- a/Engine/Source/Runtime/Engine/Classes/GameFramework/Actor.h
+++ b/Engine/Source/Runtime/Engine/Classes/GameFramework/Actor.h
@@ -1186,7 +1186,6 @@ public:
 	{
 	public:
 		FActorTransactionAnnotation(const AActor* Actor);
-		virtual ~FActorTransactionAnnotation() {}
 
 		virtual void AddReferencedObjects(FReferenceCollector& Collector) override;
 
