Core: Added FString::Reset() which allows emptying a string without reallocating memory
[INTEGRATE] Change 2494318 by [REDACTED_PERSON][REDACTED_HANDLE] on 2015/03/27 15:38:07

diff --git a/Engine/Source/Runtime/Core/Public/Containers/UnrealString.h b/Engine/Source/Runtime/Core/Public/Containers/UnrealString.h
--- a/Engine/Source/Runtime/Core/Public/Containers/UnrealString.h
+++ b/Engine/Source/Runtime/Core/Public/Containers/UnrealString.h
@@ -279,6 +279,17 @@ public:
 		return Data.Num() <= 1;
 	}
 
+	/**
+	 * Empties the string, but doesn't change memory allocation, unless the new size is larger than the current string.
+	 *
+	 * @param NewReservedSize The expected usage size (in characters, not including the terminator) after calling this function.
+	 */
+	FORCEINLINE void Reset(int32 NewReservedSize = 0)
+	{
+		const int32 NewSizeIncludingTerminator = (NewReservedSize > 0) ? (NewReservedSize + 1) : 0;
+		Data.Reset(NewSizeIncludingTerminator);
+	}
+
 	/**
 	 * Remove unallocated empty character space from the end of this string
 	 */
