Core - Fixed a crash in GetFilenameOnDisk

diff --git a/Engine/Source/Runtime/Core/Public/HAL/IPlatformFileLogWrapper.h b/Engine/Source/Runtime/Core/Public/HAL/IPlatformFileLogWrapper.h
--- a/Engine/Source/Runtime/Core/Public/HAL/IPlatformFileLogWrapper.h
+++ b/Engine/Source/Runtime/Core/Public/HAL/IPlatformFileLogWrapper.h
@@ -209,7 +209,7 @@ public:
 		double StartTime = FPlatformTime::Seconds();
 		FString Result = LowerLevel->GetFilenameOnDisk(Filename);
 		float ThisTime = 1000.0f * float(FPlatformTime::Seconds() - StartTime);
-		FILE_LOG(LogPlatformFile, Log, TEXT("GetFilenameOnDisk return %llx [%s]"), *Result, ThisTime);
+		FILE_LOG(LogPlatformFile, Log, TEXT("GetFilenameOnDisk return %s [%fms]"), *Result, ThisTime);
 		return Result;
 	}
 	virtual IFileHandle*	OpenRead(const TCHAR* Filename, bool bAllowWrite) override
