Added support for HasSizingFrame in window definition on Mac

diff --git a/Engine/Source/Runtime/Core/Private/Mac/MacWindow.cpp b/Engine/Source/Runtime/Core/Private/Mac/MacWindow.cpp
--- a/Engine/Source/Runtime/Core/Private/Mac/MacWindow.cpp
+++ b/Engine/Source/Runtime/Core/Private/Mac/MacWindow.cpp
@@ -68,7 +68,7 @@ void FMacWindow::Initialize( FMacApplication* const Application, const TSharedRe
 		{
 			WindowStyle |= NSMiniaturizableWindowMask;
 		}
-		if( Definition->SupportsMaximize )
+		if( Definition->SupportsMaximize || Definition->HasSizingFrame )
 		{
 			WindowStyle |= NSResizableWindowMask;
 		}
