Properly set bIsProcessingDeferredEvents if FMacApplication::ProcessDeferredEvents() is called recursively

diff --git a/Engine/Source/Runtime/Core/Private/Mac/MacApplication.cpp b/Engine/Source/Runtime/Core/Private/Mac/MacApplication.cpp
--- a/Engine/Source/Runtime/Core/Private/Mac/MacApplication.cpp
+++ b/Engine/Source/Runtime/Core/Private/Mac/MacApplication.cpp
@@ -176,6 +176,7 @@ void FMacApplication::ProcessDeferredEvents(const float TimeDelta)
 	DeferredEvents.Empty();
 	EventsMutex.Unlock();
 
+	const bool bAlreadyProcessingDeferredEvents = bIsProcessingDeferredEvents;
 	bIsProcessingDeferredEvents = true;
 
 	for (int32 Index = 0; Index < EventsToProcess.Num(); ++Index)
@@ -183,7 +184,7 @@ void FMacApplication::ProcessDeferredEvents(const float TimeDelta)
 		ProcessEvent(EventsToProcess[Index]);
 	}
 
-	bIsProcessingDeferredEvents = false;
+	bIsProcessingDeferredEvents = bAlreadyProcessingDeferredEvents;
 
 	InvalidateTextLayouts();
 	CloseQueuedWindows();
