Changed handling of double click on Mac to solve a problem with window zooming if user moved the cursor over the title bar during a long taks initiated by double click. Also, added some additional validation for OnCursorSet() calls in MacApplication

diff --git a/Engine/Source/Runtime/Core/Private/Mac/MacApplication.cpp b/Engine/Source/Runtime/Core/Private/Mac/MacApplication.cpp
--- a/Engine/Source/Runtime/Core/Private/Mac/MacApplication.cpp
+++ b/Engine/Source/Runtime/Core/Private/Mac/MacApplication.cpp
@@ -722,7 +722,7 @@ void FMacApplication::ProcessMouseMovedEvent(const FDeferredMacEvent& Event)
 		}
 	}
 
-	if (Event.Window.IsValid() && !DraggedWindow && !GetCapture())
+	if (Event.Window.IsValid() && EventWindow->GetWindowHandle() && !DraggedWindow && !GetCapture())
 	{
 		MessageHandler->OnCursorSet();
 	}
@@ -754,32 +754,31 @@ void FMacApplication::ProcessMouseDownEvent(const FDeferredMacEvent& Event)
 	{
 		if (Button == LastPressedMouseButton && (Event.ClickCount % 2) == 0)
 		{
-			MessageHandler->OnMouseDoubleClick(EventWindow, Button);
-
-			if (Event.Window.IsValid())
+			bool IsMouseOverTitleBar = false;
+			const bool IsMovable = IsWindowMovable(EventWindow.ToSharedRef(), &IsMouseOverTitleBar);
+			if (IsMouseOverTitleBar)
 			{
-				bool IsMouseOverTitleBar = false;
-				const bool IsMovable = IsWindowMovable(EventWindow.ToSharedRef(), &IsMouseOverTitleBar);
-				if (IsMouseOverTitleBar)
+				const bool bShouldMinimize = [[NSUserDefaults standardUserDefaults] boolForKey:@"AppleMiniaturizeOnDoubleClick"];
+				if (bShouldMinimize)
 				{
-					const bool bShouldMinimize = [[NSUserDefaults standardUserDefaults] boolForKey:@"AppleMiniaturizeOnDoubleClick"];
-					if (bShouldMinimize)
-					{
-						MainThreadCall(^{ [EventWindow->GetWindowHandle() performMiniaturize:nil]; }, NSDefaultRunLoopMode, true);
-					}
-					else if (!FPlatformMisc::IsRunningOnMavericks())
-					{
-						MainThreadCall(^{ [EventWindow->GetWindowHandle() zoom:nil]; }, NSDefaultRunLoopMode, true);
-					}
+					MainThreadCall(^{ [EventWindow->GetWindowHandle() performMiniaturize:nil]; }, NSDefaultRunLoopMode, true);
+				}
+				else if (!FPlatformMisc::IsRunningOnMavericks())
+				{
+					MainThreadCall(^{ [EventWindow->GetWindowHandle() zoom:nil]; }, NSDefaultRunLoopMode, true);
 				}
 			}
+			else
+			{
+				MessageHandler->OnMouseDoubleClick(EventWindow, Button);
+			}
 		}
 		else
 		{
 			MessageHandler->OnMouseDown(EventWindow, Button);
 		}
 
-		if (EventWindow->GetWindowHandle() && !DraggedWindow && !GetCapture())
+		if (Event.Window.IsValid() && EventWindow->GetWindowHandle() && !DraggedWindow && !GetCapture())
 		{
 			MessageHandler->OnCursorSet();
 		}
@@ -840,7 +839,8 @@ void FMacApplication::ProcessScrollWheelEvent(const FDeferredMacEvent& Event)
 		MessageHandler->OnMouseWheel(DeltaY);
 	}
 
-	if (Event.Window.IsValid() && !DraggedWindow && !GetCapture())
+	TSharedPtr<FMacWindow> EventWindow = Event.Window.Pin();
+	if (Event.Window.IsValid() && EventWindow->GetWindowHandle() && !DraggedWindow && !GetCapture())
 	{
 		MessageHandler->OnCursorSet();
 	}
