Temp workaround for COTF failures for monolithic platforms.

diff --git a/Engine/Source/Runtime/Core/Private/Modules/ModuleManager.cpp b/Engine/Source/Runtime/Core/Private/Modules/ModuleManager.cpp
--- a/Engine/Source/Runtime/Core/Private/Modules/ModuleManager.cpp
+++ b/Engine/Source/Runtime/Core/Private/Modules/ModuleManager.cpp
@@ -34,6 +34,8 @@ FModuleManager& FModuleManager::Get()
 	{
 		ModuleManager = new FModuleManager();
 
+		//temp workaround for IPlatformFile being used for FPaths::DirectoryExists before main() sets up the commandline.
+#if PLATFORM_DESKTOP
 		// Ensure that dependency dlls can be found in restricted sub directories
 		const TCHAR* RestrictedFolderNames[] = { TEXT("NoRedist"), TEXT("NotForLicensees"), TEXT("CarefullyRedist") };
 		FString ModuleDir = FPlatformProcess::GetModulesDirectory();
@@ -45,6 +47,7 @@ FModuleManager& FModuleManager::Get()
 				ModuleManager->AddBinariesDirectory(*RestrictedFolder, false);
 			}
 		}
+#endif
 	}
 
 	return *ModuleManager;
