#ue4 - Perf: Optimize FTransform::TransformVector() and FTransform::TransformVectorNoScale() to remove unnecessary work.

- Vectorized versions are 25% faster than before.
- Remove useless translation since W is zero.
- Remove copy between FVector4 and FVector.

diff --git a/Engine/Source/Runtime/Core/Public/Math/Transform.h b/Engine/Source/Runtime/Core/Public/Math/Transform.h
--- a/Engine/Source/Runtime/Core/Public/Math/Transform.h
+++ b/Engine/Source/Runtime/Core/Public/Math/Transform.h
@@ -1185,37 +1185,40 @@ FORCEINLINE FVector FTransform::TransformPositionNoScale(const FVector& V) const
 
 FORCEINLINE FVector FTransform::TransformVector(const FVector& V) const
 {
-	return TransformFVector4(FVector4(V.X,V.Y,V.Z,0.0f));
+	DiagnosticCheckNaN_Rotate();
+	DiagnosticCheckNaN_Scale3D();
+	return Rotation.RotateVector(Scale3D*V);
 }
 
 FORCEINLINE FVector FTransform::TransformVectorNoScale(const FVector& V) const
 {
-	return TransformFVector4NoScale(FVector4(V.X,V.Y,V.Z,0.0f));
+	DiagnosticCheckNaN_Rotate();
+	return Rotation.RotateVector(V);
 }
 
 // do backward operation when inverse, translation -> rotation -> scale
 FORCEINLINE FVector FTransform::InverseTransformPosition(const FVector &V) const
 {
-	return ( Rotation.Inverse() * (V-Translation) ) * GetSafeScaleReciprocal(Scale3D);
+	return ( Rotation.UnrotateVector(V-Translation) ) * GetSafeScaleReciprocal(Scale3D);
 }
 
 // do backward operation when inverse, translation -> rotation
 FORCEINLINE FVector FTransform::InverseTransformPositionNoScale(const FVector &V) const
 {
-	return ( Rotation.Inverse() * (V-Translation) );
+	return ( Rotation.UnrotateVector(V-Translation) );
 }
 
 
 // do backward operation when inverse, translation -> rotation -> scale
 FORCEINLINE FVector FTransform::InverseTransformVector(const FVector &V) const
 {
-	return ( Rotation.Inverse() * V ) * GetSafeScaleReciprocal(Scale3D);
+	return ( Rotation.UnrotateVector(V) ) * GetSafeScaleReciprocal(Scale3D);
 }
 
 // do backward operation when inverse, translation -> rotation
 FORCEINLINE FVector FTransform::InverseTransformVectorNoScale(const FVector &V) const
 {
-	return ( Rotation.Inverse() * V );
+	return ( Rotation.UnrotateVector(V) );
 }
 
 FORCEINLINE FTransform FTransform::operator*(const FTransform& Other) const
diff --git a/Engine/Source/Runtime/Core/Public/Math/TransformVectorized.h b/Engine/Source/Runtime/Core/Public/Math/TransformVectorized.h
--- a/Engine/Source/Runtime/Core/Public/Math/TransformVectorized.h
+++ b/Engine/Source/Runtime/Core/Public/Math/TransformVectorized.h
@@ -1387,7 +1387,8 @@ FORCEINLINE FTransform FTransform::GetScaled(FVector InScale) const
 
 FORCEINLINE FVector4 FTransform::TransformFVector4NoScale(const FVector4& V) const
 {
-	DiagnosticCheckNaN_All();
+	DiagnosticCheckNaN_Rotate();
+	DiagnosticCheckNaN_Translate();
 
 	// if not, this won't work
 	checkSlow (V.W == 0.f || V.W == 1.f);
@@ -1424,12 +1425,11 @@ FORCEINLINE FVector4 FTransform::TransformFVector4(const FVector4& V) const
 
 	//Transform using QST is following
 	//QST(P) = Q*S*P*-Q + T where Q = quaternion, S = scale, T = translation
-	const VectorRegister Scale = Scale3D;
 	const VectorRegister InverseRotation = MAKE_QUATINV_VECTORREGISTER(Rotation);	
 
 	//FQuat Transform = Rotation*FQuat(Scale*V.X, Scale*V.Y, Scale*V.Z, 0.f)*Rotation.Inverse();
 	const VectorRegister InputVectorW0 = VectorSet_W0(InputVector);
-	const VectorRegister ScaledVec = VectorMultiply(Scale, InputVectorW0);
+	const VectorRegister ScaledVec = VectorMultiply(Scale3D, InputVectorW0);
 	const VectorRegister TempStorage = VectorQuaternionMultiply2(Rotation, ScaledVec);
 	const VectorRegister RotatedVec = VectorQuaternionMultiply2(TempStorage, InverseRotation);
 
@@ -1456,12 +1456,36 @@ FORCEINLINE FVector FTransform::TransformPositionNoScale(const FVector& V) const
 
 FORCEINLINE FVector FTransform::TransformVector(const FVector& V) const
 {
-	return TransformFVector4(FVector4(V.X,V.Y,V.Z,0.0f));
+	DiagnosticCheckNaN_Rotate();
+	DiagnosticCheckNaN_Scale3D();
+
+	const VectorRegister InputVectorW0 = VectorLoadFloat3_W0(&V);
+	const VectorRegister InverseRotation = MAKE_QUATINV_VECTORREGISTER(Rotation);
+
+	//FQuat Transform = Rotation*FQuat(Scale*V.X, Scale*V.Y, Scale*V.Z, 0.f)*Rotation.Inverse();
+	const VectorRegister ScaledVec = VectorMultiply(Scale3D, InputVectorW0);
+	const VectorRegister TempStorage = VectorQuaternionMultiply2(Rotation, ScaledVec);
+	const VectorRegister RotatedVec = VectorQuaternionMultiply2(TempStorage, InverseRotation);
+
+	FVector Result;
+	VectorStoreFloat3(RotatedVec, &Result);
+	return Result;
 }
 
 FORCEINLINE FVector FTransform::TransformVectorNoScale(const FVector& V) const
 {
-	return TransformFVector4NoScale(FVector4(V.X,V.Y,V.Z,0.0f));
+	DiagnosticCheckNaN_Rotate();
+
+	const VectorRegister InputVectorW0 = VectorLoadFloat3_W0(&V);
+	const VectorRegister InverseRotation = MAKE_QUATINV_VECTORREGISTER(Rotation);
+
+	//FQuat Transform = Rotation*FQuat(V.X, V.Y, V.Z, 0.f)*Rotation.Inverse()
+	const VectorRegister TempStorage = VectorQuaternionMultiply2(Rotation, InputVectorW0);
+	const VectorRegister RotatedVec = VectorQuaternionMultiply2(TempStorage, InverseRotation);
+
+	FVector Result;
+	VectorStoreFloat3(RotatedVec, &Result);
+	return Result;
 }
 
 // do backward operation when inverse, translation -> rotation -> scale
