PLATFORM_COMPILER_HAS_RVALUE_REFERENCES removed.

diff --git a/Engine/Source/Runtime/Core/Public/Containers/BitArray.h b/Engine/Source/Runtime/Core/Public/Containers/BitArray.h
--- a/Engine/Source/Runtime/Core/Public/Containers/BitArray.h
+++ b/Engine/Source/Runtime/Core/Public/Containers/BitArray.h
@@ -179,6 +179,14 @@ public:
 		Init(Value,InNumBits);
 	}
 
+	/**
+	 * Move constructor.
+	 */
+	FORCEINLINE TBitArray(TBitArray&& Other)
+	{
+		MoveOrCopy(*this, Other);
+	}
+
 	/**
 	 * Copy constructor.
 	 */
@@ -189,6 +197,19 @@ public:
 		*this = Copy;
 	}
 
+	/**
+	 * Move assignment.
+	 */
+	FORCEINLINE TBitArray& operator=(TBitArray&& Other)
+	{
+		if (this != &Other)
+		{
+			MoveOrCopy(*this, Other);
+		}
+
+		return *this;
+	}
+
 	/**
 	 * Assignment operator.
 	 */
@@ -211,8 +232,6 @@ public:
 		return *this;
 	}
 
-#if PLATFORM_COMPILER_HAS_RVALUE_REFERENCES
-
 private:
 	template <typename BitArrayType>
 	static FORCEINLINE typename TEnableIf<TContainerTraits<BitArrayType>::MoveWillEmptyContainer>::Type MoveOrCopy(BitArrayType& ToArray, BitArrayType& FromArray)
@@ -232,29 +251,6 @@ private:
 	}
 
 public:
-	/**
-	 * Move constructor.
-	 */
-	FORCEINLINE TBitArray(TBitArray&& Other)
-	{
-		MoveOrCopy(*this, Other);
-	}
-
-	/**
-	 * Move assignment.
-	 */
-	FORCEINLINE TBitArray& operator=(TBitArray&& Other)
-	{
-		if (this != &Other)
-		{
-			MoveOrCopy(*this, Other);
-		}
-
-		return *this;
-	}
-
-#endif
-
 	/**
 	 * Serializer
 	 */
@@ -660,9 +656,7 @@ private:
 template<typename Allocator>
 struct TContainerTraits<TBitArray<Allocator> > : public TContainerTraitsBase<TBitArray<Allocator> >
 {
-	enum { MoveWillEmptyContainer =
-		PLATFORM_COMPILER_HAS_RVALUE_REFERENCES &&
-		TAllocatorTraits<Allocator>::SupportsMove };
+	enum { MoveWillEmptyContainer = TAllocatorTraits<Allocator>::SupportsMove };
 };
 
 
diff --git a/Engine/Source/Runtime/Core/Public/Containers/ChunkedArray.h b/Engine/Source/Runtime/Core/Public/Containers/ChunkedArray.h
--- a/Engine/Source/Runtime/Core/Public/Containers/ChunkedArray.h
+++ b/Engine/Source/Runtime/Core/Public/Containers/ChunkedArray.h
@@ -26,8 +26,6 @@ public:
 		}
 	}
 
-#if PLATFORM_COMPILER_HAS_RVALUE_REFERENCES
-
 private:
 	template <typename ArrayType>
 	FORCEINLINE static typename TEnableIf<TContainerTraits<ArrayType>::MoveWillEmptyContainer>::Type MoveOrCopy(ArrayType& ToArray, ArrayType& FromArray)
@@ -59,8 +57,6 @@ public:
 		return *this;
 	}
 
-#endif
-
 #if PLATFORM_COMPILER_HAS_DEFAULTED_FUNCTIONS
 
 	TChunkedArray(const TChunkedArray&) = default;
@@ -229,9 +225,7 @@ protected:
 template <typename ElementType, uint32 TargetBytesPerChunk>
 struct TContainerTraits<TChunkedArray<ElementType, TargetBytesPerChunk> > : public TContainerTraitsBase<TChunkedArray<ElementType, TargetBytesPerChunk> >
 {
-	enum { MoveWillEmptyContainer =
-		PLATFORM_COMPILER_HAS_RVALUE_REFERENCES &&
-		TContainerTraits<typename TChunkedArray<ElementType, TargetBytesPerChunk>::ChunksType>::MoveWillEmptyContainer };
+	enum { MoveWillEmptyContainer = TContainerTraits<typename TChunkedArray<ElementType, TargetBytesPerChunk>::ChunksType>::MoveWillEmptyContainer };
 };
 
 
diff --git a/Engine/Source/Runtime/Core/Public/Containers/DynamicRHIResourceArray.h b/Engine/Source/Runtime/Core/Public/Containers/DynamicRHIResourceArray.h
--- a/Engine/Source/Runtime/Core/Public/Containers/DynamicRHIResourceArray.h
+++ b/Engine/Source/Runtime/Core/Public/Containers/DynamicRHIResourceArray.h
@@ -36,24 +36,32 @@ public:
 
 #if PLATFORM_COMPILER_HAS_DEFAULTED_FUNCTIONS
 
+	TResourceArray(TResourceArray&&) = default;
 	TResourceArray(const TResourceArray&) = default;
+	TResourceArray& operator=(TResourceArray&&) = default;
 	TResourceArray& operator=(const TResourceArray&) = default;
 
-	#if PLATFORM_COMPILER_HAS_RVALUE_REFERENCES
-
-		TResourceArray(TResourceArray&&) = default;
-		TResourceArray& operator=(TResourceArray&&) = default;
-
-	#endif
-
 #else
 
+	FORCEINLINE TResourceArray(TResourceArray&& Other)
+		: Super          ((Super&&)Other)
+		, bNeedsCPUAccess(Other.bNeedsCPUAccess)
+	{
+	}
+
 	FORCEINLINE TResourceArray(const TResourceArray& Other)
 		: Super          ((const Super&)Other)
 		, bNeedsCPUAccess(Other.bNeedsCPUAccess)
 	{
 	}
 
+	FORCEINLINE TResourceArray& operator=(TResourceArray&& Other)
+	{
+		(Super&)*this   = (Super&&)Other;
+		bNeedsCPUAccess = Other.bNeedsCPUAccess;
+		return *this;
+	}
+
 	FORCEINLINE TResourceArray& operator=(const TResourceArray& Other)
 	{
 		(Super&)*this   = (const Super&)Other;
@@ -61,23 +69,6 @@ public:
 		return *this;
 	}
 
-	#if PLATFORM_COMPILER_HAS_RVALUE_REFERENCES
-
-		FORCEINLINE TResourceArray(TResourceArray&& Other)
-			: Super          ((Super&&)Other)
-			, bNeedsCPUAccess(Other.bNeedsCPUAccess)
-		{
-		}
-
-		FORCEINLINE TResourceArray& operator=(TResourceArray&& Other)
-		{
-			(Super&)*this   = (Super&&)Other;
-			bNeedsCPUAccess = Other.bNeedsCPUAccess;
-			return *this;
-		}
-
-	#endif
-
 #endif
 
 	// FResourceArrayInterface
@@ -176,7 +167,5 @@ private:
 template< typename ElementType, uint32 Alignment >
 struct TContainerTraits<TResourceArray<ElementType, Alignment> > : public TContainerTraitsBase<TResourceArray<ElementType, Alignment> >
 {
-	enum { MoveWillEmptyContainer =
-		PLATFORM_COMPILER_HAS_RVALUE_REFERENCES &&
-		TContainerTraits<typename TResourceArray<ElementType, Alignment>::Super>::MoveWillEmptyContainer };
+	enum { MoveWillEmptyContainer = TContainerTraits<typename TResourceArray<ElementType, Alignment>::Super>::MoveWillEmptyContainer };
 };
diff --git a/Engine/Source/Runtime/Core/Public/Containers/SparseArray.h b/Engine/Source/Runtime/Core/Public/Containers/SparseArray.h
--- a/Engine/Source/Runtime/Core/Public/Containers/SparseArray.h
+++ b/Engine/Source/Runtime/Core/Public/Containers/SparseArray.h
@@ -461,6 +461,12 @@ public:
 	,	NumFreeIndices(0)
 	{}
 
+	/** Move constructor. */
+	TSparseArray(TSparseArray&& InCopy)
+	{
+		MoveOrCopy(*this, InCopy);
+	}
+
 	/** Copy constructor. */
 	TSparseArray(const TSparseArray& InCopy)
 	:	FirstFreeIndex(-1)
@@ -469,6 +475,16 @@ public:
 		*this = InCopy;
 	}
 
+	/** Move assignment operator. */
+	TSparseArray& operator=(TSparseArray&& InCopy)
+	{
+		if(this != &InCopy)
+		{
+			MoveOrCopy(*this, InCopy);
+		}
+		return *this;
+	}
+
 	/** Copy assignment operator. */
 	TSparseArray& operator=(const TSparseArray& InCopy)
 	{
@@ -511,8 +527,6 @@ public:
 		return *this;
 	}
 
-#if PLATFORM_COMPILER_HAS_RVALUE_REFERENCES
-
 private:
 	template <typename SparseArrayType>
 	FORCEINLINE static typename TEnableIf<TContainerTraits<SparseArrayType>::MoveWillEmptyContainer>::Type MoveOrCopy(SparseArrayType& ToArray, SparseArrayType& FromArray)
@@ -533,23 +547,6 @@ private:
 	}
 
 public:
-	/** Move constructor. */
-	TSparseArray(TSparseArray&& InCopy)
-	{
-		MoveOrCopy(*this, InCopy);
-	}
-
-	/** Move assignment operator. */
-	TSparseArray& operator=(TSparseArray&& InCopy)
-	{
-		if(this != &InCopy)
-		{
-			MoveOrCopy(*this, InCopy);
-		}
-		return *this;
-	}
-#endif
-
 	// Accessors.
 	ElementType& operator[](int32 Index)
 	{
@@ -805,7 +802,6 @@ template<typename ElementType, typename Allocator>
 struct TContainerTraits<TSparseArray<ElementType, Allocator> > : public TContainerTraitsBase<TSparseArray<ElementType, Allocator> >
 {
 	enum { MoveWillEmptyContainer =
-		PLATFORM_COMPILER_HAS_RVALUE_REFERENCES &&
 		TContainerTraits<typename TSparseArray<ElementType, Allocator>::DataType>::MoveWillEmptyContainer &&
 		TContainerTraits<typename TSparseArray<ElementType, Allocator>::AllocationBitArrayType>::MoveWillEmptyContainer };
 };
diff --git a/Engine/Source/Runtime/Core/Public/Containers/StaticArray.h b/Engine/Source/Runtime/Core/Public/Containers/StaticArray.h
--- a/Engine/Source/Runtime/Core/Public/Containers/StaticArray.h
+++ b/Engine/Source/Runtime/Core/Public/Containers/StaticArray.h
@@ -21,28 +21,31 @@ public:
 		}
 	}
 
-	/** Copy constructor. */
-	TStaticArray(const TStaticArray& Other)
-	{
-		ConstructItems<TElement>((void*)Elements, (const TElement*)Other.Elements, NumElements);
-	}
-
-#if PLATFORM_COMPILER_HAS_RVALUE_REFERENCES
-
 	/** Move constructor. */
 	TStaticArray(TStaticArray&& Other)
 	{
 		MoveConstructItems((void*)Elements, (const TElement*)Other.Elements, NumElements);
 	}
 
+	/** Copy constructor. */
+	TStaticArray(const TStaticArray& Other)
+	{
+		ConstructItems<TElement>((void*)Elements, (const TElement*)Other.Elements, NumElements);
+	}
+
 	/** Move assignment operator. */
 	TStaticArray& operator=(TStaticArray&& Other)
 	{
 		MoveAssignItems((TElement*)Elements, (const TElement*)Other.Elements, NumElements);
 		return *this;
 	}
-
-#endif
+	
+	/** Assignment operator. */
+	TStaticArray& operator=(const TStaticArray& Other)
+	{
+		CopyAssignItems((TElement*)Elements, (const TElement*)Other.Elements, NumElements);
+		return *this;
+	}
 
 	/** Destructor. */
 	~TStaticArray()
@@ -59,13 +62,6 @@ public:
 		}
 		return Ar;
 	}
-	
-	/** Assignment operator. */
-	TStaticArray& operator=(const TStaticArray& Other)
-	{
-		CopyAssignItems((TElement*)Elements, (const TElement*)Other.Elements, NumElements);
-		return *this;
-	}
 
 	// Accessors.
 	TElement& operator[](uint32 Index)
@@ -142,44 +138,35 @@ public:
 
 #if PLATFORM_COMPILER_HAS_DEFAULTED_FUNCTIONS
 
+	TStaticArray2(TStaticArray2&&) = default;
 	TStaticArray2(const TStaticArray2&) = default;
+	TStaticArray2& operator=(TStaticArray2&&) = default;
 	TStaticArray2& operator=(const TStaticArray2&) = default;
 
-	#if PLATFORM_COMPILER_HAS_RVALUE_REFERENCES
-
-		TStaticArray2(TStaticArray2&&) = default;
-		TStaticArray2& operator=(TStaticArray2&&) = default;
-
-	#endif
-
 #else
 
+	FORCEINLINE TStaticArray2(TStaticArray2&& Other)
+		: Super((Super&&)Other)
+	{
+	}
+
 	FORCEINLINE TStaticArray2(const TStaticArray2& Other)
 		: Super((const Super&)Other)
 	{
 	}
 
+	FORCEINLINE TStaticArray2& operator=(TStaticArray2&& Other)
+	{
+		(Super&)*this = (Super&&)Other;
+		return *this;
+	}
+
 	FORCEINLINE TStaticArray2& operator=(const TStaticArray2& Other)
 	{
 		(Super&)*this = (const Super&)Other;
 		return *this;
 	}
 
-	#if PLATFORM_COMPILER_HAS_RVALUE_REFERENCES
-
-		FORCEINLINE TStaticArray2(TStaticArray2&& Other)
-			: Super((Super&&)Other)
-		{
-		}
-
-		FORCEINLINE TStaticArray2& operator=(TStaticArray2&& Other)
-		{
-			(Super&)*this = (Super&&)Other;
-			return *this;
-		}
-
-	#endif
-
 #endif
 };
 
@@ -204,44 +191,35 @@ public:
 
 #if PLATFORM_COMPILER_HAS_DEFAULTED_FUNCTIONS
 
+	TStaticArray3(TStaticArray3&&) = default;
 	TStaticArray3(const TStaticArray3&) = default;
+	TStaticArray3& operator=(TStaticArray3&&) = default;
 	TStaticArray3& operator=(const TStaticArray3&) = default;
 
-	#if PLATFORM_COMPILER_HAS_RVALUE_REFERENCES
-
-		TStaticArray3(TStaticArray3&&) = default;
-		TStaticArray3& operator=(TStaticArray3&&) = default;
-
-	#endif
-
 #else
 
+	FORCEINLINE TStaticArray3(TStaticArray3&& Other)
+		: Super((Super&&)Other)
+	{
+	}
+
 	FORCEINLINE TStaticArray3(const TStaticArray3& Other)
 		: Super((const Super&)Other)
 	{
 	}
 
+	FORCEINLINE TStaticArray3& operator=(TStaticArray3&& Other)
+	{
+		(Super&)*this = (Super&&)Other;
+		return *this;
+	}
+
 	FORCEINLINE TStaticArray3& operator=(const TStaticArray3& Other)
 	{
 		(Super&)*this = (const Super&)Other;
 		return *this;
 	}
 
-	#if PLATFORM_COMPILER_HAS_RVALUE_REFERENCES
-
-		FORCEINLINE TStaticArray3(TStaticArray3&& Other)
-			: Super((Super&&)Other)
-		{
-		}
-
-		FORCEINLINE TStaticArray3& operator=(TStaticArray3&& Other)
-		{
-			(Super&)*this = (Super&&)Other;
-			return *this;
-		}
-
-	#endif
-
 #endif
 };
 
@@ -268,44 +246,35 @@ public:
 
 #if PLATFORM_COMPILER_HAS_DEFAULTED_FUNCTIONS
 
+	TStaticArray4(TStaticArray4&&) = default;
 	TStaticArray4(const TStaticArray4&) = default;
+	TStaticArray4& operator=(TStaticArray4&&) = default;
 	TStaticArray4& operator=(const TStaticArray4&) = default;
 
-	#if PLATFORM_COMPILER_HAS_RVALUE_REFERENCES
-
-		TStaticArray4(TStaticArray4&&) = default;
-		TStaticArray4& operator=(TStaticArray4&&) = default;
-
-	#endif
-
 #else
 
+	FORCEINLINE TStaticArray4(TStaticArray4&& Other)
+		: Super((Super&&)Other)
+	{
+	}
+
 	FORCEINLINE TStaticArray4(const TStaticArray4& Other)
 		: Super((const Super&)Other)
 	{
 	}
 
+	FORCEINLINE TStaticArray4& operator=(TStaticArray4&& Other)
+	{
+		(Super&)*this = (Super&&)Other;
+		return *this;
+	}
+
 	FORCEINLINE TStaticArray4& operator=(const TStaticArray4& Other)
 	{
 		(Super&)*this = (const Super&)Other;
 		return *this;
 	}
 
-	#if PLATFORM_COMPILER_HAS_RVALUE_REFERENCES
-
-		FORCEINLINE TStaticArray4(TStaticArray4&& Other)
-			: Super((Super&&)Other)
-		{
-		}
-
-		FORCEINLINE TStaticArray4& operator=(TStaticArray4&& Other)
-		{
-			(Super&)*this = (Super&&)Other;
-			return *this;
-		}
-
-	#endif
-
 #endif
 };
 
diff --git a/Engine/Source/Runtime/Core/Public/Containers/UnrealString.h b/Engine/Source/Runtime/Core/Public/Containers/UnrealString.h
--- a/Engine/Source/Runtime/Core/Public/Containers/UnrealString.h
+++ b/Engine/Source/Runtime/Core/Public/Containers/UnrealString.h
@@ -1420,7 +1420,7 @@ public:
 template<>
 struct TContainerTraits<FString> : public TContainerTraitsBase<FString>
 {
-	enum { MoveWillEmptyContainer = PLATFORM_COMPILER_HAS_RVALUE_REFERENCES && TContainerTraits<FString::DataType>::MoveWillEmptyContainer };
+	enum { MoveWillEmptyContainer = TContainerTraits<FString::DataType>::MoveWillEmptyContainer };
 };
 
 template<> struct TIsZeroConstructType<FString> { enum { Value = true }; };
@@ -1680,24 +1680,32 @@ public:
 
 #if PLATFORM_COMPILER_HAS_DEFAULTED_FUNCTIONS
 
+	FStringOutputDevice(FStringOutputDevice&&) = default;
 	FStringOutputDevice(const FStringOutputDevice&) = default;
+	FStringOutputDevice& operator=(FStringOutputDevice&&) = default;
 	FStringOutputDevice& operator=(const FStringOutputDevice&) = default;
 
-	#if PLATFORM_COMPILER_HAS_RVALUE_REFERENCES
-
-		FStringOutputDevice(FStringOutputDevice&&) = default;
-		FStringOutputDevice& operator=(FStringOutputDevice&&) = default;
-
-	#endif
-
 #else
 
+	FORCEINLINE FStringOutputDevice(FStringOutputDevice&& Other)
+		: FString      ((FString&&)Other)
+		, FOutputDevice((FOutputDevice&&)Other)
+	{
+	}
+
 	FORCEINLINE FStringOutputDevice(const FStringOutputDevice& Other)
 		: FString      ((const FString&)Other)
 		, FOutputDevice((const FOutputDevice&)Other)
 	{
 	}
 
+	FORCEINLINE FStringOutputDevice& operator=(FStringOutputDevice&& Other)
+	{
+		(FString&)*this       = (FString&&)Other;
+		(FOutputDevice&)*this = (FOutputDevice&&)Other;
+		return *this;
+	}
+
 	FORCEINLINE FStringOutputDevice& operator=(const FStringOutputDevice& Other)
 	{
 		(FString&)*this       = (const FString&)Other;
@@ -1705,23 +1713,6 @@ public:
 		return *this;
 	}
 
-	#if PLATFORM_COMPILER_HAS_RVALUE_REFERENCES
-
-		FORCEINLINE FStringOutputDevice(FStringOutputDevice&& Other)
-			: FString      ((FString&&)Other)
-			, FOutputDevice((FOutputDevice&&)Other)
-		{
-		}
-
-		FORCEINLINE FStringOutputDevice& operator=(FStringOutputDevice&& Other)
-		{
-			(FString&)*this       = (FString&&)Other;
-			(FOutputDevice&)*this = (FOutputDevice&&)Other;
-			return *this;
-		}
-
-	#endif
-
 #endif
 	// Make += operator virtual.
 	virtual FString& operator+=(const FString& Other)
@@ -1815,8 +1806,6 @@ public:
 
 #endif
 
-#if PLATFORM_COMPILER_HAS_RVALUE_REFERENCES
-
 	FORCEINLINE FStringOutputDeviceCountLines(FStringOutputDeviceCountLines&& Other)
 		: Super    ((Super&&)Other)
 		, LineCount(Other.LineCount)
@@ -1835,6 +1824,4 @@ public:
 		}
 		return *this;
 	}
-
-#endif
 };
diff --git a/Engine/Source/Runtime/Core/Public/HAL/Platform.h b/Engine/Source/Runtime/Core/Public/HAL/Platform.h
--- a/Engine/Source/Runtime/Core/Public/HAL/Platform.h
+++ b/Engine/Source/Runtime/Core/Public/HAL/Platform.h
@@ -156,9 +156,6 @@
 #ifndef PLATFORM_COMPILER_HAS_GENERIC_KEYWORD
 	#define PLATFORM_COMPILER_HAS_GENERIC_KEYWORD	0
 #endif
-#ifndef PLATFORM_COMPILER_HAS_RVALUE_REFERENCES
-	#define PLATFORM_COMPILER_HAS_RVALUE_REFERENCES	1
-#endif
 #ifndef PLATFORM_COMPILER_HAS_DEFAULTED_FUNCTIONS
 	#define PLATFORM_COMPILER_HAS_DEFAULTED_FUNCTIONS	1
 #endif
diff --git a/Engine/Source/Runtime/Core/Public/Misc/OutputDevice.h b/Engine/Source/Runtime/Core/Public/Misc/OutputDevice.h
--- a/Engine/Source/Runtime/Core/Public/Misc/OutputDevice.h
+++ b/Engine/Source/Runtime/Core/Public/Misc/OutputDevice.h
@@ -257,47 +257,38 @@ public:
 
 #if PLATFORM_COMPILER_HAS_DEFAULTED_FUNCTIONS
 
+	FOutputDevice(FOutputDevice&&) = default;
 	FOutputDevice(const FOutputDevice&) = default;
+	FOutputDevice& operator=(FOutputDevice&&) = default;
 	FOutputDevice& operator=(const FOutputDevice&) = default;
 
-	#if PLATFORM_COMPILER_HAS_RVALUE_REFERENCES
-
-		FOutputDevice(FOutputDevice&&) = default;
-		FOutputDevice& operator=(FOutputDevice&&) = default;
-
-	#endif
-
 #else
 
+	FORCEINLINE FOutputDevice(FOutputDevice&& Other)
+		 : bSuppressEventTag      (Other.bSuppressEventTag)
+		 , bAutoEmitLineTerminator(Other.bAutoEmitLineTerminator)
+	{
+	}
+
 	FORCEINLINE FOutputDevice(const FOutputDevice& Other)
 		 : bSuppressEventTag      (Other.bSuppressEventTag)
 		 , bAutoEmitLineTerminator(Other.bAutoEmitLineTerminator)
 	{
 	}
 
-	FORCEINLINE FOutputDevice& operator=(const FOutputDevice& Other)
+	FORCEINLINE FOutputDevice& operator=(FOutputDevice&& Other)
 	{
 		bSuppressEventTag       = Other.bSuppressEventTag;
 		bAutoEmitLineTerminator = Other.bAutoEmitLineTerminator;
 		return *this;
 	}
 
-	#if PLATFORM_COMPILER_HAS_RVALUE_REFERENCES
-
-		FORCEINLINE FOutputDevice(FOutputDevice&& Other)
-			 : bSuppressEventTag      (Other.bSuppressEventTag)
-			 , bAutoEmitLineTerminator(Other.bAutoEmitLineTerminator)
-		{
-		}
-
-		FORCEINLINE FOutputDevice& operator=(FOutputDevice&& Other)
-		{
-			bSuppressEventTag       = Other.bSuppressEventTag;
-			bAutoEmitLineTerminator = Other.bAutoEmitLineTerminator;
-			return *this;
-		}
-
-	#endif
+	FORCEINLINE FOutputDevice& operator=(const FOutputDevice& Other)
+	{
+		bSuppressEventTag       = Other.bSuppressEventTag;
+		bAutoEmitLineTerminator = Other.bAutoEmitLineTerminator;
+		return *this;
+	}
 
 #endif
 
diff --git a/Engine/Source/Runtime/Engine/Classes/AI/Navigation/NavigationTypes.h b/Engine/Source/Runtime/Engine/Classes/AI/Navigation/NavigationTypes.h
--- a/Engine/Source/Runtime/Engine/Classes/AI/Navigation/NavigationTypes.h
+++ b/Engine/Source/Runtime/Engine/Classes/AI/Navigation/NavigationTypes.h
@@ -603,15 +603,10 @@ public:
 #if PLATFORM_COMPILER_HAS_DEFAULTED_FUNCTIONS
 
 	TNavStatArray() = default;
-	TNavStatArray(const TNavStatArray&) = default;
-	TNavStatArray& operator=(const TNavStatArray&) = default;
-
-#if PLATFORM_COMPILER_HAS_RVALUE_REFERENCES
-
 	TNavStatArray(TNavStatArray&&) = default;
+	TNavStatArray(const TNavStatArray&) = default;
 	TNavStatArray& operator=(TNavStatArray&&) = default;
-
-#endif
+	TNavStatArray& operator=(const TNavStatArray&) = default;
 
 #else
 
@@ -624,14 +619,6 @@ public:
 	{
 	}
 
-	FORCEINLINE TNavStatArray& operator=(const TNavStatArray& Other)
-	{
-		(Super&)*this = (const Super&)Other;
-		return *this;
-	}
-
-#if PLATFORM_COMPILER_HAS_RVALUE_REFERENCES
-
 	FORCEINLINE TNavStatArray(TNavStatArray&& Other)
 		: Super((Super&&)Other)
 	{
@@ -643,7 +630,11 @@ public:
 		return *this;
 	}
 
-#endif
+	FORCEINLINE TNavStatArray& operator=(const TNavStatArray& Other)
+	{
+		(Super&)*this = (const Super&)Other;
+		return *this;
+	}
 
 #endif
 };
@@ -651,11 +642,7 @@ public:
 template<typename InElementType>
 struct TContainerTraits<TNavStatArray<InElementType> > : public TContainerTraitsBase<TNavStatArray<InElementType> >
 {
-	enum {
-		MoveWillEmptyContainer =
-		PLATFORM_COMPILER_HAS_RVALUE_REFERENCES &&
-		TContainerTraits<typename TNavStatArray<InElementType>::Super>::MoveWillEmptyContainer
-	};
+	enum { MoveWillEmptyContainer = TContainerTraits<typename TNavStatArray<InElementType>::Super>::MoveWillEmptyContainer };
 };
 
 //----------------------------------------------------------------------//
