Use of PLATFORM_COMPILER_HAS_DEFAULT_FUNCTION_TEMPLATE_ARGUMENTS macro instead of context-specific solutions.
Some whitespace fixes.

diff --git a/Engine/Source/Runtime/Core/Public/Templates/Function.h b/Engine/Source/Runtime/Core/Public/Templates/Function.h
--- a/Engine/Source/Runtime/Core/Public/Templates/Function.h
+++ b/Engine/Source/Runtime/Core/Public/Templates/Function.h
@@ -14,15 +14,6 @@
 	#define ENABLE_TFUNCTIONREF_VISUALIZATION 0
 #endif
 
-// VS2012 doesn't support defaulted function template arguments, so we need a workaround for that.
-// The workaround technically changes the signature of the constructors, but we assume that the modified
-// signature will not be used.
-#if defined(_MSC_VER) && _MSC_VER < 1800
-	#define ENABLE_TFUNCTIONREF_WORKAROUND 1
-#else
-	#define ENABLE_TFUNCTIONREF_WORKAROUND 0
-#endif
-
 /**
  * TFunction<FuncType>
  *
@@ -325,7 +316,7 @@ public:
 	/**
 	 * Constructor which binds a TFunctionRef to a non-const lvalue function object.
 	 */
-#if ENABLE_TFUNCTIONREF_WORKAROUND
+#if !PLATFORM_COMPILER_HAS_DEFAULT_FUNCTION_TEMPLATE_ARGUMENTS
 	template <typename FunctorType>
 	TFunctionRef(FunctorType& Functor, typename TEnableIf<!TIsFunction<FunctorType>::Value && !TAreTypesEqual<TFunctionRef, FunctorType>::Value>::Type* = nullptr)
 #else
@@ -342,7 +333,7 @@ public:
 	/**
 	 * Constructor which binds a TFunctionRef to an rvalue or const lvalue function object.
 	 */
-#if ENABLE_TFUNCTIONREF_WORKAROUND
+#if !PLATFORM_COMPILER_HAS_DEFAULT_FUNCTION_TEMPLATE_ARGUMENTS
 	template <typename FunctorType>
 	TFunctionRef(const FunctorType& Functor, typename TEnableIf<!TIsFunction<FunctorType>::Value && !TAreTypesEqual<TFunctionRef, FunctorType>::Value>::Type* = nullptr)
 #else
@@ -359,7 +350,7 @@ public:
 	/**
 	 * Constructor which binds a TFunctionRef to a function pointer.
 	 */
-#if ENABLE_TFUNCTIONREF_WORKAROUND
+#if !PLATFORM_COMPILER_HAS_DEFAULT_FUNCTION_TEMPLATE_ARGUMENTS
 	template <typename FunctionType>
 	TFunctionRef(FunctionType* Function, typename TEnableIf<TIsFunction<FunctionType>::Value>::Type* = nullptr)
 #else
@@ -541,7 +532,7 @@ public:
 	/**
 	 * Constructor which binds a TFunction to any function object.
 	 */
-#if ENABLE_TFUNCTIONREF_WORKAROUND
+#if !PLATFORM_COMPILER_HAS_DEFAULT_FUNCTION_TEMPLATE_ARGUMENTS
 	template <typename FunctorType>
 	TFunction(FunctorType&& InFunc, typename TEnableIf<!TAreTypesEqual<TFunction, typename TDecay<FunctorType>::Type>::Value>::Type* = nullptr)
 #else
@@ -646,43 +637,41 @@ public:
 	SAFE_BOOL_OPERATORS(TFunction)
 
 private:
-
-
 	UE4Function_Private::IFunction_OwnedObject* Func;
 };
 
-	/**
-	 * Nullptr equality operator.
-	 */
-	template <typename FuncType>
-	bool operator==(TYPE_OF_NULLPTR, const TFunction<FuncType>& Func)
-	{
-		return !Func;
-	}
+/**
+ * Nullptr equality operator.
+ */
+template <typename FuncType>
+bool operator==(TYPE_OF_NULLPTR, const TFunction<FuncType>& Func)
+{
+	return !Func;
+}
 
-	/**
-	 * Nullptr equality operator.
-	 */
-	template <typename FuncType>
-	bool operator==(const TFunction<FuncType>& Func, TYPE_OF_NULLPTR)
-	{
-		return !Func;
-	}
+/**
+ * Nullptr equality operator.
+ */
+template <typename FuncType>
+bool operator==(const TFunction<FuncType>& Func, TYPE_OF_NULLPTR)
+{
+	return !Func;
+}
 
-	/**
-	 * Nullptr inequality operator.
-	 */
-	template <typename FuncType>
-	bool operator!=(TYPE_OF_NULLPTR, const TFunction<FuncType>& Func)
-	{
-		return (bool)Func;
-	}
+/**
+ * Nullptr inequality operator.
+ */
+template <typename FuncType>
+bool operator!=(TYPE_OF_NULLPTR, const TFunction<FuncType>& Func)
+{
+	return (bool)Func;
+}
 
-	/**
-	 * Nullptr inequality operator.
-	 */
-	template <typename FuncType>
-	bool operator!=(const TFunction<FuncType>& Func, TYPE_OF_NULLPTR)
-	{
-		return (bool)Func;
-	}
\ No newline at end of file
+/**
+ * Nullptr inequality operator.
+ */
+template <typename FuncType>
+bool operator!=(const TFunction<FuncType>& Func, TYPE_OF_NULLPTR)
+{
+	return (bool)Func;
+}
diff --git a/Engine/Source/Runtime/CoreUObject/Public/UObject/AssetPtr.h b/Engine/Source/Runtime/CoreUObject/Public/UObject/AssetPtr.h
--- a/Engine/Source/Runtime/CoreUObject/Public/UObject/AssetPtr.h
+++ b/Engine/Source/Runtime/CoreUObject/Public/UObject/AssetPtr.h
@@ -50,15 +50,6 @@ public:
 template <> struct TIsPODType<FAssetPtr> { enum { Value = TIsPODType<TPersistentObjectPtr<FStringAssetReference> >::Value }; };
 template <> struct TIsWeakPointerType<FAssetPtr> { enum { Value = TIsWeakPointerType<TPersistentObjectPtr<FStringAssetReference> >::Value }; };
 
-// VS2012 doesn't support defaulted function template arguments, so we need a workaround for that.
-// The workaround technically changes the signature of the constructors, but we assume that the modified
-// signature will not be used.
-#if defined(_MSC_VER) && _MSC_VER < 1800
-	#define ENABLE_TASSETPTR_CONSTRUCTOR_ENABLE_IF_WORKAROUND 1
-#else
-	#define ENABLE_TASSETPTR_CONSTRUCTOR_ENABLE_IF_WORKAROUND 0
-#endif
-
 /***
 * 
 * TAssetPtr is templatized version of the generic FAssetPtr
@@ -81,7 +72,7 @@ public:
 	 * Construct from another lazy pointer
 	 * @param Other lazy pointer to copy from
 	 */
-#if ENABLE_TASSETPTR_CONSTRUCTOR_ENABLE_IF_WORKAROUND
+#if !PLATFORM_COMPILER_HAS_DEFAULT_FUNCTION_TEMPLATE_ARGUMENTS
 	template <class U>
 	FORCEINLINE TAssetPtr(const TAssetPtr<U>& Other, typename TEnableIf<TPointerIsConvertibleFromTo<U, T>::Value>::Type* = nullptr)
 #else
