Fixed a crash at exit in FOutputDeviceConsoleMac

diff --git a/Engine/Source/Runtime/Core/Private/Mac/MacPlatformOutputDevices.cpp b/Engine/Source/Runtime/Core/Private/Mac/MacPlatformOutputDevices.cpp
--- a/Engine/Source/Runtime/Core/Private/Mac/MacPlatformOutputDevices.cpp
+++ b/Engine/Source/Runtime/Core/Private/Mac/MacPlatformOutputDevices.cpp
@@ -150,7 +150,7 @@ void FOutputDeviceMacError::HandleError()
 {
 	if (!MacApplication && [[NSApp orderedWindows] count] == 1)
 	{
-		exit(0);
+		_Exit(0);
 	}
 }
 @end
