TMap cleanup:
* TMapBase's and TSortableMapBase's special functions are now protected, meaning they are now only usable as a base.
* TMapBase's template parameters are no longer defaulted, because they don't need to be.
* bInAllowDuplicateKeys template parameters removed, as they can be inferred from the KeyFuncs.
* Some misleading comments removed.
* Unreachable code in TMapBase::Emplace removed, which did nothing except occasionally cause compile errors.
* static_asserts added to ensure that TMap/TMultiMap are only ever instantiated with KeyFuncs which prevent/allow duplicate keys respectively.
* Redundant code removed from TMultiMap which is (now) doing the same as the base.
* Redundant container traits removed, now that TMapBase and TSortableMapBase are not usable by themselves.

Obsolete (and now illegal) TMapBase removed from FRecastNavMeshGenerator.

diff --git a/Engine/Source/Runtime/Core/Public/Containers/Map.h b/Engine/Source/Runtime/Core/Public/Containers/Map.h
--- a/Engine/Source/Runtime/Core/Public/Containers/Map.h
+++ b/Engine/Source/Runtime/Core/Public/Containers/Map.h
@@ -164,7 +164,7 @@ struct TDefaultMapKeyFuncs : BaseKeyFuncs<TPair<KeyType,ValueType>,KeyType,bInAl
  * The base class of maps from keys to values.  Implemented using a TSet of key-value pairs with a custom KeyFuncs, 
  * with the same O(1) addition, removal, and finding. 
  **/
-template<typename KeyType,typename ValueType,bool bInAllowDuplicateKeys,typename SetAllocator = FDefaultSetAllocator,typename KeyFuncs = TDefaultMapKeyFuncs<KeyType,ValueType,bInAllowDuplicateKeys> >
+template <typename KeyType, typename ValueType, typename SetAllocator, typename KeyFuncs>
 class TMapBase
 {
 	friend struct TContainerTraits<TMapBase>;
@@ -177,7 +177,6 @@ public:
 protected:
 	typedef TPair<KeyType, ValueType> PairType;
 
-public:
 #if PLATFORM_COMPILER_HAS_DEFAULTED_FUNCTIONS
 
 	TMapBase() = default;
@@ -196,6 +195,7 @@ public:
 
 #endif
 
+public:
 	// Legacy comparison operators.  Note that these also test whether the map's key-value pairs were added in the same order!
 	friend bool LegacyCompareEqual(const TMapBase& A,const TMapBase& B)
 	{
@@ -315,7 +315,6 @@ public:
 
 	/**
 	 * Sets the value associated with a key.
-	 * If the key is already associated with any values, the existing values are replaced by the new value.
 	 *
 	 * @param InKey - The key to associate the value with.
 	 * @param InValue - The value to associate with the key.
@@ -328,7 +327,6 @@ public:
 
 	/**
 	 * Sets a default value associated with a key.
-	 * If the key is already associated with any values, the existing values are replaced by the new value.
 	 *
 	 * @param InKey - The key to associate the value with.
 	 * @return A reference to the value as stored in the map.  The reference is only valid until the next change to any key in the map.
@@ -338,7 +336,6 @@ public:
 
 	/**
 	 * Sets the value associated with a key.
-	 * If the key is already associated with any values, the existing values are replaced by the new value.
 	 *
 	 * @param InKey - The key to associate the value with.
 	 * @param InValue - The value to associate with the key.
@@ -347,17 +344,6 @@ public:
 	template <typename InitKeyType, typename InitValueType>
 	ValueType& Emplace(InitKeyType&& InKey, InitValueType&& InValue)
 	{
-		// Remove existing values associated with the specified key.
-		// This is only necessary if the TSet allows duplicate keys; otherwise TSet::Add replaces the existing key-value pair.
-		if(KeyFuncs::bAllowDuplicateKeys)
-		{
-			for(typename PairSetType::TKeyIterator It(Pairs,KeyType(InKey));It;++It)
-			{
-				It.RemoveCurrent();
-			}
-		}
-
-		// Add the key-value pair to the set.  TSet::Add will replace any existing key-value pair that has the same key.
 		const FSetElementId PairId = Pairs.Emplace(TPairInitializer<InitKeyType&&, InitValueType&&>(Forward<InitKeyType>(InKey), Forward<InitValueType>(InValue)));
 
 		return Pairs[PairId].Value;
@@ -365,7 +351,6 @@ public:
 
 	/**
 	 * Sets a default value associated with a key.
-	 * If the key is already associated with any values, the existing values are replaced by the new value.
 	 *
 	 * @param InKey - The key to associate the value with.
 	 * @return A reference to the value as stored in the map.  The reference is only valid until the next change to any key in the map.
@@ -373,17 +358,6 @@ public:
 	template <typename InitKeyType>
 	ValueType& Emplace(InitKeyType&& InKey)
 	{
-		// Remove existing values associated with the specified key.
-		// This is only necessary if the TSet allows duplicate keys; otherwise TSet::Add replaces the existing key-value pair.
-		if(KeyFuncs::bAllowDuplicateKeys)
-		{
-			for(typename PairSetType::TKeyIterator It(Pairs,InKey);It;++It)
-			{
-				It.RemoveCurrent();
-			}
-		}
-
-		// Add the key-value pair to the set.  TSet::Add will replace any existing key-value pair that has the same key.
 		const FSetElementId PairId = Pairs.Emplace(TKeyInitializer<InitKeyType&&>(Forward<InitKeyType>(InKey)));
 
 		return Pairs[PairId].Value;
@@ -792,13 +766,13 @@ private:
 };
 
 /** The base type of sortable maps. */
-template<typename KeyType,typename ValueType,bool bInAllowDuplicateKeys,typename SetAllocator,typename KeyFuncs>
-class TSortableMapBase : public TMapBase<KeyType,ValueType,bInAllowDuplicateKeys,SetAllocator,KeyFuncs>
+template <typename KeyType, typename ValueType, typename SetAllocator, typename KeyFuncs>
+class TSortableMapBase : public TMapBase<KeyType, ValueType, SetAllocator, KeyFuncs>
 {
 	friend struct TContainerTraits<TSortableMapBase>;
 
-public:
-	typedef TMapBase<KeyType,ValueType,bInAllowDuplicateKeys,SetAllocator,KeyFuncs> Super;
+protected:
+	typedef TMapBase<KeyType, ValueType, SetAllocator, KeyFuncs> Super;
 
 #if PLATFORM_COMPILER_HAS_DEFAULTED_FUNCTIONS
 
@@ -818,6 +792,7 @@ public:
 
 #endif
 
+public:
 	/**
 	 * Sorts the pairs array using each pair's Key as the sort criteria, then rebuilds the map's hash.
 	 * Invoked using "MyMapVar.KeySort( PREDICATE_CLASS() );"
@@ -879,12 +854,14 @@ private:
 
 /** A TMapBase specialization that only allows a single value associated with each key.*/
 template<typename KeyType,typename ValueType,typename SetAllocator /*= FDefaultSetAllocator*/,typename KeyFuncs /*= TDefaultMapKeyFuncs<KeyType,ValueType,false>*/>
-class TMap : public TSortableMapBase<KeyType,ValueType,false,SetAllocator,KeyFuncs>
+class TMap : public TSortableMapBase<KeyType, ValueType, SetAllocator, KeyFuncs>
 {
 	friend struct TContainerTraits<TMap>;
 
+	static_assert(!KeyFuncs::bAllowDuplicateKeys, "TMap cannot be instantiated with a KeyFuncs which allows duplicate keys");
+
 public:
-	typedef TSortableMapBase<KeyType,ValueType,false,SetAllocator,KeyFuncs> Super;
+	typedef TSortableMapBase<KeyType, ValueType, SetAllocator, KeyFuncs> Super;
 	typedef typename Super::KeyInitType KeyInitType;
 	typedef typename Super::KeyConstPointerType KeyConstPointerType;
 
@@ -970,12 +947,14 @@ public:
 
 /** A TMapBase specialization that allows multiple values to be associated with each key. */
 template<typename KeyType,typename ValueType,typename SetAllocator /* = FDefaultSetAllocator */,typename KeyFuncs /*= TDefaultMapKeyFuncs<KeyType,ValueType,true>*/>
-class TMultiMap : public TSortableMapBase<KeyType,ValueType,true,SetAllocator,KeyFuncs>
+class TMultiMap : public TSortableMapBase<KeyType, ValueType, SetAllocator, KeyFuncs>
 {
 	friend struct TContainerTraits<TMultiMap>;
 
+	static_assert(KeyFuncs::bAllowDuplicateKeys, "TMultiMap cannot be instantiated with a KeyFuncs which disallows duplicate keys");
+
 public:
-	typedef TSortableMapBase<KeyType,ValueType,true,SetAllocator,KeyFuncs> Super;
+	typedef TSortableMapBase<KeyType, ValueType, SetAllocator, KeyFuncs> Super;
 	typedef typename Super::KeyConstPointerType KeyConstPointerType;
 	typedef typename Super::KeyInitType KeyInitType;
 	typedef typename Super::ValueInitType ValueInitType;
@@ -1055,57 +1034,6 @@ public:
 		}
 	}
 
-	/**
-	 * Adds a key-value association to the map.  The association doesn't replace any of the key's existing associations.
-	 *
-	 * @param InKey - The key to associate.
-	 * @param InValue - The value to associate.
-	 * @return A reference to the value as stored in the map; the reference is only valid until the next change to any key in the map.
-	 */
-	FORCEINLINE ValueType& Add(const KeyType&  InKey, const ValueType&  InValue) { return Emplace(         InKey ,          InValue ); }
-	FORCEINLINE ValueType& Add(const KeyType&  InKey,       ValueType&& InValue) { return Emplace(         InKey , MoveTemp(InValue)); }
-	FORCEINLINE ValueType& Add(      KeyType&& InKey, const ValueType&  InValue) { return Emplace(MoveTemp(InKey),          InValue ); }
-	FORCEINLINE ValueType& Add(      KeyType&& InKey,       ValueType&& InValue) { return Emplace(MoveTemp(InKey), MoveTemp(InValue)); }
-
-	/**
-	 * Adds a key-value association to the map.  The association doesn't replace any of the key's existing associations.
-	 * The value is default-constructed.
-	 *
-	 * @param InKey - The key to associate.
-	 * @return A reference to the value as stored in the map; the reference is only valid until the next change to any key in the map.
-	 */
-	FORCEINLINE ValueType& Add(const KeyType&  InKey) { return Emplace(         InKey ); }
-	FORCEINLINE ValueType& Add(      KeyType&& InKey) { return Emplace(MoveTemp(InKey)); }
-
-	/**
-	 * Adds a key-value association to the map.  The association doesn't replace any of the key's existing associations.
-	 *
-	 * @param InKey - The key to associate.
-	 * @param InValue - The value to associate.
-	 * @return A reference to the value as stored in the map; the reference is only valid until the next change to any key in the map.
-	 */
-	template <typename InitKeyType, typename InitValueType>
-	FORCEINLINE ValueType& Emplace(InitKeyType&& InKey, InitValueType&& InValue)
-	{
-		const FSetElementId PairId = Super::Pairs.Emplace(TPairInitializer<InitKeyType&&, InitValueType&&>(Forward<InitKeyType>(InKey), Forward<InitValueType>(InValue)));
-		return Super::Pairs[PairId].Value;
-	}
-
-	/**
-	 * Adds a key-value association to the map.  The association doesn't replace any of the key's existing associations.
-	 * The value is default-constructed.
-	 *
-	 * @param InKey - The key to associate.
-	 * @param InValue - The value to associate.
-	 * @return A reference to the value as stored in the map; the reference is only valid until the next change to any key in the map.
-	 */
-	template <typename InitKeyType>
-	FORCEINLINE ValueType& Emplace(InitKeyType&& InKey)
-	{
-		const FSetElementId KeyId = Super::Pairs.Emplace(TKeyInitializer<InitKeyType&&>(Forward<InitKeyType>(InKey)));
-		return Super::Pairs(KeyId).Value;
-	}
-
 	/**
 	 * Adds a key-value association to the map.  The association doesn't replace any of the key's existing associations.
 	 * However, if both the key and value match an existing association in the map, no new association is made and the existing association's
@@ -1131,10 +1059,12 @@ public:
 	ValueType& EmplaceUnique(InitKeyType&& InKey, InitValueType&& InValue)
 	{
 		if (ValueType* Found = FindPair(InKey, InValue))
+		{
 			return *Found;
+		}
 
 		// If there's no existing association with the same key and value, create one.
-		return Add(Forward<InitKeyType>(InKey), Forward<InitValueType>(InValue));
+		return Super::Add(Forward<InitKeyType>(InKey), Forward<InitValueType>(InValue));
 	}
 
 	/**
@@ -1247,26 +1177,14 @@ public:
 	}
 };
 
-template <typename KeyType, typename ValueType, bool bInAllowDuplicateKeys, typename SetAllocator, typename KeyFuncs>
-struct TContainerTraits<TMapBase<KeyType, ValueType, bInAllowDuplicateKeys, SetAllocator, KeyFuncs>> : public TContainerTraitsBase<TMapBase<KeyType, ValueType, bInAllowDuplicateKeys, SetAllocator, KeyFuncs>>
-{
-	enum { MoveWillEmptyContainer = TContainerTraits<typename TMapBase<KeyType, ValueType, bInAllowDuplicateKeys, SetAllocator, KeyFuncs>::PairSetType>::MoveWillEmptyContainer };
-};
-
-template <typename KeyType, typename ValueType, bool bInAllowDuplicateKeys, typename SetAllocator, typename KeyFuncs>
-struct TContainerTraits<TSortableMapBase<KeyType, ValueType, bInAllowDuplicateKeys, SetAllocator, KeyFuncs>> : public TContainerTraitsBase<TSortableMapBase<KeyType, ValueType, bInAllowDuplicateKeys, SetAllocator, KeyFuncs>>
-{
-	enum { MoveWillEmptyContainer = TContainerTraits<typename TSortableMapBase<KeyType, ValueType, bInAllowDuplicateKeys, SetAllocator, KeyFuncs>::Super>::MoveWillEmptyContainer };
-};
-
 template <typename KeyType, typename ValueType, typename SetAllocator,typename KeyFuncs>
 struct TContainerTraits<TMap<KeyType, ValueType, SetAllocator, KeyFuncs>> : public TContainerTraitsBase<TMap<KeyType, ValueType, SetAllocator, KeyFuncs>>
 {
-	enum { MoveWillEmptyContainer = TContainerTraits<typename TMap<KeyType, ValueType, SetAllocator, KeyFuncs>::Super>::MoveWillEmptyContainer };
+	enum { MoveWillEmptyContainer = TContainerTraits<typename TMap<KeyType, ValueType, SetAllocator, KeyFuncs>::PairSetType>::MoveWillEmptyContainer };
 };
 
 template <typename KeyType, typename ValueType, typename SetAllocator,typename KeyFuncs>
 struct TContainerTraits<TMultiMap<KeyType, ValueType, SetAllocator, KeyFuncs>> : public TContainerTraitsBase<TMultiMap<KeyType, ValueType, SetAllocator, KeyFuncs>>
 {
-	enum { MoveWillEmptyContainer = TContainerTraits<typename TMultiMap<KeyType, ValueType, SetAllocator, KeyFuncs>::Super>::MoveWillEmptyContainer };
+	enum { MoveWillEmptyContainer = TContainerTraits<typename TMultiMap<KeyType, ValueType, SetAllocator, KeyFuncs>::PairSetType>::MoveWillEmptyContainer };
 };
diff --git a/Engine/Source/Runtime/Engine/Public/AI/Navigation/RecastNavMeshGenerator.h b/Engine/Source/Runtime/Engine/Public/AI/Navigation/RecastNavMeshGenerator.h
--- a/Engine/Source/Runtime/Engine/Public/AI/Navigation/RecastNavMeshGenerator.h
+++ b/Engine/Source/Runtime/Engine/Public/AI/Navigation/RecastNavMeshGenerator.h
@@ -509,9 +509,6 @@ private:
 	/** */
 	FRecastNavMeshCachedData AdditionalCachedData;
 
-	/** */
-	TMapBase<const AActor*, FBox, false> ActorToAreaMap;
-
 	uint32 bInitialized:1;
 
 	uint32 bRestrictBuildingToActiveTiles:1;
