Fixed keyboard navigation in native OS X modal dialogs

diff --git a/Engine/Source/Runtime/Core/Private/Mac/MacApplication.cpp b/Engine/Source/Runtime/Core/Private/Mac/MacApplication.cpp
--- a/Engine/Source/Runtime/Core/Private/Mac/MacApplication.cpp
+++ b/Engine/Source/Runtime/Core/Private/Mac/MacApplication.cpp
@@ -297,7 +297,7 @@ void FMacApplication::DeferEvent(NSObject* Object)
 		DeferredEvent.WindowNumber = [Event windowNumber];
 		DeferredEvent.Context = [[Event context] retain];
 
-		if (DeferredEvent.Type == NSKeyDown && !bSystemModalMode)
+		if (DeferredEvent.Type == NSKeyDown)
 		{
 			DeferredEvent.Window = FindWindowByNSWindow(WindowHandle);
 			if (DeferredEvent.Window.IsValid() && DeferredEvent.Window.Pin()->OnIMKKeyDown(Event))
@@ -419,7 +419,7 @@ NSEvent* FMacApplication::HandleNSEvent(NSEvent* Event)
 {
 	NSEvent* ReturnEvent = Event;
 
-	if (MacApplication)
+	if (MacApplication && !MacApplication->bSystemModalMode)
 	{
 		const bool bIsResentEvent = [Event type] == NSApplicationDefined && [Event subtype] == RESET_EVENT_SUBTYPE;
 
