Allow access to parts of the local path to GExternalFilePath
#ue4
#android

diff --git a/Engine/Source/Runtime/Core/Private/Android/AndroidFile.cpp b/Engine/Source/Runtime/Core/Private/Android/AndroidFile.cpp
--- a/Engine/Source/Runtime/Core/Private/Android/AndroidFile.cpp
+++ b/Engine/Source/Runtime/Core/Private/Android/AndroidFile.cpp
@@ -1599,7 +1599,7 @@ private:
 		{
 			if ((AllowLocal && AndroidPath.StartsWith(TEXT("/"))) ||
 				AndroidPath.StartsWith(GFontPathBase) ||
-				AndroidPath.StartsWith(GExternalFilePath))
+				AndroidPath.StartsWith(GExternalFilePath.Left(AndroidPath.Len())))
 			{
 				// Absolute paths are only local.
 				LocalPath = AndroidPath;
