#ue4 - Add FQuatRotationMatrix, for the common case where FQuatRotationTranslationMatrix is overkill.

diff --git a/Engine/Source/Runtime/Core/Public/Math/QuatRotationTranslationMatrix.h b/Engine/Source/Runtime/Core/Public/Math/QuatRotationTranslationMatrix.h
--- a/Engine/Source/Runtime/Core/Public/Math/QuatRotationTranslationMatrix.h
+++ b/Engine/Source/Runtime/Core/Public/Math/QuatRotationTranslationMatrix.h
@@ -24,6 +24,30 @@ public:
 };
 
 
+/** Rotation matrix using quaternion rotation */
+class FQuatRotationMatrix
+	: public FQuatRotationTranslationMatrix
+{
+public:
+
+	/** Constructor
+	*
+	* @param Q rotation
+	* @param Origin translation to apply
+	*/
+	FQuatRotationMatrix(const FQuat& Q)
+		: FQuatRotationTranslationMatrix(Q, FVector::ZeroVector)
+	{
+	}
+
+	/** Matrix factory. Return an FMatrix so we don't have type conversion issues in expressions. */
+	static FMatrix Make(const FQuat& Q)
+	{
+		return FQuatRotationMatrix(Q);
+	}
+};
+
+
 FORCEINLINE FQuatRotationTranslationMatrix::FQuatRotationTranslationMatrix(const FQuat& Q, const FVector& Origin)
 {
 #if !(UE_BUILD_SHIPPING || UE_BUILD_TEST) && WITH_EDITORONLY_DATA
