#ue4 - Unit test for FQuatRotationMatrix.

diff --git a/Engine/Source/Runtime/Core/Private/Tests/Math/UnrealMathTest.cpp b/Engine/Source/Runtime/Core/Private/Tests/Math/UnrealMathTest.cpp
--- a/Engine/Source/Runtime/Core/Private/Tests/Math/UnrealMathTest.cpp
+++ b/Engine/Source/Runtime/Core/Private/Tests/Math/UnrealMathTest.cpp
@@ -842,7 +842,10 @@ bool FVectorRegisterAbstractionTest::RunTest(const FString& Parameters)
 		FV0 = Rotator0.Vector();
 		FV1 = FRotationMatrix( Rotator0 ).GetScaledAxis( EAxis::X );
 		LogTest( TEXT("Test0 Rotator::Vector()"), FV0.Equals( FV1, 1e-6f ) );
-
+		
+		FV0 = FRotationMatrix( Rotator0 ).GetScaledAxis( EAxis::X );
+		FV1 = FQuatRotationMatrix( Q0 ).GetScaledAxis( EAxis::X );
+		LogTest( TEXT("Test0 FQuatRotationMatrix"), FV0.Equals( FV1, 1e-5f ) );
 
 		Rotator0 = FRotator(45.0f,  60.0f, 120.0f);
 		Q0 = Rotator0.Quaternion();
@@ -852,6 +855,14 @@ bool FVectorRegisterAbstractionTest::RunTest(const FString& Parameters)
 		FV0 = Rotator0.Vector();
 		FV1 = FRotationMatrix( Rotator0 ).GetScaledAxis( EAxis::X );
 		LogTest( TEXT("Test1 Rotator::Vector()"), FV0.Equals( FV1, 1e-6f ) );
+
+		FV0 = FRotationMatrix( Rotator0 ).GetScaledAxis( EAxis::X );
+		FV1 = FQuatRotationMatrix( Q0 ).GetScaledAxis( EAxis::X );
+		LogTest(TEXT("Test1 FQuatRotationMatrix"), FV0.Equals(FV1, 1e-5f));
+
+		FV0 = FRotationMatrix(FRotator::ZeroRotator).GetScaledAxis(EAxis::X);
+		FV1 = FQuatRotationMatrix(FQuat::Identity).GetScaledAxis(EAxis::X);
+		LogTest(TEXT("Test2 FQuatRotationMatrix"), FV0.Equals(FV1, 1e-6f));
 	}
 
 	{
@@ -859,23 +870,24 @@ bool FVectorRegisterAbstractionTest::RunTest(const FString& Parameters)
 		const FRotator Rotator1(45.0f,  60.0f, 120.0f);
 		const FQuat Quat0(Rotator0);
 		const FQuat Quat1(Rotator1);
-
-		LogTest( TEXT("Test0 FQuat::RotateVector(const FVector& Vec)"),		TestVectorsEqual3_NoVec(Quat0.RotateVector(FVector::ForwardVector), TestQuaternionRotateVector(Quat0, FVector::ForwardVector),	1e-4f ));
-		LogTest( TEXT("Test0 FQuat::operator*(const FVector& Vec)"),		TestVectorsEqual3_NoVec(Quat0 * FVector::ForwardVector, TestQuaternionMultiplyVector(Quat0, FVector::ForwardVector), 1e-4f ));
-		LogTest( TEXT("Test1 FQuat::RotateVector(const FVector& Vec)"),		TestVectorsEqual3_NoVec(Quat0.RotateVector(FVector::RightVector), TestQuaternionRotateVector(Quat0, FVector::RightVector),	1e-4f ));
-		LogTest( TEXT("Test1 FQuat::operator*(const FVector& Vec)"),		TestVectorsEqual3_NoVec(Quat0 * FVector::RightVector, TestQuaternionMultiplyVector(Quat0, FVector::RightVector), 1e-4f ));
-		LogTest( TEXT("Test2 FQuat::RotateVector(const FVector& Vec)"),		TestVectorsEqual3_NoVec(Quat0.RotateVector(FVector::UpVector), TestQuaternionRotateVector(Quat0, FVector::UpVector),	1e-4f ));
-		LogTest( TEXT("Test2 FQuat::operator*(const FVector& Vec)"),		TestVectorsEqual3_NoVec(Quat0 * FVector::UpVector, TestQuaternionMultiplyVector(Quat0, FVector::UpVector), 1e-4f ));
-		LogTest( TEXT("Test3 FQuat::RotateVector(const FVector& Vec)"),		TestVectorsEqual3_NoVec(Quat0.RotateVector(FVector(45.0f, -60.0f, 120.0f)), TestQuaternionRotateVector(Quat0, FVector(45.0f, -60.0f, 120.0f)),	1e-4f ));
-		LogTest( TEXT("Test3 FQuat::operator*(const FVector& Vec)"),		TestVectorsEqual3_NoVec(Quat0 * FVector(45.0f, -60.0f, 120.0f), TestQuaternionMultiplyVector(Quat0, FVector(45.0f, -60.0f, 120.0f)), 1e-4f ));
-		LogTest( TEXT("Test4 FQuat::RotateVector(const FVector& Vec)"),		TestVectorsEqual3_NoVec(Quat1.RotateVector(FVector::ForwardVector), TestQuaternionRotateVector(Quat1, FVector::ForwardVector),	1e-4f ));
-		LogTest( TEXT("Test4 FQuat::operator*(const FVector& Vec)"),		TestVectorsEqual3_NoVec(Quat1 * FVector::ForwardVector, TestQuaternionMultiplyVector(Quat1, FVector::ForwardVector), 1e-4f ));
-		LogTest( TEXT("Test5 FQuat::RotateVector(const FVector& Vec)"),		TestVectorsEqual3_NoVec(Quat1.RotateVector(FVector::RightVector), TestQuaternionRotateVector(Quat1, FVector::RightVector),	1e-4f ));
-		LogTest( TEXT("Test5 FQuat::operator*(const FVector& Vec)"),		TestVectorsEqual3_NoVec(Quat1 * FVector::RightVector, TestQuaternionMultiplyVector(Quat1, FVector::RightVector), 1e-4f ));
-		LogTest( TEXT("Test6 FQuat::RotateVector(const FVector& Vec)"),		TestVectorsEqual3_NoVec(Quat1.RotateVector(FVector::UpVector), TestQuaternionRotateVector(Quat1, FVector::UpVector),	1e-4f ));
-		LogTest( TEXT("Test6 FQuat::operator*(const FVector& Vec)"),		TestVectorsEqual3_NoVec(Quat1 * FVector::UpVector, TestQuaternionMultiplyVector(Quat1, FVector::UpVector), 1e-4f ));
-		LogTest( TEXT("Test7 FQuat::RotateVector(const FVector& Vec)"),		TestVectorsEqual3_NoVec(Quat1.RotateVector(FVector(45.0f, -60.0f, 120.0f)), TestQuaternionRotateVector(Quat1, FVector(45.0f, -60.0f, 120.0f)),	1e-4f ));
-		LogTest( TEXT("Test7 FQuat::operator*(const FVector& Vec)"),		TestVectorsEqual3_NoVec(Quat1 * FVector(45.0f, -60.0f, 120.0f), TestQuaternionMultiplyVector(Quat1, FVector(45.0f, -60.0f, 120.0f)), 1e-4f ));
+		const float Tolerance = 1e-4f;
+
+		LogTest( TEXT("Test0 FQuat::RotateVector(const FVector& Vec)"),		TestVectorsEqual3_NoVec(Quat0.RotateVector(FVector::ForwardVector), TestQuaternionRotateVector(Quat0, FVector::ForwardVector),	Tolerance ));
+		LogTest( TEXT("Test0 FQuat::operator*(const FVector& Vec)"),		TestVectorsEqual3_NoVec(Quat0 * FVector::ForwardVector, TestQuaternionMultiplyVector(Quat0, FVector::ForwardVector), Tolerance ));
+		LogTest( TEXT("Test1 FQuat::RotateVector(const FVector& Vec)"),		TestVectorsEqual3_NoVec(Quat0.RotateVector(FVector::RightVector), TestQuaternionRotateVector(Quat0, FVector::RightVector),	Tolerance ));
+		LogTest( TEXT("Test1 FQuat::operator*(const FVector& Vec)"),		TestVectorsEqual3_NoVec(Quat0 * FVector::RightVector, TestQuaternionMultiplyVector(Quat0, FVector::RightVector), Tolerance ));
+		LogTest( TEXT("Test2 FQuat::RotateVector(const FVector& Vec)"),		TestVectorsEqual3_NoVec(Quat0.RotateVector(FVector::UpVector), TestQuaternionRotateVector(Quat0, FVector::UpVector),	Tolerance ));
+		LogTest( TEXT("Test2 FQuat::operator*(const FVector& Vec)"),		TestVectorsEqual3_NoVec(Quat0 * FVector::UpVector, TestQuaternionMultiplyVector(Quat0, FVector::UpVector), Tolerance ));
+		LogTest( TEXT("Test3 FQuat::RotateVector(const FVector& Vec)"),		TestVectorsEqual3_NoVec(Quat0.RotateVector(FVector(45.0f, -60.0f, 120.0f)), TestQuaternionRotateVector(Quat0, FVector(45.0f, -60.0f, 120.0f)),	Tolerance ));
+		LogTest( TEXT("Test3 FQuat::operator*(const FVector& Vec)"),		TestVectorsEqual3_NoVec(Quat0 * FVector(45.0f, -60.0f, 120.0f), TestQuaternionMultiplyVector(Quat0, FVector(45.0f, -60.0f, 120.0f)), Tolerance ));
+		LogTest( TEXT("Test4 FQuat::RotateVector(const FVector& Vec)"),		TestVectorsEqual3_NoVec(Quat1.RotateVector(FVector::ForwardVector), TestQuaternionRotateVector(Quat1, FVector::ForwardVector),	Tolerance ));
+		LogTest( TEXT("Test4 FQuat::operator*(const FVector& Vec)"),		TestVectorsEqual3_NoVec(Quat1 * FVector::ForwardVector, TestQuaternionMultiplyVector(Quat1, FVector::ForwardVector), Tolerance ));
+		LogTest( TEXT("Test5 FQuat::RotateVector(const FVector& Vec)"),		TestVectorsEqual3_NoVec(Quat1.RotateVector(FVector::RightVector), TestQuaternionRotateVector(Quat1, FVector::RightVector),	Tolerance ));
+		LogTest( TEXT("Test5 FQuat::operator*(const FVector& Vec)"),		TestVectorsEqual3_NoVec(Quat1 * FVector::RightVector, TestQuaternionMultiplyVector(Quat1, FVector::RightVector), Tolerance ));
+		LogTest( TEXT("Test6 FQuat::RotateVector(const FVector& Vec)"),		TestVectorsEqual3_NoVec(Quat1.RotateVector(FVector::UpVector), TestQuaternionRotateVector(Quat1, FVector::UpVector),	Tolerance ));
+		LogTest( TEXT("Test6 FQuat::operator*(const FVector& Vec)"),		TestVectorsEqual3_NoVec(Quat1 * FVector::UpVector, TestQuaternionMultiplyVector(Quat1, FVector::UpVector), Tolerance ));
+		LogTest( TEXT("Test7 FQuat::RotateVector(const FVector& Vec)"),		TestVectorsEqual3_NoVec(Quat1.RotateVector(FVector(45.0f, -60.0f, 120.0f)), TestQuaternionRotateVector(Quat1, FVector(45.0f, -60.0f, 120.0f)),	Tolerance ));
+		LogTest( TEXT("Test7 FQuat::operator*(const FVector& Vec)"),		TestVectorsEqual3_NoVec(Quat1 * FVector(45.0f, -60.0f, 120.0f), TestQuaternionMultiplyVector(Quat1, FVector(45.0f, -60.0f, 120.0f)), Tolerance ));
 
 	}
 
