Core: Renamed ITlsAutoCleanup to FTlsAutoCleanup (it's not a [non-virtual] interface)

diff --git a/Engine/Source/Runtime/Core/Private/HAL/PThreadRunnableThread.cpp b/Engine/Source/Runtime/Core/Private/HAL/PThreadRunnableThread.cpp
--- a/Engine/Source/Runtime/Core/Private/HAL/PThreadRunnableThread.cpp
+++ b/Engine/Source/Runtime/Core/Private/HAL/PThreadRunnableThread.cpp
@@ -13,7 +13,7 @@ uint32 FRunnableThreadPThread::Run()
 	uint32 ExitCode = 1;
 	check(Runnable);
 	
-	// Setup TLS for this thread, used by ITlsAutoCleanup objects.
+	// Setup TLS for this thread, used by FTlsAutoCleanup objects.
 	SetTls();
 
 	// Initialize the runnable object
diff --git a/Engine/Source/Runtime/Core/Private/HAL/ThreadingBase.cpp b/Engine/Source/Runtime/Core/Private/HAL/ThreadingBase.cpp
--- a/Engine/Source/Runtime/Core/Private/HAL/ThreadingBase.cpp
+++ b/Engine/Source/Runtime/Core/Private/HAL/ThreadingBase.cpp
@@ -268,7 +268,7 @@ void FRunnableThread::FreeTls()
 	check( RunnableTlsSlot );
 	FPlatformTLS::SetTlsValue( RunnableTlsSlot, nullptr );
 
-	// Delete all ITlsAutoCleanup objects created for this thread.
+	// Delete all FTlsAutoCleanup objects created for this thread.
 	for( auto& Instance : TlsInstances )
 	{
 		delete Instance;
@@ -621,7 +621,7 @@ FQueuedThreadPool* FQueuedThreadPool::Allocate()
 	FThreadSingletonInitializer
 -----------------------------------------------------------------------------*/
 
-ITlsAutoCleanup* FThreadSingletonInitializer::Get( const TFunctionRef<ITlsAutoCleanup*()>& CreateInstance, uint32& TlsSlot )
+FTlsAutoCleanup* FThreadSingletonInitializer::Get( const TFunctionRef<FTlsAutoCleanup*()>& CreateInstance, uint32& TlsSlot )
 {
 	if( TlsSlot == 0 )
 	{
@@ -633,7 +633,7 @@ ITlsAutoCleanup* FThreadSingletonInitializer::Get( const TFunctionRef<ITlsAutoCl
 			FPlatformTLS::FreeTlsSlot( ThisTlsSlot );
 		}
 	}
-	ITlsAutoCleanup* ThreadSingleton = (ITlsAutoCleanup*)FPlatformTLS::GetTlsValue( TlsSlot );
+	FTlsAutoCleanup* ThreadSingleton = (FTlsAutoCleanup*)FPlatformTLS::GetTlsValue( TlsSlot );
 	if( !ThreadSingleton )
 	{
 		ThreadSingleton = CreateInstance();
@@ -643,7 +643,7 @@ ITlsAutoCleanup* FThreadSingletonInitializer::Get( const TFunctionRef<ITlsAutoCl
 	return ThreadSingleton;
 }
 
-void ITlsAutoCleanup::Register()
+void FTlsAutoCleanup::Register()
 {
 	FRunnableThread* RunnableThread = FRunnableThread::GetRunnableThread();
 	if( RunnableThread )
diff --git a/Engine/Source/Runtime/Core/Private/Windows/WindowsRunnableThread.cpp b/Engine/Source/Runtime/Core/Private/Windows/WindowsRunnableThread.cpp
--- a/Engine/Source/Runtime/Core/Private/Windows/WindowsRunnableThread.cpp
+++ b/Engine/Source/Runtime/Core/Private/Windows/WindowsRunnableThread.cpp
@@ -58,7 +58,7 @@ uint32 FRunnableThreadWin::Run()
 	uint32 ExitCode = 1;
 	check(Runnable);
 
-	// Setup TLS for this thread, used by ITlsAutoCleanup objects.
+	// Setup TLS for this thread, used by FTlsAutoCleanup objects.
 	SetTls();
 
 	// Initialize the runnable object
diff --git a/Engine/Source/Runtime/Core/Public/HAL/ThreadingBase.h b/Engine/Source/Runtime/Core/Public/HAL/ThreadingBase.h
--- a/Engine/Source/Runtime/Core/Public/HAL/ThreadingBase.h
+++ b/Engine/Source/Runtime/Core/Public/HAL/ThreadingBase.h
@@ -4,7 +4,7 @@
 
 #include "Function.h"
 
-class ITlsAutoCleanup;
+class FTlsAutoCleanup;
 class FRunnableThread;
 
 /**
@@ -231,7 +231,7 @@ public:
 class CORE_API FRunnableThread
 {
 	friend class FThreadSingletonInitializer;
-	friend class ITlsAutoCleanup;
+	friend class FTlsAutoCleanup;
 
 	/** Index of TLS slot for FRunnableThread pointer. */
 	static uint32 RunnableTlsSlot;
@@ -354,7 +354,7 @@ protected:
 	/** Stores this instance in the runnable thread TLS slot. */
 	void SetTls();
 
-	/** Deletes all ITlsAutoCleanup objects created for this thread. */
+	/** Deletes all FTlsAutoCleanup objects created for this thread. */
 	void FreeTls();
 
 	/**
@@ -378,8 +378,8 @@ protected:
 	/** The Affinity to run the thread with. */
 	uint64 ThreadAffinityMask;
 
-	/** An array of ITlsAutoCleanup based instances that needs to be deleted before the thread will die. */
-	TArray<ITlsAutoCleanup*> TlsInstances;
+	/** An array of FTlsAutoCleanup based instances that needs to be deleted before the thread will die. */
+	TArray<FTlsAutoCleanup*> TlsInstances;
 
 	/** The priority to run the thread at. */
 	EThreadPriority ThreadPriority;
@@ -924,12 +924,12 @@ extern CORE_API bool IsInRHIThread();
 /** Thread used for RHI */
 extern CORE_API FRunnableThread* GRHIThread;
 
-/** Interface required for auto-cleanup of instances of classes stored in the TLS. */
-class CORE_API ITlsAutoCleanup
+/** Base class for objects in TLS that support auto-cleanup. */
+class CORE_API FTlsAutoCleanup
 {
 public:
 	/** Virtual destructor. */
-	virtual ~ITlsAutoCleanup()
+	virtual ~FTlsAutoCleanup()
 	{}
 
 	/** Register this instance to be auto-cleanup. */
@@ -946,7 +946,7 @@ public:
 	/**
 	* @return an instance of a singleton for the current thread.
 	*/
-	static CORE_API ITlsAutoCleanup* Get( const TFunctionRef<ITlsAutoCleanup*()>& CreateInstance, uint32& TlsSlot );
+	static CORE_API FTlsAutoCleanup* Get( const TFunctionRef<FTlsAutoCleanup*()>& CreateInstance, uint32& TlsSlot );
 };
 
 
@@ -955,7 +955,7 @@ public:
  * Calling Get() method is thread-safe.
  */
 template < class T >
-class TThreadSingleton : public ITlsAutoCleanup
+class TThreadSingleton : public FTlsAutoCleanup
 {
 	/**
 	 * @return TLS slot that holds a TThreadSingleton.
@@ -976,7 +976,7 @@ protected:
 	/**
 	 * @return a new instance of the thread singleton.
 	 */
-	static ITlsAutoCleanup* CreateInstance()
+	static FTlsAutoCleanup* CreateInstance()
 	{
 		return new T();
 	}
@@ -991,7 +991,7 @@ public:
 	 */
 	FORCEINLINE static T& Get()
 	{
-		return *(T*)FThreadSingletonInitializer::Get( [](){ return (ITlsAutoCleanup*)new T(); }, T::GetTlsSlot() );
+		return *(T*)FThreadSingletonInitializer::Get( [](){ return (FTlsAutoCleanup*)new T(); }, T::GetTlsSlot() );
 	}
 };
 
