Core: Implemented TLS auto-cleanup for TypeContainer

diff --git a/Engine/Source/Runtime/Core/Public/HAL/ThreadingBase.h b/Engine/Source/Runtime/Core/Public/HAL/ThreadingBase.h
--- a/Engine/Source/Runtime/Core/Public/HAL/ThreadingBase.h
+++ b/Engine/Source/Runtime/Core/Public/HAL/ThreadingBase.h
@@ -936,6 +936,30 @@ public:
 	void Register();
 };
 
+/** Wrapper for values to be stored in TLS that support auto-cleanup. */
+template< class T >
+class TTlsAutoCleanupValue
+	: public FTlsAutoCleanup
+{
+public:
+
+	/** Constructor. */
+	TTlsAutoCleanupValue(const T& InValue)
+		: Value(InValue)
+	{ }
+
+	/** Gets the value. */
+	T Get() const
+	{
+		return Value;
+	}
+
+private:
+
+	/** The value. */
+	T Value;
+};
+
 /**
  * Thread singleton initializer.
  */
diff --git a/Engine/Source/Runtime/Core/Public/Misc/TypeContainer.h b/Engine/Source/Runtime/Core/Public/Misc/TypeContainer.h
--- a/Engine/Source/Runtime/Core/Public/Misc/TypeContainer.h
+++ b/Engine/Source/Runtime/Core/Public/Misc/TypeContainer.h
@@ -72,7 +72,6 @@ class FTypeContainer
 	struct TFunctionInstanceProvider
 		: public IInstanceProvider
 	{
-		/** Constructor. */
 		TFunctionInstanceProvider(TFunction<TSharedPtr<void>()> InCreateFunc)
 			: CreateFunc(InCreateFunc)
 		{ }
@@ -84,7 +83,6 @@ class FTypeContainer
 			return CreateFunc();
 		}
 
-		/** Factory function that creates the instances. */
 		TFunction<TSharedPtr<void>()> CreateFunc;
 	};
 
@@ -94,12 +92,10 @@ class FTypeContainer
 	struct TSharedInstanceProvider
 		: public IInstanceProvider
 	{
-		/** Creates and initializes a new instance. */
 		TSharedInstanceProvider(const TSharedRef<T>& InInstance)
 			: Instance(InInstance)
 		{ }
 
-		/** Virtual destructor. */
 		virtual ~TSharedInstanceProvider() { }
 
 		virtual TSharedPtr<void> GetInstance() override
@@ -107,7 +103,6 @@ class FTypeContainer
 			return Instance;
 		}
 
-		/** The stored instance. */
 		TSharedRef<T> Instance;
 	};
 
@@ -116,7 +111,6 @@ class FTypeContainer
 	struct FThreadInstanceProvider
 		: public IInstanceProvider
 	{
-		/** Constructor. */
 		FThreadInstanceProvider(TFunction<TSharedPtr<void>()>&& InCreateFunc)
 			: CreateFunc(MoveTemp(InCreateFunc))
 			, TlsSlot(FPlatformTLS::AllocTlsSlot())
@@ -127,10 +121,7 @@ class FTypeContainer
 			FPlatformTLS::FreeTlsSlot(TlsSlot);
 		}
 
-		/** Factory function for creating instances. */
 		TFunction<TSharedPtr<void>()> CreateFunc;
-
-		/** Slot ID for thread-local storage of the instance. */
 		uint32 TlsSlot;
 	};
 
@@ -140,7 +131,8 @@ class FTypeContainer
 	struct TThreadInstanceProvider
 		: public FThreadInstanceProvider
 	{
-		/** Constructor. */
+		typedef TTlsAutoCleanupValue<TSharedPtr<T>> TlsValueType;
+
 		TThreadInstanceProvider(TFunction<TSharedPtr<void>()>&& InCreateFunc)
 			: FThreadInstanceProvider(MoveTemp(InCreateFunc))
 		{ }
@@ -149,41 +141,16 @@ class FTypeContainer
 
 		virtual TSharedPtr<void> GetInstance() override
 		{
-			void* Instance = FPlatformTLS::GetTlsValue(TlsSlot);
+			auto TlsValue = (TlsValueType*)FPlatformTLS::GetTlsValue(TlsSlot);
 
-			if (Instance == nullptr)
+			if (TlsValue == nullptr)
 			{
-				Instance = new TSharedPtr<T>(StaticCastSharedPtr<T>(CreateFunc()));
-
-				/** @todo gmp: this cannot possibly work, because the FRunnableThread is not yet registered.
-					A better mechanism for executing code on thread shutdown is already being worked on.
-					In the meantime we leak some memory here.
-				FRunnableThread::GetThreadRegistry().Lock();
-				{
-					const uint32 ThreadId = FPlatformTLS::GetCurrentThreadId();
-					FRunnableThread* RunnableThread = FRunnableThread::GetThreadRegistry().GetThread(ThreadId);
-
-					if (RunnableThread != nullptr)
-					{
-						Instance = new TSharedPtr<T>((T*)CreateFunc());
-						RunnableThread->OnThreadDestroyed().AddStatic(&TThreadInstance::HandleDestroyInstance, Instance);
-					}
-				}
-				FRunnableThread::GetThreadRegistry().Unlock();*/
-
-				FPlatformTLS::SetTlsValue(TlsSlot, Instance);
-				// @see ITlsAutoCleanup
-				//Instance->Register();
+				TlsValue = new TlsValueType(StaticCastSharedPtr<T>(CreateFunc()));
+				TlsValue->Register();
+				FPlatformTLS::SetTlsValue(TlsSlot, TlsValue);
 			}
 
-			return *((TSharedPtr<void>*)Instance);
-		}
-
-		/** Callback for destroying per-thread instances. */
-		static void HandleDestroyInstance(void* Instance)
-		{
-			auto InstancePtr = (TSharedPtr<T>*)Instance;
-			delete InstancePtr;
+			return TlsValue->Get();
 		}
 	};
 
