Mac: fix result of DeleteDirectory().

diff --git a/Engine/Source/Runtime/Core/Private/Apple/ApplePlatformFile.cpp b/Engine/Source/Runtime/Core/Private/Apple/ApplePlatformFile.cpp
--- a/Engine/Source/Runtime/Core/Private/Apple/ApplePlatformFile.cpp
+++ b/Engine/Source/Runtime/Core/Private/Apple/ApplePlatformFile.cpp
@@ -489,7 +489,7 @@ bool FApplePlatformFile::CreateDirectory(const TCHAR* Directory)
 }
 bool FApplePlatformFile::DeleteDirectory(const TCHAR* Directory)
 {
-	return rmdir(TCHAR_TO_UTF8(*NormalizeFilename(Directory)));
+	return rmdir(TCHAR_TO_UTF8(*NormalizeFilename(Directory))) == 0;
 }
 bool FApplePlatformFile::IterateDirectory(const TCHAR* Directory, FDirectoryVisitor& Visitor)
 {
